- Added test for binding interfaces & fixed following fallout

Major:
- Added missing input MetaSound audio interfaces
- Fixed interface bindings not working with referenced transient, non-asset MetaSounds

Minor:
- Fix builder name not being propagated to transient builders
- Fix not defaulting to success enum when creating patch builder
- Generalized pool of named/registered builders on Builder Subsystem instead of having them per type (to support future types without having to update the subystem for each additional future type.)

diff --git a/Engine/Plugins/Runtime/Metasound/Source/MetasoundEngine/Private/Interfaces/MetasoundInputFormatInterfaces.cpp b/Engine/Plugins/Runtime/Metasound/Source/MetasoundEngine/Private/Interfaces/MetasoundInputFormatInterfaces.cpp
--- a/Engine/Plugins/Runtime/Metasound/Source/MetasoundEngine/Private/Interfaces/MetasoundInputFormatInterfaces.cpp
+++ b/Engine/Plugins/Runtime/Metasound/Source/MetasoundEngine/Private/Interfaces/MetasoundInputFormatInterfaces.cpp
@@ -249,5 +249,554 @@ namespace Metasound::Engine
 		}
 	} // namespace InputFormatStereoInterface
 #undef AUDIO_PARAMETER_INTERFACE_NAMESPACE
+
+#define AUDIO_PARAMETER_INTERFACE_NAMESPACE "UE.InputFormat.Quad"
+	namespace InputFormatQuadInterface
+	{
+		const FMetasoundFrontendVersion& GetVersion()
+		{
+			static const FMetasoundFrontendVersion Version = { AUDIO_PARAMETER_INTERFACE_NAMESPACE, { 1, 0 } };
+			return Version;
+		}
+
+		namespace Inputs
+		{
+			const FName FrontLeftIn = AUDIO_PARAMETER_INTERFACE_MEMBER_DEFINE("Audio:0");
+			const FName FrontRightIn = AUDIO_PARAMETER_INTERFACE_MEMBER_DEFINE("Audio:1");
+			const FName SideLeftIn = AUDIO_PARAMETER_INTERFACE_MEMBER_DEFINE("Audio:4");
+			const FName SideRightIn = AUDIO_PARAMETER_INTERFACE_MEMBER_DEFINE("Audio:5");
+		} // namespace Inputs
+
+		Audio::FParameterInterfacePtr CreateInterface()
+		{
+			struct FInterface : public Audio::FParameterInterface
+			{
+				FInterface()
+					: FParameterInterface(InputFormatQuadInterface::GetVersion().Name, InputFormatQuadInterface::GetVersion().Number.ToInterfaceVersion())
+				{
+					UClassOptions = InputFormatPrivate::GetUClassOptions();
+					Inputs =
+					{
+						{
+							LOCTEXT("InputFormatQuadInterfaceInputFrontLeftName", "In Front Left"),
+							LOCTEXT("QuadIn_FrontLeft_AudioDescription", "Front Left input audio."),
+							GetMetasoundDataTypeName<FAudioBuffer>(),
+							Inputs::FrontLeftIn,
+							FText::GetEmpty(), // Required Text
+							100
+						},
+						{
+							LOCTEXT("InputFormatQuadInterfaceInputFrontRightName", "In Front Right"),
+							LOCTEXT("QuadIn_FrontRight_AudioDescription", "Front Right input audio."),
+							GetMetasoundDataTypeName<FAudioBuffer>(),
+							Inputs::FrontRightIn,
+							FText::GetEmpty(), // Required Text
+							101
+						},
+						{
+							LOCTEXT("InputFormatQuadInterfaceInputSideLeftName", "In Side Left"),
+							LOCTEXT("QuadIn_SideLeft_AudioDescription", "Side Left input audio."),
+							GetMetasoundDataTypeName<FAudioBuffer>(),
+							Inputs::SideLeftIn,
+							FText::GetEmpty(), // Required Text
+							104
+						},
+						{
+							LOCTEXT("InputFormatQuadInterfaceInputSideRightName", "In Side Right"),
+							LOCTEXT("QuadIn_SideRight_AudioDescription", "Side Right input audio."),
+							GetMetasoundDataTypeName<FAudioBuffer>(),
+							Inputs::SideRightIn,
+							FText::GetEmpty(), // Required Text
+							105
+						},
+					};
+				}
+			};
+
+			return MakeShared<FInterface>();
+		}
+
+		TArray<FMetasoundFrontendInterfaceBinding> CreateBindings()
+		{
+			return TArray<FMetasoundFrontendInterfaceBinding>
+			{
+				{
+					OutputFormatMonoInterface::GetVersion(),
+					InputFormatQuadInterface::GetVersion(),
+					10,
+					{
+						{
+							OutputFormatMonoInterface::Outputs::MonoOut, InputFormatQuadInterface::Inputs::FrontLeftIn
+						}
+					}
+				},
+				{
+					OutputFormatStereoInterface::GetVersion(),
+					InputFormatQuadInterface::GetVersion(),
+					20,
+					{
+						{
+							OutputFormatStereoInterface::Outputs::LeftOut, InputFormatQuadInterface::Inputs::FrontLeftIn,
+						},
+						{
+							OutputFormatStereoInterface::Outputs::RightOut, InputFormatQuadInterface::Inputs::FrontRightIn
+						}
+					}
+				},
+				{
+					OutputFormatQuadInterface::GetVersion(),
+					InputFormatQuadInterface::GetVersion(),
+					0, // Quad to quad takes priority
+					{
+						{
+							OutputFormatQuadInterface::Outputs::FrontLeftOut, InputFormatQuadInterface::Inputs::FrontLeftIn,
+						},
+						{
+							OutputFormatQuadInterface::Outputs::FrontRightOut, InputFormatQuadInterface::Inputs::FrontRightIn,
+						},
+						{
+							OutputFormatQuadInterface::Outputs::SideLeftOut, InputFormatQuadInterface::Inputs::SideLeftIn,
+						},
+						{
+							OutputFormatQuadInterface::Outputs::SideRightOut, InputFormatQuadInterface::Inputs::SideRightIn
+						}
+					}
+				},
+				{
+					OutputFormatFiveDotOneInterface::GetVersion(),
+					InputFormatQuadInterface::GetVersion(),
+					40,
+					{
+						{
+							OutputFormatFiveDotOneInterface::Outputs::FrontLeftOut, InputFormatQuadInterface::Inputs::FrontLeftIn,
+						},
+						{
+							OutputFormatFiveDotOneInterface::Outputs::FrontRightOut, InputFormatQuadInterface::Inputs::FrontRightIn,
+						},
+						{
+							OutputFormatFiveDotOneInterface::Outputs::SideLeftOut, InputFormatQuadInterface::Inputs::SideLeftIn,
+						},
+						{
+							OutputFormatFiveDotOneInterface::Outputs::SideRightOut, InputFormatQuadInterface::Inputs::SideRightIn
+						}
+					}
+				},
+				{
+					OutputFormatSevenDotOneInterface::GetVersion(),
+					InputFormatQuadInterface::GetVersion(),
+					50,
+					{
+						{
+							OutputFormatSevenDotOneInterface::Outputs::FrontLeftOut, InputFormatQuadInterface::Inputs::FrontLeftIn,
+						},
+						{
+							OutputFormatSevenDotOneInterface::Outputs::FrontRightOut, InputFormatQuadInterface::Inputs::FrontRightIn,
+						},
+						{
+							OutputFormatSevenDotOneInterface::Outputs::SideLeftOut, InputFormatQuadInterface::Inputs::SideLeftIn,
+						},
+						{
+							OutputFormatSevenDotOneInterface::Outputs::SideRightOut, InputFormatQuadInterface::Inputs::SideRightIn
+						}
+					}
+				}
+			};
+		}
+	} // namespace InputFormatQuadInterface
+#undef AUDIO_PARAMETER_INTERFACE_NAMESPACE
+
+#define AUDIO_PARAMETER_INTERFACE_NAMESPACE "UE.InputFormat.5dot1"
+	namespace InputFormatFiveDotOneInterface
+	{
+		const FMetasoundFrontendVersion& GetVersion()
+		{
+			static const FMetasoundFrontendVersion Version = { AUDIO_PARAMETER_INTERFACE_NAMESPACE, { 1, 0 } };
+			return Version;
+		}
+
+		namespace Inputs
+		{
+			const FName FrontLeftIn = AUDIO_PARAMETER_INTERFACE_MEMBER_DEFINE("Audio:0");
+			const FName FrontRightIn = AUDIO_PARAMETER_INTERFACE_MEMBER_DEFINE("Audio:1");
+			const FName FrontCenterIn = AUDIO_PARAMETER_INTERFACE_MEMBER_DEFINE("Audio:2");
+			const FName LowFrequencyIn = AUDIO_PARAMETER_INTERFACE_MEMBER_DEFINE("Audio:3");
+			const FName SideLeftIn = AUDIO_PARAMETER_INTERFACE_MEMBER_DEFINE("Audio:4");
+			const FName SideRightIn = AUDIO_PARAMETER_INTERFACE_MEMBER_DEFINE("Audio:5");
+		} // namespace Inputs
+
+		Audio::FParameterInterfacePtr CreateInterface()
+		{
+			struct FInterface : public Audio::FParameterInterface
+			{
+				FInterface()
+					: FParameterInterface(InputFormatFiveDotOneInterface::GetVersion().Name, InputFormatFiveDotOneInterface::GetVersion().Number.ToInterfaceVersion())
+				{
+					UClassOptions = InputFormatPrivate::GetUClassOptions();
+					Inputs =
+					{
+						{
+							LOCTEXT("InputFormatFiveDotOneInterfaceInputFrontLeftName", "In Front Left"),
+							LOCTEXT("FiveDotOneIn_FrontLeft_AudioDescription", "Front Left input audio."),
+							GetMetasoundDataTypeName<FAudioBuffer>(),
+							Inputs::FrontLeftIn,
+							FText::GetEmpty(), // Required Text
+							100
+						},
+						{
+							LOCTEXT("InputFormatFiveDotOneInterfaceInputFrontRightName", "In Front Right"),
+							LOCTEXT("FiveDotOneIn_FrontRight_AudioDescription", "Front Right input audio."),
+							GetMetasoundDataTypeName<FAudioBuffer>(),
+							Inputs::FrontRightIn,
+							FText::GetEmpty(), // Required Text
+							101
+						},
+						{
+							LOCTEXT("InputFormatFiveDotOneInterfaceInputFrontCenterName", "In Front Center"),
+							LOCTEXT("FiveDotOneIn_FrontCenter_AudioDescription", "Front Center input audio."),
+							GetMetasoundDataTypeName<FAudioBuffer>(),
+							Inputs::FrontCenterIn,
+							FText::GetEmpty(), // Required Text
+							102
+						},
+						{
+							LOCTEXT("InputFormatFiveDotOneInterfaceInputLowFrequencyName", "In Low Frequency"),
+							LOCTEXT("FiveDotOneIn_LowFrequency_AudioDescription", "Low Frequency input audio."),
+							GetMetasoundDataTypeName<FAudioBuffer>(),
+							Inputs::LowFrequencyIn,
+							FText::GetEmpty(), // Required Text
+							103
+						},
+						{
+							LOCTEXT("InputFormatFiveDotOneInterfaceInputSideLeftName", "In Side Left"),
+							LOCTEXT("FiveDotOneIn_SideLeft_AudioDescription", "Side Left input audio."),
+							GetMetasoundDataTypeName<FAudioBuffer>(),
+							Inputs::SideLeftIn,
+							FText::GetEmpty(), // Required Text
+							104
+						},
+						{
+							LOCTEXT("InputFormatFiveDotOneInterfaceInputSideRightName", "In Side Right"),
+							LOCTEXT("FiveDotOneIn_SideRight_AudioDescription", "Side Right input audio."),
+							GetMetasoundDataTypeName<FAudioBuffer>(),
+							Inputs::SideRightIn,
+							FText::GetEmpty(), // Required Text
+							105
+						},
+					};
+				}
+			};
+
+			return MakeShared<FInterface>();
+		}
+
+		TArray<FMetasoundFrontendInterfaceBinding> CreateBindings()
+		{
+			return TArray<FMetasoundFrontendInterfaceBinding>
+			{
+				{
+					OutputFormatMonoInterface::GetVersion(),
+					InputFormatFiveDotOneInterface::GetVersion(),
+					10,
+					{
+						{
+							OutputFormatMonoInterface::Outputs::MonoOut, InputFormatFiveDotOneInterface::Inputs::FrontLeftIn
+						}
+					}
+				},
+				{
+					OutputFormatStereoInterface::GetVersion(),
+					InputFormatFiveDotOneInterface::GetVersion(),
+					20,
+					{
+						{
+							OutputFormatStereoInterface::Outputs::LeftOut, InputFormatFiveDotOneInterface::Inputs::FrontLeftIn,
+						},
+						{
+							OutputFormatStereoInterface::Outputs::RightOut, InputFormatFiveDotOneInterface::Inputs::FrontRightIn
+						}
+					}
+				},
+				{
+					OutputFormatQuadInterface::GetVersion(),
+					InputFormatFiveDotOneInterface::GetVersion(),
+					30,
+					{
+						{
+							OutputFormatQuadInterface::Outputs::FrontLeftOut, InputFormatFiveDotOneInterface::Inputs::FrontLeftIn,
+						},
+						{
+							OutputFormatQuadInterface::Outputs::FrontRightOut, InputFormatFiveDotOneInterface::Inputs::FrontRightIn,
+						},
+						{
+							OutputFormatQuadInterface::Outputs::SideLeftOut, InputFormatFiveDotOneInterface::Inputs::SideLeftIn,
+						},
+						{
+							OutputFormatQuadInterface::Outputs::SideRightOut, InputFormatFiveDotOneInterface::Inputs::SideRightIn
+						}
+					}
+				},
+				{
+					OutputFormatFiveDotOneInterface::GetVersion(),
+					InputFormatFiveDotOneInterface::GetVersion(),
+					0, // 5.1 to 5.1 takes priority
+					{
+						{
+							OutputFormatFiveDotOneInterface::Outputs::FrontLeftOut, InputFormatFiveDotOneInterface::Inputs::FrontLeftIn,
+						},
+						{
+							OutputFormatFiveDotOneInterface::Outputs::FrontRightOut, InputFormatFiveDotOneInterface::Inputs::FrontRightIn,
+						},
+						{
+							OutputFormatFiveDotOneInterface::Outputs::FrontCenterOut, InputFormatFiveDotOneInterface::Inputs::FrontCenterIn,
+						},
+						{
+							OutputFormatFiveDotOneInterface::Outputs::LowFrequencyOut, InputFormatFiveDotOneInterface::Inputs::LowFrequencyIn,
+						},
+						{
+							OutputFormatFiveDotOneInterface::Outputs::SideLeftOut, InputFormatFiveDotOneInterface::Inputs::SideLeftIn,
+						},
+						{
+							OutputFormatFiveDotOneInterface::Outputs::SideRightOut, InputFormatFiveDotOneInterface::Inputs::SideRightIn
+						}
+					}
+				},
+				{
+					OutputFormatSevenDotOneInterface::GetVersion(),
+					InputFormatFiveDotOneInterface::GetVersion(),
+					50,
+					{
+						{
+							OutputFormatSevenDotOneInterface::Outputs::FrontLeftOut, InputFormatFiveDotOneInterface::Inputs::FrontLeftIn,
+						},
+						{
+							OutputFormatSevenDotOneInterface::Outputs::FrontRightOut, InputFormatFiveDotOneInterface::Inputs::FrontRightIn,
+						},
+						{
+							OutputFormatSevenDotOneInterface::Outputs::FrontCenterOut, InputFormatFiveDotOneInterface::Inputs::FrontCenterIn,
+						},
+						{
+							OutputFormatSevenDotOneInterface::Outputs::LowFrequencyOut, InputFormatFiveDotOneInterface::Inputs::LowFrequencyIn,
+						},
+						{
+							OutputFormatSevenDotOneInterface::Outputs::SideLeftOut, InputFormatFiveDotOneInterface::Inputs::SideLeftIn,
+						},
+						{
+							OutputFormatSevenDotOneInterface::Outputs::SideRightOut, InputFormatFiveDotOneInterface::Inputs::SideRightIn
+						}
+					}
+				}
+			};
+		}
+	} // namespace InputFormatFiveDotOneInterface
+#undef AUDIO_PARAMETER_INTERFACE_NAMESPACE
+
+#define AUDIO_PARAMETER_INTERFACE_NAMESPACE "UE.InputFormat.7dot1"
+	namespace InputFormatSevenDotOneInterface
+	{
+		const FMetasoundFrontendVersion& GetVersion()
+		{
+			static const FMetasoundFrontendVersion Version = { AUDIO_PARAMETER_INTERFACE_NAMESPACE, { 1, 0 } };
+			return Version;
+		}
+
+		namespace Inputs
+		{
+			const FName FrontLeftIn = AUDIO_PARAMETER_INTERFACE_MEMBER_DEFINE("Audio:0");
+			const FName FrontRightIn = AUDIO_PARAMETER_INTERFACE_MEMBER_DEFINE("Audio:1");
+			const FName FrontCenterIn = AUDIO_PARAMETER_INTERFACE_MEMBER_DEFINE("Audio:2");
+			const FName LowFrequencyIn = AUDIO_PARAMETER_INTERFACE_MEMBER_DEFINE("Audio:3");
+			const FName SideLeftIn = AUDIO_PARAMETER_INTERFACE_MEMBER_DEFINE("Audio:4");
+			const FName SideRightIn = AUDIO_PARAMETER_INTERFACE_MEMBER_DEFINE("Audio:5");
+			const FName BackLeftIn = AUDIO_PARAMETER_INTERFACE_MEMBER_DEFINE("Audio:6");
+			const FName BackRightIn = AUDIO_PARAMETER_INTERFACE_MEMBER_DEFINE("Audio:7");
+		} // namespace Inputs
+
+		Audio::FParameterInterfacePtr CreateInterface()
+		{
+			struct FInterface : public Audio::FParameterInterface
+			{
+				FInterface()
+					: FParameterInterface(InputFormatSevenDotOneInterface::GetVersion().Name, InputFormatSevenDotOneInterface::GetVersion().Number.ToInterfaceVersion())
+				{
+					UClassOptions = InputFormatPrivate::GetUClassOptions();
+					Inputs =
+					{
+						{
+							LOCTEXT("InputFormatSevenDotOneInterfaceInputFrontLeftName", "In Front Left"),
+							LOCTEXT("SevenDotOneIn_FrontLeft_AudioDescription", "Front Left input audio."),
+							GetMetasoundDataTypeName<FAudioBuffer>(),
+							Inputs::FrontLeftIn,
+							FText::GetEmpty(), // Required Text
+							100
+						},
+						{
+							LOCTEXT("InputFormatSevenDotOneInterfaceInputFrontRightName", "In Front Right"),
+							LOCTEXT("SevenDotOneIn_FrontRight_AudioDescription", "Front Right input audio."),
+							GetMetasoundDataTypeName<FAudioBuffer>(),
+							Inputs::FrontRightIn,
+							FText::GetEmpty(), // Required Text
+							101
+						},
+						{
+							LOCTEXT("InputFormatSevenDotOneInterfaceInputFrontCenterName", "In Front Center"),
+							LOCTEXT("SevenDotOneIn_FrontCenter_AudioDescription", "Front Center input audio."),
+							GetMetasoundDataTypeName<FAudioBuffer>(),
+							Inputs::FrontCenterIn,
+							FText::GetEmpty(), // Required Text
+							102
+						},
+						{
+							LOCTEXT("InputFormatSevenDotOneInterfaceInputLowFrequencyName", "In Low Frequency"),
+							LOCTEXT("SevenDotOneIn_LowFrequency_AudioDescription", "Low Frequency input audio."),
+							GetMetasoundDataTypeName<FAudioBuffer>(),
+							Inputs::LowFrequencyIn,
+							FText::GetEmpty(), // Required Text
+							103
+						},
+						{
+							LOCTEXT("InputFormatSevenDotOneInterfaceInputSideLeftName", "In Side Left"),
+							LOCTEXT("SevenDotOneIn_SideLeft_AudioDescription", "Side Left input audio."),
+							GetMetasoundDataTypeName<FAudioBuffer>(),
+							Inputs::SideLeftIn,
+							FText::GetEmpty(), // Required Text
+							104
+						},
+						{
+							LOCTEXT("InputFormatSevenDotOneInterfaceInputSideRightName", "In Side Right"),
+							LOCTEXT("SevenDotOneIn_SideRight_AudioDescription", "Side Right input audio."),
+							GetMetasoundDataTypeName<FAudioBuffer>(),
+							Inputs::SideRightIn,
+							FText::GetEmpty(), // Required Text
+							105
+						},
+						{
+							LOCTEXT("InputFormatSevenDotOneInterfaceInputBackLeftName", "In Back Left"),
+							LOCTEXT("SevenDotOneIn_BackLeft_AudioDescription", "Back Left input audio."),
+							GetMetasoundDataTypeName<FAudioBuffer>(),
+							Inputs::BackLeftIn,
+							FText::GetEmpty(), // Required Text
+							106
+						},
+						{
+							LOCTEXT("InputFormatSevenDotOneInterfaceInputBackRightName", "In Back Right"),
+							LOCTEXT("SevenDotOneIn_BackRight_AudioDescription", "Back Right input audio."),
+							GetMetasoundDataTypeName<FAudioBuffer>(),
+							Inputs::BackRightIn,
+							FText::GetEmpty(), // Required Text
+							107
+						},
+					};
+				}
+			};
+
+			return MakeShared<FInterface>();
+		}
+
+		TArray<FMetasoundFrontendInterfaceBinding> CreateBindings()
+		{
+			return TArray<FMetasoundFrontendInterfaceBinding>
+			{
+				{
+					OutputFormatMonoInterface::GetVersion(),
+					InputFormatSevenDotOneInterface::GetVersion(),
+						10,
+					{
+						{
+							OutputFormatMonoInterface::Outputs::MonoOut, InputFormatSevenDotOneInterface::Inputs::FrontLeftIn
+						}
+					}
+				},
+				{
+					OutputFormatStereoInterface::GetVersion(),
+					InputFormatSevenDotOneInterface::GetVersion(),
+					20,
+					{
+						{
+							OutputFormatStereoInterface::Outputs::LeftOut, InputFormatSevenDotOneInterface::Inputs::FrontLeftIn,
+						},
+						{
+							OutputFormatStereoInterface::Outputs::RightOut, InputFormatSevenDotOneInterface::Inputs::FrontRightIn
+						}
+					}
+				},
+				{
+					OutputFormatQuadInterface::GetVersion(),
+					InputFormatSevenDotOneInterface::GetVersion(),
+					30,
+					{
+						{
+							OutputFormatQuadInterface::Outputs::FrontLeftOut, InputFormatSevenDotOneInterface::Inputs::FrontLeftIn,
+						},
+						{
+							OutputFormatQuadInterface::Outputs::FrontRightOut, InputFormatSevenDotOneInterface::Inputs::FrontRightIn,
+						},
+						{
+							OutputFormatQuadInterface::Outputs::SideLeftOut, InputFormatSevenDotOneInterface::Inputs::SideLeftIn,
+						},
+						{
+							OutputFormatQuadInterface::Outputs::SideRightOut, InputFormatSevenDotOneInterface::Inputs::SideRightIn
+						}
+					}
+				},
+				{
+					OutputFormatFiveDotOneInterface::GetVersion(),
+					InputFormatSevenDotOneInterface::GetVersion(),
+					40,
+					{
+						{
+							OutputFormatFiveDotOneInterface::Outputs::FrontLeftOut, InputFormatSevenDotOneInterface::Inputs::FrontLeftIn,
+						},
+						{
+							OutputFormatFiveDotOneInterface::Outputs::FrontRightOut, InputFormatSevenDotOneInterface::Inputs::FrontRightIn,
+						},
+						{
+							OutputFormatFiveDotOneInterface::Outputs::FrontCenterOut, InputFormatSevenDotOneInterface::Inputs::FrontCenterIn,
+						},
+						{
+							OutputFormatFiveDotOneInterface::Outputs::LowFrequencyOut, InputFormatSevenDotOneInterface::Inputs::LowFrequencyIn,
+						},
+						{
+							OutputFormatFiveDotOneInterface::Outputs::SideLeftOut, InputFormatSevenDotOneInterface::Inputs::SideLeftIn,
+						},
+						{
+							OutputFormatFiveDotOneInterface::Outputs::SideRightOut, InputFormatSevenDotOneInterface::Inputs::SideRightIn
+						}
+					}
+				},
+				{
+					OutputFormatSevenDotOneInterface::GetVersion(),
+					InputFormatSevenDotOneInterface::GetVersion(),
+					0, // 7.1 to 7.1 takes priority
+					{
+						{
+							OutputFormatSevenDotOneInterface::Outputs::FrontLeftOut, InputFormatSevenDotOneInterface::Inputs::FrontLeftIn,
+						},
+						{
+							OutputFormatSevenDotOneInterface::Outputs::FrontRightOut, InputFormatSevenDotOneInterface::Inputs::FrontRightIn,
+						},
+						{
+							OutputFormatSevenDotOneInterface::Outputs::FrontCenterOut, InputFormatSevenDotOneInterface::Inputs::FrontCenterIn,
+						},
+						{
+							OutputFormatSevenDotOneInterface::Outputs::LowFrequencyOut, InputFormatSevenDotOneInterface::Inputs::LowFrequencyIn,
+						},
+						{
+							OutputFormatSevenDotOneInterface::Outputs::SideLeftOut, InputFormatSevenDotOneInterface::Inputs::SideLeftIn,
+						},
+						{
+							OutputFormatSevenDotOneInterface::Outputs::SideRightOut, InputFormatSevenDotOneInterface::Inputs::SideRightIn
+						},
+						{
+							OutputFormatSevenDotOneInterface::Outputs::BackLeftOut, InputFormatSevenDotOneInterface::Inputs::BackLeftIn,
+						},
+						{
+							OutputFormatSevenDotOneInterface::Outputs::BackRightOut, InputFormatSevenDotOneInterface::Inputs::BackRightIn
+						}
+					}
+				}
+			};
+		}
+	} // namespace InputFormatSevenDotOneInterface
+#undef AUDIO_PARAMETER_INTERFACE_NAMESPACE
 } // namespace Metasound::Engine
 #undef LOCTEXT_NAMESPACE
diff --git a/Engine/Plugins/Runtime/Metasound/Source/MetasoundEngine/Private/Interfaces/MetasoundInterface.cpp b/Engine/Plugins/Runtime/Metasound/Source/MetasoundEngine/Private/Interfaces/MetasoundInterface.cpp
--- a/Engine/Plugins/Runtime/Metasound/Source/MetasoundEngine/Private/Interfaces/MetasoundInterface.cpp
+++ b/Engine/Plugins/Runtime/Metasound/Source/MetasoundEngine/Private/Interfaces/MetasoundInterface.cpp
@@ -72,8 +72,11 @@ namespace Metasound::Engine
 
 		// Input Formats
 		{
-			Reg.RegisterInterface(MakeUnique<FInterfaceRegistryEntry>(FMetasoundFrontendInterface(InputFormatMonoInterface::CreateInterface()), IDataReference::RouterName));
-			Reg.RegisterInterface(MakeUnique<FInterfaceRegistryEntry>(InputFormatStereoInterface::CreateInterface(), IDataReference::RouterName));
+			Reg.RegisterInterface(MakeUnique<FInterfaceRegistryEntry>(InputFormatMonoInterface::CreateInterface()));
+			Reg.RegisterInterface(MakeUnique<FInterfaceRegistryEntry>(InputFormatStereoInterface::CreateInterface()));
+			Reg.RegisterInterface(MakeUnique<FInterfaceRegistryEntry>(InputFormatQuadInterface::CreateInterface()));
+			Reg.RegisterInterface(MakeUnique<FInterfaceRegistryEntry>(InputFormatFiveDotOneInterface::CreateInterface()));
+			Reg.RegisterInterface(MakeUnique<FInterfaceRegistryEntry>(InputFormatSevenDotOneInterface::CreateInterface()));
 		}
 
 		// Output Formats
diff --git a/Engine/Plugins/Runtime/Metasound/Source/MetasoundEngine/Private/MetasoundBuilderSubsystem.cpp b/Engine/Plugins/Runtime/Metasound/Source/MetasoundEngine/Private/MetasoundBuilderSubsystem.cpp
--- a/Engine/Plugins/Runtime/Metasound/Source/MetasoundEngine/Private/MetasoundBuilderSubsystem.cpp
+++ b/Engine/Plugins/Runtime/Metasound/Source/MetasoundEngine/Private/MetasoundBuilderSubsystem.cpp
@@ -55,11 +55,11 @@ namespace Metasound::Engine
 		}
 
 		template <typename BuilderClass>
-		BuilderClass& CreateTransientBuilder()
+		BuilderClass& CreateTransientBuilder(FName BuilderName = { })
 		{
 			const EObjectFlags NewObjectFlags = RF_Public | RF_Transient;
 			UPackage* TransientPackage = GetTransientPackage();
-			const FName ObjectName = MakeUniqueObjectName(TransientPackage, BuilderClass::StaticClass());
+			const FName ObjectName = MakeUniqueObjectName(TransientPackage, BuilderClass::StaticClass(), BuilderName);
 			TObjectPtr<BuilderClass> NewBuilder = NewObject<BuilderClass>(TransientPackage, ObjectName, NewObjectFlags);
 			check(NewBuilder);
 			NewBuilder->InitFrontendBuilder();
@@ -1273,7 +1273,8 @@ void UMetaSoundSourceBuilder::SetFormat(EMetaSoundOutputAudioFormat OutputFormat
 
 UMetaSoundPatchBuilder* UMetaSoundBuilderSubsystem::CreatePatchBuilder(FName BuilderName, EMetaSoundBuilderResult& OutResult)
 {
-	return &Metasound::Engine::BuilderSubsystemPrivate::CreateTransientBuilder<UMetaSoundPatchBuilder>();
+	OutResult = EMetaSoundBuilderResult::Succeeded;
+	return &Metasound::Engine::BuilderSubsystemPrivate::CreateTransientBuilder<UMetaSoundPatchBuilder>(BuilderName);
 }
 
 UMetaSoundSourceBuilder* UMetaSoundBuilderSubsystem::CreateSourceBuilder(
@@ -1292,7 +1293,7 @@ UMetaSoundSourceBuilder* UMetaSoundBuilderSubsystem::CreateSourceBuilder(
 	OnFinishedNodeInput = { };
 	AudioOutNodeInputs.Reset();
 
-	UMetaSoundSourceBuilder& NewBuilder = CreateTransientBuilder<UMetaSoundSourceBuilder>();
+	UMetaSoundSourceBuilder& NewBuilder = CreateTransientBuilder<UMetaSoundSourceBuilder>(BuilderName);
 
 	OutResult = EMetaSoundBuilderResult::Succeeded;
 	if (OutputFormat != EMetaSoundOutputAudioFormat::Mono)
@@ -1386,7 +1387,7 @@ UMetaSoundSourceBuilder* UMetaSoundBuilderSubsystem::CreateSourceBuilder(
 
 UMetaSoundPatchBuilder* UMetaSoundBuilderSubsystem::CreatePatchPresetBuilder(FName BuilderName, const TScriptInterface<IMetaSoundDocumentInterface>& ReferencedNodeClass, EMetaSoundBuilderResult& OutResult)
 {
-	UMetaSoundPatchBuilder& Builder = Metasound::Engine::BuilderSubsystemPrivate::CreateTransientBuilder<UMetaSoundPatchBuilder>();
+	UMetaSoundPatchBuilder& Builder = Metasound::Engine::BuilderSubsystemPrivate::CreateTransientBuilder<UMetaSoundPatchBuilder>(BuilderName);
 	Builder.ConvertToPreset(ReferencedNodeClass, OutResult);
 	return &Builder;
 }
@@ -1405,7 +1406,7 @@ UMetaSoundBuilderBase& UMetaSoundBuilderSubsystem::CreatePresetBuilder(FName Bui
 	else
 	{
 		checkf(false, TEXT("UClass '%s' cannot be built to a MetaSound preset"), *Class.GetFullName());
-		return Metasound::Engine::BuilderSubsystemPrivate::CreateTransientBuilder<UMetaSoundPatchBuilder>();
+		return Metasound::Engine::BuilderSubsystemPrivate::CreateTransientBuilder<UMetaSoundPatchBuilder>(BuilderName);
 	}
 }
 
@@ -1509,14 +1510,29 @@ const Metasound::Frontend::FDocumentModifyDelegates* UMetaSoundBuilderSubsystem:
 	return nullptr;
 }
 
+UMetaSoundBuilderBase* UMetaSoundBuilderSubsystem::FindBuilder(FName BuilderName)
+{
+	return NamedBuilders.FindRef(BuilderName);
+}
+
 UMetaSoundPatchBuilder* UMetaSoundBuilderSubsystem::FindPatchBuilder(FName BuilderName)
 {
-	return NamedPatchBuilders.FindRef(BuilderName);
+	if (UMetaSoundBuilderBase* Builder = FindBuilder(BuilderName))
+	{
+		return Cast<UMetaSoundPatchBuilder>(Builder);
+	}
+
+	return nullptr;
 }
 
 UMetaSoundSourceBuilder* UMetaSoundBuilderSubsystem::FindSourceBuilder(FName BuilderName)
 {
-	return NamedSourceBuilders.FindRef(BuilderName);
+	if (UMetaSoundBuilderBase* Builder = FindBuilder(BuilderName))
+	{
+		return Cast<UMetaSoundSourceBuilder>(Builder);
+	}
+
+	return nullptr;
 }
 
 void UMetaSoundBuilderSubsystem::Initialize(FSubsystemCollectionBase& Collection)
@@ -1540,22 +1556,41 @@ bool UMetaSoundBuilderSubsystem::IsInterfaceRegistered(FName InInterfaceName) co
 	return ISearchEngine::Get().FindInterfaceWithHighestVersion(InInterfaceName, Interface);
 }
 
+void UMetaSoundBuilderSubsystem::RegisterBuilder(FName BuilderName, UMetaSoundBuilderBase* Builder)
+{
+	if (Builder)
+	{
+		NamedBuilders.FindOrAdd(BuilderName) = Builder;
+	}
+}
+
 void UMetaSoundBuilderSubsystem::RegisterPatchBuilder(FName BuilderName, UMetaSoundPatchBuilder* Builder)
 {
-	NamedPatchBuilders.FindOrAdd(BuilderName) = Builder;
+	if (Builder)
+	{
+		NamedBuilders.FindOrAdd(BuilderName) = Builder;
+	}
 }
 
 void UMetaSoundBuilderSubsystem::RegisterSourceBuilder(FName BuilderName, UMetaSoundSourceBuilder* Builder)
 {
-	NamedSourceBuilders.FindOrAdd(BuilderName) = Builder;
+	if (Builder)
+	{
+		NamedBuilders.FindOrAdd(BuilderName) = Builder;
+	}
+}
+
+bool UMetaSoundBuilderSubsystem::UnregisterBuilder(FName BuilderName)
+{
+	return NamedBuilders.Remove(BuilderName) > 0;
 }
 
 bool UMetaSoundBuilderSubsystem::UnregisterPatchBuilder(FName BuilderName)
 {
-	return NamedPatchBuilders.Remove(BuilderName) > 0;
+	return NamedBuilders.Remove(BuilderName) > 0;
 }
 
 bool UMetaSoundBuilderSubsystem::UnregisterSourceBuilder(FName BuilderName)
 {
-	return NamedSourceBuilders.Remove(BuilderName) > 0;
+	return NamedBuilders.Remove(BuilderName) > 0;
 }
diff --git a/Engine/Plugins/Runtime/Metasound/Source/MetasoundEngine/Public/Interfaces/MetasoundInputFormatInterfaces.h b/Engine/Plugins/Runtime/Metasound/Source/MetasoundEngine/Public/Interfaces/MetasoundInputFormatInterfaces.h
--- a/Engine/Plugins/Runtime/Metasound/Source/MetasoundEngine/Public/Interfaces/MetasoundInputFormatInterfaces.h
+++ b/Engine/Plugins/Runtime/Metasound/Source/MetasoundEngine/Public/Interfaces/MetasoundInputFormatInterfaces.h
@@ -32,4 +32,55 @@ namespace Metasound::Engine
 		METASOUNDENGINE_API Audio::FParameterInterfacePtr CreateInterface();
 		METASOUNDENGINE_API TArray<FMetasoundFrontendInterfaceBinding> CreateBindings();
 	} // namespace InputFormatStereoInterface
+
+	namespace InputFormatQuadInterface
+	{
+		namespace Inputs
+		{
+			METASOUNDENGINE_API const extern FName FrontLeftIn;
+			METASOUNDENGINE_API const extern FName FrontRightIn;
+			METASOUNDENGINE_API const extern FName SideLeftIn;
+			METASOUNDENGINE_API const extern FName SideRightIn;
+		} // namespace Inputs
+
+		METASOUNDENGINE_API const FMetasoundFrontendVersion& GetVersion();
+		METASOUNDENGINE_API Audio::FParameterInterfacePtr CreateInterface();
+		METASOUNDENGINE_API TArray<FMetasoundFrontendInterfaceBinding> CreateBindings();
+	} // namespace InputFormatQuadInterface
+
+	namespace InputFormatFiveDotOneInterface
+	{
+		namespace Inputs
+		{
+			METASOUNDENGINE_API const extern FName FrontLeftIn;
+			METASOUNDENGINE_API const extern FName FrontRightIn;
+			METASOUNDENGINE_API const extern FName FrontCenterIn;
+			METASOUNDENGINE_API const extern FName LowFrequencyIn;
+			METASOUNDENGINE_API const extern FName SideLeftIn;
+			METASOUNDENGINE_API const extern FName SideRightIn;
+		} // namespace Inputs
+
+		METASOUNDENGINE_API const FMetasoundFrontendVersion& GetVersion();
+		METASOUNDENGINE_API Audio::FParameterInterfacePtr CreateInterface();
+		METASOUNDENGINE_API TArray<FMetasoundFrontendInterfaceBinding> CreateBindings();
+	} // namespace InputFormatFiveDotOneInterface
+
+	namespace InputFormatSevenDotOneInterface
+	{
+		namespace Inputs
+		{
+			METASOUNDENGINE_API const extern FName FrontLeftIn;
+			METASOUNDENGINE_API const extern FName FrontRightIn;
+			METASOUNDENGINE_API const extern FName FrontCenterIn;
+			METASOUNDENGINE_API const extern FName LowFrequencyIn;
+			METASOUNDENGINE_API const extern FName SideLeftIn;
+			METASOUNDENGINE_API const extern FName SideRightIn;
+			METASOUNDENGINE_API const extern FName BackLeftOut;
+			METASOUNDENGINE_API const extern FName BackRightOut;
+		} // namespace Inputs
+
+		METASOUNDENGINE_API const FMetasoundFrontendVersion& GetVersion();
+		METASOUNDENGINE_API Audio::FParameterInterfacePtr CreateInterface();
+		METASOUNDENGINE_API TArray<FMetasoundFrontendInterfaceBinding> CreateBindings();
+	} // namespace InputFormatSevenDotOneInterface
 } // namespace Metasound::Engine
diff --git a/Engine/Plugins/Runtime/Metasound/Source/MetasoundEngine/Public/MetasoundBuilderSubsystem.h b/Engine/Plugins/Runtime/Metasound/Source/MetasoundEngine/Public/MetasoundBuilderSubsystem.h
--- a/Engine/Plugins/Runtime/Metasound/Source/MetasoundEngine/Public/MetasoundBuilderSubsystem.h
+++ b/Engine/Plugins/Runtime/Metasound/Source/MetasoundEngine/Public/MetasoundBuilderSubsystem.h
@@ -504,11 +504,9 @@ class METASOUNDENGINE_API UMetaSoundBuilderSubsystem : public UEngineSubsystem,
 {
 	GENERATED_BODY()
 
+private:
 	UPROPERTY()
-	TMap<FName, TObjectPtr<UMetaSoundPatchBuilder>> NamedPatchBuilders;
-
-	UPROPERTY()
-	TMap<FName, TObjectPtr<UMetaSoundSourceBuilder>> NamedSourceBuilders;
+	TMap<FName, TObjectPtr<UMetaSoundBuilderBase>> NamedBuilders;
 
 	UPROPERTY()
 	mutable TMap<FName, TWeakObjectPtr<UMetaSoundBuilderBase>> AssetBuilders;
@@ -585,6 +583,9 @@ public:
 	UFUNCTION(BlueprintCallable, Category = "Audio|MetaSound|Builder", meta = (DisplayName = "Create MetaSound Literal From AudioParameter"))
 	UPARAM(DisplayName = "Param Literal") FMetasoundFrontendLiteral CreateMetaSoundLiteralFromParam(const FAudioParameter& Param);
 
+	UFUNCTION(BlueprintCallable, Category = "Audio|MetaSound|Builder")
+	UPARAM(DisplayName = "Builder") UMetaSoundBuilderBase* FindBuilder(FName BuilderName);
+
 	UFUNCTION(BlueprintCallable, Category = "Audio|MetaSound|Builder")
 	UPARAM(DisplayName = "Patch Builder") UMetaSoundPatchBuilder* FindPatchBuilder(FName BuilderName);
 
@@ -594,6 +595,10 @@ public:
 	UFUNCTION(BlueprintCallable, Category = "Audio|MetaSound|Builder")
 	UPARAM(DisplayName = "Is Registered") bool IsInterfaceRegistered(FName InInterfaceName) const;
 
+	// Adds builder to subsystem's registry to make it persistent and easily accessible by multiple systems or Blueprints
+	UFUNCTION(BlueprintCallable, Category = "Audio|MetaSound|Builder")
+	void RegisterBuilder(FName BuilderName, UMetaSoundBuilderBase* Builder);
+
 	// Adds builder to subsystem's registry to make it persistent and easily accessible by multiple systems or Blueprints
 	UFUNCTION(BlueprintCallable, Category = "Audio|MetaSound|Builder")
 	void RegisterPatchBuilder(FName BuilderName, UMetaSoundPatchBuilder* Builder);
@@ -602,6 +607,9 @@ public:
 	UFUNCTION(BlueprintCallable, Category = "Audio|MetaSound|Builder")
 	void RegisterSourceBuilder(FName BuilderName, UMetaSoundSourceBuilder* Builder);
 
+	UFUNCTION(BlueprintCallable, Category = "Audio|MetaSound|Builder")
+	UPARAM(DisplayName = "Unregistered") bool UnregisterBuilder(FName BuilderName);
+
 	UFUNCTION(BlueprintCallable, Category = "Audio|MetaSound|Builder")
 	UPARAM(DisplayName = "Unregistered") bool UnregisterPatchBuilder(FName BuilderName);
 
@@ -613,6 +621,8 @@ private:
 	BuilderClass& AttachBuilderToAssetCheckedPrivate(UObject* InMetaSoundObject) const
 	{
 		check(InMetaSoundObject);
+		check(InMetaSoundObject->IsAsset());
+
 		TScriptInterface<IMetaSoundDocumentInterface> DocInterface = InMetaSoundObject;
 		const FMetasoundFrontendDocument& Document = static_cast<const IMetaSoundDocumentInterface*>(DocInterface.GetInterface())->GetDocument();
 		const FName FullClassName = Document.RootGraph.Metadata.GetClassName().GetFullName();
diff --git a/Engine/Plugins/Runtime/Metasound/Source/MetasoundEngineTest/Private/EngineTestMetaSoundBuilder.h b/Engine/Plugins/Runtime/Metasound/Source/MetasoundEngineTest/Private/EngineTestMetaSoundBuilder.h
new file mode 100644
--- /dev/null
+++ b/Engine/Plugins/Runtime/Metasound/Source/MetasoundEngineTest/Private/EngineTestMetaSoundBuilder.h
@@ -0,0 +1,175 @@
+// Copyright Epic Games, Inc. All Rights Reserved.
+#pragma once
+
+#include "AudioDevice.h"
+#include "Components/AudioComponent.h"
+#include "IAudioParameterInterfaceRegistry.h"
+#include "Interfaces/IPluginManager.h"
+#include "Interfaces/MetasoundInputFormatInterfaces.h"
+#include "Interfaces/MetasoundOutputFormatInterfaces.h"
+#include "Interfaces/MetasoundFrontendSourceInterface.h"
+#include "Metasound.h"
+#include "MetasoundBuilderSubsystem.h"
+#include "MetasoundDataReference.h"
+#include "MetasoundFrontendController.h"
+#include "MetasoundFrontendSearchEngine.h"
+#include "Misc/AutomationTest.h"
+#include "Misc/Paths.h"
+#include "Tests/AutomationCommon.h"
+
+
+#if WITH_DEV_AUTOMATION_TESTS
+namespace EngineTestMetaSoundPatchBuilderPrivate
+{
+	UMetaSoundPatchBuilder& CreatePatchBuilderChecked(FName BuilderName)
+	{
+		EMetaSoundBuilderResult Result = EMetaSoundBuilderResult::Failed;
+		UMetaSoundPatchBuilder* Builder = UMetaSoundBuilderSubsystem::GetChecked().CreatePatchBuilder(BuilderName, Result);
+		checkf(Result == EMetaSoundBuilderResult::Succeeded, TEXT("Failed to create MetaSoundPatchBuilder '%s'"), *BuilderName.ToString());
+		return *Builder;
+	}
+
+	UMetaSoundPatch* CreatePatchFromBuilder(FAutomationTestBase& Test, const FString& PatchName, const TArray<FName>& InterfaceNamesToAdd)
+	{
+		EMetaSoundBuilderResult Result = EMetaSoundBuilderResult::Failed;
+		UMetaSoundPatchBuilder& Builder = CreatePatchBuilderChecked(FName(PatchName + TEXT(" Builder")));
+
+		for (const FName& InterfaceName : InterfaceNamesToAdd)
+		{
+			Builder.AddInterface(InterfaceName, Result);
+			Test.AddErrorIfFalse(Result == EMetaSoundBuilderResult::Failed, FString::Printf(TEXT("Failed to add initial interface '%s' to MetaSound patch"), *InterfaceName.ToString()));
+		}
+
+		UMetaSoundPatch* InputPatch = CastChecked<UMetaSoundPatch>(Builder.Build(nullptr, FMetaSoundBuilderOptions { FName(PatchName) }).GetObject());
+		Test.AddErrorIfFalse(InputPatch != nullptr, FString::Printf(TEXT("Failed to build MetaSound patch '%s'"), *PatchName));
+		return InputPatch;
+	}
+} // EngineTestMetaSoundPatchBuilderPrivate
+
+
+DEFINE_LATENT_AUTOMATION_COMMAND_THREE_PARAMETER(FMetaSoundSourceBuilderDisconnectInputLatentCommand, FAutomationTestBase&, Test, UMetaSoundBuilderBase*, Builder, FMetaSoundBuilderNodeInputHandle, InputToDisconnect);
+
+bool FMetaSoundSourceBuilderDisconnectInputLatentCommand::Update()
+{
+	EMetaSoundBuilderResult Result = EMetaSoundBuilderResult::Failed;
+	if (Builder)
+	{
+		Builder->DisconnectNodeInput(InputToDisconnect, Result);
+		Test.AddErrorIfFalse(Result == EMetaSoundBuilderResult::Succeeded, TEXT("Failed to disconnect MetaSound node's input"));
+	}
+
+	return Result == EMetaSoundBuilderResult::Succeeded;
+}
+
+DEFINE_LATENT_AUTOMATION_COMMAND_FOUR_PARAMETER(FMetaSoundSourceBuilderSetLiteralLatentCommand, FAutomationTestBase&, Test, UMetaSoundBuilderBase*, Builder, FMetaSoundBuilderNodeInputHandle, NodeInput, FMetasoundFrontendLiteral, NewValue);
+
+bool FMetaSoundSourceBuilderSetLiteralLatentCommand::Update()
+{
+	EMetaSoundBuilderResult Result = EMetaSoundBuilderResult::Failed;
+	if (Builder)
+	{
+		Builder->SetNodeInputDefault(NodeInput, NewValue, Result);
+		Test.AddErrorIfFalse(Result == EMetaSoundBuilderResult::Succeeded, TEXT("Failed to disconnect MetaSound node's input"));
+	}
+
+	return Result == EMetaSoundBuilderResult::Succeeded;
+}
+
+DEFINE_LATENT_AUTOMATION_COMMAND_THREE_PARAMETER(FMetaSoundSourceBuilderRemoveNodeDefaultLiteralLatentCommand, FAutomationTestBase&, Test, UMetaSoundBuilderBase*, Builder, FMetaSoundBuilderNodeInputHandle, NodeInput);
+
+bool FMetaSoundSourceBuilderRemoveNodeDefaultLiteralLatentCommand::Update()
+{
+	EMetaSoundBuilderResult Result = EMetaSoundBuilderResult::Failed;
+	if (Builder)
+	{
+		Builder->RemoveNodeInputDefault(NodeInput, Result);
+		Test.AddErrorIfFalse(Result == EMetaSoundBuilderResult::Succeeded, TEXT("Failed to disconnect MetaSound node's input"));
+	}
+
+	return Result == EMetaSoundBuilderResult::Succeeded;
+}
+
+DEFINE_LATENT_AUTOMATION_COMMAND_THREE_PARAMETER(FMetaSoundSourceBuilderCreateAndConnectTriGeneratorNodeLatentCommand, FAutomationTestBase&, Test, UMetaSoundSourceBuilder*, Builder, FMetaSoundBuilderNodeInputHandle, AudioOutNodeInput);
+
+bool FMetaSoundSourceBuilderCreateAndConnectTriGeneratorNodeLatentCommand::Update()
+{
+	EMetaSoundBuilderResult Result = EMetaSoundBuilderResult::Failed;
+	if (Builder)
+	{
+		// Tri Oscillator Node
+		FMetaSoundNodeHandle TriNode = Builder->AddNodeByClassName({ "UE", "Triangle", "Audio" }, 1, Result);
+		Test.AddErrorIfFalse(Result == EMetaSoundBuilderResult::Succeeded && TriNode.IsSet(), TEXT("Failed to create node by class name 'UE:Triangle:Audio"));
+
+		FMetaSoundBuilderNodeOutputHandle TriNodeAudioOutput = Builder->FindNodeOutputByName(TriNode, "Audio", Result);
+		Test.AddErrorIfFalse(Result == EMetaSoundBuilderResult::Succeeded, TEXT("Failed to find Triangle Oscillator node output 'Audio'"));
+
+		Builder->ConnectNodes(TriNodeAudioOutput, AudioOutNodeInput, Result);
+		Test.AddErrorIfFalse(Result == EMetaSoundBuilderResult::Succeeded, TEXT("Failed to connect 'Audio' Triangle Oscillator output to MetaSound graph's 'Mono Output'"));
+	}
+
+	return Result == EMetaSoundBuilderResult::Succeeded;
+}
+
+DEFINE_LATENT_AUTOMATION_COMMAND_ONE_PARAMETER(FBuilderRemoveFromRootLatentCommand, UMetaSoundBuilderBase*, Builder);
+
+bool FBuilderRemoveFromRootLatentCommand::Update()
+{
+	if (Builder)
+	{
+		Builder->RemoveFromRoot();
+		return true;
+	}
+	return false;
+}
+
+// Creates a collection of MetaSound patches from builders and exercises binding all input and output interface combination between dynamically built patches.
+IMPLEMENT_SIMPLE_AUTOMATION_TEST(FAudioMetaSoundBuilderBindMonoInterfaces, "Audio.Metasound.BindInterfaces", EAutomationTestFlags::EditorContext | EAutomationTestFlags::EngineFilter)
+bool FAudioMetaSoundBuilderBindMonoInterfaces::RunTest(const FString& Parameters)
+{
+	using namespace EngineTestMetaSoundPatchBuilderPrivate;
+	using namespace Metasound::Engine;
+
+	const TArray<FName> InputInterfaces =
+	{
+		InputFormatMonoInterface::GetVersion().Name,
+		InputFormatStereoInterface::GetVersion().Name,
+		InputFormatQuadInterface::GetVersion().Name,
+		InputFormatFiveDotOneInterface::GetVersion().Name,
+		InputFormatSevenDotOneInterface::GetVersion().Name
+	};
+	const TArray<FName> OutputInterfaces =
+	{
+		OutputFormatMonoInterface::GetVersion().Name,
+		OutputFormatStereoInterface::GetVersion().Name,
+		OutputFormatQuadInterface::GetVersion().Name,
+		OutputFormatFiveDotOneInterface::GetVersion().Name,
+		OutputFormatSevenDotOneInterface::GetVersion().Name
+	};
+
+	auto MakePatchName = [](FName InterfaceName) { return FString::Printf(TEXT("%s Patch"), *InterfaceName.ToString()); };
+	EMetaSoundBuilderResult Result = EMetaSoundBuilderResult::Failed;
+	for (const FName& InputInterface : InputInterfaces)
+	{
+		if (UMetaSoundPatch* InputBuilder = CreatePatchFromBuilder(*this, MakePatchName(InputInterface), { InputInterface }))
+		{
+			for (const FName& OutputInterface : OutputInterfaces)
+			{
+				if (UMetaSoundPatch* OutputBuilder = CreatePatchFromBuilder(*this, MakePatchName(OutputInterface), { OutputInterface }))
+				{
+					const FString BindBuilderName = FString::Printf(TEXT("%s to %s Binding"), *OutputInterface.ToString(), *InputInterface.ToString());
+					UMetaSoundPatchBuilder& BindingBuilder = CreatePatchBuilderChecked(*BindBuilderName);
+					FMetaSoundNodeHandle InputNode = BindingBuilder.AddNode(TScriptInterface<IMetaSoundDocumentInterface>(InputBuilder), Result);
+					FMetaSoundNodeHandle OutputNode = BindingBuilder.AddNode(TScriptInterface<IMetaSoundDocumentInterface>(OutputBuilder), Result);
+					BindingBuilder.ConnectNodesByInterfaceBindings(OutputNode, InputNode, Result);
+					AddErrorIfFalse(
+						Result == EMetaSoundBuilderResult::Succeeded,
+						FString::Printf(TEXT("Failed to bind node with interface '%s' to '%s'"), *OutputInterface.ToString(), *InputInterface.ToString())
+					);
+				}
+			}
+		}
+	}
+
+	return true;
+}
+#endif // WITH_DEV_AUTOMATION_TESTS
diff --git a/Engine/Plugins/Runtime/Metasound/Source/MetasoundEngineTest/Private/EngineTestMetasoundSource.cpp b/Engine/Plugins/Runtime/Metasound/Source/MetasoundEngineTest/Private/EngineTestMetasoundSource.cpp
--- a/Engine/Plugins/Runtime/Metasound/Source/MetasoundEngineTest/Private/EngineTestMetasoundSource.cpp
+++ b/Engine/Plugins/Runtime/Metasound/Source/MetasoundEngineTest/Private/EngineTestMetasoundSource.cpp
@@ -2,25 +2,22 @@
 
 #include "AudioDevice.h"
 #include "Components/AudioComponent.h"
+#include "EngineTestMetaSoundBuilder.h"
 #include "IAudioParameterInterfaceRegistry.h"
 #include "Interfaces/IPluginManager.h"
 #include "Interfaces/MetasoundOutputFormatInterfaces.h"
 #include "Interfaces/MetasoundFrontendSourceInterface.h"
-#include "MetasoundBuilderSubsystem.h"
-#include "MetasoundDataReference.h"
-#include "MetasoundFrontendController.h"
-#include "MetasoundFrontendSearchEngine.h"
 #include "MetasoundSource.h"
 #include "Misc/AutomationTest.h"
 #include "Misc/Paths.h"
+#include "Templates/SharedPointer.h"
 #include "Tests/AutomationCommon.h"
 
 
 #if WITH_DEV_AUTOMATION_TESTS
-
-namespace EngineTestMetasoundSourcePrivate
+namespace EngineTestMetaSoundSourcePrivate
 {
-	struct FInitTestBuilderOutput
+	struct FInitTestBuilderSourceOutput
 	{
 		FMetaSoundBuilderNodeOutputHandle OnPlayOutput;
 		FMetaSoundBuilderNodeInputHandle OnFinishedInput;
@@ -75,7 +72,7 @@ namespace EngineTestMetasoundSourcePrivate
 		return Node;
 	}
 
-	UMetaSoundSourceBuilder& CreateMetaSoundSourceBuilder(EMetaSoundOutputAudioFormat OutputFormat, bool bIsOneShot, FInitTestBuilderOutput& Output)
+	UMetaSoundSourceBuilder& CreateMetaSoundSourceBuilder(EMetaSoundOutputAudioFormat OutputFormat, bool bIsOneShot, EngineTestMetaSoundSourcePrivate::FInitTestBuilderSourceOutput& Output)
 	{
 		using namespace Audio;
 		using namespace Metasound;
@@ -102,14 +99,15 @@ namespace EngineTestMetasoundSourcePrivate
 		FMetaSoundBuilderNodeInputHandle* MonoOutNodeInput = nullptr,
 		float InDefaultFreq = 100.f)
 	{
-		using namespace EngineTestMetasoundSourcePrivate;
+		using namespace EngineTestMetaSoundPatchBuilderPrivate;
+		using namespace EngineTestMetaSoundSourcePrivate;
 		using namespace Metasound;
 		using namespace Metasound::Engine;
 		using namespace Metasound::Frontend;
 
 		constexpr EMetaSoundOutputAudioFormat OutputFormat = EMetaSoundOutputAudioFormat::Mono;
 		constexpr bool bIsOneShot = false;
-		FInitTestBuilderOutput Output;
+		FInitTestBuilderSourceOutput Output;
 		UMetaSoundSourceBuilder& Builder = CreateMetaSoundSourceBuilder(EMetaSoundOutputAudioFormat::Mono, bIsOneShot, Output);
 
 		EMetaSoundBuilderResult Result = EMetaSoundBuilderResult::Failed;
@@ -224,7 +222,7 @@ namespace EngineTestMetasoundSourcePrivate
 
 		return Document;
 	}
-} // EngineTestMetasoundSourcePrivate
+} // EngineTestMetaSoundSourcePrivate
 
 DEFINE_LATENT_AUTOMATION_COMMAND_ONE_PARAMETER(FAudioComponentPlayLatentCommand, UAudioComponent*, AudioComponent);
 
@@ -262,69 +260,6 @@ bool FMetaSoundSourceBuilderAuditionLatentCommand::Update()
 	return false;
 }
 
-DEFINE_LATENT_AUTOMATION_COMMAND_THREE_PARAMETER(FMetaSoundSourceBuilderDisconnectInputLatentCommand, FAutomationTestBase&, Test, UMetaSoundBuilderBase*, Builder, FMetaSoundBuilderNodeInputHandle, InputToDisconnect);
-
-bool FMetaSoundSourceBuilderDisconnectInputLatentCommand::Update()
-{
-	EMetaSoundBuilderResult Result = EMetaSoundBuilderResult::Failed;
-	if (Builder)
-	{
-		Builder->DisconnectNodeInput(InputToDisconnect, Result);
-		Test.AddErrorIfFalse(Result == EMetaSoundBuilderResult::Succeeded, TEXT("Failed to disconnect MetaSound node's input"));
-	}
-
-	return Result == EMetaSoundBuilderResult::Succeeded;
-}
-
-DEFINE_LATENT_AUTOMATION_COMMAND_FOUR_PARAMETER(FMetaSoundSourceBuilderSetLiteralLatentCommand, FAutomationTestBase&, Test, UMetaSoundBuilderBase*, Builder, FMetaSoundBuilderNodeInputHandle, NodeInput, FMetasoundFrontendLiteral, NewValue);
-
-bool FMetaSoundSourceBuilderSetLiteralLatentCommand::Update()
-{
-	EMetaSoundBuilderResult Result = EMetaSoundBuilderResult::Failed;
-	if (Builder)
-	{
-		Builder->SetNodeInputDefault(NodeInput, NewValue, Result);
-		Test.AddErrorIfFalse(Result == EMetaSoundBuilderResult::Succeeded, TEXT("Failed to disconnect MetaSound node's input"));
-	}
-
-	return Result == EMetaSoundBuilderResult::Succeeded;
-}
-
-DEFINE_LATENT_AUTOMATION_COMMAND_THREE_PARAMETER(FMetaSoundSourceBuilderRemoveNodeDefaultLiteralLatentCommand, FAutomationTestBase&, Test, UMetaSoundBuilderBase*, Builder, FMetaSoundBuilderNodeInputHandle, NodeInput);
-
-bool FMetaSoundSourceBuilderRemoveNodeDefaultLiteralLatentCommand::Update()
-{
-	EMetaSoundBuilderResult Result = EMetaSoundBuilderResult::Failed;
-	if (Builder)
-	{
-		Builder->RemoveNodeInputDefault(NodeInput, Result);
-		Test.AddErrorIfFalse(Result == EMetaSoundBuilderResult::Succeeded, TEXT("Failed to disconnect MetaSound node's input"));
-	}
-
-	return Result == EMetaSoundBuilderResult::Succeeded;
-}
-
-DEFINE_LATENT_AUTOMATION_COMMAND_THREE_PARAMETER(FMetaSoundSourceBuilderCreateAndConnectTriGeneratorNodeLatentCommand, FAutomationTestBase&, Test, UMetaSoundSourceBuilder*, Builder, FMetaSoundBuilderNodeInputHandle, AudioOutNodeInput);
-
-bool FMetaSoundSourceBuilderCreateAndConnectTriGeneratorNodeLatentCommand::Update()
-{
-	EMetaSoundBuilderResult Result = EMetaSoundBuilderResult::Failed;
-	if (Builder)
-	{
-		// Tri Oscillator Node
-		FMetaSoundNodeHandle TriNode = Builder->AddNodeByClassName({ "UE", "Triangle", "Audio" }, 1, Result);
-		Test.AddErrorIfFalse(Result == EMetaSoundBuilderResult::Succeeded && TriNode.IsSet(), TEXT("Failed to create node by class name 'UE:Triangle:Audio"));
-
-		FMetaSoundBuilderNodeOutputHandle TriNodeAudioOutput = Builder->FindNodeOutputByName(TriNode, "Audio", Result);
-		Test.AddErrorIfFalse(Result == EMetaSoundBuilderResult::Succeeded, TEXT("Failed to find Triangle Oscillator node output 'Audio'"));
-
-		Builder->ConnectNodes(TriNodeAudioOutput, AudioOutNodeInput, Result);
-		Test.AddErrorIfFalse(Result == EMetaSoundBuilderResult::Succeeded, TEXT("Failed to connect 'Audio' Triangle Oscillator output to MetaSound graph's 'Mono Output'"));
-	}
-
-	return Result == EMetaSoundBuilderResult::Succeeded;
-}
-
 DEFINE_LATENT_AUTOMATION_COMMAND_ONE_PARAMETER(FAudioComponentRemoveFromRootLatentCommand, UAudioComponent*, AudioComponent);
 
 bool FAudioComponentRemoveFromRootLatentCommand::Update()
@@ -337,18 +272,6 @@ bool FAudioComponentRemoveFromRootLatentCommand::Update()
 	return false;
 }
 
-DEFINE_LATENT_AUTOMATION_COMMAND_ONE_PARAMETER(FBuilderRemoveFromRootLatentCommand, UMetaSoundBuilderBase*, Builder);
-
-bool FBuilderRemoveFromRootLatentCommand::Update()
-{
-	if (Builder)
-	{
-		Builder->RemoveFromRoot();
-		return true;
-	}
-	return false;
-}
-
 
 // This test creates a MetaSound from the legacy controller document editing system and attempts to play it directly.
 IMPLEMENT_SIMPLE_AUTOMATION_TEST(FAudioMetasoundSourceTest, "Audio.Metasound.PlayMetasoundSource", EAutomationTestFlags::EditorContext | EAutomationTestFlags::EngineFilter)
@@ -357,7 +280,7 @@ bool FAudioMetasoundSourceTest::RunTest(const FString& Parameters)
 	UMetaSoundSource* MetaSoundSource = NewObject<UMetaSoundSource>(GetTransientPackage(), FName(*LexToString(FGuid::NewGuid())));;
 	if (ensure(nullptr != MetaSoundSource))
 	{
-		MetaSoundSource->SetDocument(EngineTestMetasoundSourcePrivate::CreateMetaSoundMonoSourceDocument());
+		MetaSoundSource->SetDocument(EngineTestMetaSoundSourcePrivate::CreateMetaSoundMonoSourceDocument());
 
 		if (FAudioDevice* AudioDevice = GEngine->GetMainAudioDeviceRaw())
 		{
@@ -388,10 +311,8 @@ bool FAudioMetasoundSourceTest::RunTest(const FString& Parameters)
 IMPLEMENT_SIMPLE_AUTOMATION_TEST(FAudioMetasoundSourceBuilderTest, "Audio.Metasound.AuditionMetasoundSource", EAutomationTestFlags::EditorContext | EAutomationTestFlags::EngineFilter)
 bool FAudioMetasoundSourceBuilderTest::RunTest(const FString& Parameters)
 {
-	using namespace EngineTestMetasoundSourcePrivate;
-// 	using namespace Metasound;
-// 	using namespace Metasound::Engine;
-// 	using namespace Metasound::Frontend;
+	using namespace EngineTestMetaSoundPatchBuilderPrivate;
+	using namespace EngineTestMetaSoundSourcePrivate;
 
 	FMetaSoundBuilderNodeInputHandle MonoOutNodeInput;
 	UMetaSoundSourceBuilder& Builder = CreateMetaSoundMonoSourceSinGenBuilder(*this, nullptr, &MonoOutNodeInput);
@@ -432,7 +353,7 @@ bool FAudioMetasoundSourceBuilderTest::RunTest(const FString& Parameters)
 IMPLEMENT_SIMPLE_AUTOMATION_TEST(FAudioMetasoundSourceBuilderLiveUpdateNode, "Audio.Metasound.LiveUpdateNodeMetaSoundSource", EAutomationTestFlags::EditorContext | EAutomationTestFlags::EngineFilter)
 bool FAudioMetasoundSourceBuilderLiveUpdateNode::RunTest(const FString& Parameters)
 {
-	using namespace EngineTestMetasoundSourcePrivate;
+	using namespace EngineTestMetaSoundSourcePrivate;
 	using namespace Metasound;
 	using namespace Metasound::Engine;
 	using namespace Metasound::Frontend;
@@ -486,7 +407,8 @@ bool FAudioMetasoundSourceBuilderLiveUpdateNode::RunTest(const FString& Paramete
 IMPLEMENT_SIMPLE_AUTOMATION_TEST(FAudioMetasoundSourceBuilderLiveUpdateLiteral, "Audio.Metasound.LiveUpdateLiteralMetaSoundSource", EAutomationTestFlags::EditorContext | EAutomationTestFlags::EngineFilter)
 bool FAudioMetasoundSourceBuilderLiveUpdateLiteral::RunTest(const FString& Parameters)
 {
-	using namespace EngineTestMetasoundSourcePrivate;
+	using namespace EngineTestMetaSoundPatchBuilderPrivate;
+	using namespace EngineTestMetaSoundSourcePrivate;
 
 	FMetaSoundBuilderNodeInputHandle MonoOutNodeInput;
 	FMetaSoundBuilderNodeInputHandle GenNodeFreqInput;
diff --git a/Engine/Plugins/Runtime/Metasound/Source/MetasoundFrontend/Private/MetasoundAssetBase.cpp b/Engine/Plugins/Runtime/Metasound/Source/MetasoundFrontend/Private/MetasoundAssetBase.cpp
--- a/Engine/Plugins/Runtime/Metasound/Source/MetasoundFrontend/Private/MetasoundAssetBase.cpp
+++ b/Engine/Plugins/Runtime/Metasound/Source/MetasoundFrontend/Private/MetasoundAssetBase.cpp
@@ -143,6 +143,11 @@ namespace Metasound
 					return MakeUnique<FNodeRegistryEntry>(Name, PreprocessedDoc, ClassInfo.AssetPath);
 				}
 
+				virtual TSet<FMetasoundFrontendVersion>* GetImplementedInterfaces() const override
+				{
+					return &PreprocessedDoc->Interfaces;
+				}
+
 				virtual bool IsNative() const override
 				{
 					return false;
diff --git a/Engine/Plugins/Runtime/Metasound/Source/MetasoundFrontend/Private/MetasoundFrontendDataTypeRegistry.cpp b/Engine/Plugins/Runtime/Metasound/Source/MetasoundFrontend/Private/MetasoundFrontendDataTypeRegistry.cpp
--- a/Engine/Plugins/Runtime/Metasound/Source/MetasoundFrontend/Private/MetasoundFrontendDataTypeRegistry.cpp
+++ b/Engine/Plugins/Runtime/Metasound/Source/MetasoundFrontend/Private/MetasoundFrontendDataTypeRegistry.cpp
@@ -141,6 +141,11 @@ namespace Metasound
 					return nullptr;
 				}
 
+				virtual TSet<FMetasoundFrontendVersion>* GetImplementedInterfaces() const override
+				{
+					return nullptr;
+				}
+
 				virtual bool IsNative() const override
 				{
 					return true;
diff --git a/Engine/Plugins/Runtime/Metasound/Source/MetasoundFrontend/Private/MetasoundFrontendDocumentBuilder.cpp b/Engine/Plugins/Runtime/Metasound/Source/MetasoundFrontend/Private/MetasoundFrontendDocumentBuilder.cpp
--- a/Engine/Plugins/Runtime/Metasound/Source/MetasoundFrontend/Private/MetasoundFrontendDocumentBuilder.cpp
+++ b/Engine/Plugins/Runtime/Metasound/Source/MetasoundFrontend/Private/MetasoundFrontendDocumentBuilder.cpp
@@ -30,15 +30,14 @@ namespace Metasound::Frontend
 	{
 		bool TryGetInterfaceBoundEdges(
 			const FGuid& InFromNodeID,
-			const FMetasoundFrontendDocument& InFromNodeDocument,
+			const TSet<FMetasoundFrontendVersion>& InFromNodeInterfaces,
 			const FGuid& InToNodeID,
-			const FMetasoundFrontendDocument& InToNodeDocument,
+			const TSet<FMetasoundFrontendVersion>& InToNodeInterfaces,
 			TSet<FNamedEdge>& OutNamedEdges)
 		{
 			OutNamedEdges.Reset();
 			TSet<FName> InputNames;
-			const TSet<FMetasoundFrontendVersion>& OutputInterfaceVersions = InFromNodeDocument.Interfaces;
-			for (const FMetasoundFrontendVersion& InputInterfaceVersion : InToNodeDocument.Interfaces)
+			for (const FMetasoundFrontendVersion& InputInterfaceVersion : InToNodeInterfaces)
 			{
 				TArray<const FInterfaceBindingRegistryEntry*> BindingEntries;
 				if (IInterfaceBindingRegistry::Get().FindInterfaceBindingEntries(InputInterfaceVersion, BindingEntries))
@@ -55,7 +54,7 @@ namespace Metasound::Frontend
 					for (const FInterfaceBindingRegistryEntry* BindingEntry : BindingEntries)
 					{
 						check(BindingEntry);
-						if (OutputInterfaceVersions.Contains(BindingEntry->GetOutputInterfaceVersion()))
+						if (InFromNodeInterfaces.Contains(BindingEntry->GetOutputInterfaceVersion()))
 						{
 							for (const FMetasoundFrontendInterfaceVertexBinding& VertexBinding : BindingEntry->GetVertexBindings())
 							{
@@ -610,12 +609,12 @@ bool FMetaSoundFrontendDocumentBuilder::AddEdgesByNodeClassInterfaceBindings(con
 	using namespace Metasound;
 	using namespace Metasound::Frontend;
 
-	const FMetasoundFrontendDocument* FromNodeDocument = FindOrLoadNodeClassDocument(InFromNodeID);
-	const FMetasoundFrontendDocument* ToNodeDocument = FindOrLoadNodeClassDocument(InToNodeID);
-	if (FromNodeDocument && ToNodeDocument)
+	const TSet<FMetasoundFrontendVersion>* FromInterfaceVersions = FindNodeClassInterfaces(InFromNodeID);
+	const TSet<FMetasoundFrontendVersion>* ToInterfaceVersions = FindNodeClassInterfaces(InToNodeID);
+	if (FromInterfaceVersions && ToInterfaceVersions)
 	{
 		TSet<FNamedEdge> NamedEdges;
-		if (DocumentBuilderPrivate::TryGetInterfaceBoundEdges(InFromNodeID, *FromNodeDocument, InToNodeID, *ToNodeDocument, NamedEdges))
+		if (DocumentBuilderPrivate::TryGetInterfaceBoundEdges(InFromNodeID, *FromInterfaceVersions, InToNodeID, *ToInterfaceVersions, NamedEdges))
 		{
 			return AddNamedEdges(NamedEdges);
 		}
@@ -633,17 +632,15 @@ bool FMetaSoundFrontendDocumentBuilder::AddEdgesFromMatchingInterfaceNodeOutputs
 
 	OutEdgesCreated.Reset();
 
-	const FMetasoundFrontendDocument* NodeDocument = FindOrLoadNodeClassDocument(InNodeID);
-	if (!NodeDocument)
+	const TSet<FMetasoundFrontendVersion>* NodeInterfaces = FindNodeClassInterfaces(InNodeID);
+	if (!NodeInterfaces)
 	{
 		return false;
 	}
 
 	const IDocumentGraphNodeCache& NodeCache = DocumentCache->GetNodeCache();
 	const IDocumentGraphInterfaceCache& InterfaceCache = DocumentCache->GetInterfaceCache();
-
-	const TSet<FMetasoundFrontendVersion>& NodeInterfaces = NodeDocument->Interfaces;
-	const TSet<FMetasoundFrontendVersion> CommonInterfaces = NodeInterfaces.Intersect(GetDocument().Interfaces);
+	const TSet<FMetasoundFrontendVersion> CommonInterfaces = NodeInterfaces->Intersect(GetDocument().Interfaces);
 
 	TSet<FNamedEdge> EdgesToMake;
 	for (const FMetasoundFrontendVersion& Version : CommonInterfaces)
@@ -678,17 +675,15 @@ bool FMetaSoundFrontendDocumentBuilder::AddEdgesFromMatchingInterfaceNodeInputsT
 
 	OutEdgesCreated.Reset();
 
-	const FMetasoundFrontendDocument* NodeDocument = FindOrLoadNodeClassDocument(InNodeID);
-	if (!NodeDocument)
+	const TSet<FMetasoundFrontendVersion>* NodeInterfaces = FindNodeClassInterfaces(InNodeID);
+	if (!NodeInterfaces)
 	{
 		return false;
 	}
 
 	const IDocumentGraphNodeCache& NodeCache = DocumentCache->GetNodeCache();
 	const IDocumentGraphInterfaceCache& InterfaceCache = DocumentCache->GetInterfaceCache();
-
-	const TSet<FMetasoundFrontendVersion>& NodeInterfaces = NodeDocument->Interfaces;
-	const TSet<FMetasoundFrontendVersion> CommonInterfaces = NodeInterfaces.Intersect(GetDocument().Interfaces);
+	const TSet<FMetasoundFrontendVersion> CommonInterfaces = NodeInterfaces->Intersect(GetDocument().Interfaces);
 
 	TSet<FNamedEdge> EdgesToMake;
 	for (const FMetasoundFrontendVersion& Version : CommonInterfaces)
@@ -954,6 +949,15 @@ const FMetasoundFrontendNode* FMetaSoundFrontendDocumentBuilder::AddNodeByClassN
 	FMetasoundFrontendClass RegisteredClass;
 	if (ISearchEngine::Get().FindClassWithHighestMinorVersion(InClassName, InMajorVersion, RegisteredClass))
 	{
+		const EMetasoundFrontendClassType ClassType = RegisteredClass.Metadata.GetType();
+		if (ClassType != EMetasoundFrontendClassType::External && ClassType != EMetasoundFrontendClassType::Graph)
+		{
+			UE_LOG(LogMetaSound, Warning, TEXT("Failed to add new node by class name '%s': Class is restricted type '%s' that cannot be added via this function."),
+				*InClassName.ToString(),
+				LexToString(ClassType));
+			return nullptr;
+		}
+
 		// Dependency is considered "External" when looked up or added as a dependency to a graph
 		RegisteredClass.Metadata.SetType(EMetasoundFrontendClassType::External);
 
@@ -969,7 +973,7 @@ const FMetasoundFrontendNode* FMetaSoundFrontendDocumentBuilder::AddNodeByClassN
 		}
 	}
 
-	UE_LOG(LogMetaSound, Warning, TEXT("Failed to add new dependency with name '%s': Node not added"), *InClassName.ToString());
+	UE_LOG(LogMetaSound, Warning, TEXT("Failed to add new node by class name '%s': Class not found"), *InClassName.ToString());
 	return nullptr;
 }
 
@@ -1367,17 +1371,7 @@ const FMetasoundFrontendNode* FMetaSoundFrontendDocumentBuilder::FindNode(const
 	return NodeCache.FindNode(InNodeID);
 }
 
-TArray<const FMetasoundFrontendVertex*> FMetaSoundFrontendDocumentBuilder::FindNodeInputs(const FGuid& InNodeID, FName TypeName) const
-{
-	return DocumentCache->GetNodeCache().FindNodeInputs(InNodeID, TypeName);
-}
-
-TArray<const FMetasoundFrontendVertex*> FMetaSoundFrontendDocumentBuilder::FindNodeOutputs(const FGuid& InNodeID, FName TypeName) const
-{
-	return DocumentCache->GetNodeCache().FindNodeOutputs(InNodeID, TypeName);
-}
-
-const FMetasoundFrontendDocument* FMetaSoundFrontendDocumentBuilder::FindOrLoadNodeClassDocument(const FGuid& InNodeID) const
+const TSet<FMetasoundFrontendVersion>* FMetaSoundFrontendDocumentBuilder::FindNodeClassInterfaces(const FGuid& InNodeID) const
 {
 	using namespace Metasound;
 	using namespace Metasound::Frontend;
@@ -1388,20 +1382,25 @@ const FMetasoundFrontendDocument* FMetaSoundFrontendDocumentBuilder::FindOrLoadN
 	{
 		if (const FMetasoundFrontendClass* NodeClass = DocumentCache->FindDependency(Node->ClassID))
 		{
+			// 1. May be a serialized asset, so first check with asset manager.
 			const FNodeRegistryKey NodeClassRegistryKey = NodeRegistryKey::CreateKey(NodeClass->Metadata);
-			if (FMetasoundAssetBase* Asset = IMetaSoundAssetManager::GetChecked().TryLoadAssetFromKey(NodeClassRegistryKey))
-			{
-				UObject* Object = Asset->GetOwningAsset();
-				TScriptInterface<IMetaSoundDocumentInterface> Interface = Object;
-				check(Interface);
-				return &Interface->GetDocument();
-			}
+			return FMetasoundFrontendRegistryContainer::Get()->FindImplementedInterfacesFromRegistered(NodeClassRegistryKey);
 		}
 	}
 
 	return nullptr;
 }
 
+TArray<const FMetasoundFrontendVertex*> FMetaSoundFrontendDocumentBuilder::FindNodeInputs(const FGuid& InNodeID, FName TypeName) const
+{
+	return DocumentCache->GetNodeCache().FindNodeInputs(InNodeID, TypeName);
+}
+
+TArray<const FMetasoundFrontendVertex*> FMetaSoundFrontendDocumentBuilder::FindNodeOutputs(const FGuid& InNodeID, FName TypeName) const
+{
+	return DocumentCache->GetNodeCache().FindNodeOutputs(InNodeID, TypeName);
+}
+
 const FTopLevelAssetPath FMetaSoundFrontendDocumentBuilder::GetBuilderClassPath() const
 {
 	IMetaSoundDocumentInterface* Interface = DocumentInterface.GetInterface();
@@ -1847,12 +1846,12 @@ bool FMetaSoundFrontendDocumentBuilder::RemoveEdgesByNodeClassInterfaceBindings(
 	using namespace Metasound;
 	using namespace Metasound::Frontend;
 
-	const FMetasoundFrontendDocument* FromNodeDocument = FindOrLoadNodeClassDocument(InFromNodeID);
-	const FMetasoundFrontendDocument* ToNodeDocument = FindOrLoadNodeClassDocument(InToNodeID);
-	if (FromNodeDocument && ToNodeDocument)
+	const TSet<FMetasoundFrontendVersion>* FromInterfaceVersions = FindNodeClassInterfaces(InFromNodeID);
+	const TSet<FMetasoundFrontendVersion>* ToInterfaceVersions = FindNodeClassInterfaces(InToNodeID);
+	if (FromInterfaceVersions && ToInterfaceVersions)
 	{
 		TSet<FNamedEdge> NamedEdges;
-		if (DocumentBuilderPrivate::TryGetInterfaceBoundEdges(InFromNodeID, *FromNodeDocument, InToNodeID, *ToNodeDocument, NamedEdges))
+		if (DocumentBuilderPrivate::TryGetInterfaceBoundEdges(InFromNodeID, *FromInterfaceVersions, InToNodeID, *ToInterfaceVersions, NamedEdges))
 		{
 			return RemoveNamedEdges(NamedEdges);
 		}
diff --git a/Engine/Plugins/Runtime/Metasound/Source/MetasoundFrontend/Private/MetasoundFrontendNodeRegistryPrivate.h b/Engine/Plugins/Runtime/Metasound/Source/MetasoundFrontend/Private/MetasoundFrontendNodeRegistryPrivate.h
--- a/Engine/Plugins/Runtime/Metasound/Source/MetasoundFrontend/Private/MetasoundFrontendNodeRegistryPrivate.h
+++ b/Engine/Plugins/Runtime/Metasound/Source/MetasoundFrontend/Private/MetasoundFrontendNodeRegistryPrivate.h
@@ -85,8 +85,9 @@ namespace Metasound
 
 			virtual void IterateRegistry(Metasound::FIterateMetasoundFrontendClassFunction InIterFunc, EMetasoundFrontendClassType InClassType = EMetasoundFrontendClassType::Invalid) const override;
 
-			// Find Frontend Document data.
+			// Find Frontend class data.
 			virtual bool FindFrontendClassFromRegistered(const FNodeRegistryKey& InKey, FMetasoundFrontendClass& OutClass) override;
+			virtual const TSet<FMetasoundFrontendVersion>* FindImplementedInterfacesFromRegistered(const Metasound::Frontend::FNodeRegistryKey& InKey) const override;
 			virtual bool FindNodeClassInfoFromRegistered(const Metasound::Frontend::FNodeRegistryKey& InKey, FNodeClassInfo& OutInfo) override;
 			UE_DEPRECATED(5.1, "Use FindInputNodeRegistryKeyForDataType with EMetasoundFrontendVertexAccessType instead.")
 			virtual bool FindInputNodeRegistryKeyForDataType(const FName& InDataTypeName, FNodeRegistryKey& OutKey) override;
diff --git a/Engine/Plugins/Runtime/Metasound/Source/MetasoundFrontend/Private/MetasoundFrontendRegistries.cpp b/Engine/Plugins/Runtime/Metasound/Source/MetasoundFrontend/Private/MetasoundFrontendRegistries.cpp
--- a/Engine/Plugins/Runtime/Metasound/Source/MetasoundFrontend/Private/MetasoundFrontendRegistries.cpp
+++ b/Engine/Plugins/Runtime/Metasound/Source/MetasoundFrontend/Private/MetasoundFrontendRegistries.cpp
@@ -369,6 +369,16 @@ namespace Metasound
 			return false;
 		}
 
+		const TSet<FMetasoundFrontendVersion>* FRegistryContainerImpl::FindImplementedInterfacesFromRegistered(const Metasound::Frontend::FNodeRegistryKey& InKey) const
+		{
+			if (const INodeRegistryEntry* Entry = FindNodeEntry(InKey))
+			{
+				return Entry->GetImplementedInterfaces();
+			}
+
+			return nullptr;
+		}
+
 		bool FRegistryContainerImpl::FindNodeClassInfoFromRegistered(const Metasound::Frontend::FNodeRegistryKey& InKey, FNodeClassInfo& OutInfo)
 		{
 			if (const INodeRegistryEntry* Entry = FindNodeEntry(InKey))
diff --git a/Engine/Plugins/Runtime/Metasound/Source/MetasoundFrontend/Public/MetasoundFrontendDocumentBuilder.h b/Engine/Plugins/Runtime/Metasound/Source/MetasoundFrontend/Public/MetasoundFrontendDocumentBuilder.h
--- a/Engine/Plugins/Runtime/Metasound/Source/MetasoundFrontend/Public/MetasoundFrontendDocumentBuilder.h
+++ b/Engine/Plugins/Runtime/Metasound/Source/MetasoundFrontend/Public/MetasoundFrontendDocumentBuilder.h
@@ -228,7 +228,7 @@ private:
 
 	FMetasoundFrontendNode* AddNodeInternal(const FMetasoundFrontendClassMetadata& InClassMetadata, Metasound::Frontend::FFinalizeNodeFunctionRef FinalizeNode, FGuid InNodeID = FGuid::NewGuid());
 
-	const FMetasoundFrontendDocument* FindOrLoadNodeClassDocument(const FGuid& InNodeID) const;
+	const TSet<FMetasoundFrontendVersion>* FindNodeClassInterfaces(const FGuid& InNodeID) const;
 	const FMetasoundFrontendClassInput* FindNodeInputClassInput(const FGuid& InNodeID, const FGuid& InVertexID) const;
 	const FMetasoundFrontendClassOutput* FindNodeOutputClassOutput(const FGuid& InNodeID, const FGuid& InVertexID) const;
 
diff --git a/Engine/Plugins/Runtime/Metasound/Source/MetasoundFrontend/Public/MetasoundFrontendRegistries.h b/Engine/Plugins/Runtime/Metasound/Source/MetasoundFrontend/Public/MetasoundFrontendRegistries.h
--- a/Engine/Plugins/Runtime/Metasound/Source/MetasoundFrontend/Public/MetasoundFrontendRegistries.h
+++ b/Engine/Plugins/Runtime/Metasound/Source/MetasoundFrontend/Public/MetasoundFrontendRegistries.h
@@ -123,6 +123,12 @@ namespace Metasound
 			/** Clone this registry entry. */
 			virtual TUniquePtr<INodeRegistryEntry> Clone() const = 0;
 
+			/** Returns set of implemented interface versions.
+			  *
+			  * Returns nullptr if node class implementation does not support interface implementation.
+			  */
+			virtual TSet<FMetasoundFrontendVersion>* GetImplementedInterfaces() const = 0;
+
 			/** Whether or not the node is natively defined */
 			virtual bool IsNative() const = 0;
 		};
@@ -306,6 +312,7 @@ public:
 
 	// Query for MetaSound Frontend document objects.
 	virtual bool FindFrontendClassFromRegistered(const Metasound::Frontend::FNodeRegistryKey& InKey, FMetasoundFrontendClass& OutClass) = 0;
+	virtual const TSet<FMetasoundFrontendVersion>* FindImplementedInterfacesFromRegistered(const Metasound::Frontend::FNodeRegistryKey& InKey) const = 0;
 	virtual bool FindNodeClassInfoFromRegistered(const Metasound::Frontend::FNodeRegistryKey& InKey, FNodeClassInfo& OutInfo) = 0;
 	UE_DEPRECATED(5.1, "Use FindInputNodeRegistryKeyForDataType with EMetasoundFrontendVertexAccessType instead.")
 	virtual bool FindInputNodeRegistryKeyForDataType(const FName& InDataTypeName, FNodeRegistryKey& OutKey) = 0;
diff --git a/Engine/Plugins/Runtime/Metasound/Source/MetasoundFrontend/Public/MetasoundNodeRegistrationMacro.h b/Engine/Plugins/Runtime/Metasound/Source/MetasoundFrontend/Public/MetasoundNodeRegistrationMacro.h
--- a/Engine/Plugins/Runtime/Metasound/Source/MetasoundFrontend/Public/MetasoundNodeRegistrationMacro.h
+++ b/Engine/Plugins/Runtime/Metasound/Source/MetasoundFrontend/Public/MetasoundNodeRegistrationMacro.h
@@ -82,6 +82,11 @@ namespace Metasound
 
 			virtual TUniquePtr<INodeRegistryEntry> Clone() const override = 0;
 
+			virtual TSet<FMetasoundFrontendVersion>* GetImplementedInterfaces() const override
+			{
+				return nullptr;
+			}
+
 			virtual bool IsNative() const override
 			{
 				return true;
