Fix dynamic transactor missing node.
#rnx

diff --git a/Engine/Plugins/Runtime/Metasound/Source/MetasoundEngine/Private/MetasoundSource.cpp b/Engine/Plugins/Runtime/Metasound/Source/MetasoundEngine/Private/MetasoundSource.cpp
--- a/Engine/Plugins/Runtime/Metasound/Source/MetasoundEngine/Private/MetasoundSource.cpp
+++ b/Engine/Plugins/Runtime/Metasound/Source/MetasoundEngine/Private/MetasoundSource.cpp
@@ -1023,7 +1023,7 @@ TSharedPtr<Metasound::DynamicGraph::FDynamicOperatorTransactor> UMetaSoundSource
 
 	if (bInIsEnabled && (!DynamicTransactor.IsValid()))
 	{
-		TSharedPtr<IGraph> CurrentGraph = GetRuntimeData().Graph;
+		TSharedPtr<FGraph> CurrentGraph = GetRuntimeData().Graph;
 		if (CurrentGraph.IsValid())
 		{
 			DynamicTransactor = MakeShared<FDynamicOperatorTransactor>(*CurrentGraph);
diff --git a/Engine/Plugins/Runtime/Metasound/Source/MetasoundFrontend/Private/MetasoundAssetBase.cpp b/Engine/Plugins/Runtime/Metasound/Source/MetasoundFrontend/Private/MetasoundAssetBase.cpp
--- a/Engine/Plugins/Runtime/Metasound/Source/MetasoundFrontend/Private/MetasoundAssetBase.cpp
+++ b/Engine/Plugins/Runtime/Metasound/Source/MetasoundFrontend/Private/MetasoundAssetBase.cpp
@@ -523,7 +523,7 @@ const FMetasoundFrontendDocumentModifyContext& FMetasoundAssetBase::GetModifyCon
 }
 #endif // WITH_EDITOR
 
-TSharedPtr<Metasound::IGraph, ESPMode::ThreadSafe> FMetasoundAssetBase::BuildMetasoundDocument(const FMetasoundFrontendDocument& InPreprocessedDoc, const TSet<FName>& InTransmittableInputNames) const
+TSharedPtr<Metasound::FGraph, ESPMode::ThreadSafe> FMetasoundAssetBase::BuildMetasoundDocument(const FMetasoundFrontendDocument& InPreprocessedDoc, const TSet<FName>& InTransmittableInputNames) const
 {
 	using namespace Metasound;
 	using namespace Metasound::Frontend;
@@ -537,7 +537,7 @@ TSharedPtr<Metasound::IGraph, ESPMode::ThreadSafe> FMetasoundAssetBase::BuildMet
 		UE_LOG(LogMetaSound, Error, TEXT("Failed to build MetaSound graph in asset '%s'"), *GetOwningAssetName());
 	}
 
-	TSharedPtr<Metasound::IGraph, ESPMode::ThreadSafe> SharedGraph(FrontendGraph.Release());
+	TSharedPtr<Metasound::FGraph, ESPMode::ThreadSafe> SharedGraph(FrontendGraph.Release());
 
 	return SharedGraph;
 }
@@ -903,7 +903,7 @@ const FMetasoundAssetBase::FRuntimeData& FMetasoundAssetBase::CacheRuntimeData(c
 
 	TSet<FName> TransmittableInputNames;
 	Algo::Transform(TransmittableInputs, TransmittableInputNames, [](const FMetasoundFrontendClassInput& Input) { return Input.Name; });
-	TSharedPtr<Metasound::IGraph, ESPMode::ThreadSafe> Graph = BuildMetasoundDocument(InPreprocessedDoc, TransmittableInputNames);
+	TSharedPtr<Metasound::FGraph, ESPMode::ThreadSafe> Graph = BuildMetasoundDocument(InPreprocessedDoc, TransmittableInputNames);
 
 	CachedRuntimeData =
 	{
diff --git a/Engine/Plugins/Runtime/Metasound/Source/MetasoundFrontend/Private/MetasoundFrontendGraph.cpp b/Engine/Plugins/Runtime/Metasound/Source/MetasoundFrontend/Private/MetasoundFrontendGraph.cpp
--- a/Engine/Plugins/Runtime/Metasound/Source/MetasoundFrontend/Private/MetasoundFrontendGraph.cpp
+++ b/Engine/Plugins/Runtime/Metasound/Source/MetasoundFrontend/Private/MetasoundFrontendGraph.cpp
@@ -64,31 +64,6 @@ namespace Metasound
 		}
 	}
 
-	/** Store a node on this graph. */
-	void FFrontendGraph::AddNode(FGuid InNodeID, TSharedPtr<const INode> InNode)
-	{
-		if (InNode.IsValid())
-		{
-			// There shouldn't be duplicate IDs. 
-			check(!NodeMap.Contains(InNodeID));
-
-			NodeMap.Add(InNodeID, InNode.Get());
-			StoreNode(InNode);
-		}
-	}
-
-	const INode* FFrontendGraph::FindNode(FGuid InNodeID) const
-	{
-		const INode* const* NodePtr = NodeMap.Find(InNodeID);
-
-		if (nullptr != NodePtr)
-		{
-			return *NodePtr;
-		}
-
-		return nullptr;
-	}
-
 	const INode* FFrontendGraph::FindInputNode(int32 InIndex) const
 	{
 		const INode* const* NodePtr = InputNodes.Find(InIndex);
@@ -113,52 +88,12 @@ namespace Metasound
 		return nullptr;
 	}
 
-	/** Returns true if all edges, destinations and sources refer to 
-	 * nodes stored in this graph. */
 	bool FFrontendGraph::OwnsAllReferencedNodes() const
 	{
-		const TArray<FDataEdge>& AllEdges = GetDataEdges();
-		for (const FDataEdge& Edge : AllEdges)
-		{
-			if (!StoredNodes.Contains(Edge.From.Node))
-			{
-				return false;
-			}
-
-			if (!StoredNodes.Contains(Edge.To.Node))
-			{
-				return false;
-			}
-		}
-
-		const FInputDataDestinationCollection& AllInputDestinations = GetInputDataDestinations();
-		for (auto& DestTuple : AllInputDestinations)
-		{
-			if (!StoredNodes.Contains(DestTuple.Value.Node))
-			{
-				return false;
-			}
-		}
-
-		const FOutputDataSourceCollection& AllOutputSources = GetOutputDataSources();
-		for (auto& SourceTuple : AllOutputSources)
-		{
-			if (!StoredNodes.Contains(SourceTuple.Value.Node))
-			{
-				return false;
-			}
-		}
-
+		// This function is deprecated in 5.3 and is no longer functional.
 		return true;
 	}
 
-	void FFrontendGraph::StoreNode(TSharedPtr<const INode> InNode)
-	{
-		check(InNode.IsValid());
-		StoredNodes.Add(InNode.Get());
-		NodeStorage.Add(InNode);
-	}
-
 	TUniquePtr<INode> FFrontendGraphBuilder::CreateVariableNode(const FMetasoundFrontendNode& InNode, const FMetasoundFrontendGraph& InGraph)
 	{
 		using namespace Metasound::Frontend;
@@ -827,7 +762,6 @@ namespace Metasound
 
 		if (bSuccess)
 		{
-			check(BuildGraphContext.Graph->OwnsAllReferencedNodes());
 			return MoveTemp(BuildGraphContext.Graph);
 		}
 		else
diff --git a/Engine/Plugins/Runtime/Metasound/Source/MetasoundFrontend/Public/MetasoundAssetBase.h b/Engine/Plugins/Runtime/Metasound/Source/MetasoundFrontend/Public/MetasoundAssetBase.h
--- a/Engine/Plugins/Runtime/Metasound/Source/MetasoundFrontend/Public/MetasoundAssetBase.h
+++ b/Engine/Plugins/Runtime/Metasound/Source/MetasoundFrontend/Public/MetasoundAssetBase.h
@@ -208,7 +208,7 @@ protected:
 		TArray<FMetasoundFrontendClassInput> TransmittableInputs;
 
 		// Core graph.
-		TSharedPtr<Metasound::IGraph, ESPMode::ThreadSafe> Graph;
+		TSharedPtr<Metasound::FGraph, ESPMode::ThreadSafe> Graph;
 	};
 
 	// Returns the cached runtime data.
@@ -235,5 +235,5 @@ private:
 	FGuid CurrentCachedRuntimeDataChangeID;
 	FRuntimeData CachedRuntimeData;
 
-	TSharedPtr<Metasound::IGraph, ESPMode::ThreadSafe> BuildMetasoundDocument(const FMetasoundFrontendDocument& InPreprocessDoc, const TSet<FName>& InTransmittableInputNames) const;
+	TSharedPtr<Metasound::FGraph, ESPMode::ThreadSafe> BuildMetasoundDocument(const FMetasoundFrontendDocument& InPreprocessDoc, const TSet<FName>& InTransmittableInputNames) const;
 };
diff --git a/Engine/Plugins/Runtime/Metasound/Source/MetasoundFrontend/Public/MetasoundFrontendGraph.h b/Engine/Plugins/Runtime/Metasound/Source/MetasoundFrontend/Public/MetasoundFrontendGraph.h
--- a/Engine/Plugins/Runtime/Metasound/Source/MetasoundFrontend/Public/MetasoundFrontendGraph.h
+++ b/Engine/Plugins/Runtime/Metasound/Source/MetasoundFrontend/Public/MetasoundFrontendGraph.h
@@ -2,11 +2,18 @@
 
 #pragma once
 
-#include "CoreMinimal.h"
-
-#include "MetasoundFrontend.h"
+#include "Containers/Array.h"
+#include "Containers/Map.h"
+#include "Containers/Set.h"
+#include "Containers/UnrealString.h"
+#include "MetasoundFrontendDocument.h"
 #include "MetasoundGraph.h"
+#include "MetasoundNodeConstructorParams.h"
 #include "MetasoundNodeInterface.h"
+#include "MetasoundVertex.h"
+#include "Misc/Guid.h"
+#include "Templates/SharedPointer.h"
+#include "Templates/UniquePtr.h"
 
 namespace Metasound
 {
@@ -44,21 +51,6 @@ namespace Metasound
 			 */
 			void AddOutputNode(FGuid InNodeID, int32 InIndex, const FVertexName& InVertexName, TSharedPtr<const INode> InNode);
 
-			/** Store a node on this graph. 
-			 *
-			 * @param InNodeID - The NodeID related to the parent FMetasoundFrontendClass.
-			 * @param InNode - A shared pointer to a node. 
-			 */
-			void AddNode(FGuid InNodeID, TSharedPtr<const INode> InNode);
-
-			/** Retrieve node by node ID.
-			 *
-			 * @param InNodeID - The NodeID of the requested Node.
-			 *
-			 * @return Pointer to the Node if it is stored on this graph. nullptr otherwise. 
-			 */
-			const INode* FindNode(FGuid InNodeID) const;
-
 			/** Retrieve node by input index.
 			 *
 			 * @param InIndex - The index of the requested input.
@@ -75,20 +67,13 @@ namespace Metasound
 			 */
 			const INode* FindOutputNode(int32 InIndex) const;
 
-			/** Returns true if all edges, destinations and sources refer to 
-			 * nodes stored in this graph. */
+			UE_DEPRECATED(5.3, "This function is no longer analyzes node ownership and will always return true.")
 			bool OwnsAllReferencedNodes() const;
 
 		private:
 
-			void StoreNode(TSharedPtr<const INode> InNode);
-
 			TMap<int32, const INode*> InputNodes;
 			TMap<int32, const INode*> OutputNodes;
-
-			TMap<FGuid, const INode*> NodeMap;
-			TSet<const INode*> StoredNodes;
-			TArray<TSharedPtr<const INode>> NodeStorage;
 	};
 
 	/** FFrontendGraphBuilder builds a FFrontendGraph from a FMetasoundDocument
diff --git a/Engine/Plugins/Runtime/Metasound/Source/MetasoundGraphCore/Private/MetasoundDynamicOperatorTransactor.cpp b/Engine/Plugins/Runtime/Metasound/Source/MetasoundGraphCore/Private/MetasoundDynamicOperatorTransactor.cpp
--- a/Engine/Plugins/Runtime/Metasound/Source/MetasoundGraphCore/Private/MetasoundDynamicOperatorTransactor.cpp
+++ b/Engine/Plugins/Runtime/Metasound/Source/MetasoundGraphCore/Private/MetasoundDynamicOperatorTransactor.cpp
@@ -58,7 +58,7 @@ namespace Metasound
 			}
 		}
 
-		FDynamicOperatorTransactor::FDynamicOperatorTransactor(const IGraph& InGraph)
+		FDynamicOperatorTransactor::FDynamicOperatorTransactor(const FGraph& InGraph)
 		: Graph(InGraph)
 		{
 		}
@@ -87,11 +87,9 @@ namespace Metasound
 				return;
 			}
 
-			// Cache underlying pointer to node for later use. TUniquePtr<INode> will be moved to node
-			// map thus invalidating `TUniquePtr<INode> InNode`. 
-			INode* NodePtr = InNode.Get();
+			TSharedRef<const INode> NodePtr(InNode.Release());
 
-			NodeMap.Add(InNodeID, MoveTemp(InNode));
+			Graph.AddNode(InNodeID, NodePtr);
 
 			auto CreateAddNodeTransform = [&](const FOperatorSettings& InOperatorSettings, const FMetasoundEnvironment& InEnvironment) -> TUniquePtr<IDynamicOperatorTransform>
 			{
@@ -107,7 +105,7 @@ namespace Metasound
 
 			METASOUND_TRACE_CPUPROFILER_EVENT_SCOPE(Metasound::FDynamicOperatorTransactor::RemoveNode);
 
-			if (const INode* Node = FindNode(InNodeID))
+			if (const INode* Node = Graph.FindNode(InNodeID))
 			{
 				FOperatorID OperatorID = GetOperatorID(Node);
 				auto CreateRemoveNodeTransform = [&](const FOperatorSettings& InOperatorSettings, const FMetasoundEnvironment& InEnvironment) -> TUniquePtr<IDynamicOperatorTransform>
@@ -116,12 +114,13 @@ namespace Metasound
 				};
 				EnqueueTransformOnOperatorQueues(CreateRemoveNodeTransform);
 
-				Graph.RemoveDataEdgesWithNode(*Node);
-				NodeMap.Remove(InNodeID);
+				constexpr bool bRemoveDataEdgesWithNode = true;
+				bool bRemovedNodeFromGraph = Graph.RemoveNode(InNodeID, bRemoveDataEdgesWithNode);
+				check(bRemovedNodeFromGraph); //< Should always be true because we know the node exists in the Graph from `Graph.FindNode(...)`
 			}
 			else
 			{
-				UE_LOG(LogMetaSound, Error, TEXT("No node found in dynamic transactor with ID %s"), *InNodeID.ToString());
+				UE_LOG(LogMetaSound, Error, TEXT("No node found in dynamic transactor graph with ID %s"), *InNodeID.ToString());
 			}
 		}
 
@@ -132,8 +131,8 @@ namespace Metasound
 
 			METASOUND_TRACE_CPUPROFILER_EVENT_SCOPE(Metasound::FDynamicOperatorTransactor::AddDataEdge);
 
-			const INode* FromNode = FindNode(InFromNodeID);
-			const INode* ToNode = FindNode(InToNodeID);
+			const INode* FromNode = Graph.FindNode(InFromNodeID);
+			const INode* ToNode = Graph.FindNode(InToNodeID);
 
 			if ((nullptr == ToNode) || (nullptr == FromNode))
 			{
@@ -150,8 +149,8 @@ namespace Metasound
 
 			METASOUND_TRACE_CPUPROFILER_EVENT_SCOPE(Metasound::FDynamicOperatorTransactor::RemoveDataEdge);
 
-			const INode* FromNode = FindNode(InFromNodeID);
-			const INode* ToNode = FindNode(InToNodeID);
+			const INode* FromNode = Graph.FindNode(InFromNodeID);
+			const INode* ToNode = Graph.FindNode(InToNodeID);
 			const INode* ReplacementLiteralNode = InReplacementLiteralNode.Get(); // Cache pointer because TUniquePtr<INode> will get moved
 
 			if ((nullptr == ToNode) || (nullptr == FromNode))
@@ -197,7 +196,6 @@ namespace Metasound
 
 			auto CreateRemoveEdgeTransform = [&](const FOperatorSettings& InOperatorSettings, const FMetasoundEnvironment& InEnvironment) -> TUniquePtr<IDynamicOperatorTransform>
 			{
-
 				// Add the literal node.
 				TUniquePtr<IDynamicOperatorTransform> AddNodeTransform = CreateAddOperatorTransform(*ReplacementLiteralNode, InOperatorSettings, InEnvironment);
 
@@ -232,7 +230,7 @@ namespace Metasound
 		{
 			METASOUND_TRACE_CPUPROFILER_EVENT_SCOPE(Metasound::FDynamicOperatorTransactor::SetValue);
 
-			const INode* Node = FindNode(InNodeID);
+			const INode* Node = Graph.FindNode(InNodeID);
 			const INode* LiteralNode = InLiteralNode.Get();
 
 			if (nullptr == Node)
@@ -271,7 +269,7 @@ namespace Metasound
 		{
 			METASOUND_TRACE_CPUPROFILER_EVENT_SCOPE(Metasound::FDynamicOperatorTransactor::AddInputDataDestination);
 
-			const INode* Node = FindNode(InNodeID);
+			const INode* Node = Graph.FindNode(InNodeID);
 			if (nullptr == Node)
 			{
 				UE_LOG(LogMetaSound, Error, TEXT("Cannot add Input Data Destination %s:%s because of missing node"), *InNodeID.ToString(), *InVertexName.ToString());
@@ -331,7 +329,7 @@ namespace Metasound
 		{
 			METASOUND_TRACE_CPUPROFILER_EVENT_SCOPE(Metasound::FDynamicOperatorTransactor::AddOutputDataSource);
 
-			const INode* Node = FindNode(InNodeID);
+			const INode* Node = Graph.FindNode(InNodeID);
 			if (nullptr == Node)
 			{
 				UE_LOG(LogMetaSound, Error, TEXT("Cannot add Output Data Source %s:%s because of missing node"), *InNodeID.ToString(), *InVertexName.ToString());
@@ -412,19 +410,6 @@ namespace Metasound
 			EnqueueTransformOnOperatorQueues(CreateAddEdgeTransform);
 		}
 
-		const INode* FDynamicOperatorTransactor::FindNode(const FGuid& InNodeID) const
-		{
-			if (const TUniquePtr<INode>* NodePtr = NodeMap.Find(InNodeID))
-			{
-				return NodePtr->Get();
-			}
-			else
-			{
-				UE_LOG(LogMetaSound, Error, TEXT("Graph does not contain node with ID %s"), *InNodeID.ToString());
-				return nullptr;
-			}
-		}
-
 		TUniquePtr<IDynamicOperatorTransform> FDynamicOperatorTransactor::CreateAddOperatorTransform(const INode& InNode, const FOperatorSettings& InOperatorSettings, const FMetasoundEnvironment& InEnvironment) const
 		{
 			const FOperatorID OperatorID = DirectedGraphAlgo::GetOperatorID(InNode);
diff --git a/Engine/Plugins/Runtime/Metasound/Source/MetasoundGraphCore/Private/MetasoundGraph.cpp b/Engine/Plugins/Runtime/Metasound/Source/MetasoundGraphCore/Private/MetasoundGraph.cpp
--- a/Engine/Plugins/Runtime/Metasound/Source/MetasoundGraphCore/Private/MetasoundGraph.cpp
+++ b/Engine/Plugins/Runtime/Metasound/Source/MetasoundGraphCore/Private/MetasoundGraph.cpp
@@ -5,16 +5,6 @@
 
 namespace Metasound
 {
-	FGraph::FGraph(const IGraph& InOther)
-	: InstanceName(InOther.GetInstanceName())
-	, InstanceID(InOther.GetInstanceID())
-	, Metadata(InOther.GetMetadata())
-	, Edges(InOther.GetDataEdges())
-	, InputDestinations(InOther.GetInputDataDestinations())
-	, OutputSources(InOther.GetOutputDataSources())
-	{
-	}
-
 	FGraph::FGraph(const FString& InInstanceName, const FGuid& InInstanceID)
 	: InstanceName(InInstanceName)
 	, InstanceID(InInstanceID)
@@ -179,6 +169,38 @@ namespace Metasound
 		Edges.RemoveAllSwap(ContainsNode);
 	}
 
+	void FGraph::AddNode(const FGuid& InNodeID, TSharedPtr<const INode> InNode)
+	{
+		Nodes.Add(InNodeID, InNode);
+	}
+
+	const INode* FGraph::FindNode(const FGuid& InNodeID) const
+	{
+		if (const TSharedPtr<const INode>* NodePtr = Nodes.Find(InNodeID))
+		{
+			return NodePtr->Get();
+		}
+		return nullptr;
+	}
+
+	bool FGraph::RemoveNode(const FGuid& InNodeID, bool bInRemoveDataEdgesWithNode)
+	{
+		if (bInRemoveDataEdgesWithNode)
+		{
+			TSharedPtr<const INode> NodePtr;
+			if (Nodes.RemoveAndCopyValue(InNodeID, NodePtr))
+			{
+				RemoveDataEdgesWithNode(*NodePtr);
+				return true;
+			}
+			return false;
+		}
+		else
+		{
+			return Nodes.Remove(InNodeID) != 0;
+		}
+	}
+
 	const TArray<FDataEdge>& FGraph::GetDataEdges() const
 	{
 		return Edges;
diff --git a/Engine/Plugins/Runtime/Metasound/Source/MetasoundGraphCore/Public/MetasoundDynamicOperatorTransactor.h b/Engine/Plugins/Runtime/Metasound/Source/MetasoundGraphCore/Public/MetasoundDynamicOperatorTransactor.h
--- a/Engine/Plugins/Runtime/Metasound/Source/MetasoundGraphCore/Public/MetasoundDynamicOperatorTransactor.h
+++ b/Engine/Plugins/Runtime/Metasound/Source/MetasoundGraphCore/Public/MetasoundDynamicOperatorTransactor.h
@@ -32,7 +32,7 @@ namespace Metasound
 			FDynamicOperatorTransactor(const FDynamicOperatorTransactor&) = delete;
 
 			FDynamicOperatorTransactor();
-			FDynamicOperatorTransactor(const IGraph& InGraph);
+			FDynamicOperatorTransactor(const FGraph& InGraph);
 
 			/** Create a queue for communication with a dynamic operator. */
 			TSharedRef<TSpscQueue<TUniquePtr<IDynamicOperatorTransform>>> CreateTransformQueue(const FOperatorSettings& InOperatorSettings, const FMetasoundEnvironment& InEnvironment);
@@ -88,8 +88,6 @@ namespace Metasound
 
 			void AddDataEdgeInternal(const INode& InFromNode, const FVertexName& InFromVertex, const FGuid& InToNodeID, const INode& InToNode, const FVertexName& InToVertex);
 
-			const INode* FindNode(const FGuid& InNodeID) const;
-
 			using FCreateTransformFunctionRef = TFunctionRef<TUniquePtr<IDynamicOperatorTransform>(const FOperatorSettings& InOperatorSettings, const FMetasoundEnvironment& InEnvironment)>;
 
 			TUniquePtr<IDynamicOperatorTransform> CreateAddOperatorTransform(const INode& InNode, const FOperatorSettings& InOperatorSettings, const FMetasoundEnvironment& InEnvironment) const;
@@ -107,7 +105,6 @@ namespace Metasound
 
 			TArray<FDynamicOperatorInfo> OperatorInfos;
 
-			TSortedMap<FGuid, TUniquePtr<INode>> NodeMap;
 			struct FLiteralNodeID
 			{
 				FGuid ToNode;
diff --git a/Engine/Plugins/Runtime/Metasound/Source/MetasoundGraphCore/Public/MetasoundGraph.h b/Engine/Plugins/Runtime/Metasound/Source/MetasoundGraphCore/Public/MetasoundGraph.h
--- a/Engine/Plugins/Runtime/Metasound/Source/MetasoundGraphCore/Public/MetasoundGraph.h
+++ b/Engine/Plugins/Runtime/Metasound/Source/MetasoundGraphCore/Public/MetasoundGraph.h
@@ -18,7 +18,6 @@ namespace Metasound
 	{
 		public:
 
-			FGraph(const IGraph& InOther);
 			FGraph(const FString& InInstanceName, const FGuid& InInstanceID);
 
 			virtual ~FGraph() = default;
@@ -96,6 +95,30 @@ namespace Metasound
 			/** Removes all edges with connected to the node. */
 			void RemoveDataEdgesWithNode(const INode& InNode);
 
+			/** Store a node on this graph. 
+			 *
+			 * @param InNodeID - The NodeID related to the parent FMetasoundFrontendClass.
+			 * @param InNode - A shared pointer to a node. 
+			 */
+			void AddNode(const FGuid& InNodeID, TSharedPtr<const INode> InNode);
+
+			/** Retrieve node by node ID.
+			 *
+			 * @param InNodeID - The NodeID of the requested Node.
+			 *
+			 * @return Pointer to the Node if it is stored on this graph. nullptr otherwise. 
+			 */
+			const INode* FindNode(const FGuid& InNodeID) const;
+
+			/** Removes node from graph.
+			 *
+			 * @param InNodeID - ID of node to remove.
+			 * @param bInRemoveDataEdgesWithNode - If true, will also call RemoveDataEdgesWithNode(...)
+			 *
+			 * @return true if node exists and is removed, false otherwise. 
+			 */
+			bool RemoveNode(const FGuid& InNodeID, bool bInRemoveDataEdgesWithNode = true);
+
 			/** Add an input data destination to describe how data provided 
 			 * outside this graph should be routed internally.
 			 *
@@ -164,6 +187,7 @@ namespace Metasound
 			FNodeClassMetadata Metadata;
 
 			TArray<FDataEdge> Edges;
+			TSortedMap<FGuid, TSharedPtr<const INode>> Nodes;
 
 			FInputDataDestinationCollection InputDestinations;
 			FOutputDataSourceCollection OutputSources;
