Fix ConnectNode<In/Out>sToMatchingGraphInterface<In/Out>puts
- Add tests for connecting & disconnecting nodes that match graph interfaces
- Rename test accordingly
#rnx

diff --git a/Engine/Plugins/Runtime/Metasound/Source/MetasoundEngine/Private/MetasoundBuilderSubsystem.cpp b/Engine/Plugins/Runtime/Metasound/Source/MetasoundEngine/Private/MetasoundBuilderSubsystem.cpp
--- a/Engine/Plugins/Runtime/Metasound/Source/MetasoundEngine/Private/MetasoundBuilderSubsystem.cpp
+++ b/Engine/Plugins/Runtime/Metasound/Source/MetasoundEngine/Private/MetasoundBuilderSubsystem.cpp
@@ -254,7 +254,7 @@ TArray<FMetaSoundBuilderNodeOutputHandle> UMetaSoundBuilderBase::ConnectNodeInpu
 	TArray<FMetaSoundBuilderNodeOutputHandle> ConnectedVertices;
 	Algo::Transform(NewEdges, ConnectedVertices, [this](const FMetasoundFrontendEdge* NewEdge)
 	{
-		const FMetasoundFrontendVertex* Vertex = Builder.FindNodeInput(NewEdge->FromNodeID, NewEdge->FromVertexID);
+		const FMetasoundFrontendVertex* Vertex = Builder.FindNodeOutput(NewEdge->FromNodeID, NewEdge->FromVertexID);
 		checkf(Vertex, TEXT("Edge connection reported success but vertex not found."));
 		return FMetaSoundBuilderNodeOutputHandle(NewEdge->ToNodeID, Vertex->VertexID);
 	});
diff --git a/Engine/Plugins/Runtime/Metasound/Source/MetasoundEngine/Public/MetasoundBuilderSubsystem.h b/Engine/Plugins/Runtime/Metasound/Source/MetasoundEngine/Public/MetasoundBuilderSubsystem.h
--- a/Engine/Plugins/Runtime/Metasound/Source/MetasoundEngine/Public/MetasoundBuilderSubsystem.h
+++ b/Engine/Plugins/Runtime/Metasound/Source/MetasoundEngine/Public/MetasoundBuilderSubsystem.h
@@ -158,11 +158,11 @@ public:
 	UFUNCTION(BlueprintCallable, Category = "Audio|MetaSound|Builder", meta = (ExpandEnumAsExecs = "OutResult"))
 	void ConnectNodesByInterfaceBindings(const FMetaSoundNodeHandle& FromNodeHandle, const FMetaSoundNodeHandle& ToNodeHandle, EMetaSoundBuilderResult& OutResult);
 
-	// Connects a given node's outputs to all graph outputs for shared interfaces implemented on both the node's referenced class and the builder's MetaSound graph.
+	// Connects a given node's outputs to all graph outputs for shared interfaces implemented on both the node's referenced class and the builder's MetaSound graph. Returns inputs of connected output nodes.
 	UFUNCTION(BlueprintCallable, Category = "Audio|MetaSound|Builder", meta = (ExpandEnumAsExecs = "OutResult"))
 	UPARAM(DisplayName = "Connected Graph Output Node Inputs") TArray<FMetaSoundBuilderNodeInputHandle> ConnectNodeOutputsToMatchingGraphInterfaceOutputs(const FMetaSoundNodeHandle& NodeHandle, EMetaSoundBuilderResult& OutResult);
 
-	// Connects a given node's inputs to all graph inputs for shared interfaces implemented on both the node's referenced class and the builder's MetaSound graph. Returns output input node's outputs that were connected.
+	// Connects a given node's inputs to all graph inputs for shared interfaces implemented on both the node's referenced class and the builder's MetaSound graph. Returns outputs of connected input nodes.
 	UFUNCTION(BlueprintCallable, Category = "Audio|MetaSound|Builder", meta = (ExpandEnumAsExecs = "OutResult"))
 	UPARAM(DisplayName = "Connected Graph Input Node Outputs") TArray<FMetaSoundBuilderNodeOutputHandle> ConnectNodeInputsToMatchingGraphInterfaceInputs(const FMetaSoundNodeHandle& NodeHandle, EMetaSoundBuilderResult& OutResult);
 
diff --git a/Engine/Plugins/Runtime/Metasound/Source/MetasoundEngineTest/Private/EngineTestMetasoundSource.cpp b/Engine/Plugins/Runtime/Metasound/Source/MetasoundEngineTest/Private/EngineTestMetasoundSource.cpp
--- a/Engine/Plugins/Runtime/Metasound/Source/MetasoundEngineTest/Private/EngineTestMetasoundSource.cpp
+++ b/Engine/Plugins/Runtime/Metasound/Source/MetasoundEngineTest/Private/EngineTestMetasoundSource.cpp
@@ -72,7 +72,11 @@ namespace EngineTestMetaSoundSourcePrivate
 		return Node;
 	}
 
-	UMetaSoundSourceBuilder& CreateMetaSoundSourceBuilder(EMetaSoundOutputAudioFormat OutputFormat, bool bIsOneShot, EngineTestMetaSoundSourcePrivate::FInitTestBuilderSourceOutput& Output)
+	UMetaSoundSourceBuilder& CreateSourceBuilder(
+		FAutomationTestBase& Test,
+		EMetaSoundOutputAudioFormat OutputFormat,
+		bool bIsOneShot,
+		EngineTestMetaSoundSourcePrivate::FInitTestBuilderSourceOutput& Output)
 	{
 		using namespace Audio;
 		using namespace Metasound;
@@ -88,12 +92,13 @@ namespace EngineTestMetaSoundSourcePrivate
 			Result,
 			OutputFormat,
 			bIsOneShot);
-		checkf(Result == EMetaSoundBuilderResult::Succeeded, TEXT("Failed to create MetaSoundSourceBuilder"));
+		checkf(Builder, TEXT("Failed to create MetaSoundSourceBuilder"));
+		Test.AddErrorIfFalse(Result == EMetaSoundBuilderResult::Succeeded, TEXT("Builder created but CreateSourceBuilder did not result in 'Succeeded' state"));
 
 		return *Builder;
 	}
 
-	UMetaSoundSourceBuilder& CreateMetaSoundMonoSourceSinGenBuilder(
+	UMetaSoundSourceBuilder& CreateMonoSourceSinGenBuilder(
 		FAutomationTestBase& Test,
 		FMetaSoundBuilderNodeInputHandle* GenInputNodeFreq = nullptr,
 		FMetaSoundBuilderNodeInputHandle* MonoOutNodeInput = nullptr,
@@ -108,7 +113,7 @@ namespace EngineTestMetaSoundSourcePrivate
 		constexpr EMetaSoundOutputAudioFormat OutputFormat = EMetaSoundOutputAudioFormat::Mono;
 		constexpr bool bIsOneShot = false;
 		FInitTestBuilderSourceOutput Output;
-		UMetaSoundSourceBuilder& Builder = CreateMetaSoundSourceBuilder(EMetaSoundOutputAudioFormat::Mono, bIsOneShot, Output);
+		UMetaSoundSourceBuilder& Builder = CreateSourceBuilder(Test, EMetaSoundOutputAudioFormat::Mono, bIsOneShot, Output);
 
 		EMetaSoundBuilderResult Result = EMetaSoundBuilderResult::Failed;
 		if (MonoOutNodeInput)
@@ -315,7 +320,7 @@ bool FAudioMetasoundSourceBuilderTest::RunTest(const FString& Parameters)
 	using namespace EngineTestMetaSoundSourcePrivate;
 
 	FMetaSoundBuilderNodeInputHandle MonoOutNodeInput;
-	UMetaSoundSourceBuilder& Builder = CreateMetaSoundMonoSourceSinGenBuilder(*this, nullptr, &MonoOutNodeInput);
+	UMetaSoundSourceBuilder& Builder = CreateMonoSourceSinGenBuilder(*this, nullptr, &MonoOutNodeInput);
 	Builder.AddToRoot();
 
 	if (FAudioDevice* AudioDevice = GEngine->GetMainAudioDeviceRaw())
@@ -359,7 +364,7 @@ bool FAudioMetasoundSourceBuilderLiveUpdateNode::RunTest(const FString& Paramete
 	using namespace Metasound::Frontend;
 
 	FMetaSoundBuilderNodeInputHandle MonoOutNodeInput;
-	UMetaSoundSourceBuilder& Builder = CreateMetaSoundMonoSourceSinGenBuilder(*this, nullptr, &MonoOutNodeInput);
+	UMetaSoundSourceBuilder& Builder = CreateMonoSourceSinGenBuilder(*this, nullptr, &MonoOutNodeInput);
 	Builder.AddToRoot();
 
 	if (FAudioDevice* AudioDevice = GEngine->GetMainAudioDeviceRaw())
@@ -412,7 +417,7 @@ bool FAudioMetasoundSourceBuilderLiveUpdateLiteral::RunTest(const FString& Param
 
 	FMetaSoundBuilderNodeInputHandle MonoOutNodeInput;
 	FMetaSoundBuilderNodeInputHandle GenNodeFreqInput;
-	UMetaSoundSourceBuilder& Builder = CreateMetaSoundMonoSourceSinGenBuilder(*this, &GenNodeFreqInput , &MonoOutNodeInput, 220.f);
+	UMetaSoundSourceBuilder& Builder = CreateMonoSourceSinGenBuilder(*this, &GenNodeFreqInput , &MonoOutNodeInput, 220.f);
 	Builder.AddToRoot();
 
 	if (FAudioDevice* AudioDevice = GEngine->GetMainAudioDeviceRaw())
diff --git a/Engine/Plugins/Runtime/Metasound/Source/MetasoundFrontend/Private/MetasoundFrontendDocumentBuilder.cpp b/Engine/Plugins/Runtime/Metasound/Source/MetasoundFrontend/Private/MetasoundFrontendDocumentBuilder.cpp
--- a/Engine/Plugins/Runtime/Metasound/Source/MetasoundFrontend/Private/MetasoundFrontendDocumentBuilder.cpp
+++ b/Engine/Plugins/Runtime/Metasound/Source/MetasoundFrontend/Private/MetasoundFrontendDocumentBuilder.cpp
@@ -585,19 +585,22 @@ bool FMetaSoundFrontendDocumentBuilder::AddNamedEdges(const TSet<Metasound::Fron
 		}
 	}
 
+	const TArray<FMetasoundFrontendEdge>& Edges = GetDocument().RootGraph.Graph.Edges;
+	const int32 LastIndex = Edges.Num() - 1;
 	for (FMetasoundFrontendEdge& EdgeToAdd : EdgesToAdd)
 	{
 		const FMetasoundFrontendEdge* NewEdge = AddEdge(MoveTemp(EdgeToAdd));
-		if (ensureAlwaysMsgf(NewEdge, TEXT("Failed to add MetaSound graph edge via DocumentBuilder when prior step validated edge add was valid")))
+		if (!ensureAlwaysMsgf(NewEdge, TEXT("Failed to add MetaSound graph edge via DocumentBuilder when prior step validated edge add was valid")))
 		{
-			if (OutNewEdges)
-			{
-				OutNewEdges->Add(NewEdge);
-			}
+			bSuccess = false;
 		}
-		else
+	}
+
+	if (OutNewEdges)
+	{
+		for (int32 Index = LastIndex + 1; Index < Edges.Num(); ++Index)
 		{
-			bSuccess = false;
+			OutNewEdges->Add(&Edges[Index]);
 		}
 	}
 
@@ -694,7 +697,7 @@ bool FMetaSoundFrontendDocumentBuilder::AddEdgesFromMatchingInterfaceNodeInputsT
 			Algo::Transform(RegistryEntry->GetInterface().Inputs, EdgesToMake, [this, &NodeCache, &InterfaceCache, InNodeID](const FMetasoundFrontendClassInput& Input)
 			{
 				const FMetasoundFrontendGraph& Graph = GetDocument().RootGraph.Graph;
-				const FMetasoundFrontendVertex* NodeVertex = NodeCache.FindOutputVertex(InNodeID, Input.Name);
+				const FMetasoundFrontendVertex* NodeVertex = NodeCache.FindInputVertex(InNodeID, Input.Name);
 				check(NodeVertex);
 				const FMetasoundFrontendClassInput* InputClass = InterfaceCache.FindInput(Input.Name);
 				check(InputClass);
@@ -702,7 +705,7 @@ bool FMetaSoundFrontendDocumentBuilder::AddEdgesFromMatchingInterfaceNodeInputsT
 				check(InputNode);
 				const TArray<FMetasoundFrontendVertex>& Outputs = InputNode->Interface.Outputs;
 				check(!Outputs.IsEmpty());
-				return FNamedEdge { InNodeID, NodeVertex->Name, InputNode->GetID(), Outputs.Last().Name };
+				return FNamedEdge { InputNode->GetID(), Outputs.Last().Name, InNodeID, NodeVertex->Name };
 			});
 		}
 	}
@@ -1781,7 +1784,7 @@ bool FMetaSoundFrontendDocumentBuilder::RemoveNamedEdges(const TSet<Metasound::F
 		}
 	}
 
-	for (FMetasoundFrontendEdge& EdgeToRemove : EdgesToRemove)
+	for (const FMetasoundFrontendEdge& EdgeToRemove : EdgesToRemove)
 	{
 		const bool bRemovedEdge = RemoveEdgeToNodeInput(EdgeToRemove.ToNodeID, EdgeToRemove.ToVertexID);
 		if (ensureAlwaysMsgf(bRemovedEdge, TEXT("Failed to remove MetaSound graph edge via DocumentBuilder when prior step validated edge remove was valid")))
