Missing part of earlier MetaSound builder cache corruption fix
#rnx

diff --git a/Engine/Plugins/Runtime/Metasound/Source/MetasoundFrontend/Private/MetasoundFrontendDocumentCache.cpp b/Engine/Plugins/Runtime/Metasound/Source/MetasoundFrontend/Private/MetasoundFrontendDocumentCache.cpp
--- a/Engine/Plugins/Runtime/Metasound/Source/MetasoundFrontend/Private/MetasoundFrontendDocumentCache.cpp
+++ b/Engine/Plugins/Runtime/Metasound/Source/MetasoundFrontend/Private/MetasoundFrontendDocumentCache.cpp
@@ -205,7 +205,13 @@ namespace Metasound::Frontend
 		const FMetasoundFrontendGraphClass& GraphClass = Document.RootGraph;
 		const FMetasoundFrontendClassInput& Input = GraphClass.Interface.Inputs[IndexBeingRemoved];
 		InputNameToIndex.Remove(Input.Name);
-		Algo::ForEach(InputNameToIndex, [](TPair<FName, int32>& Entry) { --Entry.Value; });
+		Algo::ForEach(InputNameToIndex, [&IndexBeingRemoved](TPair<FName, int32>& Entry)
+		{
+			if (Entry.Value > IndexBeingRemoved)
+			{
+				--Entry.Value;
+			}
+		});
 	}
 
 	void FDocumentGraphInterfaceCache::OnRemovingOutput(int32 IndexBeingRemoved)
@@ -214,7 +220,13 @@ namespace Metasound::Frontend
 		const FMetasoundFrontendGraphClass& GraphClass = Document.RootGraph;
 		const FMetasoundFrontendClassOutput& Output = GraphClass.Interface.Outputs[IndexBeingRemoved];
 		OutputNameToIndex.Remove(Output.Name);
-		Algo::ForEach(OutputNameToIndex, [](TPair<FName, int32>& Entry) { --Entry.Value; });
+		Algo::ForEach(OutputNameToIndex, [&IndexBeingRemoved](TPair<FName, int32>& Entry)
+		{
+			if (Entry.Value > IndexBeingRemoved)
+			{
+				--Entry.Value;
+			}
+		});
 	}
 
 	FDocumentGraphNodeCache::FDocumentGraphNodeCache(TSharedRef<IDocumentCache> ParentCache)
