Fix for reducing channel count to output interface other than mono (and other ways in which interfaces were potentially swapping related API members)
#rnx

diff --git a/Engine/Plugins/Runtime/Metasound/Source/MetasoundFrontend/Private/MetasoundFrontendDocumentBuilder.cpp b/Engine/Plugins/Runtime/Metasound/Source/MetasoundFrontend/Private/MetasoundFrontendDocumentBuilder.cpp
--- a/Engine/Plugins/Runtime/Metasound/Source/MetasoundFrontend/Private/MetasoundFrontendDocumentBuilder.cpp
+++ b/Engine/Plugins/Runtime/Metasound/Source/MetasoundFrontend/Private/MetasoundFrontendDocumentBuilder.cpp
@@ -266,12 +266,14 @@ namespace Metasound::Frontend
 				bool bDidEdit = false;
 				for (const FVertexPair& PairedInput : PairedInputs)
 				{
-					bDidEdit |= OutBuilder.SwapGraphInput(PairedInput.Get<0>(), PairedInput.Get<1>());
+					const bool bSwapped = OutBuilder.SwapGraphInput(PairedInput.Get<0>(), PairedInput.Get<1>());
+					bDidEdit |= bSwapped;
 				}
 
 				for (const FVertexPair& PairedOutput : PairedOutputs)
 				{
-					bDidEdit |= OutBuilder.SwapGraphOutput(PairedOutput.Get<0>(), PairedOutput.Get<1>());
+					const bool bSwapped = OutBuilder.SwapGraphOutput(PairedOutput.Get<0>(), PairedOutput.Get<1>());
+					bDidEdit |= bSwapped;
 				}
 
 				return bDidEdit;
