Add missing soundwave/soundbase properties when copying over properties from template MetaSounds

#rnx

diff --git a/Engine/Plugins/Runtime/Metasound/Source/MetasoundEditor/Private/MetasoundEditorSubsystem.cpp b/Engine/Plugins/Runtime/Metasound/Source/MetasoundEditor/Private/MetasoundEditorSubsystem.cpp
--- a/Engine/Plugins/Runtime/Metasound/Source/MetasoundEditor/Private/MetasoundEditorSubsystem.cpp
+++ b/Engine/Plugins/Runtime/Metasound/Source/MetasoundEditor/Private/MetasoundEditorSubsystem.cpp
@@ -67,7 +67,7 @@ TScriptInterface<IMetaSoundDocumentInterface> UMetaSoundEditorSubsystem::BuildTo
 			// Template sound wave settings only apply to sources 
 			if (TemplateSoundWave != nullptr && bIsSource)
 			{
-				SetSoundWaveSettingsFromTemplate(NewMetaSound, TemplateSoundWave);
+				SetSoundWaveSettingsFromTemplate(*CastChecked<USoundWave>(NewMetaSound), *TemplateSoundWave);
 			}
 
 			InitEdGraph(*NewMetaSound);
@@ -85,25 +85,39 @@ TScriptInterface<IMetaSoundDocumentInterface> UMetaSoundEditorSubsystem::BuildTo
 	return nullptr;
 }
 
-void UMetaSoundEditorSubsystem::SetSoundWaveSettingsFromTemplate(UObject* NewMetasound, const USoundWave* TemplateSoundWave)
+void UMetaSoundEditorSubsystem::SetSoundWaveSettingsFromTemplate(USoundWave& NewMetaSoundWave, const USoundWave& TemplateSoundWave) const
 {
-	USoundWave* NewMetaSoundWave = CastChecked<USoundWave>(NewMetasound);
-
-	NewMetaSoundWave->SoundClassObject = TemplateSoundWave->SoundClassObject;
-	NewMetaSoundWave->AttenuationSettings = TemplateSoundWave->AttenuationSettings;
-	NewMetaSoundWave->ModulationSettings = TemplateSoundWave->ModulationSettings;
+	// Sound 
+	NewMetaSoundWave.Volume = TemplateSoundWave.Volume;
+	NewMetaSoundWave.Pitch = TemplateSoundWave.Pitch;
+	NewMetaSoundWave.SoundClassObject = TemplateSoundWave.SoundClassObject;
+
+	// Attenuation 
+	NewMetaSoundWave.AttenuationSettings = TemplateSoundWave.AttenuationSettings;
+	NewMetaSoundWave.bDebug = TemplateSoundWave.bDebug;
+
+	// Effects 
+	NewMetaSoundWave.bEnableBusSends = TemplateSoundWave.bEnableBusSends;
+	NewMetaSoundWave.SourceEffectChain = TemplateSoundWave.SourceEffectChain;
+	NewMetaSoundWave.BusSends = TemplateSoundWave.BusSends;
+	NewMetaSoundWave.PreEffectBusSends = TemplateSoundWave.PreEffectBusSends; 
+
+	NewMetaSoundWave.bEnableBaseSubmix = TemplateSoundWave.bEnableBaseSubmix;
+	NewMetaSoundWave.SoundSubmixObject = TemplateSoundWave.SoundSubmixObject;
+	NewMetaSoundWave.bEnableSubmixSends = TemplateSoundWave.bEnableSubmixSends;
+	NewMetaSoundWave.SoundSubmixSends = TemplateSoundWave.SoundSubmixSends;
+
+	// Modulation 
+	NewMetaSoundWave.ModulationSettings = TemplateSoundWave.ModulationSettings;
 	
-	// Concurrency
-	NewMetaSoundWave->bOverrideConcurrency = TemplateSoundWave->bOverrideConcurrency;
-	NewMetaSoundWave->ConcurrencySet = TemplateSoundWave->ConcurrencySet;
-	NewMetaSoundWave->ConcurrencyOverrides = TemplateSoundWave->ConcurrencyOverrides;
-
-	// Submix and bus sends 
-	NewMetaSoundWave->bEnableSubmixSends = TemplateSoundWave->bEnableSubmixSends;
-	NewMetaSoundWave->SoundSubmixSends = TemplateSoundWave->SoundSubmixSends;
-	NewMetaSoundWave->bEnableBusSends = TemplateSoundWave->bEnableBusSends;
-	NewMetaSoundWave->BusSends = TemplateSoundWave->BusSends;
-	NewMetaSoundWave->PreEffectBusSends = TemplateSoundWave->PreEffectBusSends;
+	// Voice Management 
+	NewMetaSoundWave.VirtualizationMode = TemplateSoundWave.VirtualizationMode;
+	NewMetaSoundWave.bOverrideConcurrency = TemplateSoundWave.bOverrideConcurrency;
+	NewMetaSoundWave.ConcurrencySet = TemplateSoundWave.ConcurrencySet;
+	NewMetaSoundWave.ConcurrencyOverrides = TemplateSoundWave.ConcurrencyOverrides;
+
+	NewMetaSoundWave.bBypassVolumeScaleForPriority = TemplateSoundWave.bBypassVolumeScaleForPriority;
+	NewMetaSoundWave.Priority = TemplateSoundWave.Priority;
 }
 
 void UMetaSoundEditorSubsystem::InitAsset(UObject& InNewMetaSound, UObject* InReferencedMetaSound)
@@ -142,7 +156,7 @@ void UMetaSoundEditorSubsystem::InitAsset(UObject& InNewMetaSound, UObject* InRe
 		// Copy sound wave settings to preset for sources
 		if (&ReferencedInterface->GetBaseMetaSoundUClass() == UMetaSoundSource::StaticClass())
 		{
-			SetSoundWaveSettingsFromTemplate(&InNewMetaSound, CastChecked<USoundWave>(InReferencedMetaSound));
+			SetSoundWaveSettingsFromTemplate(*CastChecked<USoundWave>(&InNewMetaSound), *CastChecked<USoundWave>(InReferencedMetaSound));
 		}
 	}
 
diff --git a/Engine/Plugins/Runtime/Metasound/Source/MetasoundEditor/Public/MetasoundEditorSubsystem.h b/Engine/Plugins/Runtime/Metasound/Source/MetasoundEditor/Public/MetasoundEditorSubsystem.h
--- a/Engine/Plugins/Runtime/Metasound/Source/MetasoundEditor/Public/MetasoundEditorSubsystem.h
+++ b/Engine/Plugins/Runtime/Metasound/Source/MetasoundEditor/Public/MetasoundEditorSubsystem.h
@@ -40,9 +40,10 @@ public:
 	// Get the default author for a MetaSound asset
 	const FString GetDefaultAuthor();
 
-	// Copy over sound wave settings such as attenuation, modulation, and sound class from the template sound wave to the MetaSound 
-	void SetSoundWaveSettingsFromTemplate(UObject* NewMetasound, const USoundWave* TemplateSoundWave);
-
 	static UMetaSoundEditorSubsystem& GetChecked();
 	static const UMetaSoundEditorSubsystem& GetConstChecked(); 
+
+private:
+	// Copy over sound wave settings such as attenuation, modulation, and sound class from the template sound wave to the MetaSound 
+	void SetSoundWaveSettingsFromTemplate(USoundWave& NewMetasound, const USoundWave& TemplateSoundWave) const;
 };
