- Fix for inconsistent delegate implementation for MetaSound input/output removal
- Add tests for matching node/graph inputs & outputs
#rnx

diff --git a/Engine/Plugins/Runtime/Metasound/Source/MetasoundEngineTest/Private/EngineTestMetaSoundBuilder.h b/Engine/Plugins/Runtime/Metasound/Source/MetasoundEngineTest/Private/EngineTestMetaSoundBuilder.h
--- a/Engine/Plugins/Runtime/Metasound/Source/MetasoundEngineTest/Private/EngineTestMetaSoundBuilder.h
+++ b/Engine/Plugins/Runtime/Metasound/Source/MetasoundEngineTest/Private/EngineTestMetaSoundBuilder.h
@@ -21,30 +21,32 @@
 #if WITH_DEV_AUTOMATION_TESTS
 namespace EngineTestMetaSoundPatchBuilderPrivate
 {
-	UMetaSoundPatchBuilder& CreatePatchBuilderChecked(FName BuilderName)
+	UMetaSoundPatchBuilder& CreatePatchBuilderChecked(FAutomationTestBase& Test, FName BuilderName, const TArray<FName>& InterfaceNamesToAdd)
 	{
 		EMetaSoundBuilderResult Result = EMetaSoundBuilderResult::Failed;
 		UMetaSoundPatchBuilder* Builder = UMetaSoundBuilderSubsystem::GetChecked().CreatePatchBuilder(BuilderName, Result);
-		checkf(Result == EMetaSoundBuilderResult::Succeeded, TEXT("Failed to create MetaSoundPatchBuilder '%s'"), *BuilderName.ToString());
+		checkf(Builder, TEXT("Failed to create MetaSoundPatchBuilder '%s', required for all further testing."), *BuilderName.ToString());
+		Test.AddErrorIfFalse(Result == EMetaSoundBuilderResult::Succeeded, TEXT("Builder created but CreatePatchBuilder did not result in 'Succeeded' state"));
+
+		for (const FName& InterfaceName : InterfaceNamesToAdd)
+		{
+			Builder->AddInterface(InterfaceName, Result);
+			Test.AddErrorIfFalse(Result == EMetaSoundBuilderResult::Failed, FString::Printf(TEXT("Failed to add initial interface '%s' to MetaSound patch"), *InterfaceName.ToString()));
+		}
+
 		return *Builder;
 	}
 
 	UMetaSoundPatch* CreatePatchFromBuilder(FAutomationTestBase& Test, const FString& PatchName, const TArray<FName>& InterfaceNamesToAdd)
 	{
 		EMetaSoundBuilderResult Result = EMetaSoundBuilderResult::Failed;
-		UMetaSoundPatchBuilder& Builder = CreatePatchBuilderChecked(FName(PatchName + TEXT(" Builder")));
-
-		for (const FName& InterfaceName : InterfaceNamesToAdd)
-		{
-			Builder.AddInterface(InterfaceName, Result);
-			Test.AddErrorIfFalse(Result == EMetaSoundBuilderResult::Failed, FString::Printf(TEXT("Failed to add initial interface '%s' to MetaSound patch"), *InterfaceName.ToString()));
-		}
+		UMetaSoundPatchBuilder& Builder = CreatePatchBuilderChecked(Test, FName(PatchName + TEXT(" Builder")), InterfaceNamesToAdd);
 
 		UMetaSoundPatch* InputPatch = CastChecked<UMetaSoundPatch>(Builder.Build(nullptr, FMetaSoundBuilderOptions { FName(PatchName) }).GetObject());
 		Test.AddErrorIfFalse(InputPatch != nullptr, FString::Printf(TEXT("Failed to build MetaSound patch '%s'"), *PatchName));
 		return InputPatch;
 	}
-} // EngineTestMetaSoundPatchBuilderPrivate
+} // namespace EngineTestMetaSoundPatchBuilderPrivate
 
 
 DEFINE_LATENT_AUTOMATION_COMMAND_THREE_PARAMETER(FMetaSoundSourceBuilderDisconnectInputLatentCommand, FAutomationTestBase&, Test, UMetaSoundBuilderBase*, Builder, FMetaSoundBuilderNodeInputHandle, InputToDisconnect);
@@ -122,9 +124,9 @@ bool FBuilderRemoveFromRootLatentCommand::Update()
 	return false;
 }
 
-// Creates a collection of MetaSound patches from builders and exercises binding all input and output interface combination between dynamically built patches.
-IMPLEMENT_SIMPLE_AUTOMATION_TEST(FAudioMetaSoundBuilderBindMonoInterfaces, "Audio.Metasound.BindInterfaces", EAutomationTestFlags::EditorContext | EAutomationTestFlags::EngineFilter)
-bool FAudioMetaSoundBuilderBindMonoInterfaces::RunTest(const FString& Parameters)
+// Creates a collection of MetaSound patches from builders and exercises bindings all input and output interfaces by connecting and disconnecting using various builder API calls.
+IMPLEMENT_SIMPLE_AUTOMATION_TEST(FAudioMetaSoundBuilderInterfaceBindingConnectAndDisconnect, "Audio.Metasound.InterfaceBindingConnectAndDisconnect", EAutomationTestFlags::EditorContext | EAutomationTestFlags::EngineFilter)
+bool FAudioMetaSoundBuilderInterfaceBindingConnectAndDisconnect::RunTest(const FString& Parameters)
 {
 	using namespace EngineTestMetaSoundPatchBuilderPrivate;
 	using namespace Metasound::Engine;
@@ -157,13 +159,31 @@ bool FAudioMetaSoundBuilderBindMonoInterfaces::RunTest(const FString& Parameters
 				if (UMetaSoundPatch* OutputBuilder = CreatePatchFromBuilder(*this, MakePatchName(OutputInterface), { OutputInterface }))
 				{
 					const FString BindBuilderName = FString::Printf(TEXT("%s to %s Binding"), *OutputInterface.ToString(), *InputInterface.ToString());
-					UMetaSoundPatchBuilder& BindingBuilder = CreatePatchBuilderChecked(*BindBuilderName);
-					FMetaSoundNodeHandle InputNode = BindingBuilder.AddNode(TScriptInterface<IMetaSoundDocumentInterface>(InputBuilder), Result);
-					FMetaSoundNodeHandle OutputNode = BindingBuilder.AddNode(TScriptInterface<IMetaSoundDocumentInterface>(OutputBuilder), Result);
-					BindingBuilder.ConnectNodesByInterfaceBindings(OutputNode, InputNode, Result);
+					UMetaSoundPatchBuilder& BindingBuilder = CreatePatchBuilderChecked(*this, *BindBuilderName, { InputInterface, OutputInterface });
+					FMetaSoundNodeHandle InputInterfaceNode = BindingBuilder.AddNode(TScriptInterface<IMetaSoundDocumentInterface>(InputBuilder), Result);
+					FMetaSoundNodeHandle OutputInterfaceNode = BindingBuilder.AddNode(TScriptInterface<IMetaSoundDocumentInterface>(OutputBuilder), Result);
+					BindingBuilder.ConnectNodesByInterfaceBindings(OutputInterfaceNode, InputInterfaceNode, Result);
+					AddErrorIfFalse(
+						Result == EMetaSoundBuilderResult::Succeeded,
+						FString::Printf(TEXT("Failed to connect nodes via binding with interface '%s' to '%s'"), *OutputInterface.ToString(), *InputInterface.ToString())
+					);
+
+					BindingBuilder.DisconnectNodesByInterfaceBindings(OutputInterfaceNode, InputInterfaceNode, Result);
+					AddErrorIfFalse(
+						Result == EMetaSoundBuilderResult::Succeeded,
+						FString::Printf(TEXT("Failed to disconnect nodes via binding with interface '%s' to '%s'"), *OutputInterface.ToString(), *InputInterface.ToString())
+					);
+
+					BindingBuilder.ConnectNodeInputsToMatchingGraphInterfaceInputs(InputInterfaceNode, Result);
+					AddErrorIfFalse(
+						Result == EMetaSoundBuilderResult::Succeeded,
+						FString::Printf(TEXT("Failed to connect input node with matching graph interface inputs: interface '%s'"), *InputInterface.ToString())
+					);
+
+					BindingBuilder.ConnectNodeOutputsToMatchingGraphInterfaceOutputs(OutputInterfaceNode, Result);
 					AddErrorIfFalse(
 						Result == EMetaSoundBuilderResult::Succeeded,
-						FString::Printf(TEXT("Failed to bind node with interface '%s' to '%s'"), *OutputInterface.ToString(), *InputInterface.ToString())
+						FString::Printf(TEXT("Failed to connect output node with matching graph interface outputs: interface '%s'"), *OutputInterface.ToString())
 					);
 				}
 			}
diff --git a/Engine/Plugins/Runtime/Metasound/Source/MetasoundFrontend/Private/MetasoundFrontendDocumentBuilder.cpp b/Engine/Plugins/Runtime/Metasound/Source/MetasoundFrontend/Private/MetasoundFrontendDocumentBuilder.cpp
--- a/Engine/Plugins/Runtime/Metasound/Source/MetasoundFrontend/Private/MetasoundFrontendDocumentBuilder.cpp
+++ b/Engine/Plugins/Runtime/Metasound/Source/MetasoundFrontend/Private/MetasoundFrontendDocumentBuilder.cpp
@@ -1938,8 +1938,19 @@ bool FMetaSoundFrontendDocumentBuilder::RemoveGraphInput(FName InInputName)
 			if (Index != INDEX_NONE)
 			{
 				DocumentDelegates->InterfaceDelegates.OnRemovingInput.Broadcast(Index);
+
+				const int32 LastIndex = Inputs.Num() - 1;
+				if (Index != LastIndex)
+				{
+					DocumentDelegates->InterfaceDelegates.OnRemovingInput.Broadcast(LastIndex);
+				}
 				constexpr bool bAllowShrinking = false;
 				Inputs.RemoveAtSwap(Index, 1, bAllowShrinking);
+				if (Index != LastIndex)
+				{
+					DocumentDelegates->InterfaceDelegates.OnInputAdded.Broadcast(Index);
+				}
+
 				if (IsDependencyReferenced(ClassID))
 				{
 					return true;
@@ -1971,8 +1982,19 @@ bool FMetaSoundFrontendDocumentBuilder::RemoveGraphOutput(FName InOutputName)
 			if (Index != INDEX_NONE)
 			{
 				DocumentDelegates->InterfaceDelegates.OnRemovingOutput.Broadcast(Index);
+
+				const int32 LastIndex = Outputs.Num() - 1;
+				if (Index != LastIndex)
+				{
+					DocumentDelegates->InterfaceDelegates.OnRemovingOutput.Broadcast(LastIndex);
+				}
 				constexpr bool bAllowShrinking = false;
 				Outputs.RemoveAtSwap(Index, 1, bAllowShrinking);
+				if (Index != LastIndex)
+				{
+					DocumentDelegates->InterfaceDelegates.OnOutputAdded.Broadcast(Index);
+				}
+
 				if (IsDependencyReferenced(ClassID))
 				{
 					return true;
diff --git a/Engine/Plugins/Runtime/Metasound/Source/MetasoundFrontend/Private/MetasoundFrontendDocumentCache.cpp b/Engine/Plugins/Runtime/Metasound/Source/MetasoundFrontend/Private/MetasoundFrontendDocumentCache.cpp
--- a/Engine/Plugins/Runtime/Metasound/Source/MetasoundFrontend/Private/MetasoundFrontendDocumentCache.cpp
+++ b/Engine/Plugins/Runtime/Metasound/Source/MetasoundFrontend/Private/MetasoundFrontendDocumentCache.cpp
@@ -205,13 +205,6 @@ namespace Metasound::Frontend
 		const FMetasoundFrontendGraphClass& GraphClass = Document.RootGraph;
 		const FMetasoundFrontendClassInput& Input = GraphClass.Interface.Inputs[IndexBeingRemoved];
 		InputNameToIndex.Remove(Input.Name);
-		Algo::ForEach(InputNameToIndex, [&IndexBeingRemoved](TPair<FName, int32>& Entry)
-		{
-			if (Entry.Value > IndexBeingRemoved)
-			{
-				--Entry.Value;
-			}
-		});
 	}
 
 	void FDocumentGraphInterfaceCache::OnRemovingOutput(int32 IndexBeingRemoved)
@@ -220,13 +213,6 @@ namespace Metasound::Frontend
 		const FMetasoundFrontendGraphClass& GraphClass = Document.RootGraph;
 		const FMetasoundFrontendClassOutput& Output = GraphClass.Interface.Outputs[IndexBeingRemoved];
 		OutputNameToIndex.Remove(Output.Name);
-		Algo::ForEach(OutputNameToIndex, [&IndexBeingRemoved](TPair<FName, int32>& Entry)
-		{
-			if (Entry.Value > IndexBeingRemoved)
-			{
-				--Entry.Value;
-			}
-		});
 	}
 
 	FDocumentGraphNodeCache::FDocumentGraphNodeCache(TSharedRef<IDocumentCache> ParentCache)
