[Metasounds] Bind Inputs / Outputs follow up (Automated Test Fixes)

diff --git a/Engine/Plugins/Runtime/Metasound/Source/MetasoundEngine/Private/MetasoundWaveTableEvaluateNode.cpp b/Engine/Plugins/Runtime/Metasound/Source/MetasoundEngine/Private/MetasoundWaveTableEvaluateNode.cpp
--- a/Engine/Plugins/Runtime/Metasound/Source/MetasoundEngine/Private/MetasoundWaveTableEvaluateNode.cpp
+++ b/Engine/Plugins/Runtime/Metasound/Source/MetasoundEngine/Private/MetasoundWaveTableEvaluateNode.cpp
@@ -111,12 +111,22 @@ namespace Metasound
 			using namespace WaveTableEvaluateNode;
 			
 			FInputVertexInterfaceData& Inputs = InVertexData.GetInputs();
-			Inputs.BindReadVertex(METASOUND_GET_PARAM_NAME(WaveTableParam), WaveTableReadRef);
-			Inputs.BindReadVertex(METASOUND_GET_PARAM_NAME(InputParam), InputReadRef);
-			Inputs.BindReadVertex("Interpolation", InterpModeReadRef);
 
 			FOutputVertexInterfaceData& Outputs = InVertexData.GetOutputs();
-			Outputs.BindReadVertex(METASOUND_GET_PARAM_NAME(OutParam), OutWriteRef);
+		}
+
+		virtual void BindInputs(FInputVertexInterfaceData& InOutVertexData) override
+		{
+			using namespace WaveTableEvaluateNode;
+			InOutVertexData.BindReadVertex(METASOUND_GET_PARAM_NAME(WaveTableParam), WaveTableReadRef);
+			InOutVertexData.BindReadVertex(METASOUND_GET_PARAM_NAME(InputParam), InputReadRef);
+			InOutVertexData.BindReadVertex("Interpolation", InterpModeReadRef);
+		}
+
+		virtual void BindOutputs(FOutputVertexInterfaceData& InOutVertexData) override
+		{
+			using namespace WaveTableEvaluateNode;
+			InOutVertexData.BindReadVertex(METASOUND_GET_PARAM_NAME(OutParam), OutWriteRef);
 		}
 		
 		virtual FDataReferenceCollection GetInputs() const override
diff --git a/Engine/Plugins/Runtime/Metasound/Source/MetasoundStandardNodes/Private/MetasoundAudioFormats.cpp b/Engine/Plugins/Runtime/Metasound/Source/MetasoundStandardNodes/Private/MetasoundAudioFormats.cpp
--- a/Engine/Plugins/Runtime/Metasound/Source/MetasoundStandardNodes/Private/MetasoundAudioFormats.cpp
+++ b/Engine/Plugins/Runtime/Metasound/Source/MetasoundStandardNodes/Private/MetasoundAudioFormats.cpp
@@ -260,25 +260,32 @@ namespace Metasound
 
 			virtual ~FInputOperator() {}
 
-
-			virtual FDataReferenceCollection GetInputs() const override
+			virtual void BindInputs(FInputVertexInterfaceData& InOutVertexData) override
 			{
-				FDataReferenceCollection Inputs;
-
-				Inputs.AddDataReadReference(InputName, Stereo);
-
-				return Inputs;
+				InOutVertexData.BindReadVertex(InputName, Stereo);
 			}
 
-			virtual FDataReferenceCollection GetOutputs() const override
+			virtual void BindOutputs(FOutputVertexInterfaceData& InOutVertexData) override
 			{
-				FDataReferenceCollection Outputs;
 				using namespace StereoAudioFormatVertexKeys;
+				InOutVertexData.BindReadVertex(METASOUND_GET_PARAM_NAME(LeftChannelVertex), Left);
+				InOutVertexData.BindReadVertex(METASOUND_GET_PARAM_NAME(RightChannelVertex), Right);
+			}
 
-				Outputs.AddDataReadReference(METASOUND_GET_PARAM_NAME(LeftChannelVertex), Left);
-				Outputs.AddDataReadReference(METASOUND_GET_PARAM_NAME(RightChannelVertex), Right);
+			virtual FDataReferenceCollection GetInputs() const override
+			{
+				// This should never be called. Bind(...) is called instead. This method
+				// exists as a stop-gap until the API can be deprecated and removed.
+				checkNoEntry();
+				return {};
+			}
 
-				return Outputs;
+			virtual FDataReferenceCollection GetOutputs() const override
+			{
+				// This should never be called. Bind(...) is called instead. This method
+				// exists as a stop-gap until the API can be deprecated and removed.
+				checkNoEntry();
+				return {};
 			}
 
 		private:
@@ -397,22 +404,30 @@ namespace Metasound
 
 			virtual ~FOutputOperator() {}
 
-			virtual FDataReferenceCollection GetInputs() const override
+			virtual void BindInputs(FInputVertexInterfaceData& InOutVertexData) override
 			{
-				FDataReferenceCollection Inputs;
+				InOutVertexData.BindReadVertex(OutputName, Center);
+			}
 
-				Inputs.AddDataReadReference(OutputName, Center);
+			virtual void BindOutputs(FOutputVertexInterfaceData& InOutVertexData) override
+			{
+				InOutVertexData.BindReadVertex(OutputName, Mono);
+			}
 
-				return Inputs;
+			virtual FDataReferenceCollection GetInputs() const override
+			{
+				// This should never be called. Bind(...) is called instead. This method
+				// exists as a stop-gap until the API can be deprecated and removed.
+				checkNoEntry();
+				return {};
 			}
 
 			virtual FDataReferenceCollection GetOutputs() const override
 			{
-				FDataReferenceCollection Outputs;
-
-				Outputs.AddDataReadReference(OutputName, Mono);
-
-				return Outputs;
+				// This should never be called. Bind(...) is called instead. This method
+				// exists as a stop-gap until the API can be deprecated and removed.
+				checkNoEntry();
+				return {};
 			}
 
 			virtual FExecuteFunction GetExecuteFunction() override
@@ -535,23 +550,30 @@ namespace Metasound
 
 			virtual ~FInputOperator() {}
 
-
-			virtual FDataReferenceCollection GetInputs() const override
+			virtual void BindInputs(FInputVertexInterfaceData& InOutVertexData) override
 			{
-				FDataReferenceCollection Inputs;
+				InOutVertexData.BindReadVertex(InputName, Mono);
+			}
 
-				Inputs.AddDataReadReference(InputName, Mono);
+			virtual void BindOutputs(FOutputVertexInterfaceData& InOutVertexData) override
+			{
+				InOutVertexData.BindReadVertex(InputName, Center);
+			}
 
-				return Inputs;
+			virtual FDataReferenceCollection GetInputs() const override
+			{
+				// This should never be called. Bind(...) is called instead. This method
+				// exists as a stop-gap until the API can be deprecated and removed.
+				checkNoEntry();
+				return {};
 			}
 
 			virtual FDataReferenceCollection GetOutputs() const override
 			{
-				FDataReferenceCollection Outputs;
-
-				Outputs.AddDataReadReference(InputName, Center);
-
-				return Outputs;
+				// This should never be called. Bind(...) is called instead. This method
+				// exists as a stop-gap until the API can be deprecated and removed.
+				checkNoEntry();
+				return {};
 			}
 
 			virtual FExecuteFunction GetExecuteFunction() override
diff --git a/Engine/Plugins/Runtime/Metasound/Source/MetasoundStandardNodes/Private/MetasoundDiffuserNode.cpp b/Engine/Plugins/Runtime/Metasound/Source/MetasoundStandardNodes/Private/MetasoundDiffuserNode.cpp
--- a/Engine/Plugins/Runtime/Metasound/Source/MetasoundStandardNodes/Private/MetasoundDiffuserNode.cpp
+++ b/Engine/Plugins/Runtime/Metasound/Source/MetasoundStandardNodes/Private/MetasoundDiffuserNode.cpp
@@ -117,17 +117,18 @@ namespace Metasound
 			}
 		}
 
-		virtual void Bind(FVertexInterfaceData& InVertexData) const override
+		virtual void BindInputs(FInputVertexInterfaceData& InOutVertexData) override
 		{
 			using namespace DiffuserNode;
+			InOutVertexData.BindReadVertex(METASOUND_GET_PARAM_NAME(InputAudio), AudioInput);
+			InOutVertexData.BindReadVertex(METASOUND_GET_PARAM_NAME(InputDiffusionDepth), Depth);
+			InOutVertexData.BindReadVertex(METASOUND_GET_PARAM_NAME(InputFeedbackGain), Feedback);
+		}
 
-			FInputVertexInterfaceData& Inputs = InVertexData.GetInputs();
-			Inputs.BindReadVertex(METASOUND_GET_PARAM_NAME(InputAudio), AudioInput);
-			Inputs.BindReadVertex(METASOUND_GET_PARAM_NAME(InputDiffusionDepth), Depth);
-			Inputs.BindReadVertex(METASOUND_GET_PARAM_NAME(InputFeedbackGain), Feedback);
-
-			FOutputVertexInterfaceData& Outputs = InVertexData.GetOutputs();
-			Outputs.BindReadVertex(METASOUND_GET_PARAM_NAME(OutputAudio), AudioOutput);
+		virtual void BindOutputs(FOutputVertexInterfaceData& InOutVertexData) override
+		{
+			using namespace DiffuserNode;
+			InOutVertexData.BindReadVertex(METASOUND_GET_PARAM_NAME(OutputAudio), AudioOutput);
 		}
 
 		virtual FDataReferenceCollection GetInputs() const override
diff --git a/Engine/Plugins/Runtime/Metasound/Source/MetasoundStandardNodes/Private/MetasoundDopplerPitchShiftNode.cpp b/Engine/Plugins/Runtime/Metasound/Source/MetasoundStandardNodes/Private/MetasoundDopplerPitchShiftNode.cpp
--- a/Engine/Plugins/Runtime/Metasound/Source/MetasoundStandardNodes/Private/MetasoundDopplerPitchShiftNode.cpp
+++ b/Engine/Plugins/Runtime/Metasound/Source/MetasoundStandardNodes/Private/MetasoundDopplerPitchShiftNode.cpp
@@ -32,7 +32,8 @@ namespace Metasound
 
 		FDopplerPitchShiftOperator(const FCreateOperatorParams& InParams, const FAudioBufferReadRef& InAudioInput, const FFloatReadRef& InPitchShift, const FFloatReadRef& InDelayLength);
 
-		virtual void Bind(FVertexInterfaceData& InVertexData) const override;
+		virtual void BindInputs(FInputVertexInterfaceData& InOutVertexData) override;
+		virtual void BindOutputs(FOutputVertexInterfaceData& InOutVertexData) override;
 		virtual FDataReferenceCollection GetInputs() const override;
 		virtual FDataReferenceCollection GetOutputs() const override;
 		void Reset(const IOperator::FResetParams& InParams);
@@ -60,17 +61,18 @@ namespace Metasound
 		Reset(InParams);
 	}
 
-	void FDopplerPitchShiftOperator::Bind(FVertexInterfaceData& InVertexData) const
+	void FDopplerPitchShiftOperator::BindInputs(FInputVertexInterfaceData& InOutVertexData)
 	{
 		using namespace DopplerPitchShift;
+		InOutVertexData.BindReadVertex(METASOUND_GET_PARAM_NAME(InParamAudioInput), AudioInput);
+		InOutVertexData.BindReadVertex(METASOUND_GET_PARAM_NAME(InParamPitchShift), PitchShift);
+		InOutVertexData.BindReadVertex(METASOUND_GET_PARAM_NAME(InParamDelayLength), DelayLength);
+	}
 
-		FInputVertexInterfaceData& Inputs = InVertexData.GetInputs();
-		Inputs.BindReadVertex(METASOUND_GET_PARAM_NAME(InParamAudioInput), FAudioBufferReadRef(AudioInput));
-		Inputs.BindReadVertex(METASOUND_GET_PARAM_NAME(InParamPitchShift), FFloatReadRef(PitchShift));
-		Inputs.BindReadVertex(METASOUND_GET_PARAM_NAME(InParamDelayLength), FFloatReadRef(DelayLength));
-
-		FOutputVertexInterfaceData& Outputs = InVertexData.GetOutputs();
-		Outputs.BindReadVertex(METASOUND_GET_PARAM_NAME(OutParamAudio), FAudioBufferReadRef(AudioOutput));
+	void FDopplerPitchShiftOperator::BindOutputs(FOutputVertexInterfaceData& InOutVertexData)
+	{
+		using namespace DopplerPitchShift;
+		InOutVertexData.BindReadVertex(METASOUND_GET_PARAM_NAME(OutParamAudio), AudioOutput);
 	}
 
 	FDataReferenceCollection FDopplerPitchShiftOperator::GetInputs() const
diff --git a/Engine/Plugins/Runtime/Metasound/Source/MetasoundStandardNodes/Private/MetasoundMidiScaleToArrayNode.cpp b/Engine/Plugins/Runtime/Metasound/Source/MetasoundStandardNodes/Private/MetasoundMidiScaleToArrayNode.cpp
--- a/Engine/Plugins/Runtime/Metasound/Source/MetasoundStandardNodes/Private/MetasoundMidiScaleToArrayNode.cpp
+++ b/Engine/Plugins/Runtime/Metasound/Source/MetasoundStandardNodes/Private/MetasoundMidiScaleToArrayNode.cpp
@@ -199,13 +199,13 @@ namespace Metasound
 
 	void FMusicalScaleToNoteArrayOperator::BindInputs(FInputVertexInterfaceData& InOutVertexData)
 	{
-		InOutVertexData.BindReadVertex(METASOUND_GET_PARAM_NAME(ParamScaleDegreesPreset), FEnumMusicalScaleReadRef(Scale));
-		InOutVertexData.BindReadVertex(METASOUND_GET_PARAM_NAME(ParamChordTonesOnly), FBoolReadRef(bChordTonesOnly));
+		InOutVertexData.BindReadVertex(METASOUND_GET_PARAM_NAME(ParamScaleDegreesPreset), Scale);
+		InOutVertexData.BindReadVertex(METASOUND_GET_PARAM_NAME(ParamChordTonesOnly), bChordTonesOnly);
 	}
 
 	void FMusicalScaleToNoteArrayOperator::BindOutputs(FOutputVertexInterfaceData& InOutVertexData)
 	{
-		InOutVertexData.BindReadVertex(METASOUND_GET_PARAM_NAME(ParamNoteArrayOutput), FArrayScaleDegreeWriteRef(ScaleArrayOutput));
+		InOutVertexData.BindReadVertex(METASOUND_GET_PARAM_NAME(ParamNoteArrayOutput), ScaleArrayOutput);
 	}
 
 	FDataReferenceCollection FMusicalScaleToNoteArrayOperator::GetInputs() const
