Fix ensures when migrating MetaSound assets
- Create new class names for temporary packages created during migration to avoid duplicate node registration
- Change duplicate node registration ensure to if check with better logging of asset causing collision

diff --git a/Engine/Plugins/Runtime/Metasound/Source/MetasoundEditor/Private/MetasoundEditorModule.cpp b/Engine/Plugins/Runtime/Metasound/Source/MetasoundEditor/Private/MetasoundEditorModule.cpp
--- a/Engine/Plugins/Runtime/Metasound/Source/MetasoundEditor/Private/MetasoundEditorModule.cpp
+++ b/Engine/Plugins/Runtime/Metasound/Source/MetasoundEditor/Private/MetasoundEditorModule.cpp
@@ -17,6 +17,7 @@
 #include "MetasoundAssetSubsystem.h"
 #include "MetasoundAssetTypeActions.h"
 #include "MetasoundAudioBuffer.h"
+#include "MetasoundBuilderSubsystem.h"
 #include "MetasoundDetailCustomization.h"
 #include "MetasoundEditorGraph.h"
 #include "MetasoundEditorGraphBuilder.h"
@@ -37,6 +38,7 @@
 #include "MetasoundUObjectRegistry.h"
 #include "Modules/ModuleInterface.h"
 #include "Modules/ModuleManager.h"
+#include "PackageMigrationContext.h"
 #include "PropertyEditorDelegates.h"
 #include "PropertyEditorModule.h"
 #include "Styling/CoreStyle.h"
@@ -822,6 +824,8 @@ namespace Metasound
 					constexpr bool bIncludeNamespace = false;
 					return FGraphBuilder::GetDisplayName(*NodeHandle, bIncludeNamespace);
 				});
+
+				AssetTools.GetOnPackageMigration().AddRaw(this, &FModule::OnPackageMigration);
 			}
 
 			virtual void ShutdownModule() override
@@ -840,6 +844,7 @@ namespace Metasound
 					{
 						AssetTools.UnregisterAssetTypeActions(AssetAction.ToSharedRef());
 					}
+					AssetTools.GetOnPackageMigration().RemoveAll(this);
 				}
 
 				if (GraphConnectionFactory.IsValid())
@@ -865,6 +870,59 @@ namespace Metasound
 				PinTypes.Reset();
 			}
 
+			void OnPackageMigration(UE::AssetTools::FPackageMigrationContext& MigrationContext)
+			{
+				// Migration can create temporary new packages that use the same name 
+				// (and therefore node registry key) as the asset migrated. 
+				// So generate new class names to avoid registry key collisions. 
+				if (MigrationContext.GetCurrentStep() == UE::AssetTools::FPackageMigrationContext::EPackageMigrationStep::InstancedPackagesLoaded)
+				{
+					// Gather the new MetaSound assets
+					TMap<FMetasoundFrontendClassName, FMetasoundFrontendClassName> OldToNewReferencedClassNames;
+					TArray<UMetaSoundBuilderBase*> NewMetaSoundAssetBuilders;
+					TArray<FMetasoundAssetBase*> NewMetaSoundAssets;
+					for (const UE::AssetTools::FPackageMigrationContext::FMigrationPackageData& MigrationPackageData : MigrationContext.GetMigrationPackagesData())
+					{
+						UPackage* Package = MigrationPackageData.GetInstancedPackage();
+						if (Package)
+						{
+							UObject* MainAsset = Package->FindAssetInPackage();
+							// Only apply to MetaSound assets 
+							if (IMetasoundUObjectRegistry::Get().IsRegisteredClass(MainAsset))
+							{
+								FMetasoundAssetBase* MetaSoundAsset = IMetasoundUObjectRegistry::Get().GetObjectAsAssetBase(MainAsset);
+								if (MetaSoundAsset)
+								{
+									UMetaSoundBuilderBase& MetaSoundBuilder = UMetaSoundBuilderSubsystem::GetChecked().AttachBuilderToAssetChecked(*MainAsset);
+									NewMetaSoundAssetBuilders.Add(&MetaSoundBuilder);
+									NewMetaSoundAssets.Add(MetaSoundAsset);
+								}
+							}
+						}
+					}
+
+					// Assign new class names and cache mapping with old one
+					for (UMetaSoundBuilderBase* MetaSoundBuilder : NewMetaSoundAssetBuilders)
+					{
+						const FMetasoundFrontendClassName NewReferencedClassName = FMetasoundFrontendClassName(/*Namespace=*/{}, FName(FGuid::NewGuid().ToString()), /*Variant=*/{});
+						const FMetasoundFrontendClassName OldReferencedClassName = MetaSoundBuilder->GetRootGraphClassName();
+
+						MetaSoundBuilder->RenameRootGraphClass(NewReferencedClassName);
+						OldToNewReferencedClassNames.FindOrAdd(OldReferencedClassName) = NewReferencedClassName;
+					}
+
+					// Fix up dependencies 
+					for (UMetaSoundBuilderBase* MetaSoundBuilder : NewMetaSoundAssetBuilders)
+					{
+						MetaSoundBuilder->UpdateDependencyClassNames(OldToNewReferencedClassNames);
+					}
+					for (FMetasoundAssetBase* MetaSoundAsset : NewMetaSoundAssets)
+					{
+						MetaSoundAsset->RebuildReferencedAssetClasses();
+					}
+				}
+			}
+			
 			TArray<TSharedPtr<FAssetTypeActions_Base>> AssetActions;
 			TMap<EMetasoundFrontendLiteralType, const TSubclassOf<UMetasoundEditorGraphMemberDefaultLiteral>> InputDefaultLiteralClassRegistry;
 			TMap<FName, FEdGraphPinType> PinTypes;
diff --git a/Engine/Plugins/Runtime/Metasound/Source/MetasoundEngine/Private/MetasoundBuilderSubsystem.cpp b/Engine/Plugins/Runtime/Metasound/Source/MetasoundEngine/Private/MetasoundBuilderSubsystem.cpp
--- a/Engine/Plugins/Runtime/Metasound/Source/MetasoundEngine/Private/MetasoundBuilderSubsystem.cpp
+++ b/Engine/Plugins/Runtime/Metasound/Source/MetasoundEngine/Private/MetasoundBuilderSubsystem.cpp
@@ -624,6 +624,11 @@ void UMetaSoundBuilderBase::RemoveNodeInputDefault(const FMetaSoundBuilderNodeIn
 	OutResult = bInputDefaultRemoved ? EMetaSoundBuilderResult::Succeeded : EMetaSoundBuilderResult::Failed;
 }
 
+void UMetaSoundBuilderBase::RenameRootGraphClass(const FMetasoundFrontendClassName& InName)
+{
+	Builder.RenameRootGraphClass(InName);
+}
+
 #if WITH_EDITOR
 void UMetaSoundBuilderBase::SetAuthor(const FString& InAuthor)
 {
@@ -682,6 +687,11 @@ FMetasoundFrontendVersion UMetaSoundBuilderBase::FindNodeClassVersion(const FMet
 	return FMetasoundFrontendVersion::GetInvalid();
 }
 
+FMetasoundFrontendClassName UMetaSoundBuilderBase::GetRootGraphClassName() const
+{
+	return Builder.GetDocument().RootGraph.Metadata.GetClassName();
+}
+
 void UMetaSoundBuilderBase::GetNodeInputData(const FMetaSoundBuilderNodeInputHandle& InputHandle, FName& Name, FName& DataType, EMetaSoundBuilderResult& OutResult)
 {
 	if (const FMetasoundFrontendVertex* Vertex = Builder.FindNodeInput(InputHandle.NodeID, InputHandle.VertexID))
@@ -778,6 +788,11 @@ void UMetaSoundBuilderBase::ConvertToPreset(const TScriptInterface<IMetaSoundDoc
 	}
 }
 
+void UMetaSoundBuilderBase::UpdateDependencyClassNames(const TMap<FMetasoundFrontendClassName, FMetasoundFrontendClassName>& OldToNewReferencedClassNames)
+{
+	Builder.UpdateDependencyClassNames(OldToNewReferencedClassNames);
+}
+
 TScriptInterface<IMetaSoundDocumentInterface> UMetaSoundPatchBuilder::Build(UObject* Parent, const FMetaSoundBuilderOptions& InBuilderOptions) const
 {
 	return &BuildInternal<UMetaSoundPatch>(Parent, InBuilderOptions);
diff --git a/Engine/Plugins/Runtime/Metasound/Source/MetasoundEngine/Public/MetasoundBuilderSubsystem.h b/Engine/Plugins/Runtime/Metasound/Source/MetasoundEngine/Public/MetasoundBuilderSubsystem.h
--- a/Engine/Plugins/Runtime/Metasound/Source/MetasoundEngine/Public/MetasoundBuilderSubsystem.h
+++ b/Engine/Plugins/Runtime/Metasound/Source/MetasoundEngine/Public/MetasoundBuilderSubsystem.h
@@ -202,7 +202,7 @@ public:
 	// all connections were found and removed, failed if any connections were not.
 	UFUNCTION(BlueprintCallable, Category = "Audio|MetaSound|Builder", meta = (ExpandEnumAsExecs = "OutResult"))
 	void DisconnectNodesByInterfaceBindings(const FMetaSoundNodeHandle& FromNodeHandle, const FMetaSoundNodeHandle& ToNodeHandle, EMetaSoundBuilderResult& OutResult);
-
+	
 	// Returns graph input node by the given name if it exists, or an invalid handle if not found.
 	UFUNCTION(BlueprintCallable, Category = "Audio|MetaSound|Builder", meta = (ExpandEnumAsExecs = "OutResult"))
 	UPARAM(DisplayName = "Node Handle") FMetaSoundNodeHandle FindGraphInputNode(FName InputName, EMetaSoundBuilderResult& OutResult);
@@ -255,6 +255,9 @@ public:
 	UFUNCTION(BlueprintCallable, Category = "Audio|MetaSound|Builder", meta = (ExpandEnumAsExecs = "OutResult"))
 	UPARAM(DisplayName = "Node ClassVersion") FMetasoundFrontendVersion FindNodeClassVersion(const FMetaSoundNodeHandle& NodeHandle, EMetaSoundBuilderResult& OutResult);
 
+	// Returns the document's root graph class name
+	FMetasoundFrontendClassName GetRootGraphClassName() const;
+
 	// Returns node input's data if valid (including things like name and datatype).
 	UFUNCTION(BlueprintCallable, Category = "Audio|MetaSound|Builder", meta = (ExpandEnumAsExecs = "OutResult"))
 	void GetNodeInputData(const FMetaSoundBuilderNodeInputHandle& InputHandle, FName& Name, FName& DataType, EMetaSoundBuilderResult& OutResult);
@@ -324,6 +327,9 @@ public:
 	// Returns success if value was removed, false if not removed (i.e. wasn't set to begin with).
 	UFUNCTION(BlueprintCallable, Category = "Audio|MetaSound|Builder", meta = (ExpandEnumAsExecs = "OutResult"))
 	void RemoveNodeInputDefault(const FMetaSoundBuilderNodeInputHandle& InputHandle, EMetaSoundBuilderResult& OutResult);
+	
+	// Rename the document's root graph class with a guid and optional namespace and variant
+	void RenameRootGraphClass(const FMetasoundFrontendClassName& InName);
 
 #if WITH_EDITOR
 	// Sets the author of the MetaSound.
@@ -338,6 +344,9 @@ public:
 	UFUNCTION(BlueprintCallable, Category = "Audio|MetaSound|Builder", meta = (ExpandEnumAsExecs = "OutResult"))
 	void SetGraphInputDefault(FName InputName, const FMetasoundFrontendLiteral& Literal, EMetaSoundBuilderResult& OutResult);
 
+	// Update dependency class names given a map of old to new referenced class names. 
+	void UpdateDependencyClassNames(const TMap<FMetasoundFrontendClassName, FMetasoundFrontendClassName>& OldToNewReferencedClassNames);
+
 	virtual TScriptInterface<IMetaSoundDocumentInterface> Build(UObject* Parent, const FMetaSoundBuilderOptions& Options) const PURE_VIRTUAL(UMetaSoundBuilderBase::Build, return { }; );
 
 	// Returns the base MetaSound UClass the builder is operating on (ex. MetaSoundSource, MetaSoundPatch, etc. not to be confused with the underlying document builder's MetaSound class type)
diff --git a/Engine/Plugins/Runtime/Metasound/Source/MetasoundFrontend/Private/MetasoundFrontendDocumentBuilder.cpp b/Engine/Plugins/Runtime/Metasound/Source/MetasoundFrontend/Private/MetasoundFrontendDocumentBuilder.cpp
--- a/Engine/Plugins/Runtime/Metasound/Source/MetasoundFrontend/Private/MetasoundFrontendDocumentBuilder.cpp
+++ b/Engine/Plugins/Runtime/Metasound/Source/MetasoundFrontend/Private/MetasoundFrontendDocumentBuilder.cpp
@@ -2097,6 +2097,20 @@ bool FMetaSoundFrontendDocumentBuilder::RemoveNode(const FGuid& InNodeID)
 	return false;
 }
 
+bool FMetaSoundFrontendDocumentBuilder::RenameRootGraphClass(const FMetasoundFrontendClassName& InName)
+{
+	FGuid NameGuid;
+	if (FGuid::Parse(InName.Name.ToString(), NameGuid))
+	{
+		return Metasound::Frontend::FRenameRootGraphClass::Generate(GetDocument(), NameGuid, InName.Namespace, InName.Variant);
+	}
+	else
+	{
+		UE_LOG(LogMetaSound, Error, TEXT("Attempting to rename a root graph class with class name '%s' which contains an invalid name guid."), *InName.GetFullName().ToString());
+		return false;
+	}
+}
+
 #if WITH_EDITOR
 void FMetaSoundFrontendDocumentBuilder::SetAuthor(const FString& InAuthor)
 {
@@ -2346,3 +2360,20 @@ bool FMetaSoundFrontendDocumentBuilder::SwapGraphOutput(const FMetasoundFrontend
 
 	return true;
 }
+
+bool FMetaSoundFrontendDocumentBuilder::UpdateDependencyClassNames(const TMap<FMetasoundFrontendClassName, FMetasoundFrontendClassName>& OldToNewReferencedClassNames)
+{
+	for (FMetasoundFrontendClass& Dependency : GetDocument().Dependencies)
+	{
+		const FMetasoundFrontendClassName OldName = Dependency.Metadata.GetClassName();
+		if (const FMetasoundFrontendClassName* NewName = OldToNewReferencedClassNames.Find(OldName))
+		{
+			const int32* DependencyIndex = DocumentCache->FindDependencyIndex(Dependency.ID);
+			check(DependencyIndex);
+			GetDocumentDelegates().OnRenamingDependencyClass.Broadcast(*DependencyIndex, *NewName);
+			Dependency.Metadata.SetClassName(*NewName);
+		}
+	}
+	
+	return true;
+}
diff --git a/Engine/Plugins/Runtime/Metasound/Source/MetasoundFrontend/Private/MetasoundFrontendDocumentCache.cpp b/Engine/Plugins/Runtime/Metasound/Source/MetasoundFrontend/Private/MetasoundFrontendDocumentCache.cpp
--- a/Engine/Plugins/Runtime/Metasound/Source/MetasoundFrontend/Private/MetasoundFrontendDocumentCache.cpp
+++ b/Engine/Plugins/Runtime/Metasound/Source/MetasoundFrontend/Private/MetasoundFrontendDocumentCache.cpp
@@ -36,6 +36,7 @@ namespace Metasound::Frontend
 
 		OutDelegates.OnDependencyAdded.AddSP(this, &FDocumentCache::OnDependencyAdded);
 		OutDelegates.OnRemovingDependency.AddSP(this, &FDocumentCache::OnRemovingDependency);
+		OutDelegates.OnRenamingDependencyClass.AddSP(this, &FDocumentCache::OnRenamingDependencyClass);
 
 		if (!EdgeCache.IsValid())
 		{
@@ -132,6 +133,20 @@ namespace Metasound::Frontend
 		IDToIndex.Remove(DependencyBeingRemoved.ID);
 	}
 
+	void FDocumentCache::OnRenamingDependencyClass(const int32 IndexBeingRenamed, const FMetasoundFrontendClassName& NewName)
+	{
+		using namespace Metasound::Frontend;
+
+		const FMetasoundFrontendClass& DependencyBeingRemoved = GetDocument().Dependencies[IndexBeingRenamed];
+		const FNodeRegistryKey OldKey = NodeRegistryKey::CreateKey(DependencyBeingRemoved.Metadata);
+		FMetasoundFrontendClassMetadata NewMetadata = DependencyBeingRemoved.Metadata;
+		NewMetadata.SetClassName(NewName);
+		const FNodeRegistryKey NewKey = NodeRegistryKey::CreateKey(NewMetadata);
+
+		KeyToIndex.Add(NewKey, IndexBeingRenamed);
+		KeyToIndex.Remove(OldKey);
+	}
+
 	FDocumentGraphInterfaceCache::FDocumentGraphInterfaceCache(TSharedRef<IDocumentCache> ParentCache)
 		: Parent(ParentCache)
 	{
diff --git a/Engine/Plugins/Runtime/Metasound/Source/MetasoundFrontend/Private/MetasoundFrontendDocumentCache.h b/Engine/Plugins/Runtime/Metasound/Source/MetasoundFrontend/Private/MetasoundFrontendDocumentCache.h
--- a/Engine/Plugins/Runtime/Metasound/Source/MetasoundFrontend/Private/MetasoundFrontendDocumentCache.h
+++ b/Engine/Plugins/Runtime/Metasound/Source/MetasoundFrontend/Private/MetasoundFrontendDocumentCache.h
@@ -141,6 +141,7 @@ namespace Metasound::Frontend
 	private:
 		void OnDependencyAdded(int32 InNewIndex);
 		void OnRemovingDependency(int32 IndexBeingRemoved);
+		void OnRenamingDependencyClass(const int32 IndexBeingRenamed, const FMetasoundFrontendClassName& NewName);
 
 		// Cache of dependency (Class) ID to corresponding class dependency index
 		TSortedMap<FGuid, int32> IDToIndex;
diff --git a/Engine/Plugins/Runtime/Metasound/Source/MetasoundFrontend/Private/MetasoundFrontendRegistries.cpp b/Engine/Plugins/Runtime/Metasound/Source/MetasoundFrontend/Private/MetasoundFrontendRegistries.cpp
--- a/Engine/Plugins/Runtime/Metasound/Source/MetasoundFrontend/Private/MetasoundFrontendRegistries.cpp
+++ b/Engine/Plugins/Runtime/Metasound/Source/MetasoundFrontend/Private/MetasoundFrontendRegistries.cpp
@@ -220,12 +220,14 @@ namespace Metasound
 				Key = NodeRegistryKey::CreateKey(Entry->GetClassInfo());
 
 				// check to see if an identical node was already registered, and log
-				ensureAlwaysMsgf(
-					!RegisteredNodes.Contains(Key),
-					TEXT("Node with registry key '%s' already registered. "
-						"The previously registered node will be overwritten. "
-						"This can happen if two classes share the same name or if METASOUND_REGISTER_NODE is defined in a public header."),
-					*Key);
+				if (RegisteredNodes.Contains(Key))
+				{
+					TSharedRef<INodeRegistryEntry, ESPMode::ThreadSafe> Node = RegisteredNodes[Key];
+					const FNodeClassInfo& ClassInfo = Node->GetClassInfo();
+					UE_LOG(LogMetaSound, Error, TEXT("Node with registry key '%s' already registered by asset '%s'. The previously registered node will be overwritten. " 
+						"This can happen if two classes share the same name or if METASOUND_REGISTER_NODE is defined in a public header."
+					), *Key, *ClassInfo.AssetPath.ToString());
+				}
 
 				// Store update to newly registered node in history so nodes
 				// can be queried by transaction ID
diff --git a/Engine/Plugins/Runtime/Metasound/Source/MetasoundFrontend/Private/MetasoundFrontendTransform.cpp b/Engine/Plugins/Runtime/Metasound/Source/MetasoundFrontend/Private/MetasoundFrontendTransform.cpp
--- a/Engine/Plugins/Runtime/Metasound/Source/MetasoundFrontend/Private/MetasoundFrontendTransform.cpp
+++ b/Engine/Plugins/Runtime/Metasound/Source/MetasoundFrontend/Private/MetasoundFrontendTransform.cpp
@@ -1002,5 +1002,13 @@ namespace Metasound
 
 			return true;
 		}
+
+		bool FRenameRootGraphClass::Transform(FMetasoundFrontendDocument& InOutDocument) const
+		{
+			FMetasoundFrontendClassMetadata& Metadata = InOutDocument.RootGraph.Metadata;
+			Metadata.SetClassName(NewClassName);
+
+			return true;
+		}
 	} // namespace Frontend
 } // namespace Metasound
diff --git a/Engine/Plugins/Runtime/Metasound/Source/MetasoundFrontend/Public/MetasoundFrontendDocument.h b/Engine/Plugins/Runtime/Metasound/Source/MetasoundFrontend/Public/MetasoundFrontendDocument.h
--- a/Engine/Plugins/Runtime/Metasound/Source/MetasoundFrontend/Public/MetasoundFrontendDocument.h
+++ b/Engine/Plugins/Runtime/Metasound/Source/MetasoundFrontend/Public/MetasoundFrontendDocument.h
@@ -1227,6 +1227,11 @@ struct METASOUNDFRONTEND_API FMetasoundFrontendClassName
 	// Return string version of full name.
 	FString ToString() const;
 
+	friend FORCEINLINE uint32 GetTypeHash(const FMetasoundFrontendClassName& ClassName)
+	{
+		return GetTypeHash(ClassName.GetFullName());
+	}
+
 	METASOUNDFRONTEND_API friend bool operator==(const FMetasoundFrontendClassName& InLHS, const FMetasoundFrontendClassName& InRHS);
 
 	METASOUNDFRONTEND_API friend bool operator!=(const FMetasoundFrontendClassName& InLHS, const FMetasoundFrontendClassName& InRHS);
diff --git a/Engine/Plugins/Runtime/Metasound/Source/MetasoundFrontend/Public/MetasoundFrontendDocumentBuilder.h b/Engine/Plugins/Runtime/Metasound/Source/MetasoundFrontend/Public/MetasoundFrontendDocumentBuilder.h
--- a/Engine/Plugins/Runtime/Metasound/Source/MetasoundFrontend/Public/MetasoundFrontendDocumentBuilder.h
+++ b/Engine/Plugins/Runtime/Metasound/Source/MetasoundFrontend/Public/MetasoundFrontendDocumentBuilder.h
@@ -209,6 +209,7 @@ public:
 	bool RemoveNode(const FGuid& InNodeID);
 	bool RemoveNodeInputDefault(const FGuid& InNodeID, const FGuid& InVertexID);
 	bool RemoveUnusedDependencies();
+	bool RenameRootGraphClass(const FMetasoundFrontendClassName& InName);
 
 #if WITH_EDITOR
 	void SetAuthor(const FString& InAuthor);
@@ -219,6 +220,7 @@ public:
 
 	bool SwapGraphInput(const FMetasoundFrontendClassVertex& InExistingInputVertex, const FMetasoundFrontendClassVertex& NewInputVertex);
 	bool SwapGraphOutput(const FMetasoundFrontendClassVertex& InExistingOutputVertex, const FMetasoundFrontendClassVertex& NewOutputVertex);
+	bool UpdateDependencyClassNames(const TMap<FMetasoundFrontendClassName, FMetasoundFrontendClassName>& OldToNewReferencedClassNames);
 
 private:
 	using FFinalizeNodeFunctionRef = TFunctionRef<void(FMetasoundFrontendNode&, const Metasound::Frontend::FNodeRegistryKey&)>;
diff --git a/Engine/Plugins/Runtime/Metasound/Source/MetasoundFrontend/Public/MetasoundFrontendDocumentModifyDelegates.h b/Engine/Plugins/Runtime/Metasound/Source/MetasoundFrontend/Public/MetasoundFrontendDocumentModifyDelegates.h
--- a/Engine/Plugins/Runtime/Metasound/Source/MetasoundFrontend/Public/MetasoundFrontendDocumentModifyDelegates.h
+++ b/Engine/Plugins/Runtime/Metasound/Source/MetasoundFrontend/Public/MetasoundFrontendDocumentModifyDelegates.h
@@ -8,6 +8,7 @@
 
 DECLARE_MULTICAST_DELEGATE_OneParam(FOnMetaSoundFrontendDocumentMutateArray, int32 /* Index */);
 DECLARE_MULTICAST_DELEGATE_OneParam(FOnMetaSoundFrontendDocumentMutateInterfaceArray, const FMetasoundFrontendInterface& /* Interface */);
+DECLARE_MULTICAST_DELEGATE_TwoParams(FOnMetaSoundFrontendDocumentRenameClass, const int32 /* Index */, const FMetasoundFrontendClassName& /* NewName */);
 DECLARE_MULTICAST_DELEGATE_ThreeParams(FOnMetaSoundFrontendDocumentMutateNodeInputLiteralArray, int32 /* NodeIndex */, int32 /* VertexIndex */, int32 /* LiteralIndex */);
 
 
@@ -44,6 +45,7 @@ namespace Metasound::Frontend
 	{
 		FOnMetaSoundFrontendDocumentMutateArray OnDependencyAdded;
 		FOnMetaSoundFrontendDocumentMutateArray OnRemovingDependency;
+		FOnMetaSoundFrontendDocumentRenameClass OnRenamingDependencyClass;
 
 		FInterfaceModifyDelegates InterfaceDelegates;
 		FNodeModifyDelegates NodeDelegates;
diff --git a/Engine/Plugins/Runtime/Metasound/Source/MetasoundFrontend/Public/MetasoundFrontendTransform.h b/Engine/Plugins/Runtime/Metasound/Source/MetasoundFrontend/Public/MetasoundFrontendTransform.h
--- a/Engine/Plugins/Runtime/Metasound/Source/MetasoundFrontend/Public/MetasoundFrontendTransform.h
+++ b/Engine/Plugins/Runtime/Metasound/Source/MetasoundFrontend/Public/MetasoundFrontendTransform.h
@@ -230,12 +230,19 @@ namespace Metasound
 				return FRenameRootGraphClass(GeneratedClassName).Transform(InDocument);
 			}
 
+			static bool Generate(FMetasoundFrontendDocument& InDocument, const FGuid& InGuid, const FName Namespace = { }, const FName Variant = { })
+			{
+				const FMetasoundFrontendClassName GeneratedClassName = { Namespace, *InGuid.ToString(), Variant };
+				return FRenameRootGraphClass(GeneratedClassName).Transform(InDocument);
+			}
+
 			FRenameRootGraphClass(const FMetasoundFrontendClassName InClassName)
 				: NewClassName(InClassName)
 			{
 			}
 
 			bool Transform(FDocumentHandle InDocument) const override;
+			bool Transform(FMetasoundFrontendDocument& InOutDocument) const override;
 		};
 	}
 }
