[Metasounds] Allow for zero-length attack/decay phases (skip attack / decay)

diff --git a/Engine/Plugins/Runtime/Metasound/Source/MetasoundStandardNodes/Private/MetasoundADEnvelopeNode.cpp b/Engine/Plugins/Runtime/Metasound/Source/MetasoundStandardNodes/Private/MetasoundADEnvelopeNode.cpp
--- a/Engine/Plugins/Runtime/Metasound/Source/MetasoundStandardNodes/Private/MetasoundADEnvelopeNode.cpp
+++ b/Engine/Plugins/Runtime/Metasound/Source/MetasoundStandardNodes/Private/MetasoundADEnvelopeNode.cpp
@@ -162,7 +162,7 @@ namespace Metasound
 				for (int32 i = StartFrame; i < EndFrame; ++i)
 				{
 					// We are in attack
-					if (InState.CurrentSampleIndex <= InState.AttackSampleCount)
+					if (InState.CurrentSampleIndex < InState.AttackSampleCount)
 					{
 						float AttackFraction = (float)InState.CurrentSampleIndex++ / InState.AttackSampleCount;
 						float EnvValue = FMath::Pow(AttackFraction, InState.AttackCurveFactor);
@@ -189,7 +189,7 @@ namespace Metasound
 						else if (InState.bLooping)
 						{
 							InState.StartingEnvelopeValue = 0.0f;
-							InState.CurrentEnvelopeValue = 0.0f;
+							InState.CurrentEnvelopeValue = InState.AttackSampleCount? 0.f : 1.f;
 							InState.CurrentSampleIndex = 0;
 							OutFinishedFrames.Add(i);
 						}
@@ -369,12 +369,21 @@ namespace Metasound
 		{
 			float AttackTimeSeconds = AttackTime->GetSeconds();
 			float DecayTimeSeconds = DecayTime->GetSeconds();
-			EnvState.AttackSampleCount = FMath::Max(1, SampleRate * AttackTimeSeconds);
-			EnvState.DecaySampleCount = FMath::Max(1, SampleRate * DecayTimeSeconds);
+			EnvState.AttackSampleCount = FMath::Max(0, SampleRate * AttackTimeSeconds);
+
+			// if our attack phase is zero, force decay phase to be at least a single sample
+			const int32 DecaySampleCountMin = (EnvState.AttackSampleCount? 0 : 1);
+			EnvState.DecaySampleCount = FMath::Max(DecaySampleCountMin, SampleRate * DecayTimeSeconds);
 			EnvState.AttackCurveFactor = FMath::Max(KINDA_SMALL_NUMBER, *AttackCurveFactor);
 			EnvState.DecayCurveFactor = FMath::Max(KINDA_SMALL_NUMBER, *DecayCurveFactor);
 			EnvState.bLooping = *bLooping;
 			EnvState.bHardReset = *bHardReset;
+
+			// if there is no attack phase, we jump to 1.f on the first frame
+			if(EnvState.AttackSampleCount == 0)
+			{
+				EnvState.CurrentEnvelopeValue = 1.f;
+			}
 		}
 
 		void Reset(const IOperator::FResetParams& InParams)
