[Metasound] Set graph inputs to node defaults when using single-node helper.

Also add node tests which were not marked for add.

#rnx

diff --git a/Engine/Plugins/Runtime/Metasound/Source/MetasoundEngineTest/Private/NodeTestGraphBuilder.cpp b/Engine/Plugins/Runtime/Metasound/Source/MetasoundEngineTest/Private/NodeTestGraphBuilder.cpp
--- a/Engine/Plugins/Runtime/Metasound/Source/MetasoundEngineTest/Private/NodeTestGraphBuilder.cpp
+++ b/Engine/Plugins/Runtime/Metasound/Source/MetasoundEngineTest/Private/NodeTestGraphBuilder.cpp
@@ -148,6 +148,13 @@ namespace Metasound::Test
 			FOutputHandle OutputToConnect = InputNode->GetOutputWithVertexName(Input->GetName());
 			FInputHandle InputToConnect = NodeHandle->GetInputWithVertexName(Input->GetName());
 
+			// set the input to the default, if there is one
+			if (const FMetasoundFrontendLiteral* Default = InputToConnect->GetClassDefaultLiteral(); nullptr != Default)
+			{
+				const FGuid InputId = Builder.RootGraph->GetVertexIDForInputVertex(Input->GetName());
+				Builder.RootGraph->SetDefaultInput(InputId, *Default);
+			}
+
 			if (!InputToConnect->Connect(*OutputToConnect))
 			{
 				return nullptr;
