[Metasound] Remove unnecessary lambda capture when a generator is destroyed

diff --git a/Engine/Plugins/Runtime/Metasound/Source/MetasoundEngine/Private/MetasoundGeneratorHandle.cpp b/Engine/Plugins/Runtime/Metasound/Source/MetasoundEngine/Private/MetasoundGeneratorHandle.cpp
--- a/Engine/Plugins/Runtime/Metasound/Source/MetasoundEngine/Private/MetasoundGeneratorHandle.cpp
+++ b/Engine/Plugins/Runtime/Metasound/Source/MetasoundEngine/Private/MetasoundGeneratorHandle.cpp
@@ -136,7 +136,7 @@ void UMetasoundGeneratorHandle::AttachGeneratorDelegates()
 			// We are in the audio render (or control) thread here, so create a "dispatch task" to be
 			// executed later on the game thread...
 			FFunctionGraphTask::CreateAndDispatchWhenReady(
-				[WeakGeneratorHandlePtr, InAudioComponentId, InGenerator]()
+				[WeakGeneratorHandlePtr, InAudioComponentId]()
 				{
 					METASOUND_TRACE_CPUPROFILER_EVENT_SCOPE(UMetasoundGeneratorHandle::CallingGeneratorDetachedDelegates);
 					check(IsInGameThread());
@@ -145,7 +145,7 @@ void UMetasoundGeneratorHandle::AttachGeneratorDelegates()
 					// hasn't been garbage collected.
 					if (UMetasoundGeneratorHandle* TheHandle = WeakGeneratorHandlePtr.Get())
 					{
-						TheHandle->OnSourceDestroyedAGenerator(InAudioComponentId, InGenerator);
+						TheHandle->OnSourceDestroyedAGenerator(InAudioComponentId);
 					}
 				},
 				StatId, nullptr, ENamedThreads::GameThread);
@@ -440,7 +440,7 @@ void UMetasoundGeneratorHandle::OnSourceCreatedAGenerator(uint64 InAudioComponen
 	}
 }
 
-void UMetasoundGeneratorHandle::OnSourceDestroyedAGenerator(uint64 InAudioComponentId, TSharedPtr<Metasound::FMetasoundGenerator> InGenerator)
+void UMetasoundGeneratorHandle::OnSourceDestroyedAGenerator(uint64 InAudioComponentId)
 {
 	check(IsInGameThread());
 
diff --git a/Engine/Plugins/Runtime/Metasound/Source/MetasoundEngine/Public/MetasoundGeneratorHandle.h b/Engine/Plugins/Runtime/Metasound/Source/MetasoundEngine/Public/MetasoundGeneratorHandle.h
--- a/Engine/Plugins/Runtime/Metasound/Source/MetasoundEngine/Public/MetasoundGeneratorHandle.h
+++ b/Engine/Plugins/Runtime/Metasound/Source/MetasoundEngine/Public/MetasoundGeneratorHandle.h
@@ -125,7 +125,7 @@ private:
 	 * Generator creation and destruction delegates we register with the UMetaSoundSource
 	 */
 	void OnSourceCreatedAGenerator(uint64 InAudioComponentId, TSharedPtr<Metasound::FMetasoundGenerator> InGenerator);
-	void OnSourceDestroyedAGenerator(uint64 InAudioComponentId, TSharedPtr<Metasound::FMetasoundGenerator> InGenerator);
+	void OnSourceDestroyedAGenerator(uint64 InAudioComponentId);
 	
 	TWeakObjectPtr<UAudioComponent> AudioComponent;
 	TWeakObjectPtr<UMetaSoundSource> CachedMetasoundSource;
