Add missing forward of early return value causing false negative when attempting to add MetaSound Interface via Builder API

diff --git a/Engine/Plugins/Runtime/Metasound/Source/MetasoundEngineTest/Private/EngineTestMetasoundSource.cpp b/Engine/Plugins/Runtime/Metasound/Source/MetasoundEngineTest/Private/EngineTestMetasoundSource.cpp
--- a/Engine/Plugins/Runtime/Metasound/Source/MetasoundEngineTest/Private/EngineTestMetasoundSource.cpp
+++ b/Engine/Plugins/Runtime/Metasound/Source/MetasoundEngineTest/Private/EngineTestMetasoundSource.cpp
@@ -12,6 +12,7 @@
 #include "Misc/Paths.h"
 #include "Templates/SharedPointer.h"
 #include "Tests/AutomationCommon.h"
+#include "Sound/SoundAttenuation.h"
 
 
 #if WITH_DEV_AUTOMATION_TESTS
@@ -466,4 +467,63 @@ bool FAudioMetasoundSourceBuilderLiveUpdateLiteral::RunTest(const FString& Param
 
 	return true;
 }
+
+// This test creates a MetaSound source from a SourceBuilder, then adds and finally removes an interface using the builder API, and verifies it as well as its members were added to the document.
+IMPLEMENT_SIMPLE_AUTOMATION_TEST(FAudioMetasoundSourceBuilderMutateInterface, "Audio.Metasound.MutateInterface", EAutomationTestFlags::EditorContext | EAutomationTestFlags::EngineFilter)
+bool FAudioMetasoundSourceBuilderMutateInterface::RunTest(const FString& Parameters)
+{
+	using namespace Audio;
+	using namespace EngineTestMetaSoundPatchBuilderPrivate;
+	using namespace EngineTestMetaSoundSourcePrivate;
+	using namespace Metasound::Frontend;
+
+	constexpr EMetaSoundOutputAudioFormat OutputFormat = EMetaSoundOutputAudioFormat::Mono;
+	constexpr bool bIsOneShot = false;
+	FInitTestBuilderSourceOutput Output;
+	UMetaSoundSourceBuilder& Builder = CreateSourceBuilder(*this, EMetaSoundOutputAudioFormat::Mono, bIsOneShot, Output);
+	Builder.AddToRoot();
+
+	EMetaSoundBuilderResult Result;
+
+	// Test interface output mutation with oneshot interface
+	Builder.AddInterface(SourceOneShotInterface::GetVersion().Name, Result);
+	AddErrorIfFalse(Result == EMetaSoundBuilderResult::Succeeded, "Returned failed state when adding 'OneShot' Interface to MetaSound using AddInterface Builder API call");
+
+	Builder.FindGraphOutputNode(SourceOneShotInterface::Outputs::OnFinished, Result);
+	AddErrorIfFalse(Result == EMetaSoundBuilderResult::Succeeded, "Failed to add 'OnFinished' output to MetaSound using AddInterface Builder API call");
+
+	bool bIsDeclared = Builder.InterfaceIsDeclared(SourceOneShotInterface::GetVersion().Name);
+	AddErrorIfFalse(bIsDeclared, "'OneShot' Interface added but is not member of declaration list on MetaSound asset.");
+
+	Builder.RemoveInterface(SourceOneShotInterface::GetVersion().Name, Result);
+	AddErrorIfFalse(Result == EMetaSoundBuilderResult::Succeeded, "Returned failed state when removing 'OneShot' Interface from MetaSound using RemoveInterface Builder API call");
+
+	Builder.FindGraphOutputNode(SourceOneShotInterface::Outputs::OnFinished, Result);
+	AddErrorIfFalse(Result == EMetaSoundBuilderResult::Failed, "Failed to remove 'OnFinished' output to MetaSound using RemoveInterface Builder API call");
+
+	bIsDeclared = Builder.InterfaceIsDeclared(SourceOneShotInterface::GetVersion().Name);
+	AddErrorIfFalse(!bIsDeclared, "'OneShot' Interface removed but remains member of declaration list on MetaSound asset.");
+
+
+	// Test input mutation with attenuation interface
+	Builder.AddInterface(AttenuationInterface::Name, Result);
+	AddErrorIfFalse(Result == EMetaSoundBuilderResult::Succeeded, "Returned failed state when adding 'Attenuation' Interface to MetaSound using AddInterface Builder API call");
+
+	Builder.FindGraphInputNode(AttenuationInterface::Inputs::Distance, Result);
+	AddErrorIfFalse(Result == EMetaSoundBuilderResult::Succeeded, "Failed to add 'Distance' input to MetaSound using AddInterface Builder API call");
+
+	bIsDeclared = Builder.InterfaceIsDeclared(AttenuationInterface::Name);
+	AddErrorIfFalse(bIsDeclared, "'Attenuation' Interface added but is not member of declaration list on MetaSound asset.");
+
+	Builder.RemoveInterface(AttenuationInterface::Name, Result);
+	AddErrorIfFalse(Result == EMetaSoundBuilderResult::Succeeded, "Returned failed state when removing 'Attenuation' Interface from MetaSound using RemoveInterface Builder API call");
+
+	Builder.FindGraphInputNode(AttenuationInterface::Inputs::Distance, Result);
+	AddErrorIfFalse(Result == EMetaSoundBuilderResult::Failed, "Failed to remove 'Distance' input to MetaSound using RemoveInterface Builder API call");
+
+	bIsDeclared = Builder.InterfaceIsDeclared(AttenuationInterface::Name);
+	AddErrorIfFalse(!bIsDeclared, "'Attenuation' Interface removed but remains member of declaration list on MetaSound asset.");
+
+	return true;
+}
 #endif // WITH_DEV_AUTOMATION_TESTS
diff --git a/Engine/Plugins/Runtime/Metasound/Source/MetasoundFrontend/Private/MetasoundFrontendDocumentBuilder.cpp b/Engine/Plugins/Runtime/Metasound/Source/MetasoundFrontend/Private/MetasoundFrontendDocumentBuilder.cpp
--- a/Engine/Plugins/Runtime/Metasound/Source/MetasoundFrontend/Private/MetasoundFrontendDocumentBuilder.cpp
+++ b/Engine/Plugins/Runtime/Metasound/Source/MetasoundFrontend/Private/MetasoundFrontendDocumentBuilder.cpp
@@ -906,7 +906,7 @@ bool FMetaSoundFrontendDocumentBuilder::AddInterface(FName InterfaceName)
 			TArray<FMetasoundFrontendInterface> InterfacesToAdd;
 			InterfacesToAdd.Add(Entry->GetInterface());
 			FModifyInterfaceOptions Options({ }, MoveTemp(InterfacesToAdd));
-			ModifyInterfaces(MoveTemp(Options));
+			return ModifyInterfaces(MoveTemp(Options));
 		}
 	}
 
