Fix for disconnecting MetaSound node using builder API does not update output

diff --git a/Engine/Plugins/Runtime/Metasound/Source/MetasoundEngine/Private/MetasoundSource.cpp b/Engine/Plugins/Runtime/Metasound/Source/MetasoundEngine/Private/MetasoundSource.cpp
--- a/Engine/Plugins/Runtime/Metasound/Source/MetasoundEngine/Private/MetasoundSource.cpp
+++ b/Engine/Plugins/Runtime/Metasound/Source/MetasoundEngine/Private/MetasoundSource.cpp
@@ -583,34 +583,78 @@ ISoundGeneratorPtr UMetaSoundSource::CreateSoundGenerator(const FSoundGeneratorI
 	FParameterRouter& Router = GetParameterRouter();
 	TSharedPtr<TSpscQueue<FMetaSoundParameterTransmitter::FParameter>> DataChannel = Router.FindOrCreateDataChannelForReader(InParams.AudioDeviceID, InParams.InstanceID);
 
-	TSharedPtr<const IGraph, ESPMode::ThreadSafe> MetasoundGraph = GetRuntimeData().Graph;
-	if (!MetasoundGraph.IsValid())
-	{
-		return ISoundGeneratorPtr(nullptr);
-	}
-
 	FOperatorBuilderSettings BuilderSettings = FOperatorBuilderSettings::GetDefaultSettings();
-
 	// Graph analyzer currently only enabled for preview sounds (but can theoretically be supported for all sounds)
 	BuilderSettings.bPopulateInternalDataReferences = InParams.bIsPreviewSound;
 
 	constexpr bool bBuildSynchronous = false;
-	FMetasoundGeneratorInitParams InitParams =
+
+	const bool bIsDynamic = DynamicTransactor.IsValid();
+	TSharedPtr<FMetasoundGenerator> Generator;
+
+	if (bIsDynamic)
 	{
-		InSettings,
-		MoveTemp(BuilderSettings),
-		MetasoundGraph,
-		Environment,
-		GetName(),
-		GetOutputAudioChannelOrder(),
-		MoveTemp(InDefaultParameters),
-		bBuildSynchronous,
-		DataChannel,
-		DynamicTransactor
-	};
+		// In order to ensure synchronization and avoid race conditions the current state
+		// of the graph is copied and transform queue created here. This ensures that:
+		//
+		// 1. Modifications to the underlying FGraph in the FDynamicOperatorTransactor can continue 
+		// while the generator is being constructed on an async task.  If this were not ensured, 
+		// a race condition would be introduced wherein the FGraph could be manipulated while the
+		// graph is being read while building the generator.
+		//
+		// 2. The state of the FGraph and TransformQueue are synchronized so that any additional
+		// changes applied to the FDynamicOperatorTransactor will be placed in the TransformQueue.
+		// The dynamic operator & generator will then consume these transforms after it has finished 
+		// being built.
+
+		FMetasoundDynamicGraphGeneratorInitParams InitParams
+		{
+			{
+				InSettings,
+				MoveTemp(BuilderSettings),
+				MakeShared<FGraph>(DynamicTransactor->GetGraph()), // Make a copy of the graph.
+				Environment,
+				GetName(),
+				GetOutputAudioChannelOrder(),
+				MoveTemp(InDefaultParameters),
+				bBuildSynchronous,
+				DataChannel
+			},
+			DynamicTransactor->CreateTransformQueue(InSettings, Environment) // Create transaction queue
+		};
+		TSharedPtr<FMetasoundDynamicGraphGenerator> DynamicGenerator = MakeShared<FMetasoundDynamicGraphGenerator>(InSettings);
+		DynamicGenerator->Init(MoveTemp(InitParams));
 
-	TSharedPtr<FMetasoundGenerator> Generator = MakeShared<FMetasoundGenerator>(MoveTemp(InitParams));
-	TrackGenerator(InParams.AudioComponentId, Generator);
+		Generator = MoveTemp(DynamicGenerator);
+	}
+	else
+	{
+		TSharedPtr<const IGraph, ESPMode::ThreadSafe> MetasoundGraph = GetRuntimeData().Graph;
+		if (!MetasoundGraph.IsValid())
+		{
+			return ISoundGeneratorPtr(nullptr);
+		}
+
+		FMetasoundGeneratorInitParams InitParams
+		{
+			InSettings,
+			MoveTemp(BuilderSettings),
+			MetasoundGraph,
+			Environment,
+			GetName(),
+			GetOutputAudioChannelOrder(),
+			MoveTemp(InDefaultParameters),
+			bBuildSynchronous,
+			DataChannel
+		};
+
+		Generator = MakeShared<FMetasoundConstGraphGenerator>(MoveTemp(InitParams));
+	}
+
+	if (Generator.IsValid())
+	{
+		TrackGenerator(InParams.AudioComponentId, Generator);
+	}
 
 	return ISoundGeneratorPtr(Generator);
 }
diff --git a/Engine/Plugins/Runtime/Metasound/Source/MetasoundEngineTest/Private/NodeTestGraphBuilder.cpp b/Engine/Plugins/Runtime/Metasound/Source/MetasoundEngineTest/Private/NodeTestGraphBuilder.cpp
--- a/Engine/Plugins/Runtime/Metasound/Source/MetasoundEngineTest/Private/NodeTestGraphBuilder.cpp
+++ b/Engine/Plugins/Runtime/Metasound/Source/MetasoundEngineTest/Private/NodeTestGraphBuilder.cpp
@@ -115,7 +115,7 @@ namespace Metasound::Test
 				true // bBuildSynchronous, so we don't have to do latent tasks
 			};
 
-			return MakeUnique<FMetasoundGenerator>(MoveTemp(GeneratorInitParams));
+			return MakeUnique<FMetasoundConstGraphGenerator>(MoveTemp(GeneratorInitParams));
 		}
 
 		return nullptr;
diff --git a/Engine/Plugins/Runtime/Metasound/Source/MetasoundGenerator/Private/MetasoundGenerator.cpp b/Engine/Plugins/Runtime/Metasound/Source/MetasoundGenerator/Private/MetasoundGenerator.cpp
--- a/Engine/Plugins/Runtime/Metasound/Source/MetasoundGenerator/Private/MetasoundGenerator.cpp
+++ b/Engine/Plugins/Runtime/Metasound/Source/MetasoundGenerator/Private/MetasoundGenerator.cpp
@@ -2,6 +2,7 @@
 #include "MetasoundGenerator.h"
 
 #include "AudioParameter.h"
+#include "MetasoundDynamicOperatorTransactor.h"
 #include "MetasoundFrontendDataTypeRegistry.h"
 #include "MetasoundGeneratorBuilder.h"
 #include "MetasoundGeneratorModuleImpl.h"
@@ -144,6 +145,44 @@
 		};
 
 #endif // if ENABLE_METASOUND_GENERATOR_RENDER_TIMING
+
+#ifndef ENABLE_METASOUND_CONSTANT_OUTPUT_VERTEX_ERROR_LOG
+#define ENABLE_METASOUND_CONSTANT_OUTPUT_VERTEX_ERROR_LOG DO_CHECK
+#endif
+
+#if ENABLE_METASOUND_CONSTANT_OUTPUT_VERTEX_ERROR_LOG
+		// Check whether a vertex name represents an output audio buffer. MetaSound
+		// Generators cannot dynamically update their output audio buffers. This
+		// check is here to inform future developers of this limitation.
+		static const FLazyName ParameterNameComponentUE = "UE";
+		static const FLazyName ParameterNameComponentOutputFormat = "OutputFormat";
+		bool IsAudioOutputVertex(const FVertexName& InVertexName)
+		{
+			FName Namespace;
+			FName Remaining;
+			Audio::FParameterPath::SplitName(InVertexName, Namespace, Remaining);
+			if (ParameterNameComponentUE == Namespace)
+			{
+				FName Group;
+				FName Name;
+				Audio::FParameterPath::SplitName(Remaining, Group, Name);
+				return (ParameterNameComponentOutputFormat == Name);
+			}
+			return false;
+		}
+
+		void LogErrorIfIsAudioVertex(const FVertexName& InVertexName, const TCHAR* Message)
+		{
+			if (IsAudioOutputVertex(InVertexName))
+			{
+				UE_LOG(LogMetaSound, Error, TEXT("Cannot modify vertex %s. %s"), *InVertexName.ToString(), Message);
+			}
+		};
+#else
+		void LogErrorIfIsAudioVertex(const FVertexName& InVertexName, const TCHAR* Message)
+		{
+		}
+#endif
 	}
 }
 
@@ -203,164 +242,47 @@
 		}
 	}
 
-	void FMetasoundGeneratorInitParams::Release()
+	void FMetasoundGeneratorInitParams::Reset(FMetasoundGeneratorInitParams& InParams)
 	{
-		Graph.Reset();
-		Environment = {};
-		MetaSoundName = {};
-		AudioOutputNames = {};
-		DefaultParameters = {};
-		DataChannel.Reset();
-		DynamicOperatorTransactor.Reset();
+		InParams.Graph.Reset();
+		InParams.Environment = {};
+		InParams.MetaSoundName = {};
+		InParams.AudioOutputNames = {};
+		InParams.DefaultParameters = {};
+		InParams.DataChannel.Reset();
 	}
 
-	FMetasoundGenerator::FMetasoundGenerator(FMetasoundGeneratorInitParams&& InParams)
-		: MetasoundName(InParams.MetaSoundName)
-		, OperatorSettings(InParams.OperatorSettings)
+	void FMetasoundDynamicGraphGeneratorInitParams::Reset(FMetasoundDynamicGraphGeneratorInitParams& InParams)
+	{
+		FMetasoundGeneratorInitParams::Reset(InParams);
+		InParams.TransformQueue.Reset();
+	}
+
+	FMetasoundGenerator::FMetasoundGenerator(const FOperatorSettings& InOperatorSettings)
+		: OperatorSettings(InOperatorSettings)
 		, bIsFinishTriggered(false)
 		, bIsFinished(false)
-		, NumChannels(0)
-		, NumFramesPerExecute(0)
-		, NumSamplesPerExecute(0)
-		, OnPlayTriggerRef(FTriggerWriteRef::CreateNew(InParams.OperatorSettings))
-		, OnFinishedTriggerRef(FTriggerWriteRef::CreateNew(InParams.OperatorSettings))
 		, bPendingGraphTrigger(true)
 		, bIsNewGraphPending(false)
 		, bIsWaitingForFirstGraph(true)
-		, ParameterQueue(MoveTemp(InParams.DataChannel))
+		, NumChannels(0)
+		, NumFramesPerExecute(InOperatorSettings.GetNumFramesPerBlock())
+		, NumSamplesPerExecute(0)
+		, OnFinishedTriggerRef(FTriggerWriteRef::CreateNew(InOperatorSettings))
 #if ENABLE_METASOUND_GENERATOR_RENDER_TIMING
-		, RenderTimer(MakeUnique<MetasoundGeneratorPrivate::FRenderTimer>(InParams.OperatorSettings, 1. /* AnalysisPeriod */))
+		, RenderTimer(MakeUnique<MetasoundGeneratorPrivate::FRenderTimer>(InOperatorSettings, 1. /* AnalysisPeriod */))
 #endif // if ENABLE_METASOUND_GENERATOR_RENDER_TIMING
 
 	{
-		NumChannels = InParams.AudioOutputNames.Num();
-		NumFramesPerExecute = InParams.OperatorSettings.GetNumFramesPerBlock();
-		NumSamplesPerExecute = NumChannels * NumFramesPerExecute;
-
-		const bool bIsDynamic = InParams.DynamicOperatorTransactor.IsValid();
-
-		// Create the routing for parameter packs
-PRAGMA_DISABLE_DEPRECATION_WARNINGS
-		ParameterPackSendAddress = UMetasoundParameterPack::CreateSendAddressFromEnvironment(InParams.Environment);
-PRAGMA_ENABLE_DEPRECATION_WARNINGS
-		ParameterPackReceiver = FDataTransmissionCenter::Get().RegisterNewReceiver<FMetasoundParameterStorageWrapper>(ParameterPackSendAddress, FReceiverInitParams{ InParams.OperatorSettings });
-
-		// attempt to use operator cache instead of building a new operator.
-		bool bDidUseCachedOperator = false;
-		const bool bIsOperatorCacheEnabled = ConsoleVariables::bEnableExperimentalOneShotOperatorCache || ConsoleVariables::bEnableExperimentalOperatorCache;
-		// Dynamic operators cannot use the operator cache because they can change their internal structure. 
-		// The operator cache assumes that the operator is unchanged from it's original structure. 
-		if (bIsOperatorCacheEnabled && !bIsDynamic)
-		{
-			bUseOperatorCache = ConsoleVariables::bEnableExperimentalOperatorCache || MetasoundGeneratorPrivate::HasOneShotInterface(InParams.Graph->GetVertexInterface());
-
-			if (bUseOperatorCache)
-			{
-				bDidUseCachedOperator = TryUseCachedOperator(InParams, true /* bTriggerGenerator */);
-			}
-		}
-
-		if (!bDidUseCachedOperator)
-		{
-			BuilderTask = MakeUnique<FAsyncTask<FAsyncMetaSoundBuilder>>(this, MoveTemp(InParams), true /* bTriggerGenerator */);
-			
-			if (!InParams.bBuildSynchronous && Metasound::ConsoleVariables::bEnableAsyncMetaSoundGeneratorBuilder)
-			{
-				// Build operator asynchronously
-				BuilderTask->StartBackgroundTask(GBackgroundPriorityThreadPool);
-			}
-			else
-			{
-				// Build operator synchronously
-				BuilderTask->StartSynchronousTask();
-				BuilderTask = nullptr;
-				UpdateGraphIfPending();
-			}
-		}
 	}
 
 	FMetasoundGenerator::~FMetasoundGenerator()
 	{
-		if (BuilderTask.IsValid())
-		{
-			BuilderTask->EnsureCompletion();
-			BuilderTask = nullptr;
-		}
-
-		if (bUseOperatorCache)
-		{
-			ReleaseOperatorToCache();
-		}
-
 		// Remove routing for parameter packs
 		ParameterPackReceiver.Reset();
 		FDataTransmissionCenter::Get().UnregisterDataChannelIfUnconnected(ParameterPackSendAddress);
 	}
 
-	bool FMetasoundGenerator::TryUseCachedOperator(FMetasoundGeneratorInitParams& InInitParams, bool bInTriggerGraph)
-	{
-		using namespace MetasoundGeneratorPrivate;
-		METASOUND_TRACE_CPUPROFILER_EVENT_SCOPE(FMetasoundGenerator::TryUseCachedOperator);
-
-		FMetasoundGeneratorModule& Module = FModuleManager::GetModuleChecked<FMetasoundGeneratorModule>("MetasoundGenerator");
-		TSharedPtr<FOperatorCache> OperatorCache = Module.GetOperatorCache();
-		if (OperatorCache.IsValid())
-		{
-			// See if the cache has an operator with matching OperatorID
-			OperatorID = InInitParams.Graph->GetInstanceID();
-			FOperatorAndInputs GraphOperatorAndInputs = OperatorCache->ClaimCachedOperator(OperatorID);
-			if (GraphOperatorAndInputs.Operator.IsValid())
-			{
-				UE_LOG(LogMetasoundGenerator, VeryVerbose, TEXT("Using cached operator %s for MetaSound %s"), *LexToString(OperatorID), *InInitParams.MetaSoundName); 
-
-				// Apply and default inputs to the operator.
-				GeneratorBuilder::ApplyAudioParameters(OperatorSettings, MoveTemp(InInitParams.DefaultParameters), GraphOperatorAndInputs.Inputs);
-
-				// Reset operator internal state before playing it.
-				if (IOperator::FResetFunction Reset = GraphOperatorAndInputs.Operator->GetResetFunction())
-				{
-					IOperator::FResetParams ResetParams {OperatorSettings, InInitParams.Environment};
-					Reset(GraphOperatorAndInputs.Operator.Get(), ResetParams);
-				}
-				
-				// Create data needed for MetaSound Generator
-				FMetasoundGeneratorData Data = GeneratorBuilder::BuildGeneratorData(InInitParams, MoveTemp(GraphOperatorAndInputs), nullptr);
-
-				SetGraph(MakeUnique<FMetasoundGeneratorData>(MoveTemp(Data)), bInTriggerGraph);
-				
-				return true;
-			}
-		}
-
-		return false;
-	}
-
-	void FMetasoundGenerator::ReleaseOperatorToCache()
-	{
-		using namespace MetasoundGeneratorPrivate;
-		METASOUND_TRACE_CPUPROFILER_EVENT_SCOPE(FMetasoundGenerator::ReleaseOperatorToCache);
-
-		if (OperatorID.IsValid())
-		{
-			FMetasoundGeneratorModule& Module = FModuleManager::GetModuleChecked<FMetasoundGeneratorModule>("MetasoundGenerator");
-			TSharedPtr<FOperatorCache> OperatorCache = Module.GetOperatorCache();
-			if (OperatorCache.IsValid())
-			{
-				TUniquePtr<IOperator> GraphOperator = RootExecuter.ReleaseOperator();
-
-				if (GraphOperator.IsValid())
-				{
-					// Release graph operator and input data to the cache
-					UE_LOG(LogMetasoundGenerator, VeryVerbose, TEXT("Caching operator %s"), *LexToString(OperatorID));
-					OperatorCache->AddOperatorToCache(OperatorID, MoveTemp(GraphOperator), MoveTemp(VertexInterfaceData.GetInputs()));
-				}
-			}
-		}
-
-		// Clear out any internal references to graph data
-		ClearGraph();
-	}
-
 	FDelegateHandle FMetasoundGenerator::AddGraphSetCallback(FOnSetGraph::FDelegate&& Delegate)
 	{
 		FScopeLock SetPendingGraphLock(&PendingGraphMutex);
@@ -372,6 +294,24 @@
 		return OnSetGraph.Add(Delegate);
 	}
 
+	void FMetasoundGenerator::InitBase(const FMetasoundGeneratorInitParams& InInitParams)
+	{
+		MetasoundName = InInitParams.MetaSoundName;
+		NumChannels = InInitParams.AudioOutputNames.Num();
+		NumSamplesPerExecute = NumChannels * NumFramesPerExecute;
+
+		// Data channels
+		ParameterQueue = InInitParams.DataChannel;
+
+		// Create the routing for parameter packs
+PRAGMA_DISABLE_DEPRECATION_WARNINGS
+		ParameterPackSendAddress = UMetasoundParameterPack::CreateSendAddressFromEnvironment(InInitParams.Environment);
+PRAGMA_ENABLE_DEPRECATION_WARNINGS
+		ParameterPackReceiver = FDataTransmissionCenter::Get().RegisterNewReceiver<FMetasoundParameterStorageWrapper>(ParameterPackSendAddress, FReceiverInitParams{ OperatorSettings });
+	}
+
+
+
 	bool FMetasoundGenerator::RemoveGraphSetCallback(const FDelegateHandle& Handle)
 	{
 		return OnSetGraph.Remove(Handle);
@@ -451,7 +391,6 @@
 			}
 		}
 
-		OnPlayTriggerRef = InData->TriggerOnPlayRef;
 		OnFinishedTriggerRef = InData->TriggerOnFinishRef;
 
 		// The graph operator and graph audio output contain all the values needed
@@ -476,19 +415,29 @@
 
 		if (bTriggerGraph)
 		{
-			OnPlayTriggerRef->TriggerFrame(0);
+			InData->TriggerOnPlayRef->TriggerFrame(0);
 		}
 
 		bIsWaitingForFirstGraph = false;
 		OnSetGraph.Broadcast();
 	}
 
+	TUniquePtr<IOperator> FMetasoundGenerator::ReleaseGraphOperator()
+	{
+		return RootExecuter.ReleaseOperator();
+	}
+
+	FInputVertexInterfaceData FMetasoundGenerator::ReleaseInputVertexData()
+	{
+		return MoveTemp(VertexInterfaceData.GetInputs());
+	}
+
 	void FMetasoundGenerator::ClearGraph()
 	{
 		RootExecuter.ReleaseOperator();
 		VertexInterfaceData = FVertexInterfaceData();
 		GraphOutputAudio.Reset();
-		OnFinishedTriggerRef = OnPlayTriggerRef = TDataWriteReference<FTrigger>::CreateNew(OperatorSettings);
+		OnFinishedTriggerRef = TDataWriteReference<FTrigger>::CreateNew(OperatorSettings);
 		GraphAnalyzer.Reset();
 		ParameterSetters.Reset();
 		ParameterPackSetters.Reset();
@@ -806,4 +755,239 @@
 			ProcessPack(QueuedParameterPack.Get());
 		}
 	}
+
+	FMetasoundConstGraphGenerator::FMetasoundConstGraphGenerator(FMetasoundGeneratorInitParams&& InInitParams)
+	: FMetasoundGenerator(InInitParams.OperatorSettings)
+	{
+		Init(MoveTemp(InInitParams));
+	}
+
+	FMetasoundConstGraphGenerator::FMetasoundConstGraphGenerator(const FOperatorSettings& InOperatorSettings)
+	: FMetasoundGenerator(InOperatorSettings)
+	{
+	}
+
+	void FMetasoundConstGraphGenerator::Init(FMetasoundGeneratorInitParams&& InParams)
+	{
+		InitBase(InParams);
+		// attempt to use operator cache instead of building a new operator.
+		bool bDidUseCachedOperator = false;
+		const bool bIsOperatorCacheEnabled = ConsoleVariables::bEnableExperimentalOneShotOperatorCache || ConsoleVariables::bEnableExperimentalOperatorCache;
+		// Dynamic operators cannot use the operator cache because they can change their internal structure. 
+		// The operator cache assumes that the operator is unchanged from it's original structure. 
+		if (bIsOperatorCacheEnabled)
+		{
+			bUseOperatorCache = ConsoleVariables::bEnableExperimentalOperatorCache || MetasoundGeneratorPrivate::HasOneShotInterface(InParams.Graph->GetVertexInterface());
+
+			if (bUseOperatorCache)
+			{
+				bDidUseCachedOperator = TryUseCachedOperator(InParams, true /* bTriggerGenerator */);
+			}
+		}
+
+		if (!bDidUseCachedOperator)
+		{
+			BuildGraph(MoveTemp(InParams));
+		}
+	}
+
+	FMetasoundConstGraphGenerator::~FMetasoundConstGraphGenerator()
+	{
+		if (BuilderTask.IsValid())
+		{
+			BuilderTask->EnsureCompletion();
+			BuilderTask = nullptr;
+		}
+
+		if (bUseOperatorCache)
+		{
+			ReleaseOperatorToCache();
+		}
+	}
+
+	bool FMetasoundConstGraphGenerator::TryUseCachedOperator(FMetasoundGeneratorInitParams& InInitParams, bool bInTriggerGraph)
+	{
+		using namespace MetasoundGeneratorPrivate;
+		METASOUND_TRACE_CPUPROFILER_EVENT_SCOPE(FMetasoundConstGraphGenerator::TryUseCachedOperator);
+
+		FMetasoundGeneratorModule& Module = FModuleManager::GetModuleChecked<FMetasoundGeneratorModule>("MetasoundGenerator");
+		TSharedPtr<FOperatorCache> OperatorCache = Module.GetOperatorCache();
+		if (OperatorCache.IsValid())
+		{
+			// See if the cache has an operator with matching OperatorID
+			OperatorID = InInitParams.Graph->GetInstanceID();
+			FOperatorAndInputs GraphOperatorAndInputs = OperatorCache->ClaimCachedOperator(OperatorID);
+			if (GraphOperatorAndInputs.Operator.IsValid())
+			{
+				UE_LOG(LogMetasoundGenerator, VeryVerbose, TEXT("Using cached operator %s for MetaSound %s"), *LexToString(OperatorID), *InInitParams.MetaSoundName); 
+
+				// Apply and default inputs to the operator.
+				GeneratorBuilder::ApplyAudioParameters(OperatorSettings, MoveTemp(InInitParams.DefaultParameters), GraphOperatorAndInputs.Inputs);
+
+				// Reset operator internal state before playing it.
+				if (IOperator::FResetFunction Reset = GraphOperatorAndInputs.Operator->GetResetFunction())
+				{
+					IOperator::FResetParams ResetParams {OperatorSettings, InInitParams.Environment};
+					Reset(GraphOperatorAndInputs.Operator.Get(), ResetParams);
+				}
+				
+				// Create data needed for MetaSound Generator
+				FMetasoundGeneratorData Data = GeneratorBuilder::BuildGeneratorData(OperatorSettings, InInitParams, MoveTemp(GraphOperatorAndInputs), nullptr);
+
+				SetGraph(MakeUnique<FMetasoundGeneratorData>(MoveTemp(Data)), bInTriggerGraph);
+				
+				return true;
+			}
+		}
+
+		return false;
+	}
+
+	void FMetasoundConstGraphGenerator::ReleaseOperatorToCache()
+	{
+		using namespace MetasoundGeneratorPrivate;
+		METASOUND_TRACE_CPUPROFILER_EVENT_SCOPE(FMetasoundConstGraphGenerator::ReleaseOperatorToCache);
+
+		if (OperatorID.IsValid())
+		{
+			FMetasoundGeneratorModule& Module = FModuleManager::GetModuleChecked<FMetasoundGeneratorModule>("MetasoundGenerator");
+			TSharedPtr<FOperatorCache> OperatorCache = Module.GetOperatorCache();
+			if (OperatorCache.IsValid())
+			{
+				TUniquePtr<IOperator> GraphOperator = ReleaseGraphOperator();
+
+				if (GraphOperator.IsValid())
+				{
+					// Release graph operator and input data to the cache
+					UE_LOG(LogMetasoundGenerator, VeryVerbose, TEXT("Caching operator %s"), *LexToString(OperatorID));
+					OperatorCache->AddOperatorToCache(OperatorID, MoveTemp(GraphOperator), ReleaseInputVertexData());
+				}
+			}
+		}
+
+		// Clear out any internal references to graph data
+		ClearGraph();
+	}
+
+	void FMetasoundConstGraphGenerator::BuildGraph(FMetasoundGeneratorInitParams&& InInitParams)
+	{
+		const bool bBuildAsync = !InInitParams.bBuildSynchronous && Metasound::ConsoleVariables::bEnableAsyncMetaSoundGeneratorBuilder;
+
+		BuilderTask = MakeUnique<FAsyncTask<FAsyncMetaSoundBuilder>>(this, MoveTemp(InInitParams), true /* bTriggerGenerator */);
+		
+		if (bBuildAsync)
+		{
+			// Build operator asynchronously
+			BuilderTask->StartBackgroundTask(GBackgroundPriorityThreadPool);
+		}
+		else
+		{
+			// Build operator synchronously
+			BuilderTask->StartSynchronousTask();
+			BuilderTask = nullptr;
+			UpdateGraphIfPending();
+		}
+	}
+
+	FMetasoundDynamicGraphGenerator::FMetasoundDynamicGraphGenerator(const FOperatorSettings& InOperatorSettings)
+	: FMetasoundGenerator(InOperatorSettings)
+	{
+	}
+
+	FMetasoundDynamicGraphGenerator::~FMetasoundDynamicGraphGenerator()
+	{
+		if (BuilderTask.IsValid())
+		{
+			BuilderTask->EnsureCompletion();
+			BuilderTask = nullptr;
+		}
+	}
+
+	void FMetasoundDynamicGraphGenerator::Init(FMetasoundDynamicGraphGeneratorInitParams&& InParams)
+	{
+		InitBase(InParams);
+		AudioOutputNames = InParams.AudioOutputNames;
+		BuildGraph(MoveTemp(InParams));
+	}
+
+	void FMetasoundDynamicGraphGenerator::BuildGraph(FMetasoundDynamicGraphGeneratorInitParams&& InParams)
+	{
+		const bool bBuildAsync = !InParams.bBuildSynchronous && Metasound::ConsoleVariables::bEnableAsyncMetaSoundGeneratorBuilder;
+
+		BuilderTask = MakeUnique<FAsyncTask<FAsyncDynamicMetaSoundBuilder>>(this, MoveTemp(InParams), true /* bTriggerGenerator */);
+
+		if (bBuildAsync)
+		{
+			// Build operator asynchronously
+			BuilderTask->StartBackgroundTask(GBackgroundPriorityThreadPool);
+		}
+		else
+		{
+			// Build operator synchronously
+			BuilderTask->StartSynchronousTask();
+			BuilderTask = nullptr;
+			UpdateGraphIfPending();
+		}
+	}
+
+	void FMetasoundDynamicGraphGenerator::OnInputAdded(const FVertexName& InVertexName, const FInputVertexInterfaceData& InGraphInputs)
+	{
+		// Update vertex data and set parameters. 
+		VertexInterfaceData.GetInputs() = InGraphInputs;
+		GeneratorBuilder::AddParameterSetterIfWritable(InVertexName, VertexInterfaceData.GetInputs(), ParameterSetters, ParameterPackSetters);
+	}
+
+	void FMetasoundDynamicGraphGenerator::OnInputRemoved(const FVertexName& InVertexName, const FInputVertexInterfaceData& InGraphInputs)
+	{
+		VertexInterfaceData.GetInputs() = InGraphInputs;
+		ParameterSetters.Remove(InVertexName);
+		ParameterPackSetters.Remove(InVertexName);
+	}
+
+	void FMetasoundDynamicGraphGenerator::OnOutputAdded(const FVertexName& InVertexName, const FOutputVertexInterfaceData& InGraphOutputs)
+	{
+		MetasoundGeneratorPrivate::LogErrorIfIsAudioVertex(InVertexName, TEXT("The MetaSound Generator cannot dynamically add audio outputs"));
+		VertexInterfaceData.GetOutputs() = InGraphOutputs;
+
+		if (Frontend::SourceOneShotInterface::Outputs::OnFinished == InVertexName)
+		{
+			OnFinishedTriggerRef = InGraphOutputs.GetDataReadReference<FTrigger>(InVertexName);
+		}
+	}
+
+	void FMetasoundDynamicGraphGenerator::OnOutputUpdated(const FVertexName& InVertexName, const FOutputVertexInterfaceData& InGraphOutputs)
+	{
+		VertexInterfaceData.GetOutputs() = InGraphOutputs;
+
+		// If it's an output audio buffer, update the internal reference	
+		int32 AudioOutputIndex = AudioOutputNames.IndexOfByKey(InVertexName);
+		if (INDEX_NONE != AudioOutputIndex)
+		{
+			GraphOutputAudio[AudioOutputIndex] = InGraphOutputs.GetDataReadReference<FAudioBuffer>(InVertexName);
+		}
+		// If it's the on-finished trigger, update the internal reference
+		else if (Frontend::SourceOneShotInterface::Outputs::OnFinished == InVertexName)
+		{
+			OnFinishedTriggerRef = InGraphOutputs.GetDataReadReference<FTrigger>(InVertexName);
+		}
+	}
+
+	void FMetasoundDynamicGraphGenerator::OnOutputRemoved(const FVertexName& InVertexName, const FOutputVertexInterfaceData& InGraphOutputs)
+	{
+		MetasoundGeneratorPrivate::LogErrorIfIsAudioVertex(InVertexName, TEXT("The MetaSound Generator cannot dynamically remove audio outputs"));
+
+		VertexInterfaceData.GetOutputs() = InGraphOutputs;
+	}
+
+	TUniquePtr<IOperator> FMetasoundDynamicGraphGenerator::ReleaseGraphOperator()
+	{
+		FMetasoundGenerator::ReleaseGraphOperator();
+
+		// Never allow the dynamic operator to leave this object because it contains 
+		// callbacks pointing to this object instance. If the dynamic operator left
+		// this object, it could attempt to execute member functions on this object
+		// after this object has been destroyed. (X╭╮X) 
+		return nullptr;
+	}
+
 } 
diff --git a/Engine/Plugins/Runtime/Metasound/Source/MetasoundGenerator/Private/MetasoundGeneratorBuilder.cpp b/Engine/Plugins/Runtime/Metasound/Source/MetasoundGenerator/Private/MetasoundGeneratorBuilder.cpp
--- a/Engine/Plugins/Runtime/Metasound/Source/MetasoundGenerator/Private/MetasoundGeneratorBuilder.cpp
+++ b/Engine/Plugins/Runtime/Metasound/Source/MetasoundGenerator/Private/MetasoundGeneratorBuilder.cpp
@@ -3,6 +3,7 @@
 #include "MetasoundGeneratorBuilder.h"
 
 #include "AudioParameter.h"
+#include "MetasoundDynamicOperatorTransactor.h"
 #include "MetasoundFrontendDataTypeRegistry.h"
 #include "MetasoundOperatorBuilder.h"
 #include "MetasoundOperatorInterface.h"
@@ -19,37 +20,31 @@
 
 namespace Metasound
 {
-	FAsyncMetaSoundBuilder::FAsyncMetaSoundBuilder(FMetasoundGenerator* InGenerator, FMetasoundGeneratorInitParams&& InInitParams, bool bInTriggerGenerator)
-		: Generator(InGenerator)
-		, InitParams(MoveTemp(InInitParams))
-		, bTriggerGenerator(bInTriggerGenerator)
+	FAsyncMetaSoundBuilderBase::FAsyncMetaSoundBuilderBase(FMetasoundGenerator* InGenerator, bool bInTriggerGenerator)
+	: Generator(InGenerator)
+	, bTriggerGenerator(bInTriggerGenerator)
 	{
 	}
 
-	void FAsyncMetaSoundBuilder::DoWork()
+	void FAsyncMetaSoundBuilderBase::SetGraphOnGenerator(FOperatorAndInputs&& InOpAndInputs, FBuildResults&& InBuildResults)
 	{
 		using namespace Audio;
 		using namespace Frontend;
 		using namespace MetasoundGeneratorPrivate;
 
-		METASOUND_TRACE_CPUPROFILER_EVENT_SCOPE_TEXT(*FString::Printf(TEXT("AsyncMetaSoundBuilder::DoWork %s"), *InitParams.MetaSoundName));
+		const FMetasoundGeneratorInitParams& InitParams = GetGeneratorParams();
 
-		// Create an instance of the new graph
-		FBuildResults BuildResults;
-		FOperatorAndInputs GraphOperatorAndInputs = GeneratorBuilder::BuildGraphOperator(InitParams, BuildResults);
-		GeneratorBuilder::LogBuildErrors(InitParams.MetaSoundName, BuildResults);
-
-		if (GraphOperatorAndInputs.Operator.IsValid())
+		if (InOpAndInputs.Operator.IsValid())
 		{
 			// Create graph analyzer
 			TUniquePtr<FGraphAnalyzer> GraphAnalyzer;
 			if (InitParams.BuilderSettings.bPopulateInternalDataReferences)
 			{
-				GraphAnalyzer = GeneratorBuilder::BuildGraphAnalyzer(MoveTemp(BuildResults.InternalDataReferences), InitParams.Environment, InitParams.OperatorSettings);
+				GraphAnalyzer = GeneratorBuilder::BuildGraphAnalyzer(MoveTemp(InBuildResults.InternalDataReferences), InitParams.Environment, Generator->OperatorSettings);
 			}
 	
 			// Collect data for generator
-			FMetasoundGeneratorData GeneratorData = GeneratorBuilder::BuildGeneratorData(InitParams, MoveTemp(GraphOperatorAndInputs), MoveTemp(GraphAnalyzer));
+			FMetasoundGeneratorData GeneratorData = GeneratorBuilder::BuildGeneratorData(Generator->OperatorSettings, InitParams, MoveTemp(InOpAndInputs), MoveTemp(GraphAnalyzer));
 
 			Generator->SetPendingGraph(MoveTemp(GeneratorData), bTriggerGenerator);
 		}
@@ -60,8 +55,81 @@ namespace Metasound
 			// Otherwise, generator will continually wait for a new generator.
 			Generator->SetPendingGraphBuildFailed();
 		}
+	}
 
-		InitParams.Release();
+	FAsyncMetaSoundBuilder::FAsyncMetaSoundBuilder(FMetasoundConstGraphGenerator* InGenerator, FMetasoundGeneratorInitParams&& InInitParams, bool bInTriggerGenerator)
+		: FAsyncMetaSoundBuilderBase(InGenerator, bInTriggerGenerator)
+		, InitParams(MoveTemp(InInitParams))
+	{
+	}
+
+	void FAsyncMetaSoundBuilder::DoWork()
+	{
+		using namespace Audio;
+		using namespace Frontend;
+		using namespace MetasoundGeneratorPrivate;
+
+		METASOUND_TRACE_CPUPROFILER_EVENT_SCOPE_TEXT(*FString::Printf(TEXT("AsyncMetaSoundBuilder::DoWork %s"), *InitParams.MetaSoundName));
+
+		// Create an instance of the new graph
+		FBuildResults BuildResults;
+		FOperatorAndInputs GraphOperatorAndInputs = GeneratorBuilder::BuildGraphOperator(Generator->OperatorSettings, InitParams, BuildResults);
+		GeneratorBuilder::LogBuildErrors(InitParams.MetaSoundName, BuildResults);
+
+		SetGraphOnGenerator(MoveTemp(GraphOperatorAndInputs), MoveTemp(BuildResults));
+
+		FMetasoundGeneratorInitParams::Reset(InitParams);
+	}
+
+	const FMetasoundGeneratorInitParams& FAsyncMetaSoundBuilder::GetGeneratorParams() const
+	{
+		return InitParams;
+	}
+
+	FAsyncDynamicMetaSoundBuilder::FAsyncDynamicMetaSoundBuilder(FMetasoundDynamicGraphGenerator* InGenerator, FMetasoundDynamicGraphGeneratorInitParams&& InInitParams, bool bInTriggerGenerator)
+	: FAsyncMetaSoundBuilderBase(InGenerator, bInTriggerGenerator)
+	, DynamicGenerator(InGenerator)
+	, InitParams(MoveTemp(InInitParams))
+	{
+	}
+
+	void FAsyncDynamicMetaSoundBuilder::DoWork()
+	{
+		using namespace Audio;
+		using namespace Frontend;
+		using namespace MetasoundGeneratorPrivate;
+		using namespace DynamicGraph;
+
+		METASOUND_TRACE_CPUPROFILER_EVENT_SCOPE_TEXT(*FString::Printf(TEXT("AsyncMetaSoundBuilder::DoWork %s"), *InitParams.MetaSoundName));
+
+		FDynamicOperatorUpdateCallbacks UpdateCallbacks;
+
+		// Note: The operator which calls these functions is owned by the FMetasoundDynamicGraphGenerator
+		// Capturing the "DynamicGenerator" pointer here is safe as long as the operator's lifetime
+		// is less than or equal to the lifetime of the FMetasoundDynamicGraphGenerator. These
+		// callbacks are triggered during calls to the base class ISoundGenerator::OnGeneratorAudio(...)
+		// and should not be triggered at any other time. 
+		//
+		// The dynamic operator should never be cached.
+		UpdateCallbacks.OnInputAdded = [DynamicGenerator = DynamicGenerator](const FVertexName& InVertexName, const FInputVertexInterfaceData& InData) { DynamicGenerator->OnInputAdded(InVertexName, InData); };
+		UpdateCallbacks.OnInputRemoved = [DynamicGenerator = DynamicGenerator](const FVertexName& InVertexName, const FInputVertexInterfaceData& InData) { DynamicGenerator->OnInputRemoved(InVertexName, InData); };
+		UpdateCallbacks.OnOutputAdded = [DynamicGenerator = DynamicGenerator](const FVertexName& InVertexName, const FOutputVertexInterfaceData& InData) { DynamicGenerator->OnOutputAdded(InVertexName, InData); };
+		UpdateCallbacks.OnOutputUpdated = [DynamicGenerator = DynamicGenerator](const FVertexName& InVertexName, const FOutputVertexInterfaceData& InData) { DynamicGenerator->OnOutputUpdated(InVertexName, InData); };
+		UpdateCallbacks.OnOutputRemoved = [DynamicGenerator = DynamicGenerator](const FVertexName& InVertexName, const FOutputVertexInterfaceData& InData) { DynamicGenerator->OnOutputRemoved(InVertexName, InData); };
+
+		// Create an instance of the new graph
+		FBuildResults BuildResults;
+		FOperatorAndInputs GraphOperatorAndInputs = GeneratorBuilder::BuildDynamicGraphOperator(Generator->OperatorSettings, InitParams, UpdateCallbacks, BuildResults);
+		GeneratorBuilder::LogBuildErrors(InitParams.MetaSoundName, BuildResults);
+
+		SetGraphOnGenerator(MoveTemp(GraphOperatorAndInputs), MoveTemp(BuildResults));
+
+		FMetasoundDynamicGraphGeneratorInitParams::Reset(InitParams);
+	}
+
+	const FMetasoundGeneratorInitParams& FAsyncDynamicMetaSoundBuilder::GetGeneratorParams() const
+	{
+		return InitParams;
 	}
 
 	namespace GeneratorBuilder
@@ -106,9 +174,10 @@ namespace Metasound
 			return MakeUnique<FGraphAnalyzer>(InOperatorSettings, InstanceID, MoveTemp(InInternalDataReferences));
 		}
 
-		FOperatorAndInputs BuildGraphOperator(FMetasoundGeneratorInitParams& InInitParams, FBuildResults& OutBuildResults)
+		FInputVertexInterfaceData BuildGraphOperatorInputs(const FOperatorSettings& InOperatorSettings, FMetasoundGeneratorInitParams& InInitParams)
 		{
-			METASOUND_TRACE_CPUPROFILER_EVENT_SCOPE_TEXT(TEXT("GeneratorBuilder::BuildGraphOperator"));
+			METASOUND_TRACE_CPUPROFILER_EVENT_SCOPE_TEXT(TEXT("GeneratorBuilder::BuildGraphOperatorInputs"));
+
 			using namespace Frontend;
 
 			// Choose which type of data reference access to create depending upon the access of the vertex.
@@ -137,7 +206,7 @@ namespace Metasound
 				{
 					FLiteral Literal = Frontend::ConvertParameterToLiteral(MoveTemp(Parameter));
 
-					TOptional<FAnyDataReference> DataReference = DataRegistry.CreateDataReference(InputVertex->DataTypeName, VertexAccessTypeToDataReferenceAccessType(InputVertex->AccessType), Literal, InInitParams.OperatorSettings);
+					TOptional<FAnyDataReference> DataReference = DataRegistry.CreateDataReference(InputVertex->DataTypeName, VertexAccessTypeToDataReferenceAccessType(InputVertex->AccessType), Literal, InOperatorSettings);
 
 					if (DataReference)
 					{
@@ -162,7 +231,7 @@ namespace Metasound
 				{
 					const FInputDataVertex& InputVertex = Binding.GetVertex();
 					EDataReferenceAccessType AccessType = VertexAccessTypeToDataReferenceAccessType(InputVertex.AccessType);
-					TOptional<FAnyDataReference> DataReference = DataRegistry.CreateDataReference(InputVertex.DataTypeName, VertexAccessTypeToDataReferenceAccessType(InputVertex.AccessType), InputVertex.GetDefaultLiteral(), InInitParams.OperatorSettings);
+					TOptional<FAnyDataReference> DataReference = DataRegistry.CreateDataReference(InputVertex.DataTypeName, VertexAccessTypeToDataReferenceAccessType(InputVertex.AccessType), InputVertex.GetDefaultLiteral(), InOperatorSettings);
 
 					if (DataReference)
 					{
@@ -174,25 +243,52 @@ namespace Metasound
 			// Reset as elements in array have been moved.
 			InInitParams.DefaultParameters.Reset();
 
-			// Create an instance of the new graph
-			FBuildGraphOperatorParams BuildParams { *InInitParams.Graph, InInitParams.OperatorSettings, InputData, InInitParams.Environment };
+			return MoveTemp(InputData);
+		}
+
+		FOperatorAndInputs BuildGraphOperator(const FOperatorSettings& InOperatorSettings, FMetasoundGeneratorInitParams& InInitParams, FBuildResults& OutBuildResults)
+		{
+			METASOUND_TRACE_CPUPROFILER_EVENT_SCOPE_TEXT(TEXT("GeneratorBuilder::BuildGraphOperator"));
+
+			// Create inputs to graph operator
 			FOperatorAndInputs OpAndInputs;
+			OpAndInputs.Inputs = BuildGraphOperatorInputs(InOperatorSettings, InInitParams);
+
+			// Create an instance of the new graph operator
+			FBuildGraphOperatorParams BuildParams { *InInitParams.Graph, InOperatorSettings, OpAndInputs.Inputs, InInitParams.Environment };
 			FOperatorBuilder Builder(InInitParams.BuilderSettings);
+			OpAndInputs.Operator = Builder.BuildGraphOperator(BuildParams, OutBuildResults);
 
-			if (InInitParams.DynamicOperatorTransactor.IsValid())
-			{
-				OpAndInputs.Operator = Builder.BuildDynamicGraphOperator(BuildParams, *InInitParams.DynamicOperatorTransactor, OutBuildResults);
-			}
-			else
-			{
-				OpAndInputs.Operator = Builder.BuildGraphOperator(BuildParams, OutBuildResults);
-			}
-			OpAndInputs.Inputs = InputData;
+			return OpAndInputs;
+		}
+
+		FOperatorAndInputs BuildDynamicGraphOperator(const FOperatorSettings& InOperatorSettings, FMetasoundDynamicGraphGeneratorInitParams& InInitParams, const DynamicGraph::FDynamicOperatorUpdateCallbacks& InCallbacks, FBuildResults& OutBuildResults)
+		{
+			METASOUND_TRACE_CPUPROFILER_EVENT_SCOPE_TEXT(TEXT("GeneratorBuilder::BuildDynamicGraphOperator"));
+
+			FOperatorAndInputs OpAndInputs;
+			OpAndInputs.Inputs = BuildGraphOperatorInputs(InOperatorSettings, InInitParams);
+
+			// Create an instance of the new graph
+			FBuildDynamicGraphOperatorParams BuildParams 
+			{ 
+				{ 
+					*InInitParams.Graph, 
+					InOperatorSettings, 
+					OpAndInputs.Inputs, 
+					InInitParams.Environment
+				},
+				MoveTemp(InInitParams.TransformQueue),
+				InCallbacks
+			};
+			FOperatorBuilder Builder(InInitParams.BuilderSettings);
+
+			OpAndInputs.Operator = Builder.BuildDynamicGraphOperator(BuildParams, OutBuildResults);
 
 			return OpAndInputs;
 		}
 
-		MetasoundGeneratorPrivate::FMetasoundGeneratorData BuildGeneratorData(const FMetasoundGeneratorInitParams& InInitParams, FOperatorAndInputs&& InGraphOperatorAndInputs, TUniquePtr<Frontend::FGraphAnalyzer> InAnalyzer)
+		MetasoundGeneratorPrivate::FMetasoundGeneratorData BuildGeneratorData(const FOperatorSettings& InOperatorSettings, const FMetasoundGeneratorInitParams& InInitParams, FOperatorAndInputs&& InGraphOperatorAndInputs, TUniquePtr<Frontend::FGraphAnalyzer> InAnalyzer)
 		{
 			using namespace Audio;
 			using namespace Frontend;
@@ -211,51 +307,27 @@ namespace Metasound
 			VertexData.GetInputs() = InGraphOperatorAndInputs.Inputs;
 
 			// Get inputs
-			FTriggerWriteRef PlayTrigger = VertexData.GetInputs().GetOrConstructDataWriteReference<FTrigger>(SourceInterface::Inputs::OnPlay, InInitParams.OperatorSettings, false);
+			FTriggerWriteRef PlayTrigger = VertexData.GetInputs().GetOrConstructDataWriteReference<FTrigger>(SourceInterface::Inputs::OnPlay, InOperatorSettings, false);
 
 			// Get outputs
-			TArray<FAudioBufferReadRef> OutputBuffers = FindOutputAudioBuffers(InInitParams.AudioOutputNames, VertexData, InInitParams.OperatorSettings, InInitParams.MetaSoundName);
-			FTriggerReadRef FinishTrigger = TDataReadReferenceFactory<FTrigger>::CreateExplicitArgs(InInitParams.OperatorSettings, false);
+			TArray<FAudioBufferReadRef> OutputBuffers = FindOutputAudioBuffers(InInitParams.AudioOutputNames, VertexData, InOperatorSettings, InInitParams.MetaSoundName);
+			FTriggerReadRef FinishTrigger = TDataReadReferenceFactory<FTrigger>::CreateExplicitArgs(InOperatorSettings, false);
 
 			if (InInitParams.Graph->GetVertexInterface().GetOutputInterface().Contains(SourceOneShotInterface::Outputs::OnFinished))
 			{
-				FinishTrigger = VertexData.GetOutputs().GetOrConstructDataReadReference<FTrigger>(SourceOneShotInterface::Outputs::OnFinished, InInitParams.OperatorSettings, false);
+				FinishTrigger = VertexData.GetOutputs().GetOrConstructDataReadReference<FTrigger>(SourceOneShotInterface::Outputs::OnFinished, InOperatorSettings, false);
 			}
 
-			// Create the parameter setter map so parameter packs can be cracked
-			// open and distributed as appropriate...
-			const IDataTypeRegistry& DataTypeRegistry = IDataTypeRegistry::Get();
+			
+			// Initialize parameter setters for routing inputs to a MetaSound Generator.
 			FParameterSetterSortedMap ParameterSetters;
 			TMap<FName, FParameterPackSetter> ParameterPackSetters;
-			FInputVertexInterfaceData& GraphInputs = VertexData.GetInputs();
-			for (const MetasoundVertexDataPrivate::FInputBinding& Binding : GraphInputs)
-			{
-				// Only assign inputs that are writable. 
-				if (EDataReferenceAccessType::Write == Binding.GetAccessType())
-				{
-					if (const FAnyDataReference* DataRef = Binding.GetDataReference())
-					{
-						const FInputDataVertex& InputVertex = Binding.GetVertex();
-						const Frontend::IParameterAssignmentFunction& PackSetter = DataTypeRegistry.GetRawAssignmentFunction(InputVertex.DataTypeName);
-						if (PackSetter)
-						{
-							FParameterPackSetter ParameterPackSetter(InputVertex.DataTypeName, DataRef->GetRaw(), PackSetter);
-							ParameterPackSetters.Add(InputVertex.VertexName, ParameterPackSetter);
-						}
-
-						Frontend::FLiteralAssignmentFunction LiteralSetter = DataTypeRegistry.GetLiteralAssignmentFunction(InputVertex.DataTypeName);
-						if (LiteralSetter)
-						{
-							ParameterSetters.Add(InputVertex.VertexName, FParameterSetter{LiteralSetter, *DataRef});
-						}
-					}
-				}
-			}
+			InitializeParameterSetters(VertexData.GetInputs(), ParameterSetters, ParameterPackSetters);
 
 			// Set data needed for graph
 			return FMetasoundGeneratorData 
 			{
-				InInitParams.OperatorSettings,
+				InOperatorSettings,
 				MoveTemp(InGraphOperatorAndInputs.Operator),
 				MoveTemp(VertexData),
 				MoveTemp(ParameterSetters),
@@ -288,5 +360,66 @@ namespace Metasound
 				}
 			}
 		}
+
+		void InitializeParameterSetters(FInputVertexInterfaceData& InputData, MetasoundGeneratorPrivate::FParameterSetterSortedMap& OutParamSetters, TMap<FName, MetasoundGeneratorPrivate::FParameterPackSetter>& OutParamPackSetters)
+		{
+			using namespace Frontend;
+
+			// Create the parameter setter map so parameters and parameter packs 
+			// can be assigned appropriately
+			const IDataTypeRegistry& DataTypeRegistry = IDataTypeRegistry::Get();
+
+			for (const MetasoundVertexDataPrivate::FInputBinding& Binding : InputData)
+			{
+				// Only assign inputs that are writable. 
+				if (EDataReferenceAccessType::Write == Binding.GetAccessType())
+				{
+					if (const FAnyDataReference* DataRef = Binding.GetDataReference())
+					{
+						const FInputDataVertex& InputVertex = Binding.GetVertex();
+						AddParameterSetter(DataTypeRegistry, InputVertex, *DataRef, OutParamSetters, OutParamPackSetters);
+					}
+				}
+			}
+		}
+
+		void AddParameterSetterIfWritable(const FVertexName& InVertexName, FInputVertexInterfaceData& InputData, MetasoundGeneratorPrivate::FParameterSetterSortedMap& OutParamSetters, TMap<FName, MetasoundGeneratorPrivate::FParameterPackSetter>& OutParamPackSetters)
+		{
+			using namespace Frontend;
+
+			// Create the parameter setter map so parameters and parameter packs 
+			// can be assigned appropriately
+			const IDataTypeRegistry& DataTypeRegistry = IDataTypeRegistry::Get();
+
+			if (const FAnyDataReference* DataRef = InputData.FindDataReference(InVertexName))
+			{
+				// Only assign inputs that are writable. 
+				if (EDataReferenceAccessType::Write == DataRef->GetAccessType())
+				{
+					const FInputDataVertex& InputVertex = InputData.GetVertex(InVertexName);
+					AddParameterSetter(DataTypeRegistry, InputVertex, *DataRef, OutParamSetters, OutParamPackSetters);
+				}
+			}
+		}
+
+		void AddParameterSetter(const Frontend::IDataTypeRegistry& InDataTypeRegistry, const FInputDataVertex& InVertex, const FAnyDataReference& InDataRef,  MetasoundGeneratorPrivate::FParameterSetterSortedMap& OutParamSetters, TMap<FName, MetasoundGeneratorPrivate::FParameterPackSetter>& OutParamPackSetters)
+		{
+			using namespace MetasoundGeneratorPrivate;
+			
+			checkf(EDataReferenceAccessType::Write == InDataRef.GetAccessType(), TEXT("Only writable inputs can have parameter setters."));
+
+			const Frontend::IParameterAssignmentFunction& PackSetter = InDataTypeRegistry.GetRawAssignmentFunction(InVertex.DataTypeName);
+			if (PackSetter)
+			{
+				FParameterPackSetter ParameterPackSetter(InVertex.DataTypeName, InDataRef.GetRaw(), PackSetter);
+				OutParamPackSetters.Add(InVertex.VertexName, ParameterPackSetter);
+			}
+
+			Frontend::FLiteralAssignmentFunction LiteralSetter = InDataTypeRegistry.GetLiteralAssignmentFunction(InVertex.DataTypeName);
+			if (LiteralSetter)
+			{
+				OutParamSetters.Add(InVertex.VertexName, FParameterSetter{LiteralSetter, InDataRef});
+			}
+		}
 	}
 }
diff --git a/Engine/Plugins/Runtime/Metasound/Source/MetasoundGenerator/Private/MetasoundGeneratorBuilder.h b/Engine/Plugins/Runtime/Metasound/Source/MetasoundGenerator/Private/MetasoundGeneratorBuilder.h
--- a/Engine/Plugins/Runtime/Metasound/Source/MetasoundGenerator/Private/MetasoundGeneratorBuilder.h
+++ b/Engine/Plugins/Runtime/Metasound/Source/MetasoundGenerator/Private/MetasoundGeneratorBuilder.h
@@ -7,6 +7,7 @@
 #include "Async/AsyncWork.h"
 #include "Containers/Map.h"
 #include "MetasoundEnvironment.h"
+#include "MetasoundGenerator.h"
 #include "MetasoundOperatorBuilder.h"
 #include "MetasoundOperatorCache.h"
 #include "MetasoundOperatorSettings.h"
@@ -18,40 +19,107 @@
 namespace Metasound
 {
 	class FMetasoundGenerator;
+	class FMetasoundFixedGraphGenerator;
+	class FMetasoundDynamicGraphGenerator;
 
-	class FAsyncMetaSoundBuilder : public FNonAbandonableTask
+	namespace Frontend
 	{
-	public:
-		FAsyncMetaSoundBuilder(FMetasoundGenerator* InGenerator, FMetasoundGeneratorInitParams&& InInitParams, bool bInTriggerGenerator);
+		class IDataTypeRegistry;
+	}
 
-		~FAsyncMetaSoundBuilder() = default;
+	namespace DynamicGraph
+	{
+		struct FDynamicOperatorUpdateCallbacks;
+	}
 
-		void DoWork();
+	// Base class for building a metasound generator
+	class FAsyncMetaSoundBuilderBase : public FNonAbandonableTask
+	{
+	public:
+		FAsyncMetaSoundBuilderBase(FMetasoundGenerator* InGenerator, bool bInTriggerGenerator);
+
+		virtual ~FAsyncMetaSoundBuilderBase() = default;
 
 		FORCEINLINE TStatId GetStatId() const
 		{
 			RETURN_QUICK_DECLARE_CYCLE_STAT(FAsyncMetaSoundBuilder, STATGROUP_ThreadPoolAsyncTasks);
 		}
+	protected:
+
+		FMetasoundGenerator* Generator;
+		bool bTriggerGenerator;
+
+		void SetGraphOnGenerator(FOperatorAndInputs&& InOpAndInputs, FBuildResults&& InBuildResults);
 
 	private:
+		virtual const FMetasoundGeneratorInitParams& GetGeneratorParams() const = 0;
 
-		FMetasoundGenerator* Generator;
+	};
+
+	// Task class for asynchronously building a FMetasoundGenerator
+	class FAsyncMetaSoundBuilder : public FAsyncMetaSoundBuilderBase
+	{
+	public:
+		FAsyncMetaSoundBuilder(FMetasoundConstGraphGenerator* InGenerator, FMetasoundGeneratorInitParams&& InInitParams, bool bInTriggerGenerator);
+
+		virtual ~FAsyncMetaSoundBuilder() = default;
+
+		void DoWork();
+
+	private:
+		virtual const FMetasoundGeneratorInitParams& GetGeneratorParams() const override;
 		FMetasoundGeneratorInitParams InitParams;
-		bool bTriggerGenerator;
+	};
+
+	// Task class for asynchronously building a FMetasoundDynamicGraphGenerator
+	class FAsyncDynamicMetaSoundBuilder : public FAsyncMetaSoundBuilderBase
+	{
+	public:
+		FAsyncDynamicMetaSoundBuilder(FMetasoundDynamicGraphGenerator* InGenerator, FMetasoundDynamicGraphGeneratorInitParams&& InInitParams, bool bInTriggerGenerator);
+
+		virtual ~FAsyncDynamicMetaSoundBuilder() = default;
+
+		void DoWork();
+
+	private:
+		virtual const FMetasoundGeneratorInitParams& GetGeneratorParams() const override;
+		FMetasoundDynamicGraphGenerator* DynamicGenerator;
+		FMetasoundDynamicGraphGeneratorInitParams InitParams;
 	};
 
 	namespace GeneratorBuilder
 	{
+		// Finds an array of audio buffer read references from the VertexData
 		TArray<FAudioBufferReadRef> FindOutputAudioBuffers(const TArray<FVertexName>& InAudioVertexNames, const FVertexInterfaceData& InVertexData, const FOperatorSettings& InOperatorSettings, const FString& InMetaSoundName);
 
+		// Log any errors in the build results
 		void LogBuildErrors(const FString& InMetaSoundName, const FBuildResults& InBuildResults);
 
+		// Build an analyzer for a graph
 		TUniquePtr<Frontend::FGraphAnalyzer> BuildGraphAnalyzer(TMap<FGuid, FDataReferenceCollection>&& InInternalDataReferences, const FMetasoundEnvironment& InEnvironment, const FOperatorSettings& InOperatorSettings);
 
-		FOperatorAndInputs BuildGraphOperator(FMetasoundGeneratorInitParams& InInitParams, FBuildResults& OutBuildResults);
+		// Instantiate the input data references for a generator.
+		FInputVertexInterfaceData BuildGraphOperatorInputs(const FOperatorSettings& InOperatorSettings, FMetasoundGeneratorInitParams& InInitParams);
+		
+		// Build a standard graph operatpr
+		FOperatorAndInputs BuildGraphOperator(const FOperatorSettings& InOperatorSettings, FMetasoundGeneratorInitParams& InInitParams, FBuildResults& OutBuildResults);
+		
+		// Build a dynamic graph operator
+		FOperatorAndInputs BuildDynamicGraphOperator(const FOperatorSettings& InOperatorSettings, FMetasoundDynamicGraphGeneratorInitParams& InInitParams, const DynamicGraph::FDynamicOperatorUpdateCallbacks& InCallbacks, FBuildResults& OutBuildResults);
 
-		MetasoundGeneratorPrivate::FMetasoundGeneratorData BuildGeneratorData(const FMetasoundGeneratorInitParams& InInitParams, FOperatorAndInputs&& InGraphOperatorAndInputs, TUniquePtr<Frontend::FGraphAnalyzer> InAnalyzer);
+		// Create all the data required to run a FMetasoundGenerator
+		MetasoundGeneratorPrivate::FMetasoundGeneratorData BuildGeneratorData(const FOperatorSettings& InOperatorSettings, const FMetasoundGeneratorInitParams& InInitParams, FOperatorAndInputs&& InGraphOperatorAndInputs, TUniquePtr<Frontend::FGraphAnalyzer> InAnalyzer);
 
+		// Apply audio parameters to the input vertex data
 		void ApplyAudioParameters(const FOperatorSettings& InOperatorSettings, TArray<FAudioParameter>&& InParameters, FInputVertexInterfaceData& InInterface);
+
+		// Initialize the maps of parameter setters for applying async input data to metasound inputs.
+		void InitializeParameterSetters(FInputVertexInterfaceData& InputData, MetasoundGeneratorPrivate::FParameterSetterSortedMap& OutParamSetters, TMap<FName, MetasoundGeneratorPrivate::FParameterPackSetter>& OutParamPackSetters);
+
+		// Add a parameter setter to the parameter map if the data bound to the input vertex is writable.
+		void AddParameterSetterIfWritable(const FVertexName& InVertexName, FInputVertexInterfaceData& InputData, MetasoundGeneratorPrivate::FParameterSetterSortedMap& OutParamSetters, TMap<FName, MetasoundGeneratorPrivate::FParameterPackSetter>& OutParamPackSetters);
+
+		// Add a parameter setter to the parameter map for the given vertex.
+		void AddParameterSetter(const Frontend::IDataTypeRegistry& InDataTypeRegistry, const FInputDataVertex& InVertex, const FAnyDataReference& InDataRef,  MetasoundGeneratorPrivate::FParameterSetterSortedMap& OutParamSetters, TMap<FName, MetasoundGeneratorPrivate::FParameterPackSetter>& OutParamPackSetters);
 	}
 }
diff --git a/Engine/Plugins/Runtime/Metasound/Source/MetasoundGenerator/Public/MetasoundGenerator.h b/Engine/Plugins/Runtime/Metasound/Source/MetasoundGenerator/Public/MetasoundGenerator.h
--- a/Engine/Plugins/Runtime/Metasound/Source/MetasoundGenerator/Public/MetasoundGenerator.h
+++ b/Engine/Plugins/Runtime/Metasound/Source/MetasoundGenerator/Public/MetasoundGenerator.h
@@ -10,12 +10,12 @@
 #include "MetasoundTrigger.h"
 #include "MetasoundVertex.h"
 #include "MetasoundVertexData.h"
+
 #include "Analysis/MetasoundFrontendGraphAnalyzer.h"
 #include "Async/AsyncWork.h"
 #include "Containers/MpscQueue.h"
 #include "Containers/SpscQueue.h"
 #include "Sound/SoundGenerator.h"
-#include "Delegates/Delegate.h"
 
 #ifndef ENABLE_METASOUND_GENERATOR_RENDER_TIMING
 #define ENABLE_METASOUND_GENERATOR_RENDER_TIMING WITH_EDITOR
@@ -23,6 +23,11 @@
 
 namespace Metasound
 {
+	namespace DynamicGraph
+	{
+		class IDynamicOperatorTransform;
+	}
+
 	namespace MetasoundGeneratorPrivate
 	{
 #if ENABLE_METASOUND_GENERATOR_RENDER_TIMING
@@ -84,17 +89,12 @@ namespace Metasound
 		TArray<FAudioParameter> DefaultParameters;
 		bool bBuildSynchronous = false;
 		TSharedPtr<TSpscQueue<FMetaSoundParameterTransmitter::FParameter>> DataChannel;
-		TSharedPtr<Metasound::DynamicGraph::FDynamicOperatorTransactor> DynamicOperatorTransactor;
 
-		void Release();
+		static void Reset(FMetasoundGeneratorInitParams& InParams);
 	};
 
-
 	DECLARE_TS_MULTICAST_DELEGATE(FOnSetGraph);
 
-	/** FMetasoundGenerator generates audio from a given metasound IOperator
-	 * which produces a multichannel audio output.
-	 */
 	class METASOUNDGENERATOR_API FMetasoundGenerator : public ISoundGenerator
 	{
 	public:
@@ -109,9 +109,10 @@ namespace Metasound
 		 *
 		 * @param InParams - The generator initialization parameters
 		 */
-		explicit FMetasoundGenerator(FMetasoundGeneratorInitParams&& InParams);
 
-		virtual ~FMetasoundGenerator() override;
+		explicit FMetasoundGenerator(const FOperatorSettings& InOperatorSettings);
+
+		virtual ~FMetasoundGenerator();
 
 		/** Set the value of a graph's input data using the assignment operator.
 		 *
@@ -224,11 +225,31 @@ namespace Metasound
 		FDelegateHandle AddGraphSetCallback(FOnSetGraph::FDelegate&& Delegate);
 		bool RemoveGraphSetCallback(const FDelegateHandle& Handle);
 
+	protected:
+
+		void InitBase(const FMetasoundGeneratorInitParams& InInitParams);
+
+
+		/** SetGraph directly sets graph. Callers must ensure that no race conditions exist. */
+		void SetGraph(TUniquePtr<MetasoundGeneratorPrivate::FMetasoundGeneratorData>&& InData, bool bTriggerGraph);
+
+
+		virtual TUniquePtr<IOperator> ReleaseGraphOperator();
+		FInputVertexInterfaceData ReleaseInputVertexData();
+
+		/** Release the graph operator and remove any references to data owned by
+		 * the graph operator.
+		 */
+		void ClearGraph();
+		bool UpdateGraphIfPending();
+
+
+
 	private:
-		bool TryUseCachedOperator(FMetasoundGeneratorInitParams& InParams, bool bTriggerGenerator);
-		void ReleaseOperatorToCache();
 
-		friend class FAsyncMetaSoundBuilder;
+		friend class FAsyncMetaSoundBuilderBase;
+
+		void SetPendingGraphBuildFailed();
 
 		/** Update the current graph operator with a new graph operator. The number of channels
 		 * of InGraphOutputAudioRef must match the existing number of channels reported by
@@ -238,17 +259,6 @@ namespace Metasound
 		 * @param bTriggerGraph - If true, "OnPlay" will be triggered on the new graph.
 		 */
 		void SetPendingGraph(MetasoundGeneratorPrivate::FMetasoundGeneratorData&& InData, bool bTriggerGraph);
-		void SetPendingGraphBuildFailed();
-
-		bool UpdateGraphIfPending();
-
-		/** Release the graph operator and remove any references to data owned by
-		 * the graph operator.
-		 */
-		void ClearGraph();
-
-		/** Internal set graph after checking compatibility. */
-		void SetGraph(TUniquePtr<MetasoundGeneratorPrivate::FMetasoundGeneratorData>&& InData, bool bTriggerGraph);
 
 		// Fill OutAudio with data in InBuffer, up to maximum number of samples.
 		// Returns the number of samples used.
@@ -259,44 +269,43 @@ namespace Metasound
 		
 		void ApplyPendingUpdatesToInputs();
 
-		FExecuter RootExecuter;
-		FVertexInterfaceData VertexInterfaceData;
 
 		bool bIsGraphBuilding;
 		bool bIsFinishTriggered;
 		bool bIsFinished;
-		bool bUseOperatorCache = false;
-		FGuid OperatorID;
+		bool bPendingGraphTrigger;
+		bool bIsNewGraphPending;
+		bool bIsWaitingForFirstGraph;
 
 		int32 FinishSample = INDEX_NONE;
 		int32 NumChannels;
 		int32 NumFramesPerExecute;
 		int32 NumSamplesPerExecute;
 
-		TArray<FAudioBufferReadRef> GraphOutputAudio;
+	protected:
+		FExecuter RootExecuter;
+		FVertexInterfaceData VertexInterfaceData;
 
-		// Triggered when metasound is played
-		FTriggerWriteRef OnPlayTriggerRef;
+		TArray<FAudioBufferReadRef> GraphOutputAudio;
 
 		// Triggered when metasound is finished
 		FTriggerReadRef OnFinishedTriggerRef;
+	private:
 
 		Audio::FAlignedFloatBuffer InterleavedAudioBuffer;
 
 		Audio::FAlignedFloatBuffer OverflowBuffer;
 
-		TUniquePtr<FAsyncTaskBase> BuilderTask;
-
 		FCriticalSection PendingGraphMutex;
 		TUniquePtr<MetasoundGeneratorPrivate::FMetasoundGeneratorData> PendingGraphData;
-		bool bPendingGraphTrigger;
-		bool bIsNewGraphPending;
-		bool bIsWaitingForFirstGraph;
 
 		TUniquePtr<Frontend::FGraphAnalyzer> GraphAnalyzer;
 
 		TSharedPtr<TSpscQueue<FMetaSoundParameterTransmitter::FParameter>> ParameterQueue;
+	protected:
+
 		MetasoundGeneratorPrivate::FParameterSetterSortedMap ParameterSetters;
+	private:
 
 		// These next items are needed to provide a destination for the FAudioDevice, etc. to
 		// send parameter packs to. Every playing metasound will have a parameter destination
@@ -304,9 +313,11 @@ namespace Metasound
 		FSendAddress ParameterPackSendAddress;
 		TReceiverPtr<FMetasoundParameterStorageWrapper> ParameterPackReceiver;
 		
+	protected:
 		// This map provides setters for all of the input nodes in the metasound graph. 
 		// It is used when processing named parameters in a parameter pack.
 		TMap<FName, MetasoundGeneratorPrivate::FParameterPackSetter> ParameterPackSetters;
+	private:
 
 		// While parameter packs may arrive via the IAudioParameterInterface system,
 		// a faster method of sending parameters is via the QueueParameterPack function 
@@ -322,4 +333,71 @@ namespace Metasound
 		TUniquePtr<MetasoundGeneratorPrivate::FRenderTimer> RenderTimer;
 #endif // if ENABLE_METASOUND_GENERATOR_RENDER_TIMING
 	};
+
+	/** FMetasoundConstGraphGenerator generates audio from a given metasound IOperator
+	 * which produces a multichannel audio output.
+	 */
+	class METASOUNDGENERATOR_API FMetasoundConstGraphGenerator : public FMetasoundGenerator
+	{
+	public:
+
+		explicit FMetasoundConstGraphGenerator(FMetasoundGeneratorInitParams&& InParams);
+
+		explicit FMetasoundConstGraphGenerator(const FOperatorSettings& InOperatorSettings);
+
+		void Init(FMetasoundGeneratorInitParams&& InParams);
+
+		virtual ~FMetasoundConstGraphGenerator() override;
+
+	private:
+		friend class FAsyncMetaSoundBuilder;
+		void BuildGraph(FMetasoundGeneratorInitParams&& InInitParams);
+		bool TryUseCachedOperator(FMetasoundGeneratorInitParams& InParams, bool bTriggerGenerator);
+		void ReleaseOperatorToCache();
+
+		TUniquePtr<FAsyncTaskBase> BuilderTask;
+		FGuid OperatorID;
+		bool bUseOperatorCache = false;
+	};
+
+	struct METASOUNDGENERATOR_API FMetasoundDynamicGraphGeneratorInitParams : FMetasoundGeneratorInitParams
+	{
+		TSharedPtr<TSpscQueue<TUniquePtr<DynamicGraph::IDynamicOperatorTransform>>> TransformQueue;
+		static void Reset(FMetasoundDynamicGraphGeneratorInitParams& InParams);
+	};
+
+	/** FMetasoundDynamicGraphGenerator generates audio from the given a dynamic operator. It also
+	 * reacts to updates to inputs and outputs of the dynamic operator.
+	 */
+	class METASOUNDGENERATOR_API FMetasoundDynamicGraphGenerator : public FMetasoundGenerator
+	{
+	public:
+
+		/** Create the generator with a graph operator and an output audio reference.
+		 *
+		 * @param InParams - The generator initialization parameters
+		 */
+		explicit FMetasoundDynamicGraphGenerator(const FOperatorSettings& InOperatorSettings);
+
+		void Init(FMetasoundDynamicGraphGeneratorInitParams&& InParams);
+
+		virtual ~FMetasoundDynamicGraphGenerator();
+
+		// The callbacks are executed when the equivalent change happens on the owned dynamic operator.
+		void OnInputAdded(const FVertexName& InVertexName, const FInputVertexInterfaceData& InInputData);
+		void OnInputRemoved(const FVertexName& InVertexName, const FInputVertexInterfaceData& InInputData);
+		void OnOutputAdded(const FVertexName& InVertexName, const FOutputVertexInterfaceData& InOutputData);
+		void OnOutputUpdated(const FVertexName& InVertexName, const FOutputVertexInterfaceData& InOutputData);
+		void OnOutputRemoved(const FVertexName& InVertexName, const FOutputVertexInterfaceData& InOutputData);
+
+	protected:
+		virtual TUniquePtr<IOperator> ReleaseGraphOperator() override;
+
+	private:
+		void BuildGraph(FMetasoundDynamicGraphGeneratorInitParams&& InParams);
+
+
+		TArray<FVertexName> AudioOutputNames;
+		TUniquePtr<FAsyncTaskBase> BuilderTask;
+	};
 }
diff --git a/Engine/Plugins/Runtime/Metasound/Source/MetasoundGraphCore/Private/MetasoundDynamicGraphAlgo.cpp b/Engine/Plugins/Runtime/Metasound/Source/MetasoundGraphCore/Private/MetasoundDynamicGraphAlgo.cpp
--- a/Engine/Plugins/Runtime/Metasound/Source/MetasoundGraphCore/Private/MetasoundDynamicGraphAlgo.cpp
+++ b/Engine/Plugins/Runtime/Metasound/Source/MetasoundGraphCore/Private/MetasoundDynamicGraphAlgo.cpp
@@ -153,8 +153,25 @@ namespace Metasound
 		{
 		}
 
+		FDynamicGraphOperatorData::FDynamicGraphOperatorData(const FOperatorSettings& InSettings)
+		: DirectedGraphAlgo::FGraphOperatorData(InSettings)
+		{
+		}
+
 		FDynamicGraphOperatorData::FDynamicGraphOperatorData(DirectedGraphAlgo::FGraphOperatorData&& InGraphOperatorData)
 		: DirectedGraphAlgo::FGraphOperatorData(MoveTemp(InGraphOperatorData))
+		{
+			InitTables();
+		}
+
+		FDynamicGraphOperatorData::FDynamicGraphOperatorData(DirectedGraphAlgo::FGraphOperatorData&& InGraphOperatorData, const FDynamicOperatorUpdateCallbacks& InCallbacks)
+		: DirectedGraphAlgo::FGraphOperatorData(MoveTemp(InGraphOperatorData))
+		, OperatorUpdateCallbacks(InCallbacks)
+		{
+			InitTables();
+		}
+
+		void FDynamicGraphOperatorData::InitTables()
 		{
 			// Populate execute/postexecute/reset stacks
 			for (const FOperatorID& OperatorID : OperatorOrder)
@@ -198,6 +215,11 @@ namespace Metasound
 					if (const FAnyDataReference* Ref = OperatorInfo->VertexData.GetOutputs().FindDataReference(OutputVertexInfo.Get<0>()))
 					{
 						InOutGraphOperatorData.VertexData.GetOutputs().SetVertex(VertexName, *Ref);
+
+						if (InOutGraphOperatorData.OperatorUpdateCallbacks.OnOutputUpdated)
+						{
+							InOutGraphOperatorData.OperatorUpdateCallbacks.OnOutputUpdated(VertexName, InOutGraphOperatorData.VertexData.GetOutputs());
+						}
 					}
 					else if (InOutGraphOperatorData.VertexData.GetOutputs().IsVertexBound(VertexName))
 					{
diff --git a/Engine/Plugins/Runtime/Metasound/Source/MetasoundGraphCore/Private/MetasoundDynamicGraphAlgo.h b/Engine/Plugins/Runtime/Metasound/Source/MetasoundGraphCore/Private/MetasoundDynamicGraphAlgo.h
--- a/Engine/Plugins/Runtime/Metasound/Source/MetasoundGraphCore/Private/MetasoundDynamicGraphAlgo.h
+++ b/Engine/Plugins/Runtime/Metasound/Source/MetasoundGraphCore/Private/MetasoundDynamicGraphAlgo.h
@@ -3,6 +3,7 @@
 #pragma once
 
 #include "Containers/Array.h"
+#include "MetasoundDynamicOperatorTransactor.h"
 #include "MetasoundGraphAlgoPrivate.h"
 #include "MetasoundOperatorInterface.h"
 #include "MetasoundVertexData.h"
@@ -70,7 +71,16 @@ namespace Metasound
 		/** Collection of data needed to support a dynamic operator*/
 		struct FDynamicGraphOperatorData : DirectedGraphAlgo::FGraphOperatorData
 		{
+			FDynamicGraphOperatorData(const FOperatorSettings& InSettings);
 			FDynamicGraphOperatorData(DirectedGraphAlgo::FGraphOperatorData&& InGraphOperatorData);
+			FDynamicGraphOperatorData(DirectedGraphAlgo::FGraphOperatorData&& InGraphOperatorData, const FDynamicOperatorUpdateCallbacks& InCallbacks);
+
+			// Initialize the Execute/PostExecute/Reset tables.
+			void InitTables();
+
+			// A collection of optional callbacks which can be invoked when various
+			// updates are made to this collection of data.
+			FDynamicOperatorUpdateCallbacks OperatorUpdateCallbacks;
 
 			TArray<FExecuteEntry> ExecuteTable;
 			TArray<FPostExecuteEntry> PostExecuteTable;
diff --git a/Engine/Plugins/Runtime/Metasound/Source/MetasoundGraphCore/Private/MetasoundDynamicOperator.cpp b/Engine/Plugins/Runtime/Metasound/Source/MetasoundGraphCore/Private/MetasoundDynamicOperator.cpp
--- a/Engine/Plugins/Runtime/Metasound/Source/MetasoundGraphCore/Private/MetasoundDynamicOperator.cpp
+++ b/Engine/Plugins/Runtime/Metasound/Source/MetasoundGraphCore/Private/MetasoundDynamicOperator.cpp
@@ -87,24 +87,22 @@ namespace Metasound
 		FDynamicOperator::FDynamicOperator(const FOperatorSettings& InSettings)
 		: DynamicOperatorData(InSettings)
 		{
+			// Ensure that a transform queue exists. 
+			if (!TransformQueue.IsValid())
+			{
+				TransformQueue = MakeShared<TSpscQueue<TUniquePtr<IDynamicOperatorTransform>>>();
+			}
 		}
 
-		FDynamicOperator::FDynamicOperator(DirectedGraphAlgo::FGraphOperatorData&& InGraphOperatorData, TSharedRef<TSpscQueue<TUniquePtr<IDynamicOperatorTransform>>> InTransformQueue)
-		: DynamicOperatorData(MoveTemp(InGraphOperatorData))
+		FDynamicOperator::FDynamicOperator(DirectedGraphAlgo::FGraphOperatorData&& InGraphOperatorData, TSharedPtr<TSpscQueue<TUniquePtr<IDynamicOperatorTransform>>> InTransformQueue, const FDynamicOperatorUpdateCallbacks& InOperatorUpdateCallbacks)
+		: DynamicOperatorData(MoveTemp(InGraphOperatorData), InOperatorUpdateCallbacks)
 		, TransformQueue(InTransformQueue)
 		{
-		}
-
-		FDataReferenceCollection FDynamicOperator::GetInputs() const
-		{
-			checkNoEntry();
-			return FDataReferenceCollection{};
-		}
-
-		FDataReferenceCollection FDynamicOperator::GetOutputs() const
-		{
-			checkNoEntry();
-			return FDataReferenceCollection{};
+			// Ensure that a transform queue exists. 
+			if (!TransformQueue.IsValid())
+			{
+				TransformQueue = MakeShared<TSpscQueue<TUniquePtr<IDynamicOperatorTransform>>>();
+			}
 		}
 
 		void FDynamicOperator::BindInputs(FInputVertexInterfaceData& InOutVertexData)
@@ -332,6 +330,12 @@ namespace Metasound
 
 			InGraphOperatorData.InputVertexMap.Add(VertexName, OperatorID);
 
+			// Update listeners that an input has been added.
+			if (InGraphOperatorData.OperatorUpdateCallbacks.OnInputAdded)
+			{
+				InGraphOperatorData.OperatorUpdateCallbacks.OnInputAdded(VertexName, InGraphOperatorData.VertexData.GetInputs());
+			}
+
 			// Propagate the data reference update through the graph
 			PropagateBindUpdate(OperatorID, VertexName, DataReference, InGraphOperatorData);
 
@@ -359,6 +363,12 @@ namespace Metasound
 				InGraphOperatorData.VertexData.GetInputs().RemoveVertex(VertexName);
 			}
 
+			// Update listeners that an input has been removed.
+			if (InGraphOperatorData.OperatorUpdateCallbacks.OnInputRemoved)
+			{
+				InGraphOperatorData.OperatorUpdateCallbacks.OnInputRemoved(VertexName, InGraphOperatorData.VertexData.GetInputs());
+			}
+
 			return EDynamicOperatorTransformQueueAction::Continue;
 		}
 
@@ -389,7 +399,6 @@ namespace Metasound
 				return EDynamicOperatorTransformQueueAction::Continue;
 			}
 
-
 			FOutputVertexInterfaceData& GraphOutputData = InGraphOperatorData.VertexData.GetOutputs();
 			{
 				// Unfreeze interface so a new vertex can be added. 
@@ -402,6 +411,12 @@ namespace Metasound
 
 			InGraphOperatorData.OutputVertexMap.Add(VertexName, OperatorID);
 
+			// Update listeners that an input has been added.
+			if (InGraphOperatorData.OperatorUpdateCallbacks.OnOutputAdded)
+			{
+				InGraphOperatorData.OperatorUpdateCallbacks.OnOutputAdded(VertexName, InGraphOperatorData.VertexData.GetOutputs());
+			}
+
 			return EDynamicOperatorTransformQueueAction::Continue;
 		}
 
@@ -421,6 +436,13 @@ namespace Metasound
 				DynamicOperatorPrivate::TScopeUnfreeze Unfreeze(InGraphOperatorData.VertexData.GetOutputs());
 				InGraphOperatorData.VertexData.GetOutputs().RemoveVertex(VertexName);
 			}
+
+			// Update listeners that an output has been removed.
+			if (InGraphOperatorData.OperatorUpdateCallbacks.OnOutputRemoved)
+			{
+				InGraphOperatorData.OperatorUpdateCallbacks.OnOutputRemoved(VertexName, InGraphOperatorData.VertexData.GetOutputs());
+			}
+
 			return EDynamicOperatorTransformQueueAction::Continue;
 		}
 
@@ -495,7 +517,8 @@ namespace Metasound
 
 			METASOUND_TRACE_CPUPROFILER_EVENT_SCOPE(Metasound::DynamicOperator::DisconnectOperators)
 
-			// Make new connection.
+			// Make new connection. We can skip propagating updates and updating callbacks
+			// here because those are handled in the FConnectOperators transform.
 			EDynamicOperatorTransformQueueAction NextAction = ConnectTransform.Transform(InGraphOperatorData);
 			check(NextAction == EDynamicOperatorTransformQueueAction::Continue); // this should always be continue since next step must be performed. 
 
diff --git a/Engine/Plugins/Runtime/Metasound/Source/MetasoundGraphCore/Private/MetasoundDynamicOperator.h b/Engine/Plugins/Runtime/Metasound/Source/MetasoundGraphCore/Private/MetasoundDynamicOperator.h
--- a/Engine/Plugins/Runtime/Metasound/Source/MetasoundGraphCore/Private/MetasoundDynamicOperator.h
+++ b/Engine/Plugins/Runtime/Metasound/Source/MetasoundGraphCore/Private/MetasoundDynamicOperator.h
@@ -5,6 +5,7 @@
 #include "Containers/Array.h"
 #include "Containers/SpscQueue.h"
 #include "MetasoundDynamicGraphAlgo.h"
+#include "MetasoundDynamicOperatorTransactor.h"
 #include "MetasoundGraphAlgoPrivate.h"
 #include "MetasoundNodeInterface.h"
 #include "MetasoundOperatorInterface.h"
@@ -42,10 +43,7 @@ namespace Metasound
 		public:
 
 			FDynamicOperator(const FOperatorSettings& InSettings);
-			FDynamicOperator(DirectedGraphAlgo::FGraphOperatorData&& InGraphOperatorData, TSharedRef<TSpscQueue<TUniquePtr<IDynamicOperatorTransform>>> TransformQueue);
-
-			virtual FDataReferenceCollection GetInputs() const override;
-			virtual FDataReferenceCollection GetOutputs() const override;
+			FDynamicOperator(DirectedGraphAlgo::FGraphOperatorData&& InGraphOperatorData, TSharedPtr<TSpscQueue<TUniquePtr<IDynamicOperatorTransform>>> TransformQueue, const FDynamicOperatorUpdateCallbacks& InOperatorUpdateCallbacks);
 
 			virtual void BindInputs(FInputVertexInterfaceData& InOutVertexData) override;
 			virtual void BindOutputs(FOutputVertexInterfaceData& InOutVertexData) override;
@@ -71,7 +69,7 @@ namespace Metasound
 
 			FDynamicGraphOperatorData DynamicOperatorData;
 
-			TSharedRef<TSpscQueue<TUniquePtr<IDynamicOperatorTransform>>> TransformQueue;
+			TSharedPtr<TSpscQueue<TUniquePtr<IDynamicOperatorTransform>>> TransformQueue;
 			bool bExecuteFenceIsSet = false;
 		};
 
diff --git a/Engine/Plugins/Runtime/Metasound/Source/MetasoundGraphCore/Private/MetasoundDynamicOperatorTransactor.cpp b/Engine/Plugins/Runtime/Metasound/Source/MetasoundGraphCore/Private/MetasoundDynamicOperatorTransactor.cpp
--- a/Engine/Plugins/Runtime/Metasound/Source/MetasoundGraphCore/Private/MetasoundDynamicOperatorTransactor.cpp
+++ b/Engine/Plugins/Runtime/Metasound/Source/MetasoundGraphCore/Private/MetasoundDynamicOperatorTransactor.cpp
@@ -361,7 +361,7 @@ namespace Metasound
 			EnqueueTransformOnOperatorQueues(CreateRemoveOutputTransform);
 		}
 
-		const IGraph& FDynamicOperatorTransactor::GetGraph() const
+		const FGraph& FDynamicOperatorTransactor::GetGraph() const
 		{
 			return Graph;
 		}
diff --git a/Engine/Plugins/Runtime/Metasound/Source/MetasoundGraphCore/Private/MetasoundOperatorBuilder.cpp b/Engine/Plugins/Runtime/Metasound/Source/MetasoundGraphCore/Private/MetasoundOperatorBuilder.cpp
--- a/Engine/Plugins/Runtime/Metasound/Source/MetasoundGraphCore/Private/MetasoundOperatorBuilder.cpp
+++ b/Engine/Plugins/Runtime/Metasound/Source/MetasoundGraphCore/Private/MetasoundOperatorBuilder.cpp
@@ -89,7 +89,7 @@ namespace Metasound
 		return TUniquePtr<IOperator>(nullptr);
 	}
 
-	TUniquePtr<IOperator> FOperatorBuilder::BuildDynamicGraphOperator(const FBuildGraphOperatorParams& InParams, DynamicGraph::FDynamicOperatorTransactor& InTransactor, FBuildResults& OutResults)
+	TUniquePtr<IOperator> FOperatorBuilder::BuildDynamicGraphOperator(const FBuildDynamicGraphOperatorParams& InParams, FBuildResults& OutResults)
 	{
 		using namespace DynamicGraph;
 
@@ -97,9 +97,7 @@ namespace Metasound
 
 		if (GraphData.IsValid())
 		{
-			TSharedRef<TSpscQueue<TUniquePtr<IDynamicOperatorTransform>>> TransformQueue = InTransactor.CreateTransformQueue(InParams.OperatorSettings, InParams.Environment);
-
-			return MakeUnique<FDynamicOperator>(MoveTemp(*GraphData), TransformQueue);
+			return MakeUnique<FDynamicOperator>(MoveTemp(*GraphData), InParams.TransformQueue, InParams.OperatorUpdateCallbacks);
 		}
 
 		return TUniquePtr<IOperator>(nullptr);
diff --git a/Engine/Plugins/Runtime/Metasound/Source/MetasoundGraphCore/Public/MetasoundDynamicOperatorTransactor.h b/Engine/Plugins/Runtime/Metasound/Source/MetasoundGraphCore/Public/MetasoundDynamicOperatorTransactor.h
--- a/Engine/Plugins/Runtime/Metasound/Source/MetasoundGraphCore/Public/MetasoundDynamicOperatorTransactor.h
+++ b/Engine/Plugins/Runtime/Metasound/Source/MetasoundGraphCore/Public/MetasoundDynamicOperatorTransactor.h
@@ -8,6 +8,7 @@
 #include "MetasoundGraph.h"
 #include "MetasoundOperatorInterface.h"
 #include "MetasoundNodeInterface.h"
+#include "Templates/Function.h"
 #include "Templates/SharedPointer.h"
 #include "Templates/UniquePtr.h"
 
@@ -24,7 +25,30 @@ namespace Metasound
 
 		using FLiteralAssignmentFunction = void(*)(const FOperatorSettings& InOperatorSettings, const FLiteral& InLiteral, const FAnyDataReference& OutDataRef);
 		using FReferenceCreationFunction = TOptional<FAnyDataReference>(*)(const FOperatorSettings& InSettings, const FLiteral& InLiteral, EDataReferenceAccessType InAccessType);
+		using FOnInputVertexUpdated = TFunction<void(const FVertexName&, const FInputVertexInterfaceData&)>;
+		using FOnOutputVertexUpdated = TFunction<void(const FVertexName&, const FOutputVertexInterfaceData&)>;
+
+		/** A collection of callbacks for handling updates to MetaSound dynamic operators.
+		 * 
+		 * Callbacks are invoked on the same thread which executes the dynamic operator. 
+		 */
+		struct FDynamicOperatorUpdateCallbacks
+		{
+			FOnInputVertexUpdated OnInputAdded;
+			FOnInputVertexUpdated OnInputRemoved;
+			FOnOutputVertexUpdated OnOutputAdded;
+			FOnOutputVertexUpdated OnOutputUpdated;
+			FOnOutputVertexUpdated OnOutputRemoved;
+		};
 
+		/** The FDynamicOperatorTransactor is used for communicating with a dynamic
+		 * MetaSound operator.
+		 *
+		 * Graph manipulations performed on the transactor are forwarded to dynamic
+		 * operators using the transform queue. Each modification is converted into
+		 * IDynamicOperatorTransforms which are consumed by dynamic operators during
+		 * their execution.
+		 */
 		class METASOUNDGRAPHCORE_API FDynamicOperatorTransactor
 		{
 		public:
@@ -82,7 +106,7 @@ namespace Metasound
 			void RemoveOutputDataSource(const FVertexName& InVertexName);
 
 			/** Return internal version of graph. */
-			const IGraph& GetGraph() const;
+			const FGraph& GetGraph() const;
 
 		private:
 
diff --git a/Engine/Plugins/Runtime/Metasound/Source/MetasoundGraphCore/Public/MetasoundOperatorBuilder.h b/Engine/Plugins/Runtime/Metasound/Source/MetasoundGraphCore/Public/MetasoundOperatorBuilder.h
--- a/Engine/Plugins/Runtime/Metasound/Source/MetasoundGraphCore/Public/MetasoundOperatorBuilder.h
+++ b/Engine/Plugins/Runtime/Metasound/Source/MetasoundGraphCore/Public/MetasoundOperatorBuilder.h
@@ -4,6 +4,7 @@
 
 #include "Containers/Array.h"
 #include "Containers/Map.h"
+#include "Containers/SpscQueue.h"
 #include "MetasoundBuilderInterface.h"
 #include "MetasoundEnvironment.h"
 #include "MetasoundGraphOperator.h"
@@ -12,6 +13,7 @@
 #include "MetasoundOperatorInterface.h"
 #include "MetasoundVertexData.h"
 #include "Templates/UniquePtr.h"
+#include "Templates/SharedPointer.h"
 
 namespace Metasound
 {
@@ -22,15 +24,24 @@ namespace Metasound
 
 	namespace OperatorBuilder
 	{
-		// Forward Declare
 		struct FBuildContext;
 	}
 
 	namespace DynamicGraph
 	{
-		class FDynamicOperatorTransactor;
+		class IDynamicOperatorTransform;
+		struct FDynamicOperatorUpdateCallbacks;
 	}
 
+	// Parameters for building a dynamic graph operator consist of the parameters
+	// needed to build a standard graph operator as well as a transform queue and
+	// set of callbacks in order to communicate and react to changes in the dynamic operator. 
+	struct FBuildDynamicGraphOperatorParams : FBuildGraphOperatorParams
+	{
+		TSharedPtr<TSpscQueue<TUniquePtr<DynamicGraph::IDynamicOperatorTransform>>> TransformQueue; 
+		const DynamicGraph::FDynamicOperatorUpdateCallbacks& OperatorUpdateCallbacks;
+	};
+
 	/** FOperatorBuilder builds an IOperator from an IGraph. */
 	class METASOUNDGRAPHCORE_API FOperatorBuilder : public IOperatorBuilder
 	{
@@ -56,14 +67,13 @@ namespace Metasound
 
 			/** Create an dynamic IOperator from an IGraph.
 			 *
-			 * @param InParams     - Params of the current build
-			 * @param InTransactor - A dynamic operator transactor for communicating to the operator.
-			 * @param OutResults   - Results data pertaining to the given build operator result.
+			 * @param InParams                  - Params of the current build
+			 * @param OutResults                - Results data pertaining to the given build operator result.
 			 *
 			 * @return A TUniquePtr to an IOperator. If the processes was unsuccessful, 
 			 *         the returned pointer will contain a nullptr and be invalid.
 			 */
-			TUniquePtr<IOperator> BuildDynamicGraphOperator(const FBuildGraphOperatorParams& InParams, DynamicGraph::FDynamicOperatorTransactor& InTransactor, FBuildResults& OutResults);
+			TUniquePtr<IOperator> BuildDynamicGraphOperator(const FBuildDynamicGraphOperatorParams& InParams,  FBuildResults& OutResults);
 
 		private:
 
