Fix for nametable/MetaSound Search Engine entry leaks due to new class name being generated on every builder build/audition call associated with a given target object
#rnx

diff --git a/Engine/Plugins/Runtime/Metasound/Source/MetasoundEngine/Public/MetasoundBuilderSubsystem.h b/Engine/Plugins/Runtime/Metasound/Source/MetasoundEngine/Public/MetasoundBuilderSubsystem.h
--- a/Engine/Plugins/Runtime/Metasound/Source/MetasoundEngine/Public/MetasoundBuilderSubsystem.h
+++ b/Engine/Plugins/Runtime/Metasound/Source/MetasoundEngine/Public/MetasoundBuilderSubsystem.h
@@ -370,6 +370,7 @@ protected:
 		using namespace Metasound::Frontend;
 
 		UClassType* MetaSound = nullptr;
+		const FMetasoundFrontendClassName* DocClassName = nullptr;
 		if (BuilderOptions.ExistingMetaSound)
 		{
 			MetaSound = CastChecked<UClassType>(BuilderOptions.ExistingMetaSound.GetObject());
@@ -380,6 +381,10 @@ protected:
 			// being registered as it is below.
 			if (MetaSound)
 			{
+				// If MetaSound already exists, preserve the class name to avoid
+				// nametable bloat & preserve potentially existing references.
+				const FMetasoundFrontendDocument& ExistingDoc = CastChecked<const UClassType>(MetaSound)->GetDocumentChecked();
+				DocClassName = &ExistingDoc.RootGraph.Metadata.GetClassName();
 				const FNodeRegistryKey& RegistryKey = MetaSound->GetRegistryKey();
 				if (NodeRegistryKey::IsValid(RegistryKey))
 				{
@@ -410,10 +415,10 @@ protected:
 		FMetasoundFrontendDocument NewDocument = GetConstBuilder().GetDocument();
 		{
 			// This is required to ensure the newly build document has a unique class
-			// identifier to avoid collisions if added to the Frontend class registry
-			// (either below or at a later point in time).
+			// identifier than that of the given builder's document copy to avoid collisions
+			// if added to the Frontend class registry (either below or at a later point in time).
 			constexpr bool bResetVersion = false;
-			FMetaSoundFrontendDocumentBuilder::InitGraphClassMetadata(NewDocument.RootGraph.Metadata, bResetVersion);
+			FMetaSoundFrontendDocumentBuilder::InitGraphClassMetadata(NewDocument.RootGraph.Metadata, bResetVersion, DocClassName);
 		}
 
 		MetaSound->SetDocument(MoveTemp(NewDocument));
diff --git a/Engine/Plugins/Runtime/Metasound/Source/MetasoundEngineTest/Private/EngineTestMetasoundSource.cpp b/Engine/Plugins/Runtime/Metasound/Source/MetasoundEngineTest/Private/EngineTestMetasoundSource.cpp
--- a/Engine/Plugins/Runtime/Metasound/Source/MetasoundEngineTest/Private/EngineTestMetasoundSource.cpp
+++ b/Engine/Plugins/Runtime/Metasound/Source/MetasoundEngineTest/Private/EngineTestMetasoundSource.cpp
@@ -7,6 +7,7 @@
 #include "Interfaces/IPluginManager.h"
 #include "Interfaces/MetasoundOutputFormatInterfaces.h"
 #include "Interfaces/MetasoundFrontendSourceInterface.h"
+#include "MetasoundLog.h"
 #include "MetasoundSource.h"
 #include "Misc/AutomationTest.h"
 #include "Misc/Paths.h"
@@ -260,8 +261,32 @@ bool FMetaSoundSourceBuilderAuditionLatentCommand::Update()
 {
 	if (Builder && AudioComponent)
 	{
+		auto GetComponentSoundClassName = [this]() -> FName
+		{
+			if (const USoundBase* InitSound = AudioComponent->GetSound())
+			{
+				const UMetaSoundSource* InitMetaSound = CastChecked<UMetaSoundSource>(InitSound);
+				return InitMetaSound->GetDocumentChecked().RootGraph.Metadata.GetClassName().GetFullName();
+			}
+
+			return { };
+		};
+
+		// This is an inline test to ensure that the first time a builder is auditioned, its generating a new unique MetaSound Class Name.
+		// Each subsequent call should maintain that name to avoid bloating the FName table/breaking references should this auditioned sound
+		// be in anyway referenced.
+		const FName InitClassName = GetComponentSoundClassName();
 		Builder->Audition(nullptr, AudioComponent, { }, true);
-		return true;
+		const FName BuiltClassName = GetComponentSoundClassName();
+
+		const bool bInitNameGenerated = InitClassName.IsNone() && !BuiltClassName.IsNone();
+		const bool bClassNameMaintained = InitClassName == BuiltClassName;
+		if (bInitNameGenerated || bClassNameMaintained)
+		{
+			return true;
+		}
+
+		UE_LOG(LogMetaSound, Error, TEXT("Latent test audition call resulted in generation of a new MetaSound class instead of re-using the existing class name"));
 	}
 	return false;
 }
@@ -355,6 +380,55 @@ bool FAudioMetasoundSourceBuilderTest::RunTest(const FString& Parameters)
 	return true;
 }
 
+// This test creates a MetaSound source from a SourceBuilder, adds a simple sine tone generator with a connected graph input frequency, and attempts to change the frequency and audition it rapidly 100 times.
+IMPLEMENT_SIMPLE_AUTOMATION_TEST(FAudioMetasoundSourceBuilderTestSpamAudition, "Audio.Metasound.SpamAuditionMetasoundSource", EAutomationTestFlags::EditorContext | EAutomationTestFlags::EngineFilter)
+bool FAudioMetasoundSourceBuilderTestSpamAudition::RunTest(const FString& Parameters)
+{
+	using namespace EngineTestMetaSoundPatchBuilderPrivate;
+	using namespace EngineTestMetaSoundSourcePrivate;
+
+	FMetaSoundBuilderNodeInputHandle MonoOutNodeInput;
+	FMetaSoundBuilderNodeInputHandle GenInputNodeFreq;
+	UMetaSoundSourceBuilder& Builder = CreateMonoSourceSinGenBuilder(*this, &GenInputNodeFreq, &MonoOutNodeInput);
+	Builder.AddToRoot();
+
+	if (FAudioDevice* AudioDevice = GEngine->GetMainAudioDeviceRaw())
+	{
+		constexpr bool bAddToRegistry = false;
+		TScriptInterface<IMetaSoundDocumentInterface> BuiltDocumentInterface = Builder.Build(nullptr, FMetaSoundBuilderOptions{ "BuildAndPlayMetasoundSource", bAddToRegistry });
+		UMetaSoundSource* MetaSoundSource = CastChecked<UMetaSoundSource>(BuiltDocumentInterface.GetObject());
+
+		if (UAudioComponent* AudioComponent = NewObject<UAudioComponent>())
+		{
+			AudioComponent->bAutoActivate = false;
+			AudioComponent->bIsUISound = true; // play while "paused"
+			AudioComponent->AudioDeviceID = AudioDevice->DeviceID;
+
+			AudioComponent->bAllowSpatialization = false;
+			AudioComponent->SetVolumeMultiplier(1.0f);
+			AudioComponent->AddToRoot();
+
+			constexpr bool bEnableLiveUpdate = false;
+			constexpr int32 NumTrials = 100;
+			for (int32 Index = 0; Index < NumTrials; ++Index)
+			{
+				FMetasoundFrontendLiteral NewValue;
+				NewValue.Set(FMath::RandRange(220.f, 2200.f));
+				ADD_LATENT_AUTOMATION_COMMAND(FMetaSoundSourceBuilderSetLiteralLatentCommand(*this, &Builder, GenInputNodeFreq, NewValue));
+				ADD_LATENT_AUTOMATION_COMMAND(FMetaSoundSourceBuilderAuditionLatentCommand(&Builder, AudioComponent, bEnableLiveUpdate));
+				ADD_LATENT_AUTOMATION_COMMAND(FEngineWaitLatentCommand(0.01f));
+			}
+
+			ADD_LATENT_AUTOMATION_COMMAND(FAudioComponentStopLatentCommand(AudioComponent));
+			ADD_LATENT_AUTOMATION_COMMAND(FEngineWaitLatentCommand(0.5f));
+			ADD_LATENT_AUTOMATION_COMMAND(FBuilderRemoveFromRootLatentCommand(&Builder));
+			ADD_LATENT_AUTOMATION_COMMAND(FAudioComponentRemoveFromRootLatentCommand(AudioComponent));
+		}
+	}
+
+	return true;
+}
+
 // This test creates a MetaSound source from a SourceBuilder, adds a simple sine tone generator with a connected graph input frequency, attempts to audition it, and then switches to a new tri tone generator during playback.
 IMPLEMENT_SIMPLE_AUTOMATION_TEST(FAudioMetasoundSourceBuilderLiveUpdateNode, "Audio.Metasound.LiveUpdateNodeMetaSoundSource", EAutomationTestFlags::EditorContext | EAutomationTestFlags::EngineFilter)
 bool FAudioMetasoundSourceBuilderLiveUpdateNode::RunTest(const FString& Parameters)
diff --git a/Engine/Plugins/Runtime/Metasound/Source/MetasoundFrontend/Private/MetasoundFrontendDocumentBuilder.cpp b/Engine/Plugins/Runtime/Metasound/Source/MetasoundFrontend/Private/MetasoundFrontendDocumentBuilder.cpp
--- a/Engine/Plugins/Runtime/Metasound/Source/MetasoundFrontend/Private/MetasoundFrontendDocumentBuilder.cpp
+++ b/Engine/Plugins/Runtime/Metasound/Source/MetasoundFrontend/Private/MetasoundFrontendDocumentBuilder.cpp
@@ -1498,9 +1498,16 @@ const FMetasoundFrontendLiteral* FMetaSoundFrontendDocumentBuilder::GetNodeInput
 	return nullptr;
 }
 
-void FMetaSoundFrontendDocumentBuilder::InitGraphClassMetadata(FMetasoundFrontendClassMetadata& InOutMetadata, bool bResetVersion)
+void FMetaSoundFrontendDocumentBuilder::InitGraphClassMetadata(FMetasoundFrontendClassMetadata& InOutMetadata, bool bResetVersion, const FMetasoundFrontendClassName* NewClassName)
 {
-	InOutMetadata.SetClassName(FMetasoundFrontendClassName(FName(), *FGuid::NewGuid().ToString(), FName()));
+	if (NewClassName)
+	{
+		InOutMetadata.SetClassName(*NewClassName);
+	}
+	else
+	{
+		InOutMetadata.SetClassName(FMetasoundFrontendClassName(FName(), *FGuid::NewGuid().ToString(), FName()));
+	}
 
 	if (bResetVersion)
 	{
diff --git a/Engine/Plugins/Runtime/Metasound/Source/MetasoundFrontend/Public/MetasoundFrontendDocumentBuilder.h b/Engine/Plugins/Runtime/Metasound/Source/MetasoundFrontend/Public/MetasoundFrontendDocumentBuilder.h
--- a/Engine/Plugins/Runtime/Metasound/Source/MetasoundFrontend/Public/MetasoundFrontendDocumentBuilder.h
+++ b/Engine/Plugins/Runtime/Metasound/Source/MetasoundFrontend/Public/MetasoundFrontendDocumentBuilder.h
@@ -182,7 +182,8 @@ public:
 	const FMetasoundFrontendLiteral* GetNodeInputClassDefault(const FGuid& InNodeID, const FGuid& InVertexID) const;
 	const FMetasoundFrontendLiteral* GetNodeInputDefault(const FGuid& InNodeID, const FGuid& InVertexID) const;
 
-	static void InitGraphClassMetadata(FMetasoundFrontendClassMetadata& InOutMetadata, bool bResetVersion = false);
+	// Initializes GraphClass Metadata, optionally resetting the version back to 1.0 and/or creating a unique class name if a name is not provided.
+	static void InitGraphClassMetadata(FMetasoundFrontendClassMetadata& InOutMetadata, bool bResetVersion = false, const FMetasoundFrontendClassName* NewClassName = nullptr);
 	void InitDocument();
 	void InitNodeLocations();
 
