Enable nulling out object parameters/proxies from BP/game thread transmitter requests.
#rnx

diff --git a/Engine/Plugins/Runtime/Metasound/Source/MetasoundEngine/Private/MetasoundSource.cpp b/Engine/Plugins/Runtime/Metasound/Source/MetasoundEngine/Private/MetasoundSource.cpp
--- a/Engine/Plugins/Runtime/Metasound/Source/MetasoundEngine/Private/MetasoundSource.cpp
+++ b/Engine/Plugins/Runtime/Metasound/Source/MetasoundEngine/Private/MetasoundSource.cpp
@@ -48,7 +48,7 @@
 namespace Metasound
 {
 	namespace SourcePrivate
-	{		
+	{
 		Frontend::FMetaSoundAssetRegistrationOptions GetInitRegistrationOptions()
 		{
 			Frontend::FMetaSoundAssetRegistrationOptions RegOptions;
@@ -332,19 +332,21 @@ void UMetaSoundSource::InitParameters(TArray<FAudioParameter>& ParametersToInit,
 	// via InitResources. If it has, this call is fast and returns the already cached RuntimeData.
 	RegisterGraphWithFrontend(GetInitRegistrationOptions());
 	const FRuntimeData& RuntimeData = GetRuntimeData();
-	const TArray<FMetasoundFrontendClassInput>& PublicInputs = RuntimeData.PublicInputs;
 
-	TMap<FName, FMetasoundFrontendVertex> PublicInputMap;
-	Algo::Transform(PublicInputs, PublicInputMap, [](const FMetasoundFrontendClassInput& Input)
+	TMap<FName, const FMetasoundFrontendVertex*> PublicInputMap;
 	{
-		return TPair<FName, FMetasoundFrontendVertex>(Input.Name, Input);
-	});
+		const TArray<FMetasoundFrontendClassInput>& PublicInputs = RuntimeData.PublicInputs;
+		Algo::Transform(PublicInputs, PublicInputMap, [](const FMetasoundFrontendClassInput& Input)
+			{
+				return TPair<FName, const FMetasoundFrontendVertex*>(Input.Name, &Input);
+			});
+	}
 
 	// Removes values that are not explicitly defined by the ParamType and returns
 	// whether or not the parameter is a valid input and should be included.
 	auto Sanitize = [&PublicInputMap](FAudioParameter& Parameter) -> bool
 	{
-		const FMetasoundFrontendVertex* Input = PublicInputMap.Find(Parameter.ParamName);
+		const FMetasoundFrontendVertex* Input = PublicInputMap.FindRef(Parameter.ParamName);
 		if (!Input)
 		{
 			return false;
@@ -438,7 +440,7 @@ void UMetaSoundSource::InitParameters(TArray<FAudioParameter>& ParametersToInit,
 		return true;
 	};
 
-	auto ConstructProxies = [this](FAudioParameter& OutParamToInit)
+	auto ConstructProxies = [this](FAudioParameter& OutParamToInit, FName VertexTypeName)
 	{
 		using namespace Metasound;
 		using namespace Metasound::Frontend;
@@ -448,14 +450,11 @@ void UMetaSoundSource::InitParameters(TArray<FAudioParameter>& ParametersToInit,
 			case EAudioParameterType::Object:
 			{
 				FDataTypeRegistryInfo DataTypeInfo;
-				if (IDataTypeRegistry::Get().GetDataTypeInfo(OutParamToInit.ObjectParam, DataTypeInfo))
-				{
-					TSharedPtr<Audio::IProxyData> ProxyPtr = IDataTypeRegistry::Get().CreateProxyFromUObject(DataTypeInfo.DataTypeName, OutParamToInit.ObjectParam);
-					OutParamToInit.ObjectProxies.Emplace(MoveTemp(ProxyPtr));
+				TSharedPtr<Audio::IProxyData> ProxyPtr = IDataTypeRegistry::Get().CreateProxyFromUObject(VertexTypeName, OutParamToInit.ObjectParam);
+				OutParamToInit.ObjectProxies.Emplace(MoveTemp(ProxyPtr));
 
-					// Null out param as it is no longer needed (nor desired to be accessed once passed to the Audio Thread)
-					OutParamToInit.ObjectParam = nullptr;
-				}
+				// Null out param as it is no longer needed (nor desired to be accessed once passed to the Audio Thread)
+				OutParamToInit.ObjectParam = nullptr;
 			}
 			break;
 
@@ -464,11 +463,8 @@ void UMetaSoundSource::InitParameters(TArray<FAudioParameter>& ParametersToInit,
 				for (TObjectPtr<UObject>& Object : OutParamToInit.ArrayObjectParam)
 				{
 					FDataTypeRegistryInfo DataTypeInfo;
-					if (IDataTypeRegistry::Get().GetDataTypeInfo(Object, DataTypeInfo))
-					{
-						TSharedPtr<Audio::IProxyData> ProxyPtr = IDataTypeRegistry::Get().CreateProxyFromUObject(DataTypeInfo.DataTypeName, Object);
-						OutParamToInit.ObjectProxies.Emplace(MoveTemp(ProxyPtr));
-					}
+					TSharedPtr<Audio::IProxyData> ProxyPtr = IDataTypeRegistry::Get().CreateProxyFromUObject(VertexTypeName, Object);
+					OutParamToInit.ObjectProxies.Emplace(MoveTemp(ProxyPtr));
 				}
 				// Reset param array as it is no longer needed (nor desired to be accessed once passed to the Audio Thread).
 				// All object manipulation hereafter should be done via proxies
@@ -481,7 +477,6 @@ void UMetaSoundSource::InitParameters(TArray<FAudioParameter>& ParametersToInit,
 		}
 	};
 
-
 	for (int32 i = ParametersToInit.Num() - 1; i >= 0; --i)
 	{
 		FAudioParameter& Parameter = ParametersToInit[i];
@@ -493,9 +488,10 @@ void UMetaSoundSource::InitParameters(TArray<FAudioParameter>& ParametersToInit,
 		
 		if (Sanitize(Parameter))
 		{
-			if (IsParameterValid(Parameter, PublicInputMap))
+			const FMetasoundFrontendVertex* Vertex = PublicInputMap.FindRef(Parameter.ParamName);
+			if (IsParameterValid(Parameter, Vertex))
 			{
-				ConstructProxies(Parameter);
+				ConstructProxies(Parameter, Vertex->TypeName);
 			}
 			else
 			{
@@ -776,15 +772,15 @@ bool UMetaSoundSource::GetAllDefaultParameters(TArray<FAudioParameter>& OutParam
 
 bool UMetaSoundSource::IsParameterValid(const FAudioParameter& InParameter) const
 {
-	TMap<FName, FMetasoundFrontendVertex> InputNameTypeMap;
-	Algo::Transform(GetDocumentChecked().RootGraph.Interface.Inputs, InputNameTypeMap, [] (const FMetasoundFrontendClassInput& Input)
+	const TArray<FMetasoundFrontendClassInput>& Inputs = GetDocumentChecked().RootGraph.Interface.Inputs;
+	const FMetasoundFrontendVertex* Vertex = Algo::FindByPredicate(Inputs, [&InParameter] (const FMetasoundFrontendClassInput& Input)
 	{
-		return TPair<FName, FMetasoundFrontendVertex>(Input.Name, Input);
+		return Input.Name == InParameter.ParamName;
 	});
-	return IsParameterValid(InParameter, InputNameTypeMap);
+	return IsParameterValid(InParameter, Vertex);
 }
 
-bool UMetaSoundSource::IsParameterValid(const FAudioParameter& InParameter, const TMap<FName, FMetasoundFrontendVertex>& InInputNameVertexMap) const
+bool UMetaSoundSource::IsParameterValid(const FAudioParameter& InParameter, const FMetasoundFrontendVertex* InVertex) const
 {
 	using namespace Metasound;
 	using namespace Metasound::Frontend;
@@ -794,12 +790,16 @@ bool UMetaSoundSource::IsParameterValid(const FAudioParameter& InParameter, cons
 		return false;
 	}
 
-	const FMetasoundFrontendVertex* Vertex = InInputNameVertexMap.Find(InParameter.ParamName);
-	if (!Vertex)
+	if (!InVertex)
+	{
+		return false;
+	}
+
+	const FName& TypeName = InVertex->TypeName;
+	if (!InParameter.TypeName.IsNone() && InParameter.TypeName != TypeName)
 	{
 		return false;
 	}
-	const FName& TypeName = Vertex->TypeName;
 
 	bool bIsValid = false;
 	switch (InParameter.ParamType)
@@ -856,10 +856,7 @@ bool UMetaSoundSource::IsParameterValid(const FAudioParameter& InParameter, cons
 
 		case EAudioParameterType::Object:
 		{
-			FDataTypeRegistryInfo DataTypeInfo;
-			bIsValid = IDataTypeRegistry::Get().GetDataTypeInfo(InParameter.ObjectParam, DataTypeInfo);
-			bIsValid &= DataTypeInfo.bIsProxyParsable || DataTypeInfo.bIsUniquePtrProxyParsable_DEPRECATED;
-			bIsValid &= DataTypeInfo.DataTypeName == TypeName;
+			bIsValid = IDataTypeRegistry::Get().IsValidUObjectForDataType(TypeName, InParameter.ObjectParam);
 		}
 		break;
 
@@ -868,13 +865,9 @@ bool UMetaSoundSource::IsParameterValid(const FAudioParameter& InParameter, cons
 			bIsValid = true;
 
 			const FName ElementTypeName = CreateElementTypeNameFromArrayTypeName(TypeName);
-			for (UObject* Object : InParameter.ArrayObjectParam)
+			for (const UObject* Object : InParameter.ArrayObjectParam)
 			{
-				FDataTypeRegistryInfo DataTypeInfo;
-				bIsValid = IDataTypeRegistry::Get().GetDataTypeInfo(Object, DataTypeInfo);
-				bIsValid &= DataTypeInfo.bIsProxyParsable || DataTypeInfo.bIsUniquePtrProxyParsable_DEPRECATED;
-				bIsValid &= DataTypeInfo.DataTypeName == ElementTypeName;
-
+				bIsValid = IDataTypeRegistry::Get().IsValidUObjectForDataType(ElementTypeName, InParameter.ObjectParam);
 				if (!bIsValid)
 				{
 					break;
diff --git a/Engine/Plugins/Runtime/Metasound/Source/MetasoundEngine/Public/MetasoundSource.h b/Engine/Plugins/Runtime/Metasound/Source/MetasoundEngine/Public/MetasoundSource.h
--- a/Engine/Plugins/Runtime/Metasound/Source/MetasoundEngine/Public/MetasoundSource.h
+++ b/Engine/Plugins/Runtime/Metasound/Source/MetasoundEngine/Public/MetasoundSource.h
@@ -238,7 +238,7 @@ private:
 	}
 
 
-	bool IsParameterValid(const FAudioParameter& InParameter, const TMap<FName, FMetasoundFrontendVertex>& InInputNameVertexMap) const;
+	bool IsParameterValid(const FAudioParameter& InParameter, const FMetasoundFrontendVertex* InVertex) const;
 
 	static Metasound::SourcePrivate::FParameterRouter& GetParameterRouter();
 
diff --git a/Engine/Plugins/Runtime/Metasound/Source/MetasoundFrontend/Private/MetasoundFrontendDataTypeRegistry.cpp b/Engine/Plugins/Runtime/Metasound/Source/MetasoundFrontend/Private/MetasoundFrontendDataTypeRegistry.cpp
--- a/Engine/Plugins/Runtime/Metasound/Source/MetasoundFrontend/Private/MetasoundFrontendDataTypeRegistry.cpp
+++ b/Engine/Plugins/Runtime/Metasound/Source/MetasoundFrontend/Private/MetasoundFrontendDataTypeRegistry.cpp
@@ -517,8 +517,9 @@ namespace Metasound
 
 				virtual UClass* GetUClassForDataType(const FName& InDataType) const override;
 
-				bool IsUObjectProxyFactory(UObject* InObject) const override;
-				TSharedPtr<Audio::IProxyData> CreateProxyFromUObject(const FName& InDataType, UObject* InObject) const override;
+				virtual bool IsUObjectProxyFactory(UObject* InObject) const override;
+				virtual TSharedPtr<Audio::IProxyData> CreateProxyFromUObject(const FName& InDataType, UObject* InObject) const override;
+				virtual bool IsValidUObjectForDataType(const FName& InDataTypeName, const UObject* InUObject) const override;
 
 				virtual FLiteral CreateDefaultLiteral(const FName& InDataType) const override;
 				virtual FLiteral CreateLiteralFromUObject(const FName& InDataType, UObject* InObject) const override;
@@ -867,6 +868,41 @@ namespace Metasound
 				return ProxyPtr;
 			}
 
+			bool FDataTypeRegistry::IsValidUObjectForDataType(const FName& InDataTypeName, const UObject* InUObject) const
+			{
+				FDataTypeRegistryInfo DataTypeInfo;
+				if (GetDataTypeInfo(InDataTypeName, DataTypeInfo))
+				{
+					if (DataTypeInfo.bIsProxyParsable || DataTypeInfo.bIsUniquePtrProxyParsable_DEPRECATED)
+					{
+						if (InUObject)
+						{
+							FScopeLock ObjectRegistryLock(&RegistryObjectMapMutex);
+
+							const UClass* ObjectClass = InUObject->GetClass();
+							while (ObjectClass != UObject::StaticClass())
+							{
+								if (const FName* DataTypeSupportedByUObject = RegisteredObjectClasses.Find(ObjectClass))
+								{
+									if (*DataTypeSupportedByUObject == InDataTypeName)
+									{
+										return true;
+									}
+								}
+
+								ObjectClass = ObjectClass->GetOwnerClass();
+							}
+						}
+						else
+						{
+							return true;
+						}
+					}
+				}
+
+				return false;
+			}
+
 			FLiteral FDataTypeRegistry::CreateLiteralFromUObject(const FName& InDataType, UObject* InObject) const
 			{
 				TSharedPtr<Audio::IProxyData> ProxyPtr = CreateProxyFromUObject(InDataType, InObject);
@@ -1199,4 +1235,3 @@ namespace Metasound
 		}
 	}
 }
-
diff --git a/Engine/Plugins/Runtime/Metasound/Source/MetasoundFrontend/Public/MetasoundFrontendDataTypeRegistry.h b/Engine/Plugins/Runtime/Metasound/Source/MetasoundFrontend/Public/MetasoundFrontendDataTypeRegistry.h
--- a/Engine/Plugins/Runtime/Metasound/Source/MetasoundFrontend/Public/MetasoundFrontendDataTypeRegistry.h
+++ b/Engine/Plugins/Runtime/Metasound/Source/MetasoundFrontend/Public/MetasoundFrontendDataTypeRegistry.h
@@ -266,6 +266,7 @@ namespace Metasound
 			virtual bool IsLiteralTypeSupported(const FName& InDataType, ELiteralType InLiteralType) const = 0;
 			virtual bool IsLiteralTypeSupported(const FName& InDataType, EMetasoundFrontendLiteralType InLiteralType) const = 0;
 			virtual bool IsUObjectProxyFactory(UObject* InObject) const = 0;
+			virtual bool IsValidUObjectForDataType(const FName& InDataTypeName, const UObject* InUObject) const = 0;
 
 			virtual UClass* GetUClassForDataType(const FName& InDataType) const = 0;
 
diff --git a/Engine/Source/Runtime/AudioExtensions/Private/AudioParameter.cpp b/Engine/Source/Runtime/AudioExtensions/Private/AudioParameter.cpp
--- a/Engine/Source/Runtime/AudioExtensions/Private/AudioParameter.cpp
+++ b/Engine/Source/Runtime/AudioExtensions/Private/AudioParameter.cpp
@@ -131,10 +131,7 @@ void FAudioParameter::Merge(const FAudioParameter& InParameter, bool bInTakeName
 
 			for (const TSharedPtr<Audio::IProxyData>& ProxyPtr : InParameter.ObjectProxies)
 			{
-				if (ensure(ProxyPtr.IsValid()))
-				{
-					ObjectProxies.Emplace(ProxyPtr);
-				}
+				ObjectProxies.Emplace(ProxyPtr);
 			}
 		}
 		break;
@@ -146,10 +143,7 @@ void FAudioParameter::Merge(const FAudioParameter& InParameter, bool bInTakeName
 			ObjectProxies.Reset();
 			for (const TSharedPtr<Audio::IProxyData>& ProxyPtr : InParameter.ObjectProxies)
 			{
-				if (ensure(ProxyPtr.IsValid()))
-				{
-					ObjectProxies.Emplace(ProxyPtr);
-				}
+				ObjectProxies.Emplace(ProxyPtr);
 			}
 		}
 		break;
@@ -165,10 +159,7 @@ void FAudioParameter::Merge(const FAudioParameter& InParameter, bool bInTakeName
 
 			for (const TSharedPtr<Audio::IProxyData>& ProxyPtr : InParameter.ObjectProxies)
 			{
-				if (ensure(ProxyPtr.IsValid()))
-				{
-					ObjectProxies.Emplace(ProxyPtr);
-				}
+				ObjectProxies.Emplace(ProxyPtr);
 			}
 		}
 		break;
