Fix MetaSound binding issues on oscillator and noise nodes
Fix crash in array random get node during automated stress tests for metasounds.
#rnx

diff --git a/Engine/Plugins/Runtime/Metasound/Source/MetasoundFrontend/Public/MetasoundArrayRandomNode.h b/Engine/Plugins/Runtime/Metasound/Source/MetasoundFrontend/Public/MetasoundArrayRandomNode.h
--- a/Engine/Plugins/Runtime/Metasound/Source/MetasoundFrontend/Public/MetasoundArrayRandomNode.h
+++ b/Engine/Plugins/Runtime/Metasound/Source/MetasoundFrontend/Public/MetasoundArrayRandomNode.h
@@ -267,14 +267,19 @@ namespace Metasound
 
 			const ArrayType& InputArrayRef = *InputArray;
 
-			// If the array size has changed, we need to reinit before getting the next value
-			if (PrevArraySize != InputArrayRef.Num())
+			// Determine if the state of the random number generator needs to be
+			// reinitialized.
+			const bool bIsArrayNonEmpty = InputArrayRef.Num() != 0; // Skip reinit if the array is empty because it represents an invalid state for this node.
+			const bool bIsArraySizeChanged = PrevArraySize != InputArrayRef.Num(); // Need to reinit for array size changes. 
+			const bool bIsSharedStateEnablementInconsistent = (*bEnableSharedState != bSharedStateInitialized); // Need to reinit if the shared state enablement has been updated.
+
+			const bool bIsStateReinitializationNeeded = bIsArrayNonEmpty && (bIsArraySizeChanged || bIsSharedStateEnablementInconsistent);
+
+			PrevArraySize = InputArrayRef.Num();
+
+			if (bIsStateReinitializationNeeded)
 			{
-				PrevArraySize = InputArrayRef.Num();
-				if (PrevArraySize != 0)
-				{
-					InitializeState(PrevArraySize);
-				}
+				InitializeState(PrevArraySize);
 			}
 
 			if (PrevArraySize == 0)
diff --git a/Engine/Plugins/Runtime/Metasound/Source/MetasoundStandardNodes/Private/MetasoundNoiseGenerator.cpp b/Engine/Plugins/Runtime/Metasound/Source/MetasoundStandardNodes/Private/MetasoundNoiseGenerator.cpp
--- a/Engine/Plugins/Runtime/Metasound/Source/MetasoundStandardNodes/Private/MetasoundNoiseGenerator.cpp
+++ b/Engine/Plugins/Runtime/Metasound/Source/MetasoundStandardNodes/Private/MetasoundNoiseGenerator.cpp
@@ -50,16 +50,15 @@ namespace Metasound
 		static const FNodeClassMetadata& GetNodeInfo();
 		static FVertexInterface DeclareVertexInterface();
 
-		FNoiseOperator(const FOperatorSettings& InSettings, FInt32ReadRef&& InReadRef);
+		FNoiseOperator(const FOperatorSettings& InSettings, FInt32ReadRef&& InSeedReadRef, FEnumNoiseTypeReadRef&& InNoiseTypeReadRef);
 		virtual void BindInputs(FInputVertexInterfaceData& InOutVertexData) override;
 		virtual void BindOutputs(FOutputVertexInterfaceData& InOutVertexData) override;
-		virtual FDataReferenceCollection GetInputs() const override;
-		virtual FDataReferenceCollection GetOutputs() const override;
 
 		virtual IOperator::FPostExecuteFunction GetPostExecuteFunction() override { return nullptr; }
 
 	protected:
 		FInt32ReadRef Seed;
+		FEnumNoiseTypeReadRef NoiseType;
 		FAudioBufferWriteRef Out;
 		int32 OldSeed;
 
@@ -113,8 +112,8 @@ namespace Metasound
 	{
 		Audio::FWhiteNoise Generator;
 
-		FNoiseOperator_White(const FOperatorSettings& InSettings, FInt32ReadRef&& InSeed)
-			: FNoiseOperator{ InSettings, MoveTemp(InSeed) }
+		FNoiseOperator_White(const FOperatorSettings& InSettings, FInt32ReadRef&& InSeed, FEnumNoiseTypeReadRef&& InNoiseTypeReadRef)
+			: FNoiseOperator{ InSettings, MoveTemp(InSeed), MoveTemp(InNoiseTypeReadRef) }
 			, Generator{ MakeGenerator<Audio::FWhiteNoise>(*Seed) }
 		{}
 
@@ -146,8 +145,8 @@ namespace Metasound
 		using FNoiseOperator::FNoiseOperator;
 		Audio::FPinkNoise Generator;
 
-		FNoiseOperator_Pink(const FOperatorSettings& InSettings, FInt32ReadRef&& InSeed)
-			: FNoiseOperator{ InSettings, MoveTemp(InSeed) }
+		FNoiseOperator_Pink(const FOperatorSettings& InSettings, FInt32ReadRef&& InSeed, FEnumNoiseTypeReadRef&& InNoiseTypeReadRef)
+			: FNoiseOperator{ InSettings, MoveTemp(InSeed), MoveTemp(InNoiseTypeReadRef) }
 			, Generator{ MakeGenerator<Audio::FPinkNoise>(*Seed) }
 		{}
 
@@ -173,8 +172,9 @@ namespace Metasound
 		}
 	};
 	
-	FNoiseOperator::FNoiseOperator(const FOperatorSettings& InSettings, FInt32ReadRef&& InReadRef)
-		: Seed{ MoveTemp(InReadRef) }
+	FNoiseOperator::FNoiseOperator(const FOperatorSettings& InSettings, FInt32ReadRef&& InSeedReadRef, FEnumNoiseTypeReadRef&& InNoiseTypeReadRef)
+		: Seed{ MoveTemp(InSeedReadRef) }
+		, NoiseType{ MoveTemp(InNoiseTypeReadRef) }
 		, Out{ FAudioBufferWriteRef::CreateNew(InSettings) }
 		, OldSeed(*Seed)
 	{}
@@ -184,6 +184,7 @@ namespace Metasound
 	{
 		using namespace NoiseGeneratorVertexNames;
 		InOutVertexData.BindReadVertex(METASOUND_GET_PARAM_NAME(InputSeed), Seed);
+		InOutVertexData.BindReadVertex(METASOUND_GET_PARAM_NAME(InputType), NoiseType);
 	}
 
 	void FNoiseOperator::BindOutputs(FOutputVertexInterfaceData& InOutVertexData)
@@ -192,22 +193,6 @@ namespace Metasound
 		InOutVertexData.BindReadVertex(METASOUND_GET_PARAM_NAME(OutAudio), Out);
 	}
 
-	FDataReferenceCollection FNoiseOperator::GetInputs() const
-	{
-		// This should never be called. Bind(...) is called instead. This method
-		// exists as a stop-gap until the API can be deprecated and removed.
-		checkNoEntry();
-		return {};
-	}
-
-	FDataReferenceCollection FNoiseOperator::GetOutputs() const
-	{
-		// This should never be called. Bind(...) is called instead. This method
-		// exists as a stop-gap until the API can be deprecated and removed.
-		checkNoEntry();
-		return {};
-	}
-
 	FVertexInterface FNoiseOperator::DeclareVertexInterface()
 	{
 		using namespace NoiseGeneratorVertexNames; 
@@ -272,9 +257,9 @@ namespace Metasound
 		{
 		default:
 		case ENoiseType::White:
-			return MakeUnique<FNoiseOperator_White>(InParams.OperatorSettings, MoveTemp(Seed));
+			return MakeUnique<FNoiseOperator_White>(InParams.OperatorSettings, MoveTemp(Seed), MoveTemp(Type));
 		case ENoiseType::Pink:
-			return MakeUnique<FNoiseOperator_Pink>(InParams.OperatorSettings, MoveTemp(Seed));
+			return MakeUnique<FNoiseOperator_Pink>(InParams.OperatorSettings, MoveTemp(Seed), MoveTemp(Type));
 		}
 		checkNoEntry();
 		return nullptr;
diff --git a/Engine/Plugins/Runtime/Metasound/Source/MetasoundStandardNodes/Private/MetasoundOscillatorNodes.cpp b/Engine/Plugins/Runtime/Metasound/Source/MetasoundStandardNodes/Private/MetasoundOscillatorNodes.cpp
--- a/Engine/Plugins/Runtime/Metasound/Source/MetasoundStandardNodes/Private/MetasoundOscillatorNodes.cpp
+++ b/Engine/Plugins/Runtime/Metasound/Source/MetasoundStandardNodes/Private/MetasoundOscillatorNodes.cpp
@@ -35,6 +35,7 @@ namespace Metasound
 		METASOUND_PARAM(PhaseOffsetPin, "Phase Offset", "Phase Offset In Degrees (0..360)")
 		METASOUND_PARAM(GlideFactorPin, "Glide", "The amount of glide to use when changing frequencies. 0.0 = no glide, 1.0 = lots of glide.")
 		METASOUND_PARAM(AudioOutPin, "Audio", "The output audio")
+		static const FLazyName TypePinName("Type");
 	}
 
 	// Base class of Oscillator factories which holds common the interface.
@@ -118,22 +119,6 @@ namespace Metasound
 			InOutVertexData.BindReadVertex(METASOUND_GET_PARAM_NAME(AudioOutPin), AudioBuffer);
 		}
 
-		virtual FDataReferenceCollection GetInputs() const override
-		{
-			// This should never be called. Bind(...) is called instead. This method
-			// exists as a stop-gap until the API can be deprecated and removed.
-			checkNoEntry();
-			return {};
-		}
-
-		virtual FDataReferenceCollection GetOutputs() const override
-		{
-			// This should never be called. Bind(...) is called instead. This method
-			// exists as a stop-gap until the API can be deprecated and removed.
-			checkNoEntry();
-			return {};
-		}
-
 		void ResetPhase(float InPhaseInDegrees)
 		{
 			float LinearPhase = FMath::Clamp(InPhaseInDegrees, 0.f, 360.f) / 360.f;
@@ -186,13 +171,14 @@ namespace Metasound
 	};
 
 	// Generic Oscillator operator for NON-FM Operators. 
-	template<typename GeneratorPolicy>
-	class TOscillatorOperator final : public TOscillatorOperatorBase<GeneratorPolicy, TOscillatorOperator<GeneratorPolicy>>
+	template<typename GeneratorPolicy, typename TGeneratorEnumReadRefType>
+	class TOscillatorOperator final : public TOscillatorOperatorBase<GeneratorPolicy, TOscillatorOperator<GeneratorPolicy, TGeneratorEnumReadRefType>>
 	{
-		using Super = TOscillatorOperatorBase<GeneratorPolicy, TOscillatorOperator<GeneratorPolicy>>;
+		using Super = TOscillatorOperatorBase<GeneratorPolicy, TOscillatorOperator<GeneratorPolicy, TGeneratorEnumReadRefType>>;
 	public:
-		TOscillatorOperator(const FOscillatorOperatorConstructParams& InConstructParams)
+		TOscillatorOperator(const FOscillatorOperatorConstructParams& InConstructParams, const TGeneratorEnumReadRefType& InGeneratorType)
 			: Super(InConstructParams)
+			, GeneratorType(InGeneratorType)
 		{}
 
 		void Generate(int32 InStartFrame, int32 InEndFrame, float InClampedFreq, float InClampedGlideEase)
@@ -212,41 +198,39 @@ namespace Metasound
 				});
 			}			
 		}
+
+		virtual void BindInputs(FInputVertexInterfaceData& InOutVertexData) override
+		{
+			using namespace OscillatorCommonVertexNames;
+			Super::BindInputs(InOutVertexData);
+			InOutVertexData.BindReadVertex(TypePinName, GeneratorType);
+		}
+
+	private:
+		TGeneratorEnumReadRefType GeneratorType;
 	};
 
 	// Generic Oscillator operator for FM Operators.
-	template<typename TGeneratorPolicy>
-	class TOscillatorOperatorFM final : public TOscillatorOperatorBase<TGeneratorPolicy, TOscillatorOperatorFM<TGeneratorPolicy>>
+	template<typename TGeneratorPolicy, typename TGeneratorEnumReadRefType>
+	class TOscillatorOperatorFM final : public TOscillatorOperatorBase<TGeneratorPolicy, TOscillatorOperatorFM<TGeneratorPolicy, TGeneratorEnumReadRefType>>
 	{
-		using Super = TOscillatorOperatorBase<TGeneratorPolicy, TOscillatorOperatorFM<TGeneratorPolicy>>;
+		using Super = TOscillatorOperatorBase<TGeneratorPolicy, TOscillatorOperatorFM<TGeneratorPolicy, TGeneratorEnumReadRefType>>;
 	public:
-		TOscillatorOperatorFM(const FOscillatorOperatorConstructParams& InCommonParams, const FAudioBufferReadRef& InFmData)
-			: Super(InCommonParams), Fm(InFmData)
+		TOscillatorOperatorFM(const FOscillatorOperatorConstructParams& InCommonParams, const FAudioBufferReadRef& InFmData, const TGeneratorEnumReadRefType& InGeneratorType)
+			: Super(InCommonParams)
+			, Fm(InFmData)
+			, GeneratorType(InGeneratorType)
 		{}
 
 
 		virtual void BindInputs(FInputVertexInterfaceData& InOutVertexData) override
 		{
 			using namespace OscillatorCommonVertexNames;
+			Super::BindInputs(InOutVertexData);
+			InOutVertexData.BindReadVertex(TypePinName, GeneratorType);
 			InOutVertexData.BindReadVertex(METASOUND_GET_PARAM_NAME(FrequencyModPin), Fm);
 		}
 
-		virtual FDataReferenceCollection GetInputs() const override
-		{
-			// This should never be called. Bind(...) is called instead. This method
-			// exists as a stop-gap until the API can be deprecated and removed.
-			checkNoEntry();
-			return {};
-		}
-
-		virtual FDataReferenceCollection GetOutputs() const override
-		{
-			// This should never be called. Bind(...) is called instead. This method
-			// exists as a stop-gap until the API can be deprecated and removed.
-			checkNoEntry();
-			return {};
-		}
-
 		void Generate(int32 InStartFrame, int32 InEndFrame, float InClampedFreq, float InClampedGlideEase)
 		{
 			int32 NumFrames = InEndFrame - InStartFrame;
@@ -267,6 +251,7 @@ namespace Metasound
 
 	private:
 		FAudioBufferReadRef Fm;
+		TGeneratorEnumReadRefType GeneratorType;
 	};
 
 	FOscilatorNodeBase::FOscilatorNodeBase(const FVertexName& InInstanceName, const FGuid& InInstanceID, const FNodeClassMetadata& InInfo, const TSharedRef<IOperatorFactory, ESPMode::ThreadSafe>& InFactory, float InDefaultFrequency, float InDefaultGlideFactor, bool bInDefaultEnablement)
@@ -358,19 +343,19 @@ namespace Metasound
 			FOscillatorOperatorConstructParams OpParams
 			{
 				Settings,
-				InputData.GetOrConstructDataReadReference<bool>(METASOUND_GET_PARAM_NAME(EnabledPin), SineNode.GetDefaultEnablement()),
-				InputData.GetOrConstructDataReadReference<float>(METASOUND_GET_PARAM_NAME(OscBaseFrequencyPin), SineNode.GetDefaultFrequency()),
-				InputData.GetOrConstructDataReadReference<float>(METASOUND_GET_PARAM_NAME(PhaseOffsetPin), SineNode.GetDefaultPhaseOffset()),
-				InputData.GetOrConstructDataReadReference<FTrigger>(METASOUND_GET_PARAM_NAME(OscPhaseResetPin), Settings),
-				InputData.GetOrConstructDataReadReference<float>(METASOUND_GET_PARAM_NAME(GlideFactorPin), SineNode.GetDefaultGlideFactor()),
-				InputData.GetOrConstructDataReadReference<bool>(METASOUND_GET_PARAM_NAME(BiPolarPin), true)
+				InputData.GetOrCreateDefaultDataReadReference<bool>(METASOUND_GET_PARAM_NAME(EnabledPin), Settings),
+				InputData.GetOrCreateDefaultDataReadReference<float>(METASOUND_GET_PARAM_NAME(OscBaseFrequencyPin), Settings),
+				InputData.GetOrCreateDefaultDataReadReference<float>(METASOUND_GET_PARAM_NAME(PhaseOffsetPin), Settings),
+				InputData.GetOrCreateDefaultDataReadReference<FTrigger>(METASOUND_GET_PARAM_NAME(OscPhaseResetPin), Settings),
+				InputData.GetOrCreateDefaultDataReadReference<float>(METASOUND_GET_PARAM_NAME(GlideFactorPin), Settings),
+				InputData.GetOrCreateDefaultDataReadReference<bool>(METASOUND_GET_PARAM_NAME(BiPolarPin), Settings)
 			};
 
 			// TODO: Make this a static prop. For now its a pin.
 			
 			// Check to see if we have an FM input connected.
 			bool bHasFM = InputData.IsVertexBound(METASOUND_GET_PARAM_NAME(FrequencyModPin));
-			FEnumSineGenerationTypeReadRef Type = InputData.GetOrConstructDataReadReference<FEnumSineGenerationType>(METASOUND_GET_PARAM_NAME(SineType), ESineGenerationType::Wavetable);
+			FEnumSineGenerationTypeReadRef Type = InputData.GetOrCreateDefaultDataReadReference<FEnumSineGenerationType>(METASOUND_GET_PARAM_NAME(SineType), Settings);
 			if (bHasFM)
 			{
 
@@ -379,9 +364,9 @@ namespace Metasound
 				switch (*Type)
 				{
 				default:
-				case ESineGenerationType::Sinf: return MakeUnique<TOscillatorOperatorFM<FSinfWithFm>>(OpParams, FmBuffer);
-				case ESineGenerationType::Bhaskara: return MakeUnique<TOscillatorOperatorFM<FBhaskaraWithFm>>(OpParams, FmBuffer);
-				case ESineGenerationType::Wavetable: return MakeUnique<TOscillatorOperatorFM<FSineWaveTableWithFm>>(OpParams, FmBuffer);
+				case ESineGenerationType::Sinf: return MakeUnique<TOscillatorOperatorFM<FSinfWithFm, FEnumSineGenerationTypeReadRef>>(OpParams, FmBuffer, Type);
+				case ESineGenerationType::Bhaskara: return MakeUnique<TOscillatorOperatorFM<FBhaskaraWithFm, FEnumSineGenerationTypeReadRef>>(OpParams, FmBuffer, Type);
+				case ESineGenerationType::Wavetable: return MakeUnique<TOscillatorOperatorFM<FSineWaveTableWithFm, FEnumSineGenerationTypeReadRef>>(OpParams, FmBuffer, Type);
 				}
 			}
 			else //HasFM
@@ -389,10 +374,10 @@ namespace Metasound
 				switch (*Type)
 				{
 				default:
-				case ESineGenerationType::Sinf: return MakeUnique<TOscillatorOperator<FSinf>>(OpParams);
-				case ESineGenerationType::Rotation: return MakeUnique<TOscillatorOperator<F2DRotatorGenerateBlock>>(OpParams);
-				case ESineGenerationType::Bhaskara: return MakeUnique<TOscillatorOperator<FBhaskara>>(OpParams);
-				case ESineGenerationType::Wavetable: return MakeUnique<TOscillatorOperator<FSineWaveTable>>(OpParams);
+				case ESineGenerationType::Sinf: return MakeUnique<TOscillatorOperator<FSinf, FEnumSineGenerationTypeReadRef>>(OpParams, Type);
+				case ESineGenerationType::Rotation: return MakeUnique<TOscillatorOperator<F2DRotatorGenerateBlock, FEnumSineGenerationTypeReadRef>>(OpParams, Type);
+				case ESineGenerationType::Bhaskara: return MakeUnique<TOscillatorOperator<FBhaskara, FEnumSineGenerationTypeReadRef>>(OpParams, Type);
+				case ESineGenerationType::Wavetable: return MakeUnique<TOscillatorOperator<FSineWaveTable, FEnumSineGenerationTypeReadRef>>(OpParams, Type);
 				}
 			} // HasFM
 			return nullptr;
@@ -488,17 +473,17 @@ namespace Metasound
 		using namespace OscillatorCommonVertexNames;
 		using namespace SawOscilatorVertexNames;
 
-		FSawGenerationTypeReadRef Type = InputData.GetOrConstructDataReadReference<FEnumSawGenerationType>(METASOUND_GET_PARAM_NAME(SawType));
+		FSawGenerationTypeReadRef Type = InputData.GetOrCreateDefaultDataReadReference<FEnumSawGenerationType>(METASOUND_GET_PARAM_NAME(SawType), Settings);
 
 		FOscillatorOperatorConstructParams OpParams
 		{
 			Settings,
-			InputData.GetOrConstructDataReadReference<bool>(METASOUND_GET_PARAM_NAME(EnabledPin), Node.GetDefaultEnablement()),
-			InputData.GetOrConstructDataReadReference<float>(METASOUND_GET_PARAM_NAME(OscBaseFrequencyPin), Node.GetDefaultFrequency()),
-			InputData.GetOrConstructDataReadReference<float>(METASOUND_GET_PARAM_NAME(PhaseOffsetPin), Node.GetDefaultPhaseOffset()),
-			InputData.GetOrConstructDataReadReference<FTrigger>(METASOUND_GET_PARAM_NAME(OscPhaseResetPin), Settings),
-			InputData.GetOrConstructDataReadReference<float>(METASOUND_GET_PARAM_NAME(GlideFactorPin), Node.GetDefaultGlideFactor()),
-			InputData.GetOrConstructDataReadReference<bool>(METASOUND_GET_PARAM_NAME(BiPolarPin), true)
+			InputData.GetOrCreateDefaultDataReadReference<bool>(METASOUND_GET_PARAM_NAME(EnabledPin), Settings),
+			InputData.GetOrCreateDefaultDataReadReference<float>(METASOUND_GET_PARAM_NAME(OscBaseFrequencyPin), Settings),
+			InputData.GetOrCreateDefaultDataReadReference<float>(METASOUND_GET_PARAM_NAME(PhaseOffsetPin), Settings),
+			InputData.GetOrCreateDefaultDataReadReference<FTrigger>(METASOUND_GET_PARAM_NAME(OscPhaseResetPin), Settings),
+			InputData.GetOrCreateDefaultDataReadReference<float>(METASOUND_GET_PARAM_NAME(GlideFactorPin), Settings),
+			InputData.GetOrCreateDefaultDataReadReference<bool>(METASOUND_GET_PARAM_NAME(BiPolarPin), Settings)
 		};
 
 		bool bHasFM = InputData.IsVertexBound(METASOUND_GET_PARAM_NAME(FrequencyModPin));
@@ -509,8 +494,8 @@ namespace Metasound
 			switch (*Type)
 			{
 			default:
-			case ESawGenerationType::Trivial: return MakeUnique<TOscillatorOperatorFM<FSawWithFm>>(OpParams, FmBuffer);
-			case ESawGenerationType::PolySmooth: return MakeUnique<TOscillatorOperatorFM<FSawPolysmoothWithFm>>(OpParams, FmBuffer);
+			case ESawGenerationType::Trivial: return MakeUnique<TOscillatorOperatorFM<FSawWithFm, FSawGenerationTypeReadRef>>(OpParams, FmBuffer, Type);
+			case ESawGenerationType::PolySmooth: return MakeUnique<TOscillatorOperatorFM<FSawPolysmoothWithFm, FSawGenerationTypeReadRef>>(OpParams, FmBuffer, Type);
 			}
 		}
 		else
@@ -518,8 +503,8 @@ namespace Metasound
 			switch (*Type)
 			{
 			default:
-			case ESawGenerationType::Trivial: return MakeUnique<TOscillatorOperator<FSaw>>(OpParams);
-			case ESawGenerationType::PolySmooth: return MakeUnique<TOscillatorOperator<FSawPolysmooth>>(OpParams);
+			case ESawGenerationType::Trivial: return MakeUnique<TOscillatorOperator<FSaw, FSawGenerationTypeReadRef>>(OpParams, Type);
+			case ESawGenerationType::PolySmooth: return MakeUnique<TOscillatorOperator<FSawPolysmooth, FSawGenerationTypeReadRef>>(OpParams, Type);
 			}
 		}
 		return nullptr;
@@ -554,16 +539,34 @@ namespace Metasound
 		DEFINE_METASOUND_ENUM_ENTRY(ESquareGenerationType::Trivial, "SquareTrivialDescription", "Trivial", "SquareTrivialDescriptionTT", "The most basic raw implementation. Does not obey pulse width."),
 		//DEFINE_METASOUND_ENUM_ENTRY(ESquareGenerationType::Wavetable, "SquareWavetableDescription", "Wavetable", "SquareWavetableDescriptionTT", "Use a Wavetable interpolation to generate the Waveform")
 	DEFINE_METASOUND_ENUM_END()
+
+	namespace SquareOscillatorVertexNames
+	{
+		METASOUND_PARAM(SquarePulseWidthPin, "Pulse Width", "The Width of the square part of the wave")
+		METASOUND_PARAM(SquareTypePin, "Type", "The generator type to make the squarewave")
+	}
 	
 	template<typename GeneratorPolicy>
 	class FSquareOperator final : public TOscillatorOperatorBase<GeneratorPolicy, FSquareOperator<GeneratorPolicy>>
 	{
 		using Super = TOscillatorOperatorBase<GeneratorPolicy, FSquareOperator<GeneratorPolicy>>;
 	public:
-		FSquareOperator(const FOscillatorOperatorConstructParams& InConstructParams, const FFloatReadRef& InPulseWidth)
+		FSquareOperator(const FOscillatorOperatorConstructParams& InConstructParams, const FFloatReadRef& InPulseWidth, const FSquareGenerationTypeReadRef& InGeneratorType)
 			: Super(InConstructParams)
 			, PulseWidth(InPulseWidth)
+			, GeneratorType(InGeneratorType)
 		{}
+
+		virtual void BindInputs(FInputVertexInterfaceData& InOutVertexData) override
+		{
+			using namespace SquareOscillatorVertexNames;
+			using namespace OscillatorCommonVertexNames;
+
+			Super::BindInputs(InOutVertexData);
+			InOutVertexData.BindReadVertex(METASOUND_GET_PARAM_NAME(SquarePulseWidthPin), PulseWidth);
+			InOutVertexData.BindReadVertex(TypePinName, GeneratorType);
+		}
+
 		void Generate(int32 InStartFrame, int32 InEndFrame, float InClampedFreq, float InClampedGlideEase)
 		{
 			int32 NumFrames = InEndFrame - InStartFrame;
@@ -584,6 +587,7 @@ namespace Metasound
 
 	private:
 		FFloatReadRef PulseWidth;
+		FSquareGenerationTypeReadRef GeneratorType;
 	};
 
 	template<typename GeneratorPolicy>
@@ -591,10 +595,11 @@ namespace Metasound
 	{
 		using Super = TOscillatorOperatorBase<GeneratorPolicy, FSquareOperatorFM<GeneratorPolicy>>;
 	public:
-		FSquareOperatorFM(const FOscillatorOperatorConstructParams& InConstructParams, const FFloatReadRef& InPulseWidth, const FAudioBufferReadRef& InFm)
+		FSquareOperatorFM(const FOscillatorOperatorConstructParams& InConstructParams, const FFloatReadRef& InPulseWidth, const FAudioBufferReadRef& InFm, const FSquareGenerationTypeReadRef& InGeneratorType)
 			: Super(InConstructParams)
 			, PulseWidth(InPulseWidth)
 			, FM(InFm)
+			, GeneratorType(InGeneratorType)
 		{
 			check(InFm->GetData());
 			check(InConstructParams.Settings.GetNumFramesPerBlock() == InFm->Num());
@@ -618,16 +623,23 @@ namespace Metasound
 			}
 		}
 
+		virtual void BindInputs(FInputVertexInterfaceData& InOutVertexData) override
+		{
+			using namespace SquareOscillatorVertexNames;
+			using namespace OscillatorCommonVertexNames;
+
+			Super::BindInputs(InOutVertexData);
+			InOutVertexData.BindReadVertex(METASOUND_GET_PARAM_NAME(SquarePulseWidthPin), PulseWidth);
+			InOutVertexData.BindReadVertex(METASOUND_GET_PARAM_NAME(FrequencyModPin), FM);
+			InOutVertexData.BindReadVertex(TypePinName, GeneratorType);
+		}
+
 	private:
 		FFloatReadRef PulseWidth;
 		FAudioBufferReadRef FM;
+		FSquareGenerationTypeReadRef GeneratorType;
 	};
 
-	namespace SquareOscillatorVertexNames
-	{
-		METASOUND_PARAM(SquarePulseWidthPin, "Pulse Width", "The Width of the square part of the wave")
-		METASOUND_PARAM(SquareTypePin, "Type", "The generator type to make the squarewave")
-	}
 
 	class FSquareOscilatorNode::FFactory : public FOscilatorFactoryBase
 	{
@@ -647,18 +659,18 @@ namespace Metasound
 			FOscillatorOperatorConstructParams OpParams
 			{
 				Settings,
-				InputData.GetOrConstructDataReadReference<bool>(METASOUND_GET_PARAM_NAME(EnabledPin), Node.GetDefaultEnablement()),
-				InputData.GetOrConstructDataReadReference<float>(METASOUND_GET_PARAM_NAME(OscBaseFrequencyPin), Node.GetDefaultFrequency()),
-				InputData.GetOrConstructDataReadReference<float>(METASOUND_GET_PARAM_NAME(PhaseOffsetPin), Node.GetDefaultPhaseOffset()),
-				InputData.GetOrConstructDataReadReference<FTrigger>(METASOUND_GET_PARAM_NAME(OscPhaseResetPin), Settings),
-				InputData.GetOrConstructDataReadReference<float>(METASOUND_GET_PARAM_NAME(GlideFactorPin), Node.GetDefaultGlideFactor()),
-				InputData.GetOrConstructDataReadReference<bool>(METASOUND_GET_PARAM_NAME(BiPolarPin), true)
+				InputData.GetOrCreateDefaultDataReadReference<bool>(METASOUND_GET_PARAM_NAME(EnabledPin), Settings),
+				InputData.GetOrCreateDefaultDataReadReference<float>(METASOUND_GET_PARAM_NAME(OscBaseFrequencyPin), Settings),
+				InputData.GetOrCreateDefaultDataReadReference<float>(METASOUND_GET_PARAM_NAME(PhaseOffsetPin), Settings),
+				InputData.GetOrCreateDefaultDataReadReference<FTrigger>(METASOUND_GET_PARAM_NAME(OscPhaseResetPin), Settings),
+				InputData.GetOrCreateDefaultDataReadReference<float>(METASOUND_GET_PARAM_NAME(GlideFactorPin), Settings),
+				InputData.GetOrCreateDefaultDataReadReference<bool>(METASOUND_GET_PARAM_NAME(BiPolarPin), Settings)
 			};
 			
-			FSquareGenerationTypeReadRef Type = InputData.GetOrConstructDataReadReference<FEnumSquareGenerationType>(METASOUND_GET_PARAM_NAME(SquareTypePin));
+			FSquareGenerationTypeReadRef Type = InputData.GetOrCreateDefaultDataReadReference<FEnumSquareGenerationType>(METASOUND_GET_PARAM_NAME(SquareTypePin), Settings);
 
 			bool bHasFM = InputData.IsVertexBound(METASOUND_GET_PARAM_NAME(FrequencyModPin));
-			FFloatReadRef PulseWidth = InputData.GetOrConstructDataReadReference<float>(METASOUND_GET_PARAM_NAME(SquarePulseWidthPin), 0.5f);
+			FFloatReadRef PulseWidth = InputData.GetOrCreateDefaultDataReadReference<float>(METASOUND_GET_PARAM_NAME(SquarePulseWidthPin), Settings);
 
 			if (bHasFM)
 			{
@@ -666,8 +678,8 @@ namespace Metasound
 				switch (*Type)
 				{
 				default:				
-				case ESquareGenerationType::Trivial: return MakeUnique<FSquareOperatorFM<FSquareWithFm>>(OpParams, PulseWidth, FmBuffer);
-				case ESquareGenerationType::PolySmooth: return MakeUnique<FSquareOperatorFM<FSquarePolysmoothWithFm>>(OpParams, PulseWidth, FmBuffer);
+				case ESquareGenerationType::Trivial: return MakeUnique<FSquareOperatorFM<FSquareWithFm>>(OpParams, PulseWidth, FmBuffer, Type);
+				case ESquareGenerationType::PolySmooth: return MakeUnique<FSquareOperatorFM<FSquarePolysmoothWithFm>>(OpParams, PulseWidth, FmBuffer, Type);
 				}
 			}
 			else
@@ -675,8 +687,8 @@ namespace Metasound
 				switch (*Type)
 				{
 				default:
-				case ESquareGenerationType::Trivial: return MakeUnique<FSquareOperator<FSquare>>(OpParams, PulseWidth);
-				case ESquareGenerationType::PolySmooth: return MakeUnique<FSquareOperator<FSquarePolysmooth>>(OpParams, PulseWidth);
+				case ESquareGenerationType::Trivial: return MakeUnique<FSquareOperator<FSquare>>(OpParams, PulseWidth, Type);
+				case ESquareGenerationType::PolySmooth: return MakeUnique<FSquareOperator<FSquarePolysmooth>>(OpParams, PulseWidth, Type);
 				}
 			}
 			return nullptr;
@@ -767,17 +779,17 @@ namespace Metasound
 			using namespace OscillatorCommonVertexNames;
 			using namespace TriangleOscilatorVertexNames;
 			
-			FTriangleGenerationTypeReadRef Type = InputData.GetOrConstructDataReadReference<FEnumTriangleGenerationType>(METASOUND_GET_PARAM_NAME(TriangeTypePin));
+			FTriangleGenerationTypeReadRef Type = InputData.GetOrCreateDefaultDataReadReference<FEnumTriangleGenerationType>(METASOUND_GET_PARAM_NAME(TriangeTypePin), Settings);
 		
 			FOscillatorOperatorConstructParams OpParams
 			{
 				Settings,
-				InputData.GetOrConstructDataReadReference<bool>(METASOUND_GET_PARAM_NAME(EnabledPin), Node.GetDefaultEnablement()),
-				InputData.GetOrConstructDataReadReference<float>(METASOUND_GET_PARAM_NAME(OscBaseFrequencyPin), Node.GetDefaultFrequency()),
-				InputData.GetOrConstructDataReadReference<float>(METASOUND_GET_PARAM_NAME(PhaseOffsetPin), Node.GetDefaultPhaseOffset()),
-				InputData.GetOrConstructDataReadReference<FTrigger>(METASOUND_GET_PARAM_NAME(OscPhaseResetPin), Settings),
-				InputData.GetOrConstructDataReadReference<float>(METASOUND_GET_PARAM_NAME(GlideFactorPin), Node.GetDefaultGlideFactor()),
-				InputData.GetOrConstructDataReadReference<bool>(METASOUND_GET_PARAM_NAME(BiPolarPin), true)
+				InputData.GetOrCreateDefaultDataReadReference<bool>(METASOUND_GET_PARAM_NAME(EnabledPin), Settings),
+				InputData.GetOrCreateDefaultDataReadReference<float>(METASOUND_GET_PARAM_NAME(OscBaseFrequencyPin), Settings),
+				InputData.GetOrCreateDefaultDataReadReference<float>(METASOUND_GET_PARAM_NAME(PhaseOffsetPin), Settings),
+				InputData.GetOrCreateDefaultDataReadReference<FTrigger>(METASOUND_GET_PARAM_NAME(OscPhaseResetPin), Settings),
+				InputData.GetOrCreateDefaultDataReadReference<float>(METASOUND_GET_PARAM_NAME(GlideFactorPin), Settings),
+				InputData.GetOrCreateDefaultDataReadReference<bool>(METASOUND_GET_PARAM_NAME(BiPolarPin), Settings)
 			};
 
 			bool bHasFM = InputData.IsVertexBound(METASOUND_GET_PARAM_NAME(FrequencyModPin));
@@ -788,8 +800,8 @@ namespace Metasound
 				switch (*Type)
 				{
 				default:
-				case ETriangleGenerationType::PolySmooth: return MakeUnique<TOscillatorOperatorFM<FTrianglePolysmoothWithFm>>(OpParams, FmBuffer);
-				case ETriangleGenerationType::Trivial: return MakeUnique<TOscillatorOperatorFM<FTriangleWithFm>>(OpParams, FmBuffer);
+				case ETriangleGenerationType::PolySmooth: return MakeUnique<TOscillatorOperatorFM<FTrianglePolysmoothWithFm, FTriangleGenerationTypeReadRef>>(OpParams, FmBuffer, Type);
+				case ETriangleGenerationType::Trivial: return MakeUnique<TOscillatorOperatorFM<FTriangleWithFm, FTriangleGenerationTypeReadRef>>(OpParams, FmBuffer, Type);
 				}
 			}
 			else
@@ -797,8 +809,8 @@ namespace Metasound
 				switch (*Type)
 				{
 				default:
-				case ETriangleGenerationType::PolySmooth: return MakeUnique<TOscillatorOperator<FTrianglePolysmooth>>(OpParams);
-				case ETriangleGenerationType::Trivial: return MakeUnique<TOscillatorOperator<FTriangle>>(OpParams);
+				case ETriangleGenerationType::PolySmooth: return MakeUnique<TOscillatorOperator<FTrianglePolysmooth, FTriangleGenerationTypeReadRef>>(OpParams, Type);
+				case ETriangleGenerationType::Trivial: return MakeUnique<TOscillatorOperator<FTriangle, FTriangleGenerationTypeReadRef>>(OpParams, Type);
 				}
 			}
 			return nullptr;
@@ -927,24 +939,23 @@ namespace Metasound
 			return Info;
 		}
 		
-		static TUniquePtr<IOperator> CreateOperator(const FCreateOperatorParams& InParams, FBuildErrorArray& OutErrors)
+		static TUniquePtr<IOperator> CreateOperator(const FBuildOperatorParams& InParams, FBuildResults& OutResults)
 		{
 			const FLfoNode& Node = static_cast<const FLfoNode&>(InParams.Node);
-			const FDataReferenceCollection& InputCol = InParams.InputDataReferences;
+			const FInputVertexInterfaceData& InputData = InParams.InputData;
 			const FOperatorSettings& Settings = InParams.OperatorSettings;
-			const FInputVertexInterface& InputInterface = GetVertexInterface().GetInputInterface();
 
 			using namespace LfoVertexNames;
 
 			return MakeUnique<FLfoOperator>(
 				  Settings
-				, InputCol.GetDataReadReferenceOrConstructWithVertexDefault<float>(InputInterface, METASOUND_GET_PARAM_NAME(LfoBaseFrequencyPin), Settings)
-				, InputCol.GetDataReadReferenceOrConstructWithVertexDefault<FEnumLfoWaveshapeType>(InputInterface, METASOUND_GET_PARAM_NAME(WaveshapePin), Settings)
-				, InputCol.GetDataReadReferenceOrConstructWithVertexDefault<float>(InputInterface, METASOUND_GET_PARAM_NAME(MinOutputValuePin), Settings)
-				, InputCol.GetDataReadReferenceOrConstructWithVertexDefault<float>(InputInterface, METASOUND_GET_PARAM_NAME(MaxOutputValuePin), Settings)
-				, InputCol.GetDataReadReferenceOrConstruct<FTrigger>(METASOUND_GET_PARAM_NAME(LfoPhaseResetPin), Settings)
-				, InputCol.GetDataReadReferenceOrConstructWithVertexDefault<float>(InputInterface, METASOUND_GET_PARAM_NAME(PhaseOffsetPin), Settings)
-				, InputCol.GetDataReadReferenceOrConstructWithVertexDefault<float>(InputInterface, METASOUND_GET_PARAM_NAME(LfoPulseWidthPin), Settings)
+				, InputData.GetOrCreateDefaultDataReadReference<float>(METASOUND_GET_PARAM_NAME(LfoBaseFrequencyPin), Settings)
+				, InputData.GetOrCreateDefaultDataReadReference<FEnumLfoWaveshapeType>(METASOUND_GET_PARAM_NAME(WaveshapePin), Settings)
+				, InputData.GetOrCreateDefaultDataReadReference<float>(METASOUND_GET_PARAM_NAME(MinOutputValuePin), Settings)
+				, InputData.GetOrCreateDefaultDataReadReference<float>(METASOUND_GET_PARAM_NAME(MaxOutputValuePin), Settings)
+				, InputData.GetOrCreateDefaultDataReadReference<FTrigger>(METASOUND_GET_PARAM_NAME(LfoPhaseResetPin), Settings)
+				, InputData.GetOrCreateDefaultDataReadReference<float>(METASOUND_GET_PARAM_NAME(PhaseOffsetPin), Settings)
+				, InputData.GetOrCreateDefaultDataReadReference<float>(METASOUND_GET_PARAM_NAME(LfoPulseWidthPin), Settings)
 			);
 		}
 
@@ -984,22 +995,6 @@ namespace Metasound
 			InOutVertexData.BindReadVertex(METASOUND_GET_PARAM_NAME(LfoOutPin), Output);
 		}
 
-		virtual FDataReferenceCollection GetInputs() const override
-		{
-			// This should never be called. Bind(...) is called instead. This method
-			// exists as a stop-gap until the API can be deprecated and removed.
-			checkNoEntry();
-			return {};
-		}
-
-		virtual FDataReferenceCollection GetOutputs() const override
-		{
-			// This should never be called. Bind(...) is called instead. This method
-			// exists as a stop-gap until the API can be deprecated and removed.
-			checkNoEntry();
-			return {};
-		}
-
 		void ResetPhase()
 		{
 			float ClampedDegrees = FMath::Clamp(*PhaseOffset, 0.f, 360.f);
