Enable MetaSound deterministic id generation at runtime

#rnx

diff --git a/Engine/Plugins/Runtime/Metasound/Source/MetasoundFrontend/Private/MetasoundAssetBase.cpp b/Engine/Plugins/Runtime/Metasound/Source/MetasoundFrontend/Private/MetasoundAssetBase.cpp
--- a/Engine/Plugins/Runtime/Metasound/Source/MetasoundFrontend/Private/MetasoundAssetBase.cpp
+++ b/Engine/Plugins/Runtime/Metasound/Source/MetasoundFrontend/Private/MetasoundAssetBase.cpp
@@ -347,6 +347,9 @@ bool FMetasoundAssetBase::VersionAsset()
 #if WITH_EDITOR
 	const bool bIsDeterministic = MetaSoundEnableDeterministicIDGenerationInEditorCVar != 0 && !IsRunningCookCommandlet();
 	FDocumentIDGenerator::FScopeDeterminism DeterminismScope(bIsDeterministic);
+#else
+	const bool bIsDeterministic = MetaSoundEnableRuntimeDeterministicIDGeneration != 0 && !IsRunningCookCommandlet();
+	FDocumentIDGenerator::FScopeDeterminism DeterminismScope(bIsDeterministic);
 #endif // WITH_EDITOR
 
 	FName AssetName;
@@ -846,6 +849,9 @@ bool FMetasoundAssetBase::AutoUpdate(bool bInLogWarningsOnDroppedConnection)
 #if WITH_EDITOR
 	const bool bIsDeterministic = MetaSoundEnableDeterministicIDGenerationInEditorCVar != 0 && !IsRunningCookCommandlet();
 	FDocumentIDGenerator::FScopeDeterminism DeterminismScope(bIsDeterministic);
+#else
+	const bool bIsDeterministic = MetaSoundEnableRuntimeDeterministicIDGeneration != 0 && !IsRunningCookCommandlet();
+	FDocumentIDGenerator::FScopeDeterminism DeterminismScope(bIsDeterministic);
 #endif // WITH_EDITOR
 
 	FString OwningAssetName = GetOwningAssetName();
diff --git a/Engine/Plugins/Runtime/Metasound/Source/MetasoundFrontend/Private/MetasoundFrontendDocument.cpp b/Engine/Plugins/Runtime/Metasound/Source/MetasoundFrontend/Private/MetasoundFrontendDocument.cpp
--- a/Engine/Plugins/Runtime/Metasound/Source/MetasoundFrontend/Private/MetasoundFrontendDocument.cpp
+++ b/Engine/Plugins/Runtime/Metasound/Source/MetasoundFrontend/Private/MetasoundFrontendDocument.cpp
@@ -430,7 +430,8 @@ FMetasoundFrontendClassInterface FMetasoundFrontendClassInterface::GenerateClass
 			const bool bIsDeterministic = MetaSoundEnableDeterministicIDGenerationInEditorCVar != 0 && !IsRunningCookCommandlet();
 			ClassInput.VertexID = bIsDeterministic ? FClassIDGenerator::Get().CreateInputID(ClassInput) : FGuid::NewGuid();
 #else 
-			ClassInput.VertexID = FGuid::NewGuid();
+			const bool bIsDeterministic = MetaSoundEnableRuntimeDeterministicIDGeneration != 0 && !IsRunningCookCommandlet();
+			ClassInput.VertexID = bIsDeterministic ? FClassIDGenerator::Get().CreateInputID(ClassInput) : FGuid::NewGuid();
 #endif // WITH_EDITOR
 
 #if WITH_EDITOR
@@ -491,7 +492,8 @@ FMetasoundFrontendClassInterface FMetasoundFrontendClassInterface::GenerateClass
 			const bool bIsDeterministic = MetaSoundEnableDeterministicIDGenerationInEditorCVar != 0 && !IsRunningCookCommandlet();
 			ClassOutput.VertexID = bIsDeterministic ? FClassIDGenerator::Get().CreateOutputID(ClassOutput) : FGuid::NewGuid();
 #else 
-			ClassOutput.VertexID = FGuid::NewGuid();
+			const bool bIsDeterministic = MetaSoundEnableRuntimeDeterministicIDGeneration != 0 && !IsRunningCookCommandlet();
+			ClassOutput.VertexID = bIsDeterministic ? FClassIDGenerator::Get().CreateOutputID(ClassOutput) : FGuid::NewGuid();
 #endif // WITH_EDITOR
 #if WITH_EDITOR
 			const FDataVertexMetadata& VertexMetadata = OutputVertex.Metadata;
@@ -762,7 +764,8 @@ FMetasoundFrontendClassInput::FMetasoundFrontendClassInput(const Audio::FParamet
 	const bool bIsDeterministic = MetaSoundEnableDeterministicIDGenerationInEditorCVar != 0 && !IsRunningCookCommandlet();
 	VertexID = bIsDeterministic ? FClassIDGenerator::Get().CreateInputID(InInput) : FGuid::NewGuid();
 #else 
-	VertexID = FGuid::NewGuid();
+	const bool bIsDeterministic = MetaSoundEnableRuntimeDeterministicIDGeneration != 0 && !IsRunningCookCommandlet();
+	VertexID = bIsDeterministic ? FClassIDGenerator::Get().CreateInputID(InInput) : FGuid::NewGuid();
 #endif // WITH_EDITOR
 
 #if WITH_EDITOR
@@ -795,7 +798,8 @@ FMetasoundFrontendClassOutput::FMetasoundFrontendClassOutput(const Audio::FParam
 	const bool bIsDeterministic = MetaSoundEnableDeterministicIDGenerationInEditorCVar != 0 && !IsRunningCookCommandlet();
 	VertexID = bIsDeterministic ? FClassIDGenerator::Get().CreateOutputID(Output) : FGuid::NewGuid();
 #else 
-	VertexID = FGuid::NewGuid();
+	const bool bIsDeterministic = MetaSoundEnableRuntimeDeterministicIDGeneration != 0 && !IsRunningCookCommandlet();
+	VertexID = bIsDeterministic ? FClassIDGenerator::Get().CreateOutputID(Output) : FGuid::NewGuid();
 #endif // WITH_EDITOR
 
 #if WITH_EDITOR
diff --git a/Engine/Plugins/Runtime/Metasound/Source/MetasoundFrontend/Private/MetasoundFrontendDocumentIdGenerator.cpp b/Engine/Plugins/Runtime/Metasound/Source/MetasoundFrontend/Private/MetasoundFrontendDocumentIdGenerator.cpp
--- a/Engine/Plugins/Runtime/Metasound/Source/MetasoundFrontend/Private/MetasoundFrontendDocumentIdGenerator.cpp
+++ b/Engine/Plugins/Runtime/Metasound/Source/MetasoundFrontend/Private/MetasoundFrontendDocumentIdGenerator.cpp
@@ -14,6 +14,14 @@ namespace Metasound
 			TEXT("0: Disabled (default), !0: Enabled"),
 			ECVF_Default);
 
+		int32 MetaSoundEnableRuntimeDeterministicIDGeneration = 0;
+		FAutoConsoleVariableRef CVarMetaSoundEnableRuntimeDeterministicIDGeneration(
+			TEXT("au.MetaSound.EnableRuntimeDeterministicIDGeneration"),
+			MetaSoundEnableRuntimeDeterministicIDGeneration,
+			TEXT("Enable deterministic ID generation in MetaSounds for AutoUpdate at runtime (non editor) \n")
+			TEXT("0: Disabled (default), !0: Enabled"),
+			ECVF_Default);
+
 		FDocumentIDGenerator::FScopeDeterminism::FScopeDeterminism(bool bInIsDeterministic)
 		{
 			FDocumentIDGenerator& IDGen = Get();
diff --git a/Engine/Plugins/Runtime/Metasound/Source/MetasoundFrontend/Public/MetasoundFrontendDocumentIdGenerator.h b/Engine/Plugins/Runtime/Metasound/Source/MetasoundFrontend/Public/MetasoundFrontendDocumentIdGenerator.h
--- a/Engine/Plugins/Runtime/Metasound/Source/MetasoundFrontend/Public/MetasoundFrontendDocumentIdGenerator.h
+++ b/Engine/Plugins/Runtime/Metasound/Source/MetasoundFrontend/Public/MetasoundFrontendDocumentIdGenerator.h
@@ -10,6 +10,7 @@ namespace Metasound
 	namespace Frontend
 	{
 		extern METASOUNDFRONTEND_API int32 MetaSoundEnableDeterministicIDGenerationInEditorCVar;
+		extern METASOUNDFRONTEND_API int32 MetaSoundEnableRuntimeDeterministicIDGeneration;
 
 		/*** 
 		 * For generating IDs using a given document. 
