Fix ensures with empty MetaSound object literals

diff --git a/Engine/Plugins/Runtime/Metasound/Source/MetasoundEditor/Private/MetasoundEditorGraphBuilder.cpp b/Engine/Plugins/Runtime/Metasound/Source/MetasoundEditor/Private/MetasoundEditorGraphBuilder.cpp
--- a/Engine/Plugins/Runtime/Metasound/Source/MetasoundEditor/Private/MetasoundEditorGraphBuilder.cpp
+++ b/Engine/Plugins/Runtime/Metasound/Source/MetasoundEditor/Private/MetasoundEditorGraphBuilder.cpp
@@ -1090,10 +1090,22 @@ namespace Metasound
 							}
 						}
 					}
-					
+
 					if (!bObjectFound)
 					{
-						OutDefaultLiteral.Clear();
+						// If the class default literal is the default (type is None), then the literal should be set to that.
+						// However, if the class default literal is set to an object, the literal should be set to a valid, null object.
+						// This is used for reset to default behavior, where an valid object literal with a null value is a separate case
+						// from an inherited cleared default literal. 
+						const FMetasoundFrontendLiteral* ClassDefaultLiteral = InputHandle->GetClassDefaultLiteral();
+						if (ClassDefaultLiteral && ClassDefaultLiteral->GetType() == EMetasoundFrontendLiteralType::None)
+						{
+							OutDefaultLiteral.Clear();	
+						}
+						else
+						{
+							OutDefaultLiteral.Set(static_cast<UObject*>(nullptr));
+						}
 					}
 				}
 				break;
diff --git a/Engine/Plugins/Runtime/Metasound/Source/MetasoundEditor/Private/MetasoundEditorGraphMemberDefaults.cpp b/Engine/Plugins/Runtime/Metasound/Source/MetasoundEditor/Private/MetasoundEditorGraphMemberDefaults.cpp
--- a/Engine/Plugins/Runtime/Metasound/Source/MetasoundEditor/Private/MetasoundEditorGraphMemberDefaults.cpp
+++ b/Engine/Plugins/Runtime/Metasound/Source/MetasoundEditor/Private/MetasoundEditorGraphMemberDefaults.cpp
@@ -411,7 +411,11 @@ EMetasoundFrontendLiteralType UMetasoundEditorGraphMemberDefaultObject::GetLiter
 void UMetasoundEditorGraphMemberDefaultObject::SetFromLiteral(const FMetasoundFrontendLiteral& InLiteral)
 {
 	UObject* Object = nullptr;
-	ensure(InLiteral.TryGet(Object));
+	// If default constructed (type is None), then we don't need to try to get the object, and just return nullptr
+	if (InLiteral.GetType() != EMetasoundFrontendLiteralType::None)
+	{
+		ensure(InLiteral.TryGet(Object));
+	}
 	Default.Object = Object;
 }
 
diff --git a/Engine/Plugins/Runtime/Metasound/Source/MetasoundFrontend/Private/MetasoundFrontendNodeController.cpp b/Engine/Plugins/Runtime/Metasound/Source/MetasoundFrontend/Private/MetasoundFrontendNodeController.cpp
--- a/Engine/Plugins/Runtime/Metasound/Source/MetasoundFrontend/Private/MetasoundFrontendNodeController.cpp
+++ b/Engine/Plugins/Runtime/Metasound/Source/MetasoundFrontend/Private/MetasoundFrontendNodeController.cpp
@@ -121,7 +121,34 @@ namespace Metasound
 				{
 					if (VertexLiteral.VertexID == InVertexLiteral.VertexID)
 					{
-						if (ensure(VertexLiteral.Value.GetType() == InVertexLiteral.Value.GetType()))
+						// Type None implies forwarding to the node instance to default construct the literal 
+						// so setting to/from that class default literal is an exception to requiring a type match
+						const bool bIsMatchingType = VertexLiteral.Value.GetType() == InVertexLiteral.Value.GetType();
+						bool bSettingToClassDefaultLiteral = false;
+						bool bSettingFromClassDefaultLiteralToMatchingType = false;
+
+						if (!bIsMatchingType)
+						{
+							auto GetInputWithVertexID = [InVertexID = VertexLiteral.VertexID](const FMetasoundFrontendClassInput& ClassInput)
+							{
+								return ClassInput.VertexID == InVertexID;
+							};
+							const FMetasoundFrontendClassInput* ClassInput = ClassPtr.Get()->Interface.Inputs.FindByPredicate(GetInputWithVertexID);
+							if (ClassInput)
+							{
+								// Check if setting back to class default literal
+								const FMetasoundFrontendLiteral& ClassDefaultLiteral = ClassInput->DefaultLiteral;
+								bSettingToClassDefaultLiteral = ClassDefaultLiteral.IsEqual(InVertexLiteral.Value);
+
+								// Check if setting from class default literal (which may have a None type) to an appropriate type 
+								FDataTypeRegistryInfo DataTypeInfo;
+								IDataTypeRegistry::Get().GetDataTypeInfo(ClassInput->TypeName, DataTypeInfo);
+								const EMetasoundFrontendLiteralType ClassInputLiteralType = static_cast<EMetasoundFrontendLiteralType>(DataTypeInfo.PreferredLiteralType);
+								bSettingFromClassDefaultLiteralToMatchingType |= ClassDefaultLiteral.IsEqual(VertexLiteral.Value) && ClassInputLiteralType == InVertexLiteral.Value.GetType();
+							}
+						}
+
+						if (ensure(bIsMatchingType || bSettingToClassDefaultLiteral || bSettingFromClassDefaultLiteralToMatchingType))
 						{
 							VertexLiteral = InVertexLiteral;
 						}
