Fix for MetaSound transaction stack manipulation not updating builder caches if builder is associated with edited object/asset & edge cache not initializing properly

diff --git a/Engine/Plugins/Runtime/Metasound/Source/MetasoundEngine/Private/MetasoundBuilderSubsystem.cpp b/Engine/Plugins/Runtime/Metasound/Source/MetasoundEngine/Private/MetasoundBuilderSubsystem.cpp
--- a/Engine/Plugins/Runtime/Metasound/Source/MetasoundEngine/Private/MetasoundBuilderSubsystem.cpp
+++ b/Engine/Plugins/Runtime/Metasound/Source/MetasoundEngine/Private/MetasoundBuilderSubsystem.cpp
@@ -630,6 +630,11 @@ void UMetaSoundBuilderBase::RenameRootGraphClass(const FMetasoundFrontendClassNa
 }
 
 #if WITH_EDITOR
+void UMetaSoundBuilderBase::ReloadCache()
+{
+	Builder.ReloadCache();
+}
+
 void UMetaSoundBuilderBase::SetAuthor(const FString& InAuthor)
 {
 	Builder.SetAuthor(InAuthor);
@@ -1571,6 +1576,17 @@ bool UMetaSoundBuilderSubsystem::IsInterfaceRegistered(FName InInterfaceName) co
 	return ISearchEngine::Get().FindInterfaceWithHighestVersion(InInterfaceName, Interface);
 }
 
+#if WITH_EDITOR
+void UMetaSoundBuilderSubsystem::PostBuilderAssetTransaction(const FMetasoundFrontendClassName& InClassName)
+{
+	TWeakObjectPtr<UMetaSoundBuilderBase> BuilderPtr = AssetBuilders.FindRef(InClassName.GetFullName());
+	if (UMetaSoundBuilderBase* Builder = BuilderPtr.Get())
+	{
+		Builder->ReloadCache();
+	}
+}
+#endif // WITH_EDITOR
+
 void UMetaSoundBuilderSubsystem::RegisterBuilder(FName BuilderName, UMetaSoundBuilderBase* Builder)
 {
 	if (Builder)
diff --git a/Engine/Plugins/Runtime/Metasound/Source/MetasoundEngine/Private/MetasoundEngineAsset.h b/Engine/Plugins/Runtime/Metasound/Source/MetasoundEngine/Private/MetasoundEngineAsset.h
--- a/Engine/Plugins/Runtime/Metasound/Source/MetasoundEngine/Private/MetasoundEngineAsset.h
+++ b/Engine/Plugins/Runtime/Metasound/Source/MetasoundEngine/Private/MetasoundEngineAsset.h
@@ -4,6 +4,7 @@
 
 #include "Metasound.h"
 #include "MetasoundAssetManager.h"
+#include "MetasoundBuilderSubsystem.h"
 #include "MetasoundUObjectRegistry.h"
 #include "Serialization/Archive.h"
 
@@ -26,6 +27,11 @@ namespace Metasound
 		static void PostEditUndo(TMetaSoundObject& InMetaSound)
 		{
 			InMetaSound.GetModifyContext().SetForceRefreshViews();
+
+			const FMetasoundFrontendDocument& Document = InMetaSound.GetDocumentChecked();
+			const FMetasoundFrontendClassName& ClassName = Document.RootGraph.Metadata.GetClassName();
+			UMetaSoundBuilderSubsystem::GetChecked().PostBuilderAssetTransaction(ClassName);
+
 			if (UMetasoundEditorGraphBase* Graph = Cast<UMetasoundEditorGraphBase>(InMetaSound.GetGraph()))
 			{
 				Graph->RegisterGraphWithFrontend();
diff --git a/Engine/Plugins/Runtime/Metasound/Source/MetasoundEngine/Public/MetasoundBuilderSubsystem.h b/Engine/Plugins/Runtime/Metasound/Source/MetasoundEngine/Public/MetasoundBuilderSubsystem.h
--- a/Engine/Plugins/Runtime/Metasound/Source/MetasoundEngine/Public/MetasoundBuilderSubsystem.h
+++ b/Engine/Plugins/Runtime/Metasound/Source/MetasoundEngine/Public/MetasoundBuilderSubsystem.h
@@ -27,6 +27,7 @@ class UMetaSound;
 class UMetaSoundPatch;
 class UMetaSoundSource;
 
+struct FMetasoundFrontendClassName;
 struct FMetasoundFrontendVersion;
 
 enum class EMetaSoundOutputAudioFormat : uint8;
@@ -332,6 +333,10 @@ public:
 	void RenameRootGraphClass(const FMetasoundFrontendClassName& InName);
 
 #if WITH_EDITOR
+	// Primarily used by editor transaction stack to avoid corruption when object is
+	// changed outside of the builder API. Generally discouraged for direct use otherwise.
+	void ReloadCache();
+
 	// Sets the author of the MetaSound.
 	void SetAuthor(const FString& InAuthor);
 #endif // WITH_EDITOR
@@ -609,6 +614,11 @@ public:
 	UFUNCTION(BlueprintCallable, Category = "Audio|MetaSound|Builder")
 	UPARAM(DisplayName = "Is Registered") bool IsInterfaceRegistered(FName InInterfaceName) const;
 
+#if WITH_EDITOR
+	// Updates builder data should a builder-managed asset be updated via the transaction stack.
+	void PostBuilderAssetTransaction(const FMetasoundFrontendClassName& InClassName);
+#endif // WITH_EDITOR
+
 	// Adds builder to subsystem's registry to make it persistent and easily accessible by multiple systems or Blueprints
 	UFUNCTION(BlueprintCallable, Category = "Audio|MetaSound|Builder")
 	void RegisterBuilder(FName BuilderName, UMetaSoundBuilderBase* Builder);
diff --git a/Engine/Plugins/Runtime/Metasound/Source/MetasoundFrontend/Private/MetasoundFrontendDocumentBuilder.cpp b/Engine/Plugins/Runtime/Metasound/Source/MetasoundFrontend/Private/MetasoundFrontendDocumentBuilder.cpp
--- a/Engine/Plugins/Runtime/Metasound/Source/MetasoundFrontend/Private/MetasoundFrontendDocumentBuilder.cpp
+++ b/Engine/Plugins/Runtime/Metasound/Source/MetasoundFrontend/Private/MetasoundFrontendDocumentBuilder.cpp
@@ -494,7 +494,7 @@ FMetaSoundFrontendDocumentBuilder::FMetaSoundFrontendDocumentBuilder(TScriptInte
 {
 	if (DocumentInterface)
 	{
-		ReloadCache();
+		ReloadCacheInternal();
 	}
 }
 
@@ -504,7 +504,7 @@ FMetaSoundFrontendDocumentBuilder::FMetaSoundFrontendDocumentBuilder(TScriptInte
 {
 	if (DocumentInterface)
 	{
-		ReloadCache();
+		ReloadCacheInternal();
 	}
 }
 
@@ -1069,7 +1069,7 @@ void FMetaSoundFrontendDocumentBuilder::ClearGraph()
 	GraphClass.PresetOptions.InputsInheritingDefault.Reset();
 	GetDocument().Interfaces.Reset();
 	RemoveUnusedDependencies();
-	ReloadCache();
+	ReloadCacheInternal();
 }
 
 bool FMetaSoundFrontendDocumentBuilder::ContainsEdge(const FMetasoundFrontendEdge& InEdge) const
@@ -1136,7 +1136,7 @@ bool FMetaSoundFrontendDocumentBuilder::ConvertToPreset(const FMetasoundFrontend
 	FRebuildPresetRootGraph RebuildPresetRootGraph(InReferencedDocument);
 	if (RebuildPresetRootGraph.Transform(GetDocument()))
 	{
-		ReloadCache();
+		ReloadCacheInternal();
 		return true;
 	}
 	return false;
@@ -1640,7 +1640,7 @@ bool FMetaSoundFrontendDocumentBuilder::ModifyInterfaces(Metasound::Frontend::FM
 	return Context.Execute(*this, Doc, *DocumentDelegates);
 }
 
-void FMetaSoundFrontendDocumentBuilder::ReloadCache()
+void FMetaSoundFrontendDocumentBuilder::ReloadCacheInternal()
 {
 	using namespace Metasound::Frontend;
 	if (!DocumentCache.IsValid())
@@ -2119,6 +2119,11 @@ bool FMetaSoundFrontendDocumentBuilder::RenameRootGraphClass(const FMetasoundFro
 }
 
 #if WITH_EDITOR
+void FMetaSoundFrontendDocumentBuilder::ReloadCache()
+{
+	ReloadCacheInternal();
+}
+
 void FMetaSoundFrontendDocumentBuilder::SetAuthor(const FString& InAuthor)
 {
 	FMetasoundFrontendClassMetadata& ClassMetadata = GetDocument().RootGraph.Metadata;
diff --git a/Engine/Plugins/Runtime/Metasound/Source/MetasoundFrontend/Public/MetasoundFrontendDocumentBuilder.h b/Engine/Plugins/Runtime/Metasound/Source/MetasoundFrontend/Public/MetasoundFrontendDocumentBuilder.h
--- a/Engine/Plugins/Runtime/Metasound/Source/MetasoundFrontend/Public/MetasoundFrontendDocumentBuilder.h
+++ b/Engine/Plugins/Runtime/Metasound/Source/MetasoundFrontend/Public/MetasoundFrontendDocumentBuilder.h
@@ -213,6 +213,11 @@ public:
 	bool RenameRootGraphClass(const FMetasoundFrontendClassName& InName);
 
 #if WITH_EDITOR
+	// Primarily used by editor transaction stack to avoid corruption when object is
+	// changed outside of the builder API. Generally discouraged for direct use otherwise
+	// as it can be expensive and builder API should manage internal cache state directly.
+	void ReloadCache();
+
 	void SetAuthor(const FString& InAuthor);
 #endif // WITH_EDITOR
 
@@ -235,7 +240,7 @@ private:
 	const FMetasoundFrontendClassInput* FindNodeInputClassInput(const FGuid& InNodeID, const FGuid& InVertexID) const;
 	const FMetasoundFrontendClassOutput* FindNodeOutputClassOutput(const FGuid& InNodeID, const FGuid& InVertexID) const;
 
-	void ReloadCache();
+	void ReloadCacheInternal();
 
 	bool SetGraphInputInheritsDefault(FName InName, bool bInputInheritsDefault);
 
