Fix for incorrect element type being used from earlier change to support null object ptr passing to MetaSound parameter system
#rnx

diff --git a/Engine/Plugins/Runtime/Metasound/Source/MetasoundEngine/Private/MetasoundSource.cpp b/Engine/Plugins/Runtime/Metasound/Source/MetasoundEngine/Private/MetasoundSource.cpp
--- a/Engine/Plugins/Runtime/Metasound/Source/MetasoundEngine/Private/MetasoundSource.cpp
+++ b/Engine/Plugins/Runtime/Metasound/Source/MetasoundEngine/Private/MetasoundSource.cpp
@@ -337,9 +337,9 @@ void UMetaSoundSource::InitParameters(TArray<FAudioParameter>& ParametersToInit,
 	{
 		const TArray<FMetasoundFrontendClassInput>& PublicInputs = RuntimeData.PublicInputs;
 		Algo::Transform(PublicInputs, PublicInputMap, [](const FMetasoundFrontendClassInput& Input)
-			{
-				return TPair<FName, const FMetasoundFrontendVertex*>(Input.Name, &Input);
-			});
+		{
+			return TPair<FName, const FMetasoundFrontendVertex*>(Input.Name, &Input);
+		});
 	}
 
 	// Removes values that are not explicitly defined by the ParamType and returns
@@ -449,7 +449,6 @@ void UMetaSoundSource::InitParameters(TArray<FAudioParameter>& ParametersToInit,
 		{
 			case EAudioParameterType::Object:
 			{
-				FDataTypeRegistryInfo DataTypeInfo;
 				TSharedPtr<Audio::IProxyData> ProxyPtr = IDataTypeRegistry::Get().CreateProxyFromUObject(VertexTypeName, OutParamToInit.ObjectParam);
 				OutParamToInit.ObjectProxies.Emplace(MoveTemp(ProxyPtr));
 
@@ -460,10 +459,10 @@ void UMetaSoundSource::InitParameters(TArray<FAudioParameter>& ParametersToInit,
 
 			case EAudioParameterType::ObjectArray:
 			{
+				const FName ElementTypeName = CreateElementTypeNameFromArrayTypeName(VertexTypeName);
 				for (TObjectPtr<UObject>& Object : OutParamToInit.ArrayObjectParam)
 				{
-					FDataTypeRegistryInfo DataTypeInfo;
-					TSharedPtr<Audio::IProxyData> ProxyPtr = IDataTypeRegistry::Get().CreateProxyFromUObject(VertexTypeName, Object);
+					TSharedPtr<Audio::IProxyData> ProxyPtr = IDataTypeRegistry::Get().CreateProxyFromUObject(ElementTypeName, Object);
 					OutParamToInit.ObjectProxies.Emplace(MoveTemp(ProxyPtr));
 				}
 				// Reset param array as it is no longer needed (nor desired to be accessed once passed to the Audio Thread).
