Fix for new asset not getting registered until opened in MetaSound editor
#rnx

diff --git a/Engine/Plugins/Runtime/Metasound/Source/MetasoundEditor/Private/MetasoundFactory.cpp b/Engine/Plugins/Runtime/Metasound/Source/MetasoundEditor/Private/MetasoundFactory.cpp
--- a/Engine/Plugins/Runtime/Metasound/Source/MetasoundEditor/Private/MetasoundFactory.cpp
+++ b/Engine/Plugins/Runtime/Metasound/Source/MetasoundEditor/Private/MetasoundFactory.cpp
@@ -84,8 +84,16 @@ UMetaSoundFactory::UMetaSoundFactory(const FObjectInitializer& ObjectInitializer
 
 UObject* UMetaSoundFactory::FactoryCreateNew(UClass* InClass, UObject* InParent, FName InName, EObjectFlags InFlags, UObject* InContext, FFeedbackContext* InFeedbackContext)
 {
+	using namespace Metasound::Editor;
 	using namespace Metasound::FactoryPrivate;
-	return CreateNewMetaSoundObject<UMetaSoundPatch>(InParent, InName, InFlags, ReferencedMetaSoundObject);
+
+	if (UMetaSoundPatch* NewPatch = CreateNewMetaSoundObject<UMetaSoundPatch>(InParent, InName, InFlags, ReferencedMetaSoundObject))
+	{
+		FGraphBuilder::RegisterGraphWithFrontend(*NewPatch);
+		return NewPatch;
+	}
+
+	return nullptr;
 }
 
 UMetaSoundSourceFactory::UMetaSoundSourceFactory(const FObjectInitializer& ObjectInitializer)
@@ -96,15 +104,21 @@ UMetaSoundSourceFactory::UMetaSoundSourceFactory(const FObjectInitializer& Objec
 
 UObject* UMetaSoundSourceFactory::FactoryCreateNew(UClass* InClass, UObject* InParent, FName InName, EObjectFlags InFlags, UObject* InContext, FFeedbackContext* InFeedbackContext)
 {
+	using namespace Metasound::Editor;
 	using namespace Metasound::FactoryPrivate;
-	UMetaSoundSource* NewSource = CreateNewMetaSoundObject<UMetaSoundSource>(InParent, InName, InFlags, ReferencedMetaSoundObject);
 
-	// Copy over referenced fields that are specific to sources
-	if (UMetaSoundSource* ReferencedMetaSound = Cast<UMetaSoundSource>(ReferencedMetaSoundObject))
+	if (UMetaSoundSource* NewSource = CreateNewMetaSoundObject<UMetaSoundSource>(InParent, InName, InFlags, ReferencedMetaSoundObject))
 	{
-		NewSource->OutputFormat = ReferencedMetaSound->OutputFormat;
+		// Copy over referenced fields that are specific to sources
+		if (UMetaSoundSource* ReferencedMetaSound = Cast<UMetaSoundSource>(ReferencedMetaSoundObject))
+		{
+			NewSource->OutputFormat = ReferencedMetaSound->OutputFormat;
+		}
+
+		FGraphBuilder::RegisterGraphWithFrontend(*NewSource);
+		return NewSource;
 	}
 
-	return NewSource;
+	return nullptr;
 }
 
