Nuke temporary cvar for disabling MetaSound Builder audition live updates
#rnx

diff --git a/Engine/Plugins/Runtime/Metasound/Source/MetasoundEngine/Private/MetasoundBuilderSubsystem.cpp b/Engine/Plugins/Runtime/Metasound/Source/MetasoundEngine/Private/MetasoundBuilderSubsystem.cpp
--- a/Engine/Plugins/Runtime/Metasound/Source/MetasoundEngine/Private/MetasoundBuilderSubsystem.cpp
+++ b/Engine/Plugins/Runtime/Metasound/Source/MetasoundEngine/Private/MetasoundBuilderSubsystem.cpp
@@ -32,8 +32,6 @@ namespace Metasound::Engine
 {
 	namespace BuilderSubsystemPrivate
 	{
-		static int32 LiveUpdateEnabledCVar = 0;
-
 		template <typename TLiteralType>
 		FMetasoundFrontendLiteral CreatePODMetaSoundLiteral(const TLiteralType& Value, FName& OutDataType)
 		{
@@ -68,13 +66,6 @@ namespace Metasound::Engine
 	} // namespace BuilderSubsystemPrivate
 } // namespace Metasound::Engine
 
-FAutoConsoleVariableRef CVarMetaSoundEditorAudioConnectSpacing(
-	TEXT("au.MetaSound.Builder.Audition.LiveUpdatesEnabled"),
-	Metasound::Engine::BuilderSubsystemPrivate::LiveUpdateEnabledCVar,
-	TEXT("Whether or not the MetaSound builder's audition feature supports live updates (mutating a MetaSound graph while it is being executed at runtime).\n")
-	TEXT("Values: 0 (Disabled), !=0 (Enabled)"),
-	ECVF_Default);
-
 
 FMetaSoundBuilderNodeOutputHandle UMetaSoundBuilderBase::AddGraphInputNode(FName Name, FName DataType, FMetasoundFrontendLiteral DefaultValue, EMetaSoundBuilderResult& OutResult, bool bIsConstructorInput)
 {
@@ -892,14 +883,7 @@ void UMetaSoundSourceBuilder::Audition(UObject* Parent, UAudioComponent* AudioCo
 
 	AudioComponent->SetSound(AuditionSound.Get());
 
-	if (bLiveUpdatesEnabled && !BuilderSubsystemPrivate::LiveUpdateEnabledCVar)
-	{
-		UE_LOG(LogMetaSound, Warning, TEXT("Attempting to audition MetaSoundBuilder '%s' with live update enabled, but feature's console variable is disabled. Request is being ignored and updates will not be audible."),
-			*GetFullName(),
-			* AuditionSound->GetFullName());
-	}
-
-	const bool bEnableDynamicGenerators = bLiveUpdatesEnabled && BuilderSubsystemPrivate::LiveUpdateEnabledCVar;
+	const bool bEnableDynamicGenerators = bLiveUpdatesEnabled;
 	AuditionSound->EnableDynamicGenerators(bEnableDynamicGenerators);
 
 	if (CreateGenerator.IsBound())
@@ -919,7 +903,7 @@ bool UMetaSoundSourceBuilder::ExecuteAuditionableTransaction(FAuditionableTransa
 
 	METASOUND_TRACE_CPUPROFILER_EVENT_SCOPE(UMetaSoundSourceBuilder::ExecuteAuditionableTransaction);
 
-	if (LiveUpdateEnabledCVar && AuditionSound.IsValid())
+	if (AuditionSound.IsValid())
 	{
 		TSharedPtr<FDynamicOperatorTransactor> Transactor = AuditionSound->GetDynamicGeneratorTransactor();
 		if (Transactor.IsValid())
@@ -972,11 +956,6 @@ bool UMetaSoundSourceBuilder::GetLiveUpdatesEnabled() const
 {
 	using namespace Metasound::Engine::BuilderSubsystemPrivate;
 
-	if (!LiveUpdateEnabledCVar)
-	{
-		return false;
-	}
-
 	if (!AuditionSound.IsValid())
 	{
 		return false;
