Fix-up false negative in MetaSound engine test(s)
#rnx

diff --git a/Engine/Plugins/Runtime/Metasound/Source/MetasoundEngineTest/Private/EngineTestMetaSoundBuilder.h b/Engine/Plugins/Runtime/Metasound/Source/MetasoundEngineTest/Private/EngineTestMetaSoundBuilder.h
--- a/Engine/Plugins/Runtime/Metasound/Source/MetasoundEngineTest/Private/EngineTestMetaSoundBuilder.h
+++ b/Engine/Plugins/Runtime/Metasound/Source/MetasoundEngineTest/Private/EngineTestMetaSoundBuilder.h
@@ -31,7 +31,7 @@ namespace EngineTestMetaSoundPatchBuilderPrivate
 		for (const FName& InterfaceName : InterfaceNamesToAdd)
 		{
 			Builder->AddInterface(InterfaceName, Result);
-			Test.AddErrorIfFalse(Result == EMetaSoundBuilderResult::Failed, FString::Printf(TEXT("Failed to add initial interface '%s' to MetaSound patch"), *InterfaceName.ToString()));
+			Test.AddErrorIfFalse(Result == EMetaSoundBuilderResult::Succeeded, FString::Printf(TEXT("Failed to add initial interface '%s' to MetaSound patch"), *InterfaceName.ToString()));
 		}
 
 		return *Builder;
@@ -39,7 +39,6 @@ namespace EngineTestMetaSoundPatchBuilderPrivate
 
 	UMetaSoundPatch* CreatePatchFromBuilder(FAutomationTestBase& Test, const FString& PatchName, const TArray<FName>& InterfaceNamesToAdd)
 	{
-		EMetaSoundBuilderResult Result = EMetaSoundBuilderResult::Failed;
 		UMetaSoundPatchBuilder& Builder = CreatePatchBuilderChecked(Test, FName(PatchName + TEXT(" Builder")), InterfaceNamesToAdd);
 
 		UMetaSoundPatch* InputPatch = CastChecked<UMetaSoundPatch>(Builder.Build(nullptr, FMetaSoundBuilderOptions { FName(PatchName) }).GetObject());
