Conform test component initialization in MetaSound Engine Tests & fix false negative error in case where component is built with invalid initial state of built sound.
#rnx

diff --git a/Engine/Plugins/Runtime/Metasound/Source/MetasoundEngineTest/Private/EngineTestMetasoundSource.cpp b/Engine/Plugins/Runtime/Metasound/Source/MetasoundEngineTest/Private/EngineTestMetasoundSource.cpp
--- a/Engine/Plugins/Runtime/Metasound/Source/MetasoundEngineTest/Private/EngineTestMetasoundSource.cpp
+++ b/Engine/Plugins/Runtime/Metasound/Source/MetasoundEngineTest/Private/EngineTestMetasoundSource.cpp
@@ -26,7 +26,7 @@ namespace EngineTestMetaSoundSourcePrivate
 		TArray<FMetaSoundBuilderNodeInputHandle> AudioOutNodeInputs;
 	};
 
-	static FString GetPluginContentDirectory()
+	FString GetPluginContentDirectory()
 	{
 		TSharedPtr<IPlugin> Plugin = IPluginManager::Get().FindPlugin(TEXT("Metasound"));
 		if (ensure(Plugin.IsValid()))
@@ -35,7 +35,8 @@ namespace EngineTestMetaSoundSourcePrivate
 		}
 		return FString();
 	}
-	static FString GetPathToTestFilesDir()
+
+	FString GetPathToTestFilesDir()
 	{
 		FString OutPath =  FPaths::Combine(GetPluginContentDirectory(), TEXT("Test"));
 
@@ -45,7 +46,7 @@ namespace EngineTestMetaSoundSourcePrivate
 		return OutPath;
 	}
 
-	static FString GetPathToGeneratedFilesDir()
+	FString GetPathToGeneratedFilesDir()
 	{
 		FString OutPath = FPaths::Combine(FPaths::ProjectSavedDir(), TEXT("Metasounds"));
 
@@ -55,7 +56,7 @@ namespace EngineTestMetaSoundSourcePrivate
 		return OutPath;
 	}
 
-	static FString GetPathToGeneratedAssetsDir()
+	FString GetPathToGeneratedAssetsDir()
 	{
 		FString OutPath = TEXT("/Game/Metasound/Generated/");
 		FPaths::NormalizeDirectoryName(OutPath);
@@ -165,7 +166,7 @@ namespace EngineTestMetaSoundSourcePrivate
 		return Builder;
 	}
 
-	FMetasoundFrontendDocument CreateMetaSoundMonoSourceDocument()
+	FMetasoundFrontendDocument CreateMonoSourceDocument()
 	{
 		using namespace Audio;
 		using namespace Metasound;
@@ -229,6 +230,36 @@ namespace EngineTestMetaSoundSourcePrivate
 
 		return Document;
 	}
+
+	UAudioComponent* CreateTestComponent(FAutomationTestBase& Test, USoundBase* Sound = nullptr, bool bAddToRoot = true)
+	{
+		if (FAudioDevice* AudioDevice = GEngine->GetMainAudioDeviceRaw())
+		{
+			UAudioComponent* AudioComponent = NewObject<UAudioComponent>();
+			Test.AddErrorIfFalse(AudioComponent != nullptr, "Failed to create test audio component");
+
+			if (AudioComponent)
+			{
+				AudioComponent->bAutoActivate = false;
+				AudioComponent->bIsUISound = true; // play while "paused"
+				AudioComponent->AudioDeviceID = AudioDevice->DeviceID;
+
+				AudioComponent->bAllowSpatialization = false;
+				AudioComponent->SetVolumeMultiplier(1.0f);
+
+				if (bAddToRoot)
+				{
+					AudioComponent->AddToRoot();
+				}
+
+				AudioComponent->SetSound(Sound);
+			}
+
+			return AudioComponent;
+		}
+
+		return nullptr;
+	}
 } // EngineTestMetaSoundSourcePrivate
 
 DEFINE_LATENT_AUTOMATION_COMMAND_ONE_PARAMETER(FAudioComponentPlayLatentCommand, UAudioComponent*, AudioComponent);
@@ -308,30 +339,20 @@ bool FAudioComponentRemoveFromRootLatentCommand::Update()
 IMPLEMENT_SIMPLE_AUTOMATION_TEST(FAudioMetasoundSourceTest, "Audio.Metasound.PlayMetasoundSource", EAutomationTestFlags::EditorContext | EAutomationTestFlags::EngineFilter)
 bool FAudioMetasoundSourceTest::RunTest(const FString& Parameters)
 {
+	using namespace EngineTestMetaSoundSourcePrivate;
+
 	UMetaSoundSource* MetaSoundSource = NewObject<UMetaSoundSource>(GetTransientPackage(), FName(*LexToString(FGuid::NewGuid())));;
 	if (ensure(nullptr != MetaSoundSource))
 	{
-		MetaSoundSource->SetDocument(EngineTestMetaSoundSourcePrivate::CreateMetaSoundMonoSourceDocument());
+		MetaSoundSource->SetDocument(CreateMonoSourceDocument());
 
-		if (FAudioDevice* AudioDevice = GEngine->GetMainAudioDeviceRaw())
+		if (UAudioComponent* AudioComponent = CreateTestComponent(*this, MetaSoundSource))
 		{
-
-			UAudioComponent* AudioComponent = FAudioDevice::CreateComponent(MetaSoundSource);
-			AddErrorIfFalse(AudioComponent != nullptr, "Failed to create audio component");
-
-			if (AudioComponent)
-			{
-				AudioComponent->bIsUISound = true;
-				AudioComponent->bAllowSpatialization = false;
-				AudioComponent->SetVolumeMultiplier(1.0f);
-				AudioComponent->AddToRoot();
-
-				ADD_LATENT_AUTOMATION_COMMAND(FAudioComponentPlayLatentCommand(AudioComponent));
-				ADD_LATENT_AUTOMATION_COMMAND(FEngineWaitLatentCommand(2.f));
-				ADD_LATENT_AUTOMATION_COMMAND(FAudioComponentStopLatentCommand(AudioComponent));
-				ADD_LATENT_AUTOMATION_COMMAND(FEngineWaitLatentCommand(0.5f));
-				ADD_LATENT_AUTOMATION_COMMAND(FAudioComponentRemoveFromRootLatentCommand(AudioComponent));
-			}
+			ADD_LATENT_AUTOMATION_COMMAND(FAudioComponentPlayLatentCommand(AudioComponent));
+			ADD_LATENT_AUTOMATION_COMMAND(FEngineWaitLatentCommand(2.f));
+			ADD_LATENT_AUTOMATION_COMMAND(FAudioComponentStopLatentCommand(AudioComponent));
+			ADD_LATENT_AUTOMATION_COMMAND(FEngineWaitLatentCommand(0.5f));
+			ADD_LATENT_AUTOMATION_COMMAND(FAudioComponentRemoveFromRootLatentCommand(AudioComponent));
 		}
 	}
 
@@ -349,35 +370,20 @@ bool FAudioMetasoundSourceBuilderTest::RunTest(const FString& Parameters)
 	UMetaSoundSourceBuilder& Builder = CreateMonoSourceSinGenBuilder(*this, nullptr, &MonoOutNodeInput);
 	Builder.AddToRoot();
 
-	if (FAudioDevice* AudioDevice = GEngine->GetMainAudioDeviceRaw())
+	if (UAudioComponent* AudioComponent = CreateTestComponent(*this))
 	{
-		constexpr bool bAddToRegistry = false;
-		TScriptInterface<IMetaSoundDocumentInterface> BuiltDocumentInterface = Builder.Build(nullptr, FMetaSoundBuilderOptions { "BuildAndPlayMetasoundSource", bAddToRegistry });
-		UMetaSoundSource* MetaSoundSource = CastChecked<UMetaSoundSource>(BuiltDocumentInterface.GetObject());
-		UAudioComponent* AudioComponent = FAudioDevice::CreateComponent(MetaSoundSource);
-		AddErrorIfFalse(AudioComponent != nullptr, "Failed to create audio component");
-
-		if (AudioComponent)
-		{
-			AudioComponent->bAutoActivate = false;
-			AudioComponent->bIsUISound = true; // play while "paused"
-			AudioComponent->AudioDeviceID = AudioDevice->DeviceID;
+		constexpr bool bEnableLiveUpdate = false;
+		ADD_LATENT_AUTOMATION_COMMAND(FMetaSoundSourceBuilderAuditionLatentCommand(&Builder, AudioComponent, bEnableLiveUpdate));
+		ADD_LATENT_AUTOMATION_COMMAND(FEngineWaitLatentCommand(2.f));
+		ADD_LATENT_AUTOMATION_COMMAND(FAudioComponentStopLatentCommand(AudioComponent));
+		ADD_LATENT_AUTOMATION_COMMAND(FEngineWaitLatentCommand(0.5f));
+		ADD_LATENT_AUTOMATION_COMMAND(FBuilderRemoveFromRootLatentCommand(&Builder));
+		ADD_LATENT_AUTOMATION_COMMAND(FAudioComponentRemoveFromRootLatentCommand(AudioComponent));
 
-			AudioComponent->bAllowSpatialization = false;
-			AudioComponent->SetVolumeMultiplier(1.0f);
-			AudioComponent->AddToRoot();
-
-			constexpr bool bEnableLiveUpdate = false;
-			ADD_LATENT_AUTOMATION_COMMAND(FMetaSoundSourceBuilderAuditionLatentCommand(&Builder, AudioComponent, bEnableLiveUpdate));
-			ADD_LATENT_AUTOMATION_COMMAND(FEngineWaitLatentCommand(2.f));
-			ADD_LATENT_AUTOMATION_COMMAND(FAudioComponentStopLatentCommand(AudioComponent));
-			ADD_LATENT_AUTOMATION_COMMAND(FEngineWaitLatentCommand(0.5f));
-			ADD_LATENT_AUTOMATION_COMMAND(FBuilderRemoveFromRootLatentCommand(&Builder));
-			ADD_LATENT_AUTOMATION_COMMAND(FAudioComponentRemoveFromRootLatentCommand(AudioComponent));
-		}
+		return true;
 	}
 
-	return true;
+	return false;
 }
 
 // This test creates a MetaSound source from a SourceBuilder, adds a simple sine tone generator with a connected graph input frequency, and attempts to change the frequency and audition it rapidly 100 times.
@@ -392,41 +398,28 @@ bool FAudioMetasoundSourceBuilderTestSpamAudition::RunTest(const FString& Parame
 	UMetaSoundSourceBuilder& Builder = CreateMonoSourceSinGenBuilder(*this, &GenInputNodeFreq, &MonoOutNodeInput);
 	Builder.AddToRoot();
 
-	if (FAudioDevice* AudioDevice = GEngine->GetMainAudioDeviceRaw())
+	if (UAudioComponent* AudioComponent = CreateTestComponent(*this))
 	{
-		constexpr bool bAddToRegistry = false;
-		TScriptInterface<IMetaSoundDocumentInterface> BuiltDocumentInterface = Builder.Build(nullptr, FMetaSoundBuilderOptions{ "BuildAndPlayMetasoundSource", bAddToRegistry });
-		UMetaSoundSource* MetaSoundSource = CastChecked<UMetaSoundSource>(BuiltDocumentInterface.GetObject());
-
-		if (UAudioComponent* AudioComponent = NewObject<UAudioComponent>())
+		constexpr bool bEnableLiveUpdate = false;
+		constexpr int32 NumTrials = 100;
+		for (int32 Index = 0; Index < NumTrials; ++Index)
 		{
-			AudioComponent->bAutoActivate = false;
-			AudioComponent->bIsUISound = true; // play while "paused"
-			AudioComponent->AudioDeviceID = AudioDevice->DeviceID;
-
-			AudioComponent->bAllowSpatialization = false;
-			AudioComponent->SetVolumeMultiplier(1.0f);
-			AudioComponent->AddToRoot();
+			FMetasoundFrontendLiteral NewValue;
+			NewValue.Set(FMath::RandRange(220.f, 2200.f));
+			ADD_LATENT_AUTOMATION_COMMAND(FMetaSoundSourceBuilderSetLiteralLatentCommand(*this, &Builder, GenInputNodeFreq, NewValue));
+			ADD_LATENT_AUTOMATION_COMMAND(FMetaSoundSourceBuilderAuditionLatentCommand(&Builder, AudioComponent, bEnableLiveUpdate));
+			ADD_LATENT_AUTOMATION_COMMAND(FEngineWaitLatentCommand(0.01f));
+		}
 
-			constexpr bool bEnableLiveUpdate = false;
-			constexpr int32 NumTrials = 100;
-			for (int32 Index = 0; Index < NumTrials; ++Index)
-			{
-				FMetasoundFrontendLiteral NewValue;
-				NewValue.Set(FMath::RandRange(220.f, 2200.f));
-				ADD_LATENT_AUTOMATION_COMMAND(FMetaSoundSourceBuilderSetLiteralLatentCommand(*this, &Builder, GenInputNodeFreq, NewValue));
-				ADD_LATENT_AUTOMATION_COMMAND(FMetaSoundSourceBuilderAuditionLatentCommand(&Builder, AudioComponent, bEnableLiveUpdate));
-				ADD_LATENT_AUTOMATION_COMMAND(FEngineWaitLatentCommand(0.01f));
-			}
+		ADD_LATENT_AUTOMATION_COMMAND(FAudioComponentStopLatentCommand(AudioComponent));
+		ADD_LATENT_AUTOMATION_COMMAND(FEngineWaitLatentCommand(0.5f));
+		ADD_LATENT_AUTOMATION_COMMAND(FBuilderRemoveFromRootLatentCommand(&Builder));
+		ADD_LATENT_AUTOMATION_COMMAND(FAudioComponentRemoveFromRootLatentCommand(AudioComponent));
 
-			ADD_LATENT_AUTOMATION_COMMAND(FAudioComponentStopLatentCommand(AudioComponent));
-			ADD_LATENT_AUTOMATION_COMMAND(FEngineWaitLatentCommand(0.5f));
-			ADD_LATENT_AUTOMATION_COMMAND(FBuilderRemoveFromRootLatentCommand(&Builder));
-			ADD_LATENT_AUTOMATION_COMMAND(FAudioComponentRemoveFromRootLatentCommand(AudioComponent));
-		}
+		return true;
 	}
 
-	return true;
+	return false;
 }
 
 // This test creates a MetaSound source from a SourceBuilder, adds a simple sine tone generator with a connected graph input frequency, attempts to audition it, and then switches to a new tri tone generator during playback.
@@ -442,43 +435,31 @@ bool FAudioMetasoundSourceBuilderLiveUpdateNode::RunTest(const FString& Paramete
 	UMetaSoundSourceBuilder& Builder = CreateMonoSourceSinGenBuilder(*this, nullptr, &MonoOutNodeInput);
 	Builder.AddToRoot();
 
-	if (FAudioDevice* AudioDevice = GEngine->GetMainAudioDeviceRaw())
+	if (UAudioComponent* AudioComponent = CreateTestComponent(*this))
 	{
-		UAudioComponent* AudioComponent = NewObject<UAudioComponent>();
-		AddErrorIfFalse(AudioComponent != nullptr, "Failed to create audio component");
-
-		if (AudioComponent)
-		{
-			AudioComponent->bAutoActivate = false;
-			AudioComponent->bIsUISound = true; // play while "paused"
-			AudioComponent->AudioDeviceID = AudioDevice->DeviceID;
+		constexpr bool bEnableLiveUpdate = true;
+		ADD_LATENT_AUTOMATION_COMMAND(FMetaSoundSourceBuilderAuditionLatentCommand(&Builder, AudioComponent, bEnableLiveUpdate));
+		ADD_LATENT_AUTOMATION_COMMAND(FEngineWaitLatentCommand(2.f));
 
-			AudioComponent->bAllowSpatialization = false;
-			AudioComponent->SetVolumeMultiplier(1.0f);
-			AudioComponent->AddToRoot();
+		// Disconnect graph audio output from existing sinosc output and connect to added triosc
+		ADD_LATENT_AUTOMATION_COMMAND(FMetaSoundSourceBuilderDisconnectInputLatentCommand(*this, &Builder, MonoOutNodeInput));
+		ADD_LATENT_AUTOMATION_COMMAND(FMetaSoundSourceBuilderCreateAndConnectTriGeneratorNodeLatentCommand(*this, &Builder, MonoOutNodeInput));
 
-			constexpr bool bEnableLiveUpdate = true;
-			ADD_LATENT_AUTOMATION_COMMAND(FMetaSoundSourceBuilderAuditionLatentCommand(&Builder, AudioComponent, bEnableLiveUpdate));
-			ADD_LATENT_AUTOMATION_COMMAND(FEngineWaitLatentCommand(2.f));
+		FName DataTypeName;
+		FMetasoundFrontendLiteral NewValue;
+		NewValue.Set(FMetasoundFrontendLiteral::FDefault{ });
+		ADD_LATENT_AUTOMATION_COMMAND(FMetaSoundSourceBuilderSetLiteralLatentCommand(*this, &Builder, MonoOutNodeInput, NewValue));
 
-			// Disconnect graph audio output from existing sinosc output and connect to added triosc
-			ADD_LATENT_AUTOMATION_COMMAND(FMetaSoundSourceBuilderDisconnectInputLatentCommand(*this, &Builder, MonoOutNodeInput));
-			ADD_LATENT_AUTOMATION_COMMAND(FMetaSoundSourceBuilderCreateAndConnectTriGeneratorNodeLatentCommand(*this, &Builder, MonoOutNodeInput));
+		ADD_LATENT_AUTOMATION_COMMAND(FEngineWaitLatentCommand(2.f));
+		ADD_LATENT_AUTOMATION_COMMAND(FAudioComponentStopLatentCommand(AudioComponent));
+		ADD_LATENT_AUTOMATION_COMMAND(FEngineWaitLatentCommand(0.5f));
+		ADD_LATENT_AUTOMATION_COMMAND(FAudioComponentRemoveFromRootLatentCommand(AudioComponent));
+		ADD_LATENT_AUTOMATION_COMMAND(FBuilderRemoveFromRootLatentCommand(&Builder));
 
-			FName DataTypeName;
-			FMetasoundFrontendLiteral NewValue;
-			NewValue.Set(FMetasoundFrontendLiteral::FDefault{ });
-			ADD_LATENT_AUTOMATION_COMMAND(FMetaSoundSourceBuilderSetLiteralLatentCommand(*this, &Builder, MonoOutNodeInput, NewValue));
-
-			ADD_LATENT_AUTOMATION_COMMAND(FEngineWaitLatentCommand(2.f));
-			ADD_LATENT_AUTOMATION_COMMAND(FAudioComponentStopLatentCommand(AudioComponent));
-			ADD_LATENT_AUTOMATION_COMMAND(FEngineWaitLatentCommand(0.5f));
-			ADD_LATENT_AUTOMATION_COMMAND(FAudioComponentRemoveFromRootLatentCommand(AudioComponent));
-			ADD_LATENT_AUTOMATION_COMMAND(FBuilderRemoveFromRootLatentCommand(&Builder));
-		}
+		return true;
 	}
 
-	return true;
+	return false;
 }
 
 // This test creates a MetaSound source from a SourceBuilder, adds a simple sine tone generator with a connected graph input frequency, attempts to audition it,
@@ -495,51 +476,39 @@ bool FAudioMetasoundSourceBuilderLiveUpdateLiteral::RunTest(const FString& Param
 	UMetaSoundSourceBuilder& Builder = CreateMonoSourceSinGenBuilder(*this, &GenNodeFreqInput , &MonoOutNodeInput, 220.f);
 	Builder.AddToRoot();
 
-	if (FAudioDevice* AudioDevice = GEngine->GetMainAudioDeviceRaw())
+	if (UAudioComponent* AudioComponent = CreateTestComponent(*this))
 	{
-		UAudioComponent* AudioComponent = NewObject<UAudioComponent>();
-		AddErrorIfFalse(AudioComponent != nullptr, "Failed to create audio component");
+		constexpr bool bEnableLiveUpdate = true;
+		ADD_LATENT_AUTOMATION_COMMAND(FMetaSoundSourceBuilderAuditionLatentCommand(&Builder, AudioComponent, bEnableLiveUpdate));
+		ADD_LATENT_AUTOMATION_COMMAND(FEngineWaitLatentCommand(0.25f));
 
-		if (AudioComponent)
-		{
-			AudioComponent->bAutoActivate = false;
-			AudioComponent->bIsUISound = true; // play while "paused"
-			AudioComponent->AudioDeviceID = AudioDevice->DeviceID;
-
-			AudioComponent->bAllowSpatialization = false;
-			AudioComponent->SetVolumeMultiplier(1.0f);
-			AudioComponent->AddToRoot();
+		// Disconnects freq input node output from sinosc freq input. Initially was set to 220Hz above, and node's default is 440Hz,
+		// resulting in an octive pitch up.
+		ADD_LATENT_AUTOMATION_COMMAND(FMetaSoundSourceBuilderDisconnectInputLatentCommand(*this, &Builder, GenNodeFreqInput));
+		ADD_LATENT_AUTOMATION_COMMAND(FEngineWaitLatentCommand(0.25f));
 
-			constexpr bool bEnableLiveUpdate = true;
-			ADD_LATENT_AUTOMATION_COMMAND(FMetaSoundSourceBuilderAuditionLatentCommand(&Builder, AudioComponent, bEnableLiveUpdate));
-			ADD_LATENT_AUTOMATION_COMMAND(FEngineWaitLatentCommand(0.25f));
+		// Sets literal value on the sinosc freq input to 880Hz, pitching an octive yet again from previous.
+		FName DataTypeName;
+		FMetasoundFrontendLiteral NewValue = UMetaSoundBuilderSubsystem::GetChecked().CreateFloatMetaSoundLiteral(880.f, DataTypeName);
+		AddErrorIfFalse(DataTypeName == Metasound::GetMetasoundDataTypeName<float>(),
+			"Setting MetaSound Float literal returns non-float DataTypeName.");
+		ADD_LATENT_AUTOMATION_COMMAND(FMetaSoundSourceBuilderSetLiteralLatentCommand(*this, &Builder, GenNodeFreqInput, NewValue));
+		ADD_LATENT_AUTOMATION_COMMAND(FEngineWaitLatentCommand(0.25f));
 
-			// Disconnects freq input node output from sinosc freq input. Initially was set to 220Hz above, and node's default is 440Hz,
-			// resulting in an octive pitch up.
-			ADD_LATENT_AUTOMATION_COMMAND(FMetaSoundSourceBuilderDisconnectInputLatentCommand(*this, &Builder, GenNodeFreqInput));
-			ADD_LATENT_AUTOMATION_COMMAND(FEngineWaitLatentCommand(0.25f));
+		// Removes the literal value on the sinosc freq input set to 880Hz, reverting back to the class literal of 440Hz.
+		ADD_LATENT_AUTOMATION_COMMAND(FMetaSoundSourceBuilderRemoveNodeDefaultLiteralLatentCommand(*this, &Builder, GenNodeFreqInput));
+		ADD_LATENT_AUTOMATION_COMMAND(FEngineWaitLatentCommand(0.25f));
 
-			// Sets literal value on the sinosc freq input to 880Hz, pitching an octive yet again from previous.
-			FName DataTypeName;
-			FMetasoundFrontendLiteral NewValue = UMetaSoundBuilderSubsystem::GetChecked().CreateFloatMetaSoundLiteral(880.f, DataTypeName);
-			AddErrorIfFalse(DataTypeName == Metasound::GetMetasoundDataTypeName<float>(),
-				"Setting MetaSound Float literal returns non-float DataTypeName.");
-			ADD_LATENT_AUTOMATION_COMMAND(FMetaSoundSourceBuilderSetLiteralLatentCommand(*this, &Builder, GenNodeFreqInput, NewValue));
-			ADD_LATENT_AUTOMATION_COMMAND(FEngineWaitLatentCommand(0.25f));
+		ADD_LATENT_AUTOMATION_COMMAND(FAudioComponentStopLatentCommand(AudioComponent));
+		ADD_LATENT_AUTOMATION_COMMAND(FEngineWaitLatentCommand(0.25f));
 
-			// Removes the literal value on the sinosc freq input set to 880Hz, reverting back to the class literal of 440Hz.
-			ADD_LATENT_AUTOMATION_COMMAND(FMetaSoundSourceBuilderRemoveNodeDefaultLiteralLatentCommand(*this, &Builder, GenNodeFreqInput));
-			ADD_LATENT_AUTOMATION_COMMAND(FEngineWaitLatentCommand(0.25f));
-
-			ADD_LATENT_AUTOMATION_COMMAND(FAudioComponentStopLatentCommand(AudioComponent));
-			ADD_LATENT_AUTOMATION_COMMAND(FEngineWaitLatentCommand(0.25f));
+		ADD_LATENT_AUTOMATION_COMMAND(FAudioComponentRemoveFromRootLatentCommand(AudioComponent));
+		ADD_LATENT_AUTOMATION_COMMAND(FBuilderRemoveFromRootLatentCommand(&Builder));
 
-			ADD_LATENT_AUTOMATION_COMMAND(FAudioComponentRemoveFromRootLatentCommand(AudioComponent));
-			ADD_LATENT_AUTOMATION_COMMAND(FBuilderRemoveFromRootLatentCommand(&Builder));
-		}
+		return true;
 	}
 
-	return true;
+	return false;
 }
 
 // This test creates a MetaSound source from a SourceBuilder, then adds and finally removes an interface using the builder API, and verifies it as well as its members were added to the document.
