Metasound FDataTypeRegistry::IsValidUObjectForDataType uses GetSuperClass instead of GetOwnerClass when traversing class hierarchy and checks bIsExplicit.

diff --git a/Engine/Plugins/Runtime/Metasound/Source/MetasoundEngine/Private/MetasoundSource.cpp b/Engine/Plugins/Runtime/Metasound/Source/MetasoundEngine/Private/MetasoundSource.cpp
--- a/Engine/Plugins/Runtime/Metasound/Source/MetasoundEngine/Private/MetasoundSource.cpp
+++ b/Engine/Plugins/Runtime/Metasound/Source/MetasoundEngine/Private/MetasoundSource.cpp
@@ -866,7 +866,7 @@ bool UMetaSoundSource::IsParameterValid(const FAudioParameter& InParameter, cons
 			const FName ElementTypeName = CreateElementTypeNameFromArrayTypeName(TypeName);
 			for (const UObject* Object : InParameter.ArrayObjectParam)
 			{
-				bIsValid = IDataTypeRegistry::Get().IsValidUObjectForDataType(ElementTypeName, InParameter.ObjectParam);
+				bIsValid = IDataTypeRegistry::Get().IsValidUObjectForDataType(ElementTypeName, Object);
 				if (!bIsValid)
 				{
 					break;
diff --git a/Engine/Plugins/Runtime/Metasound/Source/MetasoundFrontend/Private/MetasoundFrontendDataTypeRegistry.cpp b/Engine/Plugins/Runtime/Metasound/Source/MetasoundFrontend/Private/MetasoundFrontendDataTypeRegistry.cpp
--- a/Engine/Plugins/Runtime/Metasound/Source/MetasoundFrontend/Private/MetasoundFrontendDataTypeRegistry.cpp
+++ b/Engine/Plugins/Runtime/Metasound/Source/MetasoundFrontend/Private/MetasoundFrontendDataTypeRegistry.cpp
@@ -878,19 +878,37 @@ namespace Metasound
 						if (InUObject)
 						{
 							FScopeLock ObjectRegistryLock(&RegistryObjectMapMutex);
+							FScopeLock RegistryLock(&RegistryMapMutex);
 
 							const UClass* ObjectClass = InUObject->GetClass();
 							while (ObjectClass != UObject::StaticClass())
 							{
 								if (const FName* DataTypeSupportedByUObject = RegisteredObjectClasses.Find(ObjectClass))
 								{
+									// If this is the specified data type and it corresponds to the given object, skip finding its registry entry.
+									if (*DataTypeSupportedByUObject == InDataTypeName && ObjectClass == InUObject->GetClass())
+									{
+										return true;
+									}
+
+									// Find the object's data type's registry entry to determine if it must match the given object.
+									if (const TSharedRef<IDataTypeRegistryEntry, ESPMode::ThreadSafe>* Entry = RegisteredDataTypes.Find(*DataTypeSupportedByUObject))
+									{
+										const FDataTypeRegistryInfo& Info = (*Entry)->GetDataTypeInfo();
+										if (Info.bIsExplicit)
+										{
+											return false;
+										}
+									}
+
+									// If this is the specified data type, then the object is valid.
 									if (*DataTypeSupportedByUObject == InDataTypeName)
 									{
 										return true;
 									}
 								}
 
-								ObjectClass = ObjectClass->GetOwnerClass();
+								ObjectClass = ObjectClass->GetSuperClass();
 							}
 						}
 						else
