- Make Connect... MetaSoundBuilderSubsystem API calls replace edges as needed to parallel editor behavior more
- Update MetaSoundFrontendDocumentBuilder AddEdge... API calls to support both replacement and ignoring invalid requests where applicable

diff --git a/Engine/Plugins/Runtime/Metasound/Source/MetasoundEngine/Private/MetasoundBuilderSubsystem.cpp b/Engine/Plugins/Runtime/Metasound/Source/MetasoundEngine/Private/MetasoundBuilderSubsystem.cpp
--- a/Engine/Plugins/Runtime/Metasound/Source/MetasoundEngine/Private/MetasoundBuilderSubsystem.cpp
+++ b/Engine/Plugins/Runtime/Metasound/Source/MetasoundEngine/Private/MetasoundBuilderSubsystem.cpp
@@ -203,14 +203,61 @@ bool UMetaSoundBuilderBase::ContainsNodeOutput(const FMetaSoundBuilderNodeOutput
 
 void UMetaSoundBuilderBase::ConnectNodes(const FMetaSoundBuilderNodeOutputHandle& NodeOutputHandle, const FMetaSoundBuilderNodeInputHandle& NodeInputHandle, EMetaSoundBuilderResult& OutResult)
 {
-	const FMetasoundFrontendEdge* NewEdge = Builder.AddEdge(FMetasoundFrontendEdge
+	using namespace Metasound::Frontend;
+
+	OutResult = EMetaSoundBuilderResult::Failed;
+
+	FMetasoundFrontendEdge NewEdge { NodeOutputHandle.NodeID, NodeOutputHandle.VertexID, NodeInputHandle.NodeID, NodeInputHandle.VertexID };
+	const EInvalidEdgeReason InvalidEdgeReason = Builder.IsValidEdge(NewEdge);
+
+	if (InvalidEdgeReason == Metasound::Frontend::EInvalidEdgeReason::None)
 	{
-		NodeOutputHandle.NodeID,
-		NodeOutputHandle.VertexID,
-		NodeInputHandle.NodeID,
-		NodeInputHandle.VertexID,
-	});
-	OutResult = NewEdge ? EMetaSoundBuilderResult::Succeeded : EMetaSoundBuilderResult::Failed;
+#if !NO_LOGGING
+		const FMetasoundFrontendNode* OldOutputNode = nullptr;
+		const FMetasoundFrontendVertex* OldOutputVertex = nullptr;
+		if (Builder.IsNodeInputConnected(NodeInputHandle.NodeID, NodeInputHandle.VertexID))
+		{
+			OldOutputVertex = Builder.FindNodeOutputConnectedToNodeInput(NodeInputHandle.NodeID, NodeInputHandle.VertexID, &OldOutputNode);
+		}
+#endif // !NO_LOGGING
+
+		const bool bRemovedEdge = Builder.RemoveEdgeToNodeInput(NodeInputHandle.NodeID, NodeInputHandle.VertexID);
+		Builder.AddEdge(MoveTemp(NewEdge));
+
+#if !NO_LOGGING
+		if (bRemovedEdge)
+		{
+			checkf(OldOutputNode, TEXT("MetaSound edge was removed from output but output node not found."));
+			checkf(OldOutputVertex, TEXT("MetaSound edge was removed from output but output vertex not found."));
+
+			const FMetasoundFrontendNode* InputNode = Builder.FindNode(NodeInputHandle.NodeID);
+			checkf(InputNode, TEXT("Edge was deemed valid but input parent node is missing"));
+
+			const FMetasoundFrontendVertex* InputVertex = Builder.FindNodeInput(NodeInputHandle.NodeID, NodeInputHandle.VertexID);
+			checkf(InputVertex, TEXT("Edge was deemed valid but input is missing"));
+
+			const FMetasoundFrontendNode* OutputNode = Builder.FindNode(NodeOutputHandle.NodeID);
+			checkf(OutputNode, TEXT("Edge was deemed valid but output parent node is missing"));
+
+			const FMetasoundFrontendVertex* OutputVertex = Builder.FindNodeOutput(NodeOutputHandle.NodeID, NodeOutputHandle.VertexID);
+			checkf(OutputVertex, TEXT("Edge was deemed valid but output is missing"));
+
+			UE_LOG(LogMetaSound, Verbose, TEXT("Removed connection from node output '%s:%s' to node '%s:%s' in order to connect to node output '%s:%s'"),
+				*OldOutputNode->Name.ToString(),
+				*OldOutputVertex->Name.ToString(),
+				*InputNode->Name.ToString(),
+				*InputVertex->Name.ToString(),
+				*OutputNode->Name.ToString(),
+				*OutputVertex->Name.ToString());
+		}
+#endif // !NO_LOGGING
+
+		OutResult = EMetaSoundBuilderResult::Succeeded;
+	}
+	else
+	{
+		UE_LOG(LogMetaSound, Warning, TEXT("Builder '%s' 'ConnectNodes' failed: '%s'"), *GetName(), *LexToString(InvalidEdgeReason));
+	}
 }
 
 void UMetaSoundBuilderBase::ConnectNodesByInterfaceBindings(const FMetaSoundNodeHandle& FromNodeHandle, const FMetaSoundNodeHandle& ToNodeHandle, EMetaSoundBuilderResult& OutResult)
@@ -255,40 +302,50 @@ TArray<FMetaSoundBuilderNodeOutputHandle> UMetaSoundBuilderBase::ConnectNodeInpu
 
 void UMetaSoundBuilderBase::ConnectNodeOutputToGraphOutput(FName GraphOutputName, const FMetaSoundBuilderNodeOutputHandle& NodeOutputHandle, EMetaSoundBuilderResult& OutResult)
 {
+	using namespace Metasound::Frontend;
+
 	OutResult = EMetaSoundBuilderResult::Failed;
 
 	if (const FMetasoundFrontendNode* GraphOutputNode = Builder.FindGraphOutputNode(GraphOutputName))
 	{
 		const FMetasoundFrontendVertex& InputVertex = GraphOutputNode->Interface.Inputs.Last();
-		const FMetasoundFrontendEdge* EdgeAdded = Builder.AddEdge(FMetasoundFrontendEdge
+		FMetasoundFrontendEdge NewEdge { NodeOutputHandle.NodeID, NodeOutputHandle.VertexID, GraphOutputNode->GetID(), InputVertex.VertexID };
+		const EInvalidEdgeReason InvalidEdgeReason = Builder.IsValidEdge(NewEdge);
+		if (InvalidEdgeReason == EInvalidEdgeReason::None)
 		{
-			NodeOutputHandle.NodeID,
-			NodeOutputHandle.VertexID,
-			GraphOutputNode->GetID(),
-			InputVertex.VertexID
-		});
-		OutResult = EdgeAdded ? EMetaSoundBuilderResult::Succeeded : EMetaSoundBuilderResult::Failed;
+			Builder.RemoveEdgeToNodeInput(GraphOutputNode->GetID(), InputVertex.VertexID);
+			Builder.AddEdge(MoveTemp(NewEdge));
+			OutResult = EMetaSoundBuilderResult::Succeeded;
+		}
+		else
+		{
+			UE_LOG(LogMetaSound, Warning, TEXT("Builder '%s' 'ConnectNodeOutputToGraphOutput' failed: '%s'"), *GetName(), *LexToString(InvalidEdgeReason));
+		}
 	}
 }
 
 void UMetaSoundBuilderBase::ConnectNodeInputToGraphInput(FName GraphInputName, const FMetaSoundBuilderNodeInputHandle& NodeInputHandle, EMetaSoundBuilderResult& OutResult)
 {
-	const FMetasoundFrontendNode* GraphInputNode = Builder.FindGraphInputNode(GraphInputName);
-	if (!GraphInputNode)
-	{
-		OutResult = EMetaSoundBuilderResult::Failed;
-		return;
-	}
+	using namespace Metasound::Frontend;
+
+	OutResult = EMetaSoundBuilderResult::Failed;
 
-	const FMetasoundFrontendVertex& OutputVertex = GraphInputNode->Interface.Outputs.Last();
-	const FMetasoundFrontendEdge* EdgeAdded = Builder.AddEdge(FMetasoundFrontendEdge
+	if (const FMetasoundFrontendNode* GraphInputNode = Builder.FindGraphInputNode(GraphInputName))
 	{
-		GraphInputNode->GetID(),
-		OutputVertex.VertexID,
-		NodeInputHandle.NodeID,
-		NodeInputHandle.VertexID
-	});
-	OutResult = EdgeAdded ? EMetaSoundBuilderResult::Succeeded : EMetaSoundBuilderResult::Failed;
+		const FMetasoundFrontendVertex& OutputVertex = GraphInputNode->Interface.Outputs.Last();
+		FMetasoundFrontendEdge NewEdge { GraphInputNode->GetID(), OutputVertex.VertexID, NodeInputHandle.NodeID, NodeInputHandle.VertexID };
+		const EInvalidEdgeReason InvalidEdgeReason = Builder.IsValidEdge(NewEdge);
+		if (InvalidEdgeReason == EInvalidEdgeReason::None)
+		{
+			Builder.RemoveEdgeToNodeInput(NodeInputHandle.NodeID, NodeInputHandle.VertexID);
+			Builder.AddEdge(MoveTemp(NewEdge));
+			OutResult = EMetaSoundBuilderResult::Succeeded;
+		}
+		else
+		{
+			UE_LOG(LogMetaSound, Warning, TEXT("Builder '%s' 'ConnectNodeInputToGraphInput' failed: '%s'"), *GetName(), *LexToString(InvalidEdgeReason));
+		}
+	}
 }
 
 void UMetaSoundBuilderBase::DisconnectNodes(const FMetaSoundBuilderNodeOutputHandle& NodeOutputHandle, const FMetaSoundBuilderNodeInputHandle& NodeInputHandle, EMetaSoundBuilderResult& OutResult)
diff --git a/Engine/Plugins/Runtime/Metasound/Source/MetasoundEngineTest/Private/EngineTestMetasoundSource.cpp b/Engine/Plugins/Runtime/Metasound/Source/MetasoundEngineTest/Private/EngineTestMetasoundSource.cpp
--- a/Engine/Plugins/Runtime/Metasound/Source/MetasoundEngineTest/Private/EngineTestMetasoundSource.cpp
+++ b/Engine/Plugins/Runtime/Metasound/Source/MetasoundEngineTest/Private/EngineTestMetasoundSource.cpp
@@ -307,7 +307,7 @@ bool FMetaSoundSourceBuilderAuditionLatentCommand::Update()
 		// Each subsequent call should maintain that name to avoid bloating the FName table/breaking references should this auditioned sound
 		// be in anyway referenced.
 		const FName InitClassName = GetComponentSoundClassName();
-		Builder->Audition(nullptr, AudioComponent, { }, true);
+		Builder->Audition(nullptr, AudioComponent, { }, bEnableLiveUpdates);
 		const FName BuiltClassName = GetComponentSoundClassName();
 
 		const bool bInitNameGenerated = InitClassName.IsNone() && !BuiltClassName.IsNone();
@@ -408,7 +408,7 @@ bool FAudioMetasoundSourceBuilderTestSpamAudition::RunTest(const FString& Parame
 			NewValue.Set(FMath::RandRange(220.f, 2200.f));
 			ADD_LATENT_AUTOMATION_COMMAND(FMetaSoundSourceBuilderSetLiteralLatentCommand(*this, &Builder, GenInputNodeFreq, NewValue));
 			ADD_LATENT_AUTOMATION_COMMAND(FMetaSoundSourceBuilderAuditionLatentCommand(&Builder, AudioComponent, bEnableLiveUpdate));
-			ADD_LATENT_AUTOMATION_COMMAND(FEngineWaitLatentCommand(0.01f));
+			ADD_LATENT_AUTOMATION_COMMAND(FEngineWaitLatentCommand(0.05f));
 		}
 
 		ADD_LATENT_AUTOMATION_COMMAND(FAudioComponentStopLatentCommand(AudioComponent));
@@ -442,7 +442,6 @@ bool FAudioMetasoundSourceBuilderLiveUpdateNode::RunTest(const FString& Paramete
 		ADD_LATENT_AUTOMATION_COMMAND(FEngineWaitLatentCommand(2.f));
 
 		// Disconnect graph audio output from existing sinosc output and connect to added triosc
-		ADD_LATENT_AUTOMATION_COMMAND(FMetaSoundSourceBuilderDisconnectInputLatentCommand(*this, &Builder, MonoOutNodeInput));
 		ADD_LATENT_AUTOMATION_COMMAND(FMetaSoundSourceBuilderCreateAndConnectTriGeneratorNodeLatentCommand(*this, &Builder, MonoOutNodeInput));
 
 		FName DataTypeName;
diff --git a/Engine/Plugins/Runtime/Metasound/Source/MetasoundFrontend/Private/MetasoundFrontendDocumentBuilder.cpp b/Engine/Plugins/Runtime/Metasound/Source/MetasoundFrontend/Private/MetasoundFrontendDocumentBuilder.cpp
--- a/Engine/Plugins/Runtime/Metasound/Source/MetasoundFrontend/Private/MetasoundFrontendDocumentBuilder.cpp
+++ b/Engine/Plugins/Runtime/Metasound/Source/MetasoundFrontend/Private/MetasoundFrontendDocumentBuilder.cpp
@@ -444,6 +444,32 @@
 		};
 	} // namespace DocumentBuilderPrivate
 
+	FString LexToString(const EInvalidEdgeReason& InReason)
+	{
+		switch (InReason)
+		{
+			case EInvalidEdgeReason::None:
+				return TEXT("No reason");
+
+			case EInvalidEdgeReason::MismatchedAccessType:
+				return TEXT("Mismatched Access Type");
+
+			case EInvalidEdgeReason::MismatchedDataType:
+				return TEXT("Mismatched DataType");
+
+			case EInvalidEdgeReason::MissingInput:
+				return TEXT("Missing Input");
+
+			case EInvalidEdgeReason::MissingOutput:
+				return TEXT("Missing Output");
+
+			default:
+				return TEXT("COUNT");
+		}
+
+		static_assert(static_cast<uint32>(EInvalidEdgeReason::COUNT) == 5, "Potential missing case coverage for EInvalidEdgeReason");
+	}
+
 	FModifyInterfaceOptions::FModifyInterfaceOptions(const TArray<FMetasoundFrontendInterface>& InInterfacesToRemove, const TArray<FMetasoundFrontendInterface>& InInterfacesToAdd)
 		: InterfacesToRemove(InInterfacesToRemove)
 		, InterfacesToAdd(InInterfacesToAdd)
@@ -533,25 +559,20 @@
 	return Dependency;
 }
 
-const FMetasoundFrontendEdge* FMetaSoundFrontendDocumentBuilder::AddEdge(FMetasoundFrontendEdge&& InNewEdge)
+void FMetaSoundFrontendDocumentBuilder::AddEdge(FMetasoundFrontendEdge&& InNewEdge)
 {
 	using namespace Metasound::Frontend;
 
-	if (CanAddEdge(InNewEdge))
-	{
-		FMetasoundFrontendDocument& Document = GetDocument();
-		FMetasoundFrontendGraph& Graph = Document.RootGraph.Graph;
-		const IDocumentGraphEdgeCache& EdgeCache = DocumentCache->GetEdgeCache();
-		const FMetasoundFrontendEdge& NewEdge = Graph.Edges.Add_GetRef(MoveTemp(InNewEdge));
-		const int32 NewIndex = Graph.Edges.Num() - 1;
-		DocumentDelegates->EdgeDelegates.OnEdgeAdded.Broadcast(NewIndex);
-		return &NewEdge;
-	}
+	checkf(CanAddEdge(InNewEdge), TEXT("Attempted call to AddEdge in MetaSound Builder where not valid."));
 
-	return nullptr;
+	FMetasoundFrontendDocument& Document = GetDocument();
+	FMetasoundFrontendGraph& Graph = Document.RootGraph.Graph;
+	Graph.Edges.Add(MoveTemp(InNewEdge));
+	const int32 NewIndex = Graph.Edges.Num() - 1;
+	DocumentDelegates->EdgeDelegates.OnEdgeAdded.Broadcast(NewIndex);
 }
 
-bool FMetaSoundFrontendDocumentBuilder::AddNamedEdges(const TSet<Metasound::Frontend::FNamedEdge>& EdgesToMake, TArray<const FMetasoundFrontendEdge*>* OutNewEdges)
+bool FMetaSoundFrontendDocumentBuilder::AddNamedEdges(const TSet<Metasound::Frontend::FNamedEdge>& EdgesToMake, TArray<const FMetasoundFrontendEdge*>* OutNewEdges, bool bReplaceExistingConnections)
 {
 	using namespace Metasound::Frontend;
 
@@ -565,7 +586,14 @@
 
 	bool bSuccess = true;
 
-	TArray<FMetasoundFrontendEdge> EdgesToAdd;
+	struct FNewEdgeData
+	{
+		FMetasoundFrontendEdge NewEdge;
+		const FMetasoundFrontendVertex* OutputVertex = nullptr;
+		const FMetasoundFrontendVertex* InputVertex = nullptr;
+	};
+
+	TArray<FNewEdgeData> EdgesToAdd;
 	for (const FNamedEdge& Edge : EdgesToMake)
 	{
 		const FMetasoundFrontendGraph& Graph = Document.RootGraph.Graph;
@@ -575,26 +603,79 @@
 		if (OutputVertex && InputVertex)
 		{
 			FMetasoundFrontendEdge NewEdge = { Edge.OutputNodeID, OutputVertex->VertexID, Edge.InputNodeID, InputVertex->VertexID };
-			if (CanAddEdge(NewEdge))
+			const EInvalidEdgeReason InvalidEdgeReason = IsValidEdge(NewEdge);
+			if (InvalidEdgeReason == EInvalidEdgeReason::None)
 			{
-				EdgesToAdd.Add(MoveTemp(NewEdge));
+				EdgesToAdd.Add(FNewEdgeData { MoveTemp(NewEdge), OutputVertex, InputVertex });
 			}
 			else
 			{
 				bSuccess = false;
-				UE_LOG(LogMetaSound, Warning, TEXT("Failed to add connections between MetaSound output '%s' and input '%s': Vertex types either incompatable (eg. construction pin to non-construction pin) or input already connected."), *Edge.OutputName.ToString(), *Edge.InputName.ToString());
+				UE_LOG(LogMetaSound, Error, TEXT("Failed to add connections between MetaSound output '%s' and input '%s': '%s'."), *Edge.OutputName.ToString(), *Edge.InputName.ToString(), *LexToString(InvalidEdgeReason));
 			}
 		}
 	}
 
 	const TArray<FMetasoundFrontendEdge>& Edges = GetDocument().RootGraph.Graph.Edges;
 	const int32 LastIndex = Edges.Num() - 1;
-	for (FMetasoundFrontendEdge& EdgeToAdd : EdgesToAdd)
+	for (FNewEdgeData& EdgeToAdd : EdgesToAdd)
 	{
-		const FMetasoundFrontendEdge* NewEdge = AddEdge(MoveTemp(EdgeToAdd));
-		if (!ensureAlwaysMsgf(NewEdge, TEXT("Failed to add MetaSound graph edge via DocumentBuilder when prior step validated edge add was valid")))
+		if (bReplaceExistingConnections)
 		{
+#if !NO_LOGGING
+			const FMetasoundFrontendNode* OldOutputNode = nullptr;
+			const FMetasoundFrontendVertex* OldOutputVertex = FindNodeOutputConnectedToNodeInput(EdgeToAdd.NewEdge.ToNodeID, EdgeToAdd.NewEdge.ToVertexID, &OldOutputNode);
+#endif // !NO_LOGGING
+
+			const bool bRemovedEdge = RemoveEdgeToNodeInput(EdgeToAdd.NewEdge.ToNodeID, EdgeToAdd.NewEdge.ToVertexID);
+
+#if !NO_LOGGING
+			if (bRemovedEdge)
+			{
+				checkf(OldOutputNode, TEXT("MetaSound edge was removed from output but output node not found."));
+				checkf(OldOutputVertex, TEXT("MetaSound edge was removed from output but output vertex not found."));
+
+				const FMetasoundFrontendNode* InputNode = FindNode(EdgeToAdd.NewEdge.ToNodeID);
+				checkf(InputNode, TEXT("Edge was deemed valid but input parent node is missing"));
+
+				const FMetasoundFrontendNode* OutputNode = FindNode(EdgeToAdd.NewEdge.FromNodeID);
+				checkf(OutputNode, TEXT("Edge was deemed valid but output parent node is missing"));
+
+				UE_LOG(LogMetaSound, Verbose, TEXT("Removed connection from node output '%s:%s' to node '%s:%s' in order to connect to node output '%s:%s'"),
+					*OldOutputNode->Name.ToString(),
+					*OldOutputVertex->Name.ToString(),
+					*InputNode->Name.ToString(),
+					*EdgeToAdd.InputVertex->Name.ToString(),
+					*OutputNode->Name.ToString(),
+					*EdgeToAdd.OutputVertex->Name.ToString());
+			}
+#endif // !NO_LOGGING
+
+			AddEdge(MoveTemp(EdgeToAdd.NewEdge));
+		}
+		else if (!IsNodeInputConnected(EdgeToAdd.NewEdge.ToNodeID, EdgeToAdd.NewEdge.ToVertexID))
+		{
+			AddEdge(MoveTemp(EdgeToAdd.NewEdge));
+		}
+		else
+		{
 			bSuccess = false;
+
+#if !NO_LOGGING
+			FMetasoundFrontendEdge EdgeToRemove;
+			if (const int32* EdgeIndex = DocumentCache->GetEdgeCache().FindEdgeIndexToNodeInput(EdgeToAdd.NewEdge.ToNodeID, EdgeToAdd.NewEdge.ToVertexID))
+			{
+				EdgeToRemove = Document.RootGraph.Graph.Edges[*EdgeIndex];
+			}
+
+			const FMetasoundFrontendVertex* Input = FindNodeInput(EdgeToAdd.NewEdge.ToNodeID, EdgeToAdd.NewEdge.ToVertexID);
+			checkf(Input, TEXT("Prior loop to check edge validity should protect against missing input vertex"));
+
+			const FMetasoundFrontendVertex* Output = FindNodeOutput(EdgeToAdd.NewEdge.FromNodeID, EdgeToAdd.NewEdge.FromVertexID);
+			checkf(Input, TEXT("Prior loop to check edge validity should protect against missing output vertex"));
+
+			UE_LOG(LogMetaSound, Warning, TEXT("Connection between MetaSound output '%s' and input '%s' not added: Input already connected to '%s'."), *Output->Name.ToString(), *Input->Name.ToString(), *Output->Name.ToString());
+#endif // !NO_LOGGING
 		}
 	}
 
@@ -609,7 +690,7 @@
 	return bSuccess;
 }
 
-bool FMetaSoundFrontendDocumentBuilder::AddEdgesByNodeClassInterfaceBindings(const FGuid& InFromNodeID, const FGuid& InToNodeID)
+bool FMetaSoundFrontendDocumentBuilder::AddEdgesByNodeClassInterfaceBindings(const FGuid& InFromNodeID, const FGuid& InToNodeID, bool bReplaceExistingConnections)
 {
 	using namespace Metasound;
 	using namespace Metasound::Frontend;
@@ -621,7 +702,7 @@
 		TSet<FNamedEdge> NamedEdges;
 		if (DocumentBuilderPrivate::TryGetInterfaceBoundEdges(InFromNodeID, *FromInterfaceVersions, InToNodeID, *ToInterfaceVersions, NamedEdges))
 		{
-			return AddNamedEdges(NamedEdges);
+			return AddNamedEdges(NamedEdges, nullptr, bReplaceExistingConnections);
 		}
 	}
 
@@ -629,7 +710,7 @@
 
 }
 
-bool FMetaSoundFrontendDocumentBuilder::AddEdgesFromMatchingInterfaceNodeOutputsToGraphOutputs(const FGuid& InNodeID, TArray<const FMetasoundFrontendEdge*>& OutEdgesCreated)
+bool FMetaSoundFrontendDocumentBuilder::AddEdgesFromMatchingInterfaceNodeOutputsToGraphOutputs(const FGuid& InNodeID, TArray<const FMetasoundFrontendEdge*>& OutEdgesCreated, bool bReplaceExistingConnections)
 {
 	METASOUND_TRACE_CPUPROFILER_EVENT_SCOPE(FMetaSoundFrontendDocumentBuilder::AddEdgesFromMatchingInterfaceNodeOutputsToGraphOutputs);
 
@@ -669,10 +750,10 @@
 		}
 	}
 
-	return AddNamedEdges(EdgesToMake, &OutEdgesCreated);
+	return AddNamedEdges(EdgesToMake, &OutEdgesCreated, bReplaceExistingConnections);
 }
 
-bool FMetaSoundFrontendDocumentBuilder::AddEdgesFromMatchingInterfaceNodeInputsToGraphInputs(const FGuid& InNodeID, TArray<const FMetasoundFrontendEdge*>& OutEdgesCreated)
+bool FMetaSoundFrontendDocumentBuilder::AddEdgesFromMatchingInterfaceNodeInputsToGraphInputs(const FGuid& InNodeID, TArray<const FMetasoundFrontendEdge*>& OutEdgesCreated, bool bReplaceExistingConnections)
 {
 	METASOUND_TRACE_CPUPROFILER_EVENT_SCOPE(FMetaSoundFrontendDocumentBuilder::AddEdgesFromMatchingInterfaceNodeInputsToGraphInputs);
 
@@ -712,7 +793,7 @@
 		}
 	}
 
-	return AddNamedEdges(EdgesToMake, &OutEdgesCreated);
+	return AddNamedEdges(EdgesToMake, &OutEdgesCreated, bReplaceExistingConnections);
 }
 
 const FMetasoundFrontendNode* FMetaSoundFrontendDocumentBuilder::AddGraphInput(const FMetasoundFrontendClassInput& InClassInput)
@@ -1020,22 +1101,7 @@
 
 	if (!EdgeCache.IsNodeInputConnected(InEdge.ToNodeID, InEdge.ToVertexID))
 	{
-		const IDocumentGraphNodeCache& NodeCache = DocumentCache->GetNodeCache();
-		const FMetasoundFrontendGraph& Graph = Document.RootGraph.Graph;
-		const FMetasoundFrontendVertex* OutputVertex = NodeCache.FindOutputVertex(InEdge.FromNodeID, InEdge.FromVertexID);
-		const FMetasoundFrontendVertex* InputVertex = NodeCache.FindInputVertex(InEdge.ToNodeID, InEdge.ToVertexID);
-		if (OutputVertex && InputVertex)
-		{
-			if (OutputVertex->TypeName == InputVertex->TypeName)
-			{
-				const FMetasoundFrontendClassOutput* ClassOutput = FindNodeOutputClassOutput(InEdge.FromNodeID, InEdge.FromVertexID);
-				const FMetasoundFrontendClassInput* ClassInput = FindNodeInputClassInput(InEdge.ToNodeID, InEdge.ToVertexID);
-				if (ClassInput && ClassOutput)
-				{
-					return FMetasoundFrontendClassVertex::CanConnectVertexAccessTypes(ClassOutput->AccessType, ClassInput->AccessType);
-				}
-			}
-		}
+		return IsValidEdge(InEdge) == EInvalidEdgeReason::None;
 	}
 
 	return false;
@@ -1303,6 +1369,33 @@
 	return nullptr;
 }
 
+const FMetasoundFrontendNode* FMetaSoundFrontendDocumentBuilder::FindNode(const FGuid& InNodeID) const
+{
+	using namespace Metasound::Frontend;
+	const IDocumentGraphNodeCache& NodeCache = DocumentCache->GetNodeCache();
+	return NodeCache.FindNode(InNodeID);
+}
+
+const TSet<FMetasoundFrontendVersion>* FMetaSoundFrontendDocumentBuilder::FindNodeClassInterfaces(const FGuid& InNodeID) const
+{
+	using namespace Metasound;
+	using namespace Metasound::Frontend;
+
+	const FMetasoundFrontendDocument& Document = GetDocument();
+	const IDocumentGraphNodeCache& NodeCache = DocumentCache->GetNodeCache();
+	if (const FMetasoundFrontendNode* Node = NodeCache.FindNode(InNodeID))
+	{
+		if (const FMetasoundFrontendClass* NodeClass = DocumentCache->FindDependency(Node->ClassID))
+		{
+			// 1. May be a serialized asset, so first check with asset manager.
+			const FNodeRegistryKey NodeClassRegistryKey = NodeRegistryKey::CreateKey(NodeClass->Metadata);
+			return FMetasoundFrontendRegistryContainer::Get()->FindImplementedInterfacesFromRegistered(NodeClassRegistryKey);
+		}
+	}
+
+	return nullptr;
+}
+
 const FMetasoundFrontendVertex* FMetaSoundFrontendDocumentBuilder::FindNodeInput(const FGuid& InNodeID, const FGuid& InVertexID) const
 {
 	using namespace Metasound::Frontend;
@@ -1317,18 +1410,4 @@
 	return NodeCache.FindInputVertex(InNodeID, InVertexName);
 }
 
-const FMetasoundFrontendVertex* FMetaSoundFrontendDocumentBuilder::FindNodeOutput(const FGuid& InNodeID, const FGuid& InVertexID) const
-{
-	using namespace Metasound::Frontend;
-	const IDocumentGraphNodeCache& NodeCache = DocumentCache->GetNodeCache();
-	return NodeCache.FindOutputVertex(InNodeID, InVertexID);
-}
-
-const FMetasoundFrontendVertex* FMetaSoundFrontendDocumentBuilder::FindNodeOutput(const FGuid& InNodeID, FName InVertexName) const
-{
-	using namespace Metasound::Frontend;
-	const IDocumentGraphNodeCache& NodeCache = DocumentCache->GetNodeCache();
-	return NodeCache.FindOutputVertex(InNodeID, InVertexName);
-}
-
 const FMetasoundFrontendClassInput* FMetaSoundFrontendDocumentBuilder::FindNodeInputClassInput(const FGuid& InNodeID, const FGuid& InVertexID) const
@@ -1350,6 +1429,50 @@
 	return nullptr;
 }
 
+TArray<const FMetasoundFrontendVertex*> FMetaSoundFrontendDocumentBuilder::FindNodeInputs(const FGuid& InNodeID, FName TypeName) const
+{
+	return DocumentCache->GetNodeCache().FindNodeInputs(InNodeID, TypeName);
+}
+
+TArray<const FMetasoundFrontendVertex*> FMetaSoundFrontendDocumentBuilder::FindNodeInputsConnectedToNodeOutput(const FGuid& InOutputNodeID, const FGuid& InOutputVertexID, TArray<const FMetasoundFrontendNode*>* ConnectedInputNodes) const
+{
+	using namespace Metasound::Frontend;
+
+	const IDocumentGraphEdgeCache& EdgeCache = DocumentCache->GetEdgeCache();
+	const IDocumentGraphNodeCache& NodeCache = DocumentCache->GetNodeCache();
+
+	const FMetasoundFrontendDocument& Document = GetDocument();
+
+	TArray<const FMetasoundFrontendVertex*> Inputs;
+	if (const TArray<int32>* Indices = EdgeCache.FindEdgeIndicesFromNodeOutput(InOutputNodeID, InOutputVertexID))
+	{
+		Algo::Transform(*Indices, Inputs, [&Document, &NodeCache, &ConnectedInputNodes](const int32& Index)
+		{
+			const FMetasoundFrontendEdge& Edge = Document.RootGraph.Graph.Edges[Index];
+			if (ConnectedInputNodes)
+			{
+				ConnectedInputNodes->Add(NodeCache.FindNode(Edge.ToNodeID));
+			}
+			return NodeCache.FindInputVertex(Edge.ToNodeID, Edge.ToVertexID);
+		});
+	}
+	return Inputs;
+}
+
+const FMetasoundFrontendVertex* FMetaSoundFrontendDocumentBuilder::FindNodeOutput(const FGuid& InNodeID, const FGuid& InVertexID) const
+{
+	using namespace Metasound::Frontend;
+	const IDocumentGraphNodeCache& NodeCache = DocumentCache->GetNodeCache();
+	return NodeCache.FindOutputVertex(InNodeID, InVertexID);
+}
+
+const FMetasoundFrontendVertex* FMetaSoundFrontendDocumentBuilder::FindNodeOutput(const FGuid& InNodeID, FName InVertexName) const
+{
+	using namespace Metasound::Frontend;
+	const IDocumentGraphNodeCache& NodeCache = DocumentCache->GetNodeCache();
+	return NodeCache.FindOutputVertex(InNodeID, InVertexName);
+}
+
 const FMetasoundFrontendClassOutput* FMetaSoundFrontendDocumentBuilder::FindNodeOutputClassOutput(const FGuid& InNodeID, const FGuid& InVertexID) const
 {
 	using namespace Metasound::Frontend;
@@ -1369,41 +1492,33 @@
 	return nullptr;
 }
 
-const FMetasoundFrontendNode* FMetaSoundFrontendDocumentBuilder::FindNode(const FGuid& InNodeID) const
-{
-	using namespace Metasound::Frontend;
-	const IDocumentGraphNodeCache& NodeCache = DocumentCache->GetNodeCache();
-	return NodeCache.FindNode(InNodeID);
-}
-
-const TSet<FMetasoundFrontendVersion>* FMetaSoundFrontendDocumentBuilder::FindNodeClassInterfaces(const FGuid& InNodeID) const
-{
-	using namespace Metasound;
-	using namespace Metasound::Frontend;
-
-	const FMetasoundFrontendDocument& Document = GetDocument();
-	const IDocumentGraphNodeCache& NodeCache = DocumentCache->GetNodeCache();
-	if (const FMetasoundFrontendNode* Node = NodeCache.FindNode(InNodeID))
-	{
-		if (const FMetasoundFrontendClass* NodeClass = DocumentCache->FindDependency(Node->ClassID))
-		{
-			// 1. May be a serialized asset, so first check with asset manager.
-			const FNodeRegistryKey NodeClassRegistryKey = NodeRegistryKey::CreateKey(NodeClass->Metadata);
-			return FMetasoundFrontendRegistryContainer::Get()->FindImplementedInterfacesFromRegistered(NodeClassRegistryKey);
-		}
-	}
-
-	return nullptr;
-}
-
-TArray<const FMetasoundFrontendVertex*> FMetaSoundFrontendDocumentBuilder::FindNodeInputs(const FGuid& InNodeID, FName TypeName) const
-{
-	return DocumentCache->GetNodeCache().FindNodeInputs(InNodeID, TypeName);
-}
-
 TArray<const FMetasoundFrontendVertex*> FMetaSoundFrontendDocumentBuilder::FindNodeOutputs(const FGuid& InNodeID, FName TypeName) const
 {
 	return DocumentCache->GetNodeCache().FindNodeOutputs(InNodeID, TypeName);
 }
 
+const FMetasoundFrontendVertex* FMetaSoundFrontendDocumentBuilder::FindNodeOutputConnectedToNodeInput(const FGuid& InInputNodeID, const FGuid& InInputVertexID, const FMetasoundFrontendNode** ConnectedOutputNode) const
+{
+	using namespace Metasound::Frontend;
+
+	const IDocumentGraphEdgeCache& EdgeCache = DocumentCache->GetEdgeCache();
+	if (const int32* Index = EdgeCache.FindEdgeIndexToNodeInput(InInputNodeID, InInputVertexID))
+	{
+		const FMetasoundFrontendDocument& Document = GetDocument();
+		const FMetasoundFrontendEdge& Edge = Document.RootGraph.Graph.Edges[*Index];
+		const IDocumentGraphNodeCache& NodeCache = DocumentCache->GetNodeCache();
+		if (ConnectedOutputNode)
+		{
+			(*ConnectedOutputNode) = NodeCache.FindNode(Edge.FromNodeID);
+		}
+		return NodeCache.FindOutputVertex(Edge.FromNodeID, Edge.FromVertexID);
+	}
+
+	if (ConnectedOutputNode)
+	{
+		*ConnectedOutputNode = nullptr;
+	}
+	return nullptr;
+}
+
 const FTopLevelAssetPath FMetaSoundFrontendDocumentBuilder::GetBuilderClassPath() const
@@ -1631,6 +1746,43 @@
 	return GetDocument().RootGraph.PresetOptions.bIsPreset;
 }
 
+Metasound::Frontend::EInvalidEdgeReason FMetaSoundFrontendDocumentBuilder::IsValidEdge(const FMetasoundFrontendEdge& InEdge) const
+{
+	using namespace Metasound::Frontend;
+
+	const IDocumentGraphNodeCache& NodeCache = DocumentCache->GetNodeCache();
+
+	const FMetasoundFrontendVertex* OutputVertex = NodeCache.FindOutputVertex(InEdge.FromNodeID, InEdge.FromVertexID);
+	const FMetasoundFrontendVertex* InputVertex = NodeCache.FindInputVertex(InEdge.ToNodeID, InEdge.ToVertexID);
+	if (!OutputVertex)
+	{
+		return EInvalidEdgeReason::MissingOutput;
+	}
+
+	if (!InputVertex)
+	{
+		return EInvalidEdgeReason::MissingInput;
+	}
+
+	if (OutputVertex->TypeName != InputVertex->TypeName)
+	{
+		return EInvalidEdgeReason::MismatchedDataType;
+	}
+
+	const FMetasoundFrontendClassOutput* ClassOutput = FindNodeOutputClassOutput(InEdge.FromNodeID, InEdge.FromVertexID);
+	check(ClassOutput);
+
+	const FMetasoundFrontendClassInput* ClassInput = FindNodeInputClassInput(InEdge.ToNodeID, InEdge.ToVertexID);
+	check(ClassInput);
+
+	if (!FMetasoundFrontendClassVertex::CanConnectVertexAccessTypes(ClassOutput->AccessType, ClassInput->AccessType))
+	{
+		return EInvalidEdgeReason::MismatchedAccessType;
+	}
+
+	return EInvalidEdgeReason::None;
+}
+
 bool FMetaSoundFrontendDocumentBuilder::ModifyInterfaces(Metasound::Frontend::FModifyInterfaceOptions&& InOptions)
 {
 	using namespace Metasound::Frontend;
@@ -2288,8 +2440,7 @@
 		FMetasoundFrontendEdge NewEdge = RemovedEdge;
 		NewEdge.FromNodeID = NewInputNode->GetID();
 		NewEdge.FromVertexID = NewInputNode->Interface.Outputs.Last().VertexID;
-		const bool bEdgeAdded = AddEdge(MoveTemp(NewEdge)) != nullptr;
-		checkf(bEdgeAdded, TEXT("Failed to add replacement edge when swapping paired inputs"));
+		AddEdge(MoveTemp(NewEdge));
 	}
 	return true;
 }
@@ -2366,8 +2517,7 @@
 		FMetasoundFrontendEdge NewEdge = RemovedEdge;
 		NewEdge.ToNodeID = NewOutputNode->GetID();
 		NewEdge.ToVertexID = NewOutputNode->Interface.Inputs.Last().VertexID;
-		const bool bEdgeAdded = AddEdge(MoveTemp(NewEdge)) != nullptr;
-		checkf(bEdgeAdded, TEXT("Failed to add replacement edge when swapping paired outputs"));
+		AddEdge(MoveTemp(NewEdge));
 	}
 
 	return true;
diff --git a/Engine/Plugins/Runtime/Metasound/Source/MetasoundFrontend/Public/MetasoundFrontendDocumentBuilder.h b/Engine/Plugins/Runtime/Metasound/Source/MetasoundFrontend/Public/MetasoundFrontendDocumentBuilder.h
--- a/Engine/Plugins/Runtime/Metasound/Source/MetasoundFrontend/Public/MetasoundFrontendDocumentBuilder.h
+++ b/Engine/Plugins/Runtime/Metasound/Source/MetasoundFrontend/Public/MetasoundFrontendDocumentBuilder.h
@@ -24,6 +24,18 @@ namespace Metasound::Frontend
 {
 	using FFinalizeNodeFunctionRef = TFunctionRef<void(FMetasoundFrontendNode&, const Metasound::Frontend::FNodeRegistryKey&)>;
 
+	enum class EInvalidEdgeReason : uint8
+	{
+		None = 0,
+		MismatchedAccessType,
+		MismatchedDataType,
+		MissingInput,
+		MissingOutput,
+		COUNT
+	};
+
+	METASOUNDFRONTEND_API FString LexToString(const EInvalidEdgeReason& InReason);
+
 	struct METASOUNDFRONTEND_API FNamedEdge
 	{
 		const FGuid OutputNodeID;
@@ -132,11 +144,11 @@ public:
 	FMetaSoundFrontendDocumentBuilder(TScriptInterface<IMetaSoundDocumentInterface> InDocumentInterface, TSharedRef<Metasound::Frontend::FDocumentModifyDelegates> InDocumentDelegates);
 
 	const FMetasoundFrontendClass* AddDependency(const FMetasoundFrontendClass& InClass);
-	const FMetasoundFrontendEdge* AddEdge(FMetasoundFrontendEdge&& InNewEdge);
-	bool AddNamedEdges(const TSet<Metasound::Frontend::FNamedEdge>& ConnectionsToMake, TArray<const FMetasoundFrontendEdge*>* OutNewEdges = nullptr);
-	bool AddEdgesByNodeClassInterfaceBindings(const FGuid& InFromNodeID, const FGuid& InToNodeID);
-	bool AddEdgesFromMatchingInterfaceNodeOutputsToGraphOutputs(const FGuid& InNodeID, TArray<const FMetasoundFrontendEdge*>& OutEdgesCreated);
-	bool AddEdgesFromMatchingInterfaceNodeInputsToGraphInputs(const FGuid& InNodeID, TArray<const FMetasoundFrontendEdge*>& OutEdgesCreated);
+	void AddEdge(FMetasoundFrontendEdge&& InNewEdge);
+	bool AddNamedEdges(const TSet<Metasound::Frontend::FNamedEdge>& ConnectionsToMake, TArray<const FMetasoundFrontendEdge*>* OutEdgesCreated = nullptr, bool bReplaceExistingConnections = true);
+	bool AddEdgesByNodeClassInterfaceBindings(const FGuid& InFromNodeID, const FGuid& InToNodeID, bool bReplaceExistingConnections = true);
+	bool AddEdgesFromMatchingInterfaceNodeOutputsToGraphOutputs(const FGuid& InNodeID, TArray<const FMetasoundFrontendEdge*>& OutEdgesCreated, bool bReplaceExistingConnections = true);
+	bool AddEdgesFromMatchingInterfaceNodeInputsToGraphInputs(const FGuid& InNodeID, TArray<const FMetasoundFrontendEdge*>& OutEdgesCreated, bool bReplaceExistingConnections = true);
 	const FMetasoundFrontendNode* AddGraphInput(const FMetasoundFrontendClassInput& InClassInput);
 	const FMetasoundFrontendNode* AddGraphOutput(const FMetasoundFrontendClassOutput& InClassOutput);
 	bool AddInterface(FName InterfaceName);
@@ -144,6 +156,8 @@ public:
 	const FMetasoundFrontendNode* AddGraphNode(const FMetasoundFrontendGraphClass& InClass, FGuid InNodeID = FGuid::NewGuid());
 	const FMetasoundFrontendNode* AddNodeByClassName(const FMetasoundFrontendClassName& InClassName, int32 InMajorVersion, FGuid InNodeID = FGuid::NewGuid());
 
+	// Returns whether or not the given edge can be added, which requires that its input
+	// is not already connected and the edge is valid (see function 'IsValidEdge').
 	bool CanAddEdge(const FMetasoundFrontendEdge& InEdge) const;
 
 	void ClearGraph();
@@ -167,14 +181,17 @@ public:
 	const FMetasoundFrontendClassOutput* FindGraphOutput(FName OutputName) const;
 	const FMetasoundFrontendNode* FindGraphOutputNode(FName OutputName) const;
 
+	const FMetasoundFrontendNode* FindNode(const FGuid& InNodeID) const;
+
 	const FMetasoundFrontendVertex* FindNodeInput(const FGuid& InNodeID, const FGuid& InVertexID) const;
 	const FMetasoundFrontendVertex* FindNodeInput(const FGuid& InNodeID, FName InVertexName) const;
+	TArray<const FMetasoundFrontendVertex*> FindNodeInputs(const FGuid& InNodeID, FName TypeName = FName()) const;
+	TArray<const FMetasoundFrontendVertex*> FindNodeInputsConnectedToNodeOutput(const FGuid& InOutputNodeID, const FGuid& InOutputVertexID, TArray<const FMetasoundFrontendNode*>* ConnectedInputNodes = nullptr) const;
+
 	const FMetasoundFrontendVertex* FindNodeOutput(const FGuid& InNodeID, const FGuid& InVertexID) const;
 	const FMetasoundFrontendVertex* FindNodeOutput(const FGuid& InNodeID, FName InVertexName) const;
-
-	const FMetasoundFrontendNode* FindNode(const FGuid& InNodeID) const;
-	TArray<const FMetasoundFrontendVertex*> FindNodeInputs(const FGuid& InNodeID, FName TypeName = FName()) const;
 	TArray<const FMetasoundFrontendVertex*> FindNodeOutputs(const FGuid& InNodeID, FName TypeName = FName()) const;
+	const FMetasoundFrontendVertex* FindNodeOutputConnectedToNodeInput(const FGuid& InInputNodeID, const FGuid& InInputVertexID, const FMetasoundFrontendNode** ConnectedOutputNode = nullptr) const;
 
 	const FMetasoundFrontendDocument& GetDocument() const;
 	const Metasound::Frontend::FDocumentModifyDelegates& GetDocumentDelegates() const;
@@ -195,6 +212,10 @@ public:
 	bool IsInterfaceDeclared(const FMetasoundFrontendVersion& InInterfaceVersion) const;
 	bool IsPreset() const;
 
+	// Returns whether or not the given edge is valid (i.e. represents an input and output that equate in data and access types) or malformed.
+	// Note that this does not return whether or not the given edge exists, but rather if it could be legally applied to the given edge vertices.
+	Metasound::Frontend::EInvalidEdgeReason IsValidEdge(const FMetasoundFrontendEdge& InEdge) const;
+
 	bool ModifyInterfaces(Metasound::Frontend::FModifyInterfaceOptions&& InOptions);
 
 	bool RemoveDependency(const FGuid& InClassID);
