Fix MetaSound reroute nodes not retaining connections when copy pasted

diff --git a/Engine/Plugins/Runtime/Metasound/Source/MetasoundEditor/Private/MetasoundEditor.cpp b/Engine/Plugins/Runtime/Metasound/Source/MetasoundEditor/Private/MetasoundEditor.cpp
--- a/Engine/Plugins/Runtime/Metasound/Source/MetasoundEditor/Private/MetasoundEditor.cpp
+++ b/Engine/Plugins/Runtime/Metasound/Source/MetasoundEditor/Private/MetasoundEditor.cpp
@@ -2418,7 +2418,8 @@ namespace Metasound
 									FMetasoundFrontendNodeInterface NodeInterface = FRerouteNodeTemplate::CreateNodeInterfaceFromDataType(DataType);
 
 									FNodeHandle NodeHandle = MetasoundAsset->GetRootGraphHandle()->AddTemplateNode(RerouteTemplateKey, MoveTemp(NodeInterface));
-									FGraphBuilder::InitGraphNode(NodeHandle, ExternalNode, *Metasound);
+									FGraphBuilder::InitGraphNodeIDFromNodeHandle(NodeHandle, ExternalNode);
+
 									bIsValid = NodeHandle->IsValid();
 								}
 							}
diff --git a/Engine/Plugins/Runtime/Metasound/Source/MetasoundEditor/Private/MetasoundEditorGraphBuilder.cpp b/Engine/Plugins/Runtime/Metasound/Source/MetasoundEditor/Private/MetasoundEditorGraphBuilder.cpp
--- a/Engine/Plugins/Runtime/Metasound/Source/MetasoundEditor/Private/MetasoundEditorGraphBuilder.cpp
+++ b/Engine/Plugins/Runtime/Metasound/Source/MetasoundEditor/Private/MetasoundEditorGraphBuilder.cpp
@@ -558,6 +558,11 @@ namespace Metasound
 			RebuildNodePins(*NewGraphNode);
 		}
 
+		void FGraphBuilder::InitGraphNodeIDFromNodeHandle(const Frontend::FConstNodeHandle& InNodeHandle, UMetasoundEditorGraphNode* NewGraphNode)
+		{
+			NewGraphNode->SetNodeID(InNodeHandle->GetID());
+		}
+
 		FGraphValidationResults FGraphBuilder::ValidateGraph(UObject& InMetaSound)
 		{
 			TRACE_CPUPROFILER_EVENT_SCOPE(Metasound::Editor::FGraphBuilder::ValidateGraph);
diff --git a/Engine/Plugins/Runtime/Metasound/Source/MetasoundEditor/Private/MetasoundEditorGraphBuilder.h b/Engine/Plugins/Runtime/Metasound/Source/MetasoundEditor/Private/MetasoundEditorGraphBuilder.h
--- a/Engine/Plugins/Runtime/Metasound/Source/MetasoundEditor/Private/MetasoundEditorGraphBuilder.h
+++ b/Engine/Plugins/Runtime/Metasound/Source/MetasoundEditor/Private/MetasoundEditorGraphBuilder.h
@@ -60,6 +60,7 @@ namespace Metasound
 			static const FText GraphMenuName;
 
 			static void InitGraphNode(Frontend::FNodeHandle& InNodeHandle, UMetasoundEditorGraphNode* NewGraphNode, UObject& InMetaSound);
+			static void InitGraphNodeIDFromNodeHandle(const Frontend::FConstNodeHandle& InNodeHandle, UMetasoundEditorGraphNode* NewGraphNode);
 
 			// Adds an EdGraph node to mirror the provided FNodeHandle.
 			static UMetasoundEditorGraphNode* AddNode(UObject& InMetaSound, Frontend::FNodeHandle InNodeHandle, FVector2D InLocation, bool bInSelectNewNode = true);
