Adds functionality for profiling Metasound nodes in Insights.

diff --git a/Engine/Plugins/Runtime/Metasound/Source/MetasoundGraphCore/Private/MetasoundGraphCoreModule.cpp b/Engine/Plugins/Runtime/Metasound/Source/MetasoundGraphCore/Private/MetasoundGraphCoreModule.cpp
--- a/Engine/Plugins/Runtime/Metasound/Source/MetasoundGraphCore/Private/MetasoundGraphCoreModule.cpp
+++ b/Engine/Plugins/Runtime/Metasound/Source/MetasoundGraphCore/Private/MetasoundGraphCoreModule.cpp
@@ -4,6 +4,7 @@
 
 #include "CoreMinimal.h"
 #include "MetasoundLog.h"
+#include "MetasoundProfilingOperator.h"
 #include "Modules/ModuleInterface.h"
 #include "Modules/ModuleManager.h"
 
@@ -13,6 +14,11 @@ namespace Metasound
 {
 	class FMetasoundGraphCoreModule : public IMetasoundGraphCoreModule
 	{
+	public:
+		virtual void StartupModule() override
+		{
+			Profiling::Init();
+		}
 	};
 }
 
diff --git a/Engine/Plugins/Runtime/Metasound/Source/MetasoundGraphCore/Private/MetasoundOperatorBuilder.cpp b/Engine/Plugins/Runtime/Metasound/Source/MetasoundGraphCore/Private/MetasoundOperatorBuilder.cpp
--- a/Engine/Plugins/Runtime/Metasound/Source/MetasoundGraphCore/Private/MetasoundOperatorBuilder.cpp
+++ b/Engine/Plugins/Runtime/Metasound/Source/MetasoundGraphCore/Private/MetasoundOperatorBuilder.cpp
@@ -14,6 +14,7 @@
 #include "MetasoundGraphLinter.h"
 #include "MetasoundGraphOperator.h"
 #include "MetasoundOperatorInterface.h"
+#include "MetasoundProfilingOperator.h"
 #include "MetasoundRebindableGraphOperator.h"
 #include "MetasoundThreadLocalDebug.h"
 #include "MetasoundTrace.h"
@@ -66,7 +67,7 @@ namespace Metasound
 			}
 		}
 	}
-
+	
 	FOperatorBuilder::FOperatorBuilder(const FOperatorBuilderSettings& InBuilderSettings)
 	: BuilderSettings(InBuilderSettings)
 	{
@@ -438,6 +439,8 @@ namespace Metasound
 
 		using namespace DirectedGraphAlgo;
 
+		bool ProfileOperators = BuilderSettings.bProfileOperators || Profiling::ProfileAllGraphs();
+
 		FBuildStatus BuildStatus;
 
 		// Create FOperatorInfos from Nodes
@@ -455,7 +458,14 @@ namespace Metasound
 				METASOUND_TRACE_CPUPROFILER_EVENT_SCOPE(Metasound::FOperatorBuilder::CreateOperators::CreateAndBind);
 				FBuildOperatorParams CreateParams{*Node, InOutContext.Settings, OperatorInfo.VertexData.GetInputs(), InOutContext.Environment, this};
 				FOperatorFactorySharedRef Factory = Node->GetDefaultOperatorFactory();
-				OperatorInfo.Operator = Factory->CreateOperator(CreateParams, InOutContext.Results);
+				if (ProfileOperators && Profiling::OperatorShouldBeProfiled(Node->GetMetadata()))
+				{
+					OperatorInfo.Operator = MakeUnique<ProfilingOperator>(Factory->CreateOperator(CreateParams, InOutContext.Results), Node->GetMetadata());
+				}
+				else
+				{
+					OperatorInfo.Operator = Factory->CreateOperator(CreateParams, InOutContext.Results);
+				}
 
 				if (!OperatorInfo.Operator.IsValid())
 				{
diff --git a/Engine/Plugins/Runtime/Metasound/Source/MetasoundGraphCore/Private/MetasoundProfilingOperator.cpp b/Engine/Plugins/Runtime/Metasound/Source/MetasoundGraphCore/Private/MetasoundProfilingOperator.cpp
new file mode 100644
--- /dev/null
+++ b/Engine/Plugins/Runtime/Metasound/Source/MetasoundGraphCore/Private/MetasoundProfilingOperator.cpp
@@ -0,0 +1,110 @@
+// Copyright Epic Games, Inc. All Rights Reserved.
+
+#include "MetasoundProfilingOperator.h"
+
+#include "MetasoundTrace.h"
+#include "HAL/IConsoleManager.h"
+
+DEFINE_LOG_CATEGORY_STATIC(LogMetasoundProfiling, Log, All);
+
+namespace Metasound
+{
+	namespace Profiling
+	{
+		int32 ProfileAllEnabled = 0;
+		FAutoConsoleVariableRef CVarMetaSoundProfileAllEnabled(
+			TEXT("au.MetaSound.ProfileAllGraphs"),
+			ProfileAllEnabled,
+			TEXT("Enable profiling of all MetaSound graphs. NOTE: If the node filter is set it will still apply (see au.Metasound.AddProfileNode)\n")
+			TEXT("0: Disabled (default), !0: Enabled"),
+			ECVF_Default);
+
+		TArray<FString> ProfilingNodeFilter;
+
+		void Init()
+		{
+			IConsoleManager::Get().RegisterConsoleCommand(
+				TEXT("au.Metasound.Profiling.AddNodes"),
+				TEXT("Adds the specified node class name(s) to the list of metasound nodes that will be profiled and visible in Insights."),
+				FConsoleCommandWithArgsDelegate::CreateLambda([](const TArray<FString>& Args)
+					{
+						for (const FString& NodeName : Args)
+						{
+							ProfilingNodeFilter.AddUnique(NodeName);
+						}
+					}),
+				ECVF_Default
+			);
+			IConsoleManager::Get().RegisterConsoleCommand(
+				TEXT("au.Metasound.Profiling.ListNodes"),
+				TEXT("Lists the node class names that will be profiled and visible in Insights."),
+				FConsoleCommandDelegate::CreateLambda([]()
+					{
+						UE_LOG(LogMetasoundProfiling, Display, TEXT("Metasound profiling limited to nodes with these names:"));
+						for (const FString& NodeName : ProfilingNodeFilter)
+						{
+							UE_LOG(LogMetasoundProfiling, Display, TEXT("    %s"), *NodeName);
+						}
+					}),
+				ECVF_Default
+			);
+			IConsoleManager::Get().RegisterConsoleCommand(
+				TEXT("au.Metasound.Profiling.RemoveNodes"),
+				TEXT("Removes the specified node class name(s) (or ALL if no names are provided) from the list of node types that will be profiled and visible in Insights."),
+				FConsoleCommandWithArgsDelegate::CreateLambda([](const TArray<FString>& Args)
+					{
+						if (Args.Num() == 0)
+						{
+							ProfilingNodeFilter.Empty();
+						}
+						else
+						{
+							for (const FString& NodeName : Args)
+							{
+								ProfilingNodeFilter.Remove(NodeName);
+							}
+						}
+					}),
+				ECVF_Default
+			);
+
+		}
+
+		bool OperatorShouldBeProfiled(const FNodeClassMetadata& NodeMetadata)
+		{
+			return ProfilingNodeFilter.IsEmpty() || ProfilingNodeFilter.Contains(NodeMetadata.ClassName.GetName().ToString());
+		}
+
+		bool ProfileAllGraphs()
+		{
+			return ProfileAllEnabled != 0;
+		}
+		
+		
+	}
+
+	void ProfilingOperator::StaticReset(IOperator* InOperator, const IOperator::FResetParams& InParams)
+	{
+		ProfilingOperator* DerivedOperator = static_cast<ProfilingOperator*>(InOperator);
+		check(nullptr != DerivedOperator && DerivedOperator->ResetFunction);
+		METASOUND_TRACE_CPUPROFILER_EVENT_SCOPE_TEXT(*DerivedOperator->InsightsResetEventName);
+		DerivedOperator->ResetFunction(DerivedOperator->Operator.Get(), InParams);
+	}
+
+	void ProfilingOperator::StaticExecute(IOperator* InOperator)
+	{
+		ProfilingOperator* DerivedOperator = static_cast<ProfilingOperator*>(InOperator);
+		check(nullptr != DerivedOperator && DerivedOperator->ExecuteFunction);
+		METASOUND_TRACE_CPUPROFILER_EVENT_SCOPE_TEXT(*DerivedOperator->InsightsExecuteEventName);
+		DerivedOperator->ExecuteFunction(DerivedOperator->Operator.Get());
+	}
+
+	void ProfilingOperator::StaticPostExecute(IOperator* InOperator)
+	{
+		ProfilingOperator* DerivedOperator = static_cast<ProfilingOperator*>(InOperator);
+		check(nullptr != DerivedOperator && DerivedOperator->PostExecuteFunction);
+		METASOUND_TRACE_CPUPROFILER_EVENT_SCOPE_TEXT(*DerivedOperator->InsightsPostExecuteEventName);
+		DerivedOperator->PostExecuteFunction(DerivedOperator->Operator.Get());
+	}
+
+}
\ No newline at end of file
diff --git a/Engine/Plugins/Runtime/Metasound/Source/MetasoundGraphCore/Private/MetasoundProfilingOperator.h b/Engine/Plugins/Runtime/Metasound/Source/MetasoundGraphCore/Private/MetasoundProfilingOperator.h
new file mode 100644
--- /dev/null
+++ b/Engine/Plugins/Runtime/Metasound/Source/MetasoundGraphCore/Private/MetasoundProfilingOperator.h
@@ -0,0 +1,85 @@
+// Copyright Epic Games, Inc. All Rights Reserved.
+
+#pragma once
+
+#include "MetasoundNodeInterface.h"
+#include "MetasoundOperatorInterface.h"
+
+namespace Metasound
+{
+
+	namespace Profiling
+	{
+		void Init();
+		bool OperatorShouldBeProfiled(const FNodeClassMetadata& NodeMetadata);
+		bool ProfileAllGraphs();
+	}
+
+	// This is a wrapper around any IOperator that causes its functions to be timed for Insights
+	class ProfilingOperator : public IOperator
+	{
+	public:
+		ProfilingOperator(TUniquePtr<IOperator>&& WrappedOperator, const FNodeClassMetadata& NodeMetadata)
+			: Operator(MoveTemp(WrappedOperator))
+			, ResetFunction(Operator->GetResetFunction())
+			, ExecuteFunction(Operator->GetExecuteFunction())
+			, PostExecuteFunction(Operator->GetPostExecuteFunction())
+		{
+			check(Operator);
+			FString BaseEventName = NodeMetadata.ClassName.GetName().ToString();
+			InsightsResetEventName = FString::Printf(TEXT("%s_RESET"), *BaseEventName);
+			InsightsExecuteEventName = FString::Printf(TEXT("%s_EXECUTE"), *BaseEventName);
+			InsightsPostExecuteEventName = FString::Printf(TEXT("%s_POSTEXECUTE"), *BaseEventName);
+		}
+
+		virtual ~ProfilingOperator() = default;
+
+		virtual void BindInputs(FInputVertexInterfaceData& InVertexData) override
+		{
+			Operator->BindInputs(InVertexData);
+		}
+
+		virtual void BindOutputs(FOutputVertexInterfaceData& InVertexData)
+		{
+			Operator->BindOutputs(InVertexData);
+		}
+
+		virtual FResetFunction GetResetFunction() override
+		{
+			if (!ResetFunction)
+			{
+				return nullptr;
+			}
+			return &StaticReset;
+		}
+		virtual FExecuteFunction GetExecuteFunction() override
+		{
+			if (!ExecuteFunction)
+			{
+				return nullptr;
+			}
+			return &StaticExecute;
+		}
+		virtual FPostExecuteFunction GetPostExecuteFunction() override
+		{
+			if (!PostExecuteFunction)
+			{
+				return nullptr;
+			}
+			return &StaticPostExecute;
+		}
+
+		static void StaticReset(IOperator* InOperator, const IOperator::FResetParams& InParams);
+		static void StaticExecute(IOperator* InOperator);
+		static void StaticPostExecute(IOperator* InOperator);
+
+	private:
+		TUniquePtr<IOperator> Operator;
+		FResetFunction        ResetFunction;
+		FExecuteFunction      ExecuteFunction;
+		FPostExecuteFunction  PostExecuteFunction;
+		FString               InsightsResetEventName;
+		FString               InsightsExecuteEventName;
+		FString               InsightsPostExecuteEventName;
+	};
+}
\ No newline at end of file
diff --git a/Engine/Plugins/Runtime/Metasound/Source/MetasoundGraphCore/Public/MetasoundOperatorBuilderSettings.h b/Engine/Plugins/Runtime/Metasound/Source/MetasoundGraphCore/Public/MetasoundOperatorBuilderSettings.h
--- a/Engine/Plugins/Runtime/Metasound/Source/MetasoundGraphCore/Public/MetasoundOperatorBuilderSettings.h
+++ b/Engine/Plugins/Runtime/Metasound/Source/MetasoundGraphCore/Public/MetasoundOperatorBuilderSettings.h
@@ -85,6 +85,12 @@ namespace Metasound
 		 */
 		bool bEnableOperatorRebind = false;
 
+		/** If true, the graph's operators will be wrapped in ProfilingOperator 
+		* instances and insights events will be generated for their Execute
+		* functions.
+		*/
+		bool bProfileOperators = false;
+
 		/** Return the default settings for the current build environment. */
 		static const FOperatorBuilderSettings& GetDefaultSettings();
 
