Fix crash when closing the MetaSound asset editor while dragging an input widget

diff --git a/Engine/Plugins/Runtime/Metasound/Source/MetasoundEditor/Private/MetasoundEditor.cpp b/Engine/Plugins/Runtime/Metasound/Source/MetasoundEditor/Private/MetasoundEditor.cpp
--- a/Engine/Plugins/Runtime/Metasound/Source/MetasoundEditor/Private/MetasoundEditor.cpp
+++ b/Engine/Plugins/Runtime/Metasound/Source/MetasoundEditor/Private/MetasoundEditor.cpp
@@ -3296,24 +3296,27 @@ namespace Metasound
 				}
 				else if (!NodesModified.IsEmpty() || !MembersModified.IsEmpty())
 				{
-					TArray<TWeakObjectPtr<UObject>> SelectedObjects = MetasoundDetails->GetSelectedObjects();
-					const bool bShouldRefreshDetails = Algo::AnyOf(SelectedObjects, [&NodesModified, &MembersModified](const TWeakObjectPtr<UObject>& Obj)
+					if (MetasoundDetails.IsValid())
 					{
-						if (const UMetasoundEditorGraphNode* Node = Cast<const UMetasoundEditorGraphNode>(Obj.Get()))
-						{
-							return NodesModified.Contains(Node->GetNodeID());
-						}
-						if (const UMetasoundEditorGraphMember* Member = Cast<const UMetasoundEditorGraphMember>(Obj.Get()))
+						TArray<TWeakObjectPtr<UObject>> SelectedObjects = MetasoundDetails->GetSelectedObjects();
+						const bool bShouldRefreshDetails = Algo::AnyOf(SelectedObjects, [&NodesModified, &MembersModified](const TWeakObjectPtr<UObject>& Obj)
+							{
+								if (const UMetasoundEditorGraphNode* Node = Cast<const UMetasoundEditorGraphNode>(Obj.Get()))
+								{
+									return NodesModified.Contains(Node->GetNodeID());
+								}
+								if (const UMetasoundEditorGraphMember* Member = Cast<const UMetasoundEditorGraphMember>(Obj.Get()))
+								{
+									return MembersModified.Contains(Member->GetMemberID());
+								}
+								return false;
+							});
+						if (bShouldRefreshDetails)
 						{
-							return MembersModified.Contains(Member->GetMemberID());
+							RefreshDetails();
 						}
-						return false;
-					});
-					if (bShouldRefreshDetails)
-					{
-						RefreshDetails();
 					}
-
+					
 					// TODO: Because input editor nodes are not one-to-one, this can cause multi-selection
 					// when not desired.  Once input alias templates are complete, this can be switched on.
 					// For now, callsites are just directly setting selection via ClearSelectionAndSelectNode.
diff --git a/Engine/Plugins/Runtime/Metasound/Source/MetasoundEditor/Private/SMetasoundGraphNode.cpp b/Engine/Plugins/Runtime/Metasound/Source/MetasoundEditor/Private/SMetasoundGraphNode.cpp
--- a/Engine/Plugins/Runtime/Metasound/Source/MetasoundEditor/Private/SMetasoundGraphNode.cpp
+++ b/Engine/Plugins/Runtime/Metasound/Source/MetasoundEditor/Private/SMetasoundGraphNode.cpp
@@ -62,6 +62,7 @@ namespace Metasound
 	{
 		SMetaSoundGraphNode::~SMetaSoundGraphNode()
 		{
+			// Clean up input widgets
 			UMetasoundEditorGraphNode& Node = GetMetaSoundNode();
 			if (UMetasoundEditorGraphMemberNode* MemberNode = Cast<UMetasoundEditorGraphMemberNode>(&Node))
 			{
@@ -69,17 +70,23 @@ namespace Metasound
 				{
 					if (UMetasoundEditorGraphMemberDefaultFloat* DefaultFloat = Cast<UMetasoundEditorGraphMemberDefaultFloat>(GraphMember->GetLiteral()))
 					{
+						// This may hit if the asset editor is closed while interacting with a widget 
+						// (ex. Ctrl-W is pressed mid drag before the value is committed)
+						if (bIsInputWidgetTransacting)
+						{
+							GEditor->EndTransaction();
+							if (UMetasoundEditorGraph* Graph = GraphMember->GetOwningGraph())
+							{
+								constexpr bool bPostTransaction = false;
+								GraphMember->UpdateFrontendDefaultLiteral(bPostTransaction);
+								Graph->GetModifyContext().AddMemberIDsModified({ GraphMember->GetMemberID() });
+							}
+						}
 						DefaultFloat->OnDefaultValueChanged.Remove(InputSliderOnValueChangedDelegateHandle);
 						DefaultFloat->OnRangeChanged.Remove(InputSliderOnRangeChangedDelegateHandle);
 					}
 				}
 			}
-
-			if (bIsInputWidgetTransacting)
-			{
-				GEditor->EndTransaction();
-				UE_LOG(LogMetaSound, Warning, TEXT("Unmatched MetaSound editor widget transaction."));
-			}
 		}
 
 		bool SMetaSoundGraphNode::IsVariableAccessor() const
