Only audition literal update if no edge is already set.
#rnx

diff --git a/Engine/Plugins/Runtime/Metasound/Source/MetasoundEngine/Private/MetasoundBuilderSubsystem.cpp b/Engine/Plugins/Runtime/Metasound/Source/MetasoundEngine/Private/MetasoundBuilderSubsystem.cpp
--- a/Engine/Plugins/Runtime/Metasound/Source/MetasoundEngine/Private/MetasoundBuilderSubsystem.cpp
+++ b/Engine/Plugins/Runtime/Metasound/Source/MetasoundEngine/Private/MetasoundBuilderSubsystem.cpp
@@ -1140,21 +1140,24 @@ void UMetaSoundSourceBuilder::OnNodeInputLiteralSet(int32 NodeIndex, int32 Verte
 {
 	using namespace Metasound::DynamicGraph;
 
-	ExecuteAuditionableTransaction([this, &NodeIndex, &VertexIndex, &LiteralIndex](FDynamicOperatorTransactor& Transactor)
-	{
-		using namespace Metasound;
-		using namespace Metasound::Engine;
-		using namespace Metasound::Frontend;
+	const FMetasoundFrontendGraphClass& GraphClass = Builder.GetDocument().RootGraph;
+	const FMetasoundFrontendNode& Node = GraphClass.Graph.Nodes[NodeIndex];
+	const FMetasoundFrontendVertex& Input = Node.Interface.Inputs[VertexIndex];
 
-		const FMetasoundFrontendGraphClass& GraphClass = Builder.GetDocument().RootGraph;
-		const FMetasoundFrontendNode& Node = GraphClass.Graph.Nodes[NodeIndex];
-		const FMetasoundFrontendVertex& Input = Node.Interface.Inputs[VertexIndex];
-		const FMetasoundFrontendLiteral& InputDefault = Node.InputLiterals[LiteralIndex].Value;
-		TUniquePtr<INode> LiteralNode = BuilderSubsystemPrivate::CreateDynamicNodeFromFrontendLiteral(Input.TypeName, InputDefault);
-		Transactor.SetValue(Node.GetID(), Input.Name, MoveTemp(LiteralNode));
+	// Only send the literal down if not connected, as the graph core layer will disconnect.
+	if (!Builder.IsNodeInputConnected(Node.GetID(), Input.VertexID))
+	{
+		ExecuteAuditionableTransaction([this, &Node, &Input, &LiteralIndex](FDynamicOperatorTransactor& Transactor)
+		{
+			using namespace Metasound;
+			using namespace Metasound::Engine;
 
-		return true;
-	});
+			const FMetasoundFrontendLiteral& InputDefault = Node.InputLiterals[LiteralIndex].Value;
+			TUniquePtr<INode> LiteralNode = BuilderSubsystemPrivate::CreateDynamicNodeFromFrontendLiteral(Input.TypeName, InputDefault);
+			Transactor.SetValue(Node.GetID(), Input.Name, MoveTemp(LiteralNode));
+			return true;
+		});
+	}
 }
 
 void UMetaSoundSourceBuilder::OnOutputAdded(int32 OutputIndex) const
