- Add ability to set node position explicitly via MetaSound Editor Subsystem
- Fix bug introduced recently with BuildToAsset not applying a clean class name/guid
- Fix case where builder was added to root and not freed
- MetaSoundEditorSubsystem minor clean-up
#rnx

diff --git a/Engine/Plugins/Runtime/Metasound/Source/MetasoundEditor/Private/MetasoundEditorSubsystem.cpp b/Engine/Plugins/Runtime/Metasound/Source/MetasoundEditor/Private/MetasoundEditorSubsystem.cpp
--- a/Engine/Plugins/Runtime/Metasound/Source/MetasoundEditor/Private/MetasoundEditorSubsystem.cpp
+++ b/Engine/Plugins/Runtime/Metasound/Source/MetasoundEditor/Private/MetasoundEditorSubsystem.cpp
@@ -18,16 +18,22 @@
 #define LOCTEXT_NAMESPACE "MetaSoundEditor"
 
 
-TScriptInterface<IMetaSoundDocumentInterface> UMetaSoundEditorSubsystem::BuildToAsset(UMetaSoundBuilderBase* InBuilder, const FString& Author, const FString& AssetName, const FString& PackagePath, EMetaSoundBuilderResult& OutResult, const USoundWave* TemplateSoundWave)
+TScriptInterface<IMetaSoundDocumentInterface> UMetaSoundEditorSubsystem::BuildToAsset(
+	UMetaSoundBuilderBase* InBuilder,
+	const FString& Author,
+	const FString& AssetName,
+	const FString& PackagePath,
+	EMetaSoundBuilderResult& OutResult,
+	const USoundWave* TemplateSoundWave
+)
 {
 	using namespace Metasound;
 	using namespace Metasound::Frontend;
 
+	OutResult = EMetaSoundBuilderResult::Failed;
+
 	if (InBuilder)
 	{
-		// Not about to follow this lack of const correctness down a multidecade in the works rabbit hole.
-		UClass& BuilderUClass = const_cast<UClass&>(InBuilder->GetBuilderUClass());
-
 		// AddToRoot to avoid builder getting gc'ed during CreateAsset call below, as the builder 
 		// may be unreferenced by other UObjects and it must be persistent to finish initializing.
 		const bool bWasRooted = InBuilder->IsRooted();
@@ -37,40 +43,55 @@
 		}
 
 		constexpr UFactory* Factory = nullptr;
+		// Not about to follow this lack of const correctness down a multidecade in the works rabbit hole.
+		UClass& BuilderUClass = const_cast<UClass&>(InBuilder->GetBuilderUClass());
 		if (UObject* NewMetaSound = IAssetTools::Get().CreateAsset(AssetName, PackagePath, &BuilderUClass, Factory))
 		{
 			InBuilder->InitNodeLocations();
 			InBuilder->SetAuthor(Author);
 
-			FMetaSoundBuilderOptions BuilderOptions { FName(*AssetName) };
-			BuilderOptions.ExistingMetaSound = NewMetaSound;
-			constexpr UObject* Parent = nullptr;
-			TScriptInterface<IMetaSoundDocumentInterface> DocInterface = InBuilder->Build(Parent, BuilderOptions);
+			// Initialize and Build
+			{
+				constexpr UObject* Parent = nullptr;
+				constexpr bool bForceUniqueClassName = true;
+				constexpr bool bAddToRegistry = true;
+				const FMetaSoundBuilderOptions BuilderOptions { FName(*AssetName), bForceUniqueClassName, bAddToRegistry, NewMetaSound };
+				InBuilder->Build(Parent, BuilderOptions);
+			}
 
-			const bool bIsSource = &BuilderUClass == UMetaSoundSource::StaticClass();
-			if (InBuilder->IsPreset())
+			// Apply template SoundWave settings
 			{
-				FMetasoundAssetBase* PresetAsset = IMetasoundUObjectRegistry::Get().GetObjectAsAssetBase(NewMetaSound);
-				check(PresetAsset);
-				PresetAsset->ConformObjectDataToInterfaces();
+				const bool bIsSource = &BuilderUClass == UMetaSoundSource::StaticClass();
+				if (InBuilder->IsPreset())
+				{
+					FMetasoundAssetBase* PresetAsset = IMetasoundUObjectRegistry::Get().GetObjectAsAssetBase(NewMetaSound);
+					check(PresetAsset);
+					PresetAsset->ConformObjectDataToInterfaces();
 
-				// Only use referenced object's soundwave settings for sources if not overridden 
-				if (TemplateSoundWave == nullptr && bIsSource)
-				{
-					if (const UObject* ReferencedObject = InBuilder->GetReferencedPresetAsset())
+					// Only use referenced UObject's SoundWave settings for sources if not overridden 
+					if (TemplateSoundWave == nullptr && bIsSource)
 					{
-						TemplateSoundWave = CastChecked<USoundWave>(ReferencedObject);
+						if (const UObject* ReferencedObject = InBuilder->GetReferencedPresetAsset())
+						{
+							TemplateSoundWave = CastChecked<USoundWave>(ReferencedObject);
+						}
 					}
 				}
+
+				// Template SoundWave settings only apply to sources
+				if (TemplateSoundWave != nullptr && bIsSource)
+				{
+					SetSoundWaveSettingsFromTemplate(*CastChecked<USoundWave>(NewMetaSound), *TemplateSoundWave);
+				}
 			}
 
-			// Template sound wave settings only apply to sources 
-			if (TemplateSoundWave != nullptr && bIsSource)
-			{
-				SetSoundWaveSettingsFromTemplate(*CastChecked<USoundWave>(NewMetaSound), *TemplateSoundWave);
-			}
-
 			InitEdGraph(*NewMetaSound);
+
+			if (!bWasRooted)
+			{
+				InBuilder->RemoveFromRoot();
+			}
+
 			OutResult = EMetaSoundBuilderResult::Succeeded;
 			return NewMetaSound;
 		}
@@ -81,43 +102,36 @@
 		}
 	}
 
-	OutResult = EMetaSoundBuilderResult::Failed;
 	return nullptr;
 }
 
-void UMetaSoundEditorSubsystem::SetSoundWaveSettingsFromTemplate(USoundWave& NewMetaSoundWave, const USoundWave& TemplateSoundWave) const
+UMetaSoundEditorSubsystem& UMetaSoundEditorSubsystem::GetChecked()
 {
-	// Sound 
-	NewMetaSoundWave.Volume = TemplateSoundWave.Volume;
-	NewMetaSoundWave.Pitch = TemplateSoundWave.Pitch;
-	NewMetaSoundWave.SoundClassObject = TemplateSoundWave.SoundClassObject;
-
-	// Attenuation 
-	NewMetaSoundWave.AttenuationSettings = TemplateSoundWave.AttenuationSettings;
-	NewMetaSoundWave.bDebug = TemplateSoundWave.bDebug;
-
-	// Effects 
-	NewMetaSoundWave.bEnableBusSends = TemplateSoundWave.bEnableBusSends;
-	NewMetaSoundWave.SourceEffectChain = TemplateSoundWave.SourceEffectChain;
-	NewMetaSoundWave.BusSends = TemplateSoundWave.BusSends;
-	NewMetaSoundWave.PreEffectBusSends = TemplateSoundWave.PreEffectBusSends; 
-
-	NewMetaSoundWave.bEnableBaseSubmix = TemplateSoundWave.bEnableBaseSubmix;
-	NewMetaSoundWave.SoundSubmixObject = TemplateSoundWave.SoundSubmixObject;
-	NewMetaSoundWave.bEnableSubmixSends = TemplateSoundWave.bEnableSubmixSends;
-	NewMetaSoundWave.SoundSubmixSends = TemplateSoundWave.SoundSubmixSends;
-
-	// Modulation 
-	NewMetaSoundWave.ModulationSettings = TemplateSoundWave.ModulationSettings;
-	
-	// Voice Management 
-	NewMetaSoundWave.VirtualizationMode = TemplateSoundWave.VirtualizationMode;
-	NewMetaSoundWave.bOverrideConcurrency = TemplateSoundWave.bOverrideConcurrency;
-	NewMetaSoundWave.ConcurrencySet = TemplateSoundWave.ConcurrencySet;
-	NewMetaSoundWave.ConcurrencyOverrides = TemplateSoundWave.ConcurrencyOverrides;
-
-	NewMetaSoundWave.bBypassVolumeScaleForPriority = TemplateSoundWave.bBypassVolumeScaleForPriority;
-	NewMetaSoundWave.Priority = TemplateSoundWave.Priority;
+	checkf(GEditor, TEXT("Cannot access UMetaSoundEditorSubsystem without editor loaded"));
+	UMetaSoundEditorSubsystem* EditorSubsystem = GEditor->GetEditorSubsystem<UMetaSoundEditorSubsystem>();
+	checkf(EditorSubsystem, TEXT("Failed to find initialized 'UMetaSoundEditorSubsystem"));
+	return *EditorSubsystem;
 }
 
+const UMetaSoundEditorSubsystem& UMetaSoundEditorSubsystem::GetConstChecked()
+{
+	checkf(GEditor, TEXT("Cannot access UMetaSoundEditorSubsystem without editor loaded"));
+	UMetaSoundEditorSubsystem* EditorSubsystem = GEditor->GetEditorSubsystem<UMetaSoundEditorSubsystem>();
+	checkf(EditorSubsystem, TEXT("Failed to find initialized 'UMetaSoundEditorSubsystem"));
+	return *EditorSubsystem;
+}
+
+const FString UMetaSoundEditorSubsystem::GetDefaultAuthor()
+{
+	FString Author = UKismetSystemLibrary::GetPlatformUserName();
+	if (const UMetasoundEditorSettings* EditorSettings = GetDefault<UMetasoundEditorSettings>())
+	{
+		if (!EditorSettings->DefaultAuthor.IsEmpty())
+		{
+			Author = EditorSettings->DefaultAuthor;
+		}
+	}
+	return Author;
+}
+
 void UMetaSoundEditorSubsystem::InitAsset(UObject& InNewMetaSound, UObject* InReferencedMetaSound)
@@ -191,33 +205,55 @@
 	Metasound::Editor::FGraphBuilder::RegisterGraphWithFrontend(InMetaSound, bInForceViewSynchronization);
 }
 
-const FString UMetaSoundEditorSubsystem::GetDefaultAuthor()
+void UMetaSoundEditorSubsystem::SetNodeLocation(
+	UMetaSoundBuilderBase* InBuilder,
+	const FMetaSoundNodeHandle& InNode,
+	const FVector2D& InLocation,
+	EMetaSoundBuilderResult& OutResult)
 {
-	FString Author = UKismetSystemLibrary::GetPlatformUserName();
-	if (const UMetasoundEditorSettings* EditorSettings = GetDefault<UMetasoundEditorSettings>())
+	if (InBuilder)
 	{
-		if (!EditorSettings->DefaultAuthor.IsEmpty())
-		{
-			Author = EditorSettings->DefaultAuthor;
-		}
+		InBuilder->SetNodeLocation(InNode, InLocation, OutResult);
 	}
-	return Author;
+	else
+	{
+		OutResult = EMetaSoundBuilderResult::Failed;
+	}
 }
 
-UMetaSoundEditorSubsystem& UMetaSoundEditorSubsystem::GetChecked()
+void UMetaSoundEditorSubsystem::SetSoundWaveSettingsFromTemplate(USoundWave& NewMetaSoundWave, const USoundWave& TemplateSoundWave) const
 {
-	checkf(GEditor, TEXT("Cannot access UMetaSoundEditorSubsystem without editor loaded"));
-	UMetaSoundEditorSubsystem* EditorSubsystem = GEditor->GetEditorSubsystem<UMetaSoundEditorSubsystem>();
-	checkf(EditorSubsystem, TEXT("Failed to find initialized 'UMetaSoundEditorSubsystem"));
-	return *EditorSubsystem;
+	// Sound 
+	NewMetaSoundWave.Volume = TemplateSoundWave.Volume;
+	NewMetaSoundWave.Pitch = TemplateSoundWave.Pitch;
+	NewMetaSoundWave.SoundClassObject = TemplateSoundWave.SoundClassObject;
+
+	// Attenuation 
+	NewMetaSoundWave.AttenuationSettings = TemplateSoundWave.AttenuationSettings;
+	NewMetaSoundWave.bDebug = TemplateSoundWave.bDebug;
+
+	// Effects 
+	NewMetaSoundWave.bEnableBusSends = TemplateSoundWave.bEnableBusSends;
+	NewMetaSoundWave.SourceEffectChain = TemplateSoundWave.SourceEffectChain;
+	NewMetaSoundWave.BusSends = TemplateSoundWave.BusSends;
+	NewMetaSoundWave.PreEffectBusSends = TemplateSoundWave.PreEffectBusSends; 
+
+	NewMetaSoundWave.bEnableBaseSubmix = TemplateSoundWave.bEnableBaseSubmix;
+	NewMetaSoundWave.SoundSubmixObject = TemplateSoundWave.SoundSubmixObject;
+	NewMetaSoundWave.bEnableSubmixSends = TemplateSoundWave.bEnableSubmixSends;
+	NewMetaSoundWave.SoundSubmixSends = TemplateSoundWave.SoundSubmixSends;
+
+	// Modulation 
+	NewMetaSoundWave.ModulationSettings = TemplateSoundWave.ModulationSettings;
+	
+	// Voice Management 
+	NewMetaSoundWave.VirtualizationMode = TemplateSoundWave.VirtualizationMode;
+	NewMetaSoundWave.bOverrideConcurrency = TemplateSoundWave.bOverrideConcurrency;
+	NewMetaSoundWave.ConcurrencySet = TemplateSoundWave.ConcurrencySet;
+	NewMetaSoundWave.ConcurrencyOverrides = TemplateSoundWave.ConcurrencyOverrides;
+
+	NewMetaSoundWave.bBypassVolumeScaleForPriority = TemplateSoundWave.bBypassVolumeScaleForPriority;
+	NewMetaSoundWave.Priority = TemplateSoundWave.Priority;
 }
 
-const UMetaSoundEditorSubsystem& UMetaSoundEditorSubsystem::GetConstChecked()
-{
-	checkf(GEditor, TEXT("Cannot access UMetaSoundEditorSubsystem without editor loaded"));
-	UMetaSoundEditorSubsystem* EditorSubsystem = GEditor->GetEditorSubsystem<UMetaSoundEditorSubsystem>();
-	checkf(EditorSubsystem, TEXT("Failed to find initialized 'UMetaSoundEditorSubsystem"));
-	return *EditorSubsystem;
-}
-
 #undef LOCTEXT_NAMESPACE // "MetaSoundEditor"
diff --git a/Engine/Plugins/Runtime/Metasound/Source/MetasoundEditor/Public/MetasoundEditorSubsystem.h b/Engine/Plugins/Runtime/Metasound/Source/MetasoundEditor/Public/MetasoundEditorSubsystem.h
--- a/Engine/Plugins/Runtime/Metasound/Source/MetasoundEditor/Public/MetasoundEditorSubsystem.h
+++ b/Engine/Plugins/Runtime/Metasound/Source/MetasoundEditor/Public/MetasoundEditorSubsystem.h
@@ -20,10 +20,28 @@ class METASOUNDEDITOR_API UMetaSoundEditorSubsystem : public UEditorSubsystem
 
 public:
 	// Build the given builder to a MetaSound asset
-	// @param TemplateSoundWave - SoundWave settings such as attenuation, modulation, and sound class will be copied from the optional TemplateSoundWave. 
-	// For preset builders, TemplateSoundWave will override the template values from the referenced asset
+	// @param Author - Sets the author on the given builder's document.
+	// @param AssetName - Name of the asset to build.
+	// @param PackagePath - Path of package to build asset to.
+	// @param TemplateSoundWave - SoundWave settings such as attenuation, modulation, and sound class will be copied from the optional TemplateSoundWave.
+	// For preset builders, TemplateSoundWave will override the template values from the referenced asset.
 	UFUNCTION(BlueprintCallable, Category = "Audio|MetaSound|Builder|Editor", meta = (WorldContext = "Parent", ExpandEnumAsExecs = "OutResult"))
-	UPARAM(DisplayName = "MetaSound Asset") TScriptInterface<IMetaSoundDocumentInterface> BuildToAsset(UMetaSoundBuilderBase* InBuilder, const FString& Author, const FString& AssetName, const FString& PackagePath, EMetaSoundBuilderResult& OutResult, const USoundWave* TemplateSoundWave = nullptr);
+	UPARAM(DisplayName = "MetaSound Asset") TScriptInterface<IMetaSoundDocumentInterface> BuildToAsset(
+		UPARAM(DisplayName = "Builder") UMetaSoundBuilderBase* InBuilder,
+		const FString& Author,
+		const FString& AssetName,
+		const FString& PackagePath,
+		EMetaSoundBuilderResult& OutResult,
+		UPARAM(DisplayName = "Template SoundWave") const USoundWave* TemplateSoundWave = nullptr);
+
+	// @param NodeStyle - Initial style of given nodes, including position.  If not provided, sets to default data including position according to simple
+	// algorithm making a column of inputs, outputs, and any other nodes in between.
+	UFUNCTION(BlueprintCallable, Category = "Audio|MetaSound|Builder|Editor", meta = (ExpandEnumAsExecs = "OutResult"))
+	void SetNodeLocation(
+		UPARAM(DisplayName = "Builder") UMetaSoundBuilderBase * InBuilder,
+		UPARAM(DisplayName = "Node") const FMetaSoundNodeHandle& InNode,
+		UPARAM(DisplayName = "Location") const FVector2D& InLocation,
+		EMetaSoundBuilderResult& OutResult);
 	
 	// Initialize the UObject asset, with an optional MetaSound to be referenced if the asset is a preset
 	void InitAsset(UObject& InNewMetaSound, UObject* InReferencedMetaSound = nullptr);
diff --git a/Engine/Plugins/Runtime/Metasound/Source/MetasoundEngine/Private/MetasoundBuilderSubsystem.cpp b/Engine/Plugins/Runtime/Metasound/Source/MetasoundEngine/Private/MetasoundBuilderSubsystem.cpp
--- a/Engine/Plugins/Runtime/Metasound/Source/MetasoundEngine/Private/MetasoundBuilderSubsystem.cpp
+++ b/Engine/Plugins/Runtime/Metasound/Source/MetasoundEngine/Private/MetasoundBuilderSubsystem.cpp
@@ -689,17 +689,25 @@ void UMetaSoundBuilderBase::SetAuthor(const FString& InAuthor)
 }
 #endif // WITH_EDITOR
 
+void UMetaSoundBuilderBase::SetGraphInputDefault(FName InputName, const FMetasoundFrontendLiteral& Literal, EMetaSoundBuilderResult& OutResult)
+{
+	const bool bInputDefaultSet = Builder.SetGraphInputDefault(InputName, Literal);
+	OutResult = bInputDefaultSet ? EMetaSoundBuilderResult::Succeeded : EMetaSoundBuilderResult::Failed;
+}
+
 void UMetaSoundBuilderBase::SetNodeInputDefault(const FMetaSoundBuilderNodeInputHandle& InputHandle, const FMetasoundFrontendLiteral& Literal, EMetaSoundBuilderResult& OutResult)
 {
 	const bool bInputDefaultSet = Builder.SetNodeInputDefault(InputHandle.NodeID, InputHandle.VertexID, Literal);
 	OutResult = bInputDefaultSet ? EMetaSoundBuilderResult::Succeeded : EMetaSoundBuilderResult::Failed;
 }
 
-void UMetaSoundBuilderBase::SetGraphInputDefault(FName InputName, const FMetasoundFrontendLiteral& Literal, EMetaSoundBuilderResult& OutResult)
+#if WITH_EDITOR
+void UMetaSoundBuilderBase::SetNodeLocation(const FMetaSoundNodeHandle& InNodeHandle, const FVector2D& InLocation, EMetaSoundBuilderResult& OutResult)
 {
-	const bool bInputDefaultSet = Builder.SetGraphInputDefault(InputName, Literal);
-	OutResult = bInputDefaultSet ? EMetaSoundBuilderResult::Succeeded : EMetaSoundBuilderResult::Failed;
+	const bool bLocationSet = Builder.SetNodeLocation(InNodeHandle.NodeID, InLocation);
+	OutResult = bLocationSet ? EMetaSoundBuilderResult::Succeeded : EMetaSoundBuilderResult::Failed;
 }
+#endif // WITH_EDITOR
 
 FMetaSoundNodeHandle UMetaSoundBuilderBase::FindNodeInputParent(const FMetaSoundBuilderNodeInputHandle& InputHandle, EMetaSoundBuilderResult& OutResult)
 {
diff --git a/Engine/Plugins/Runtime/Metasound/Source/MetasoundEngine/Public/MetasoundBuilderSubsystem.h b/Engine/Plugins/Runtime/Metasound/Source/MetasoundEngine/Public/MetasoundBuilderSubsystem.h
--- a/Engine/Plugins/Runtime/Metasound/Source/MetasoundEngine/Public/MetasoundBuilderSubsystem.h
+++ b/Engine/Plugins/Runtime/Metasound/Source/MetasoundEngine/Public/MetasoundBuilderSubsystem.h
@@ -100,6 +100,12 @@ struct METASOUNDENGINE_API FMetaSoundBuilderOptions
 	UPROPERTY(EditAnywhere, BlueprintReadWrite, Category = "MetaSound|Builder")
 	FName Name;
 
+	// If true, this will force regeneration of the class identifier. If the resulting MetaSound is building over
+	// an existing document, this will effectively invalidate any referencing MetaSounds and register the MetaSound
+	// as a new entry in the Frontend.
+	UPROPERTY(EditAnywhere, BlueprintReadWrite, Category = "MetaSound|Builder", meta = (AdvancedDisplay))
+	bool bForceUniqueClassName = false;
+
 	// If true, adds MetaSound to node registry, making it available
 	// for reference by other dynamically created MetaSounds.
 	UPROPERTY()
@@ -362,6 +368,10 @@ public:
 
 	virtual void InitFrontendBuilder();
 
+#if WITH_EDITOR
+	void SetNodeLocation(const FMetaSoundNodeHandle& InNodeHandle, const FVector2D& InLocation, EMetaSoundBuilderResult& OutResult);
+#endif // WITH_EDITOR
+
 protected:
 	const FMetaSoundFrontendDocumentBuilder& GetConstBuilder() const
 	{
@@ -389,7 +399,10 @@ protected:
 				// If MetaSound already exists, preserve the class name to avoid
 				// nametable bloat & preserve potentially existing references.
 				const FMetasoundFrontendDocument& ExistingDoc = CastChecked<const UClassType>(MetaSound)->GetDocumentChecked();
-				DocClassName = &ExistingDoc.RootGraph.Metadata.GetClassName();
+				if (!BuilderOptions.bForceUniqueClassName)
+				{
+					DocClassName = &ExistingDoc.RootGraph.Metadata.GetClassName();
+				}
 				const FNodeRegistryKey& RegistryKey = MetaSound->GetRegistryKey();
 				if (NodeRegistryKey::IsValid(RegistryKey))
 				{
diff --git a/Engine/Plugins/Runtime/Metasound/Source/MetasoundFrontend/Private/MetasoundFrontendDocumentBuilder.cpp b/Engine/Plugins/Runtime/Metasound/Source/MetasoundFrontend/Private/MetasoundFrontendDocumentBuilder.cpp
--- a/Engine/Plugins/Runtime/Metasound/Source/MetasoundFrontend/Private/MetasoundFrontendDocumentBuilder.cpp
+++ b/Engine/Plugins/Runtime/Metasound/Source/MetasoundFrontend/Private/MetasoundFrontendDocumentBuilder.cpp
@@ -1702,7 +1702,10 @@ void FMetaSoundFrontendDocumentBuilder::InitNodeLocations()
 			// TODO: Find consistent location for controlling node locations.
 			// Currently it is split between MetasoundEditor and MetasoundFrontend modules.
 			FMetasoundFrontendNodeStyle& Style = Node.Style;
-			Style.Display.Locations = { { FGuid::NewGuid(), NewLocation } };
+			if (Style.Display.Locations.IsEmpty())
+			{
+				Style.Display.Locations = { { FGuid::NewGuid(), NewLocation } };
+			}
 		}
 	}
 #endif // WITH_EDITORONLY_DATA
@@ -2368,6 +2371,36 @@ bool FMetaSoundFrontendDocumentBuilder::SetNodeInputDefault(const FGuid& InNodeI
 	return false;
 }
 
+#if WITH_EDITOR
+bool FMetaSoundFrontendDocumentBuilder::SetNodeLocation(const FGuid& InNodeID, const FVector2D& InLocation)
+{
+	using namespace Metasound;
+	using namespace Metasound::Frontend;
+
+	const IDocumentGraphNodeCache& NodeCache = DocumentCache->GetNodeCache();
+	if (const int32* NodeIndex = NodeCache.FindNodeIndex(InNodeID))
+	{
+		FMetasoundFrontendNode& Node = GetDocument().RootGraph.Graph.Nodes[*NodeIndex];
+		FMetasoundFrontendNodeStyle& Style = Node.Style;
+		if (Style.Display.Locations.IsEmpty())
+		{
+			Style.Display.Locations = { { FGuid::NewGuid(), InLocation } };
+		}
+		else
+		{
+			Algo::ForEach(Style.Display.Locations, [InLocation](TPair<FGuid, FVector2D>& Pair)
+			{
+				Pair.Value = InLocation;
+			});
+		}
+
+		return true;
+	}
+
+	return false;
+}
+#endif // WITH_EDITOR
+
 bool FMetaSoundFrontendDocumentBuilder::SwapGraphInput(const FMetasoundFrontendClassVertex& InExistingInputVertex, const FMetasoundFrontendClassVertex& InNewInputVertex)
 {
 	using namespace Metasound::Frontend;
diff --git a/Engine/Plugins/Runtime/Metasound/Source/MetasoundFrontend/Public/MetasoundFrontendDocumentBuilder.h b/Engine/Plugins/Runtime/Metasound/Source/MetasoundFrontend/Public/MetasoundFrontendDocumentBuilder.h
--- a/Engine/Plugins/Runtime/Metasound/Source/MetasoundFrontend/Public/MetasoundFrontendDocumentBuilder.h
+++ b/Engine/Plugins/Runtime/Metasound/Source/MetasoundFrontend/Public/MetasoundFrontendDocumentBuilder.h
@@ -240,6 +240,10 @@ public:
 	void ReloadCache();
 
 	void SetAuthor(const FString& InAuthor);
+
+	// Sets the editor-only node location of a node with the given ID to the provided location.
+	// Returns true if the node was found and the location was updated, false if not.
+	bool SetNodeLocation(const FGuid& InNodeID, const FVector2D& InLocation);
 #endif // WITH_EDITOR
 
 	bool SetGraphInputDefault(FName InputName, const FMetasoundFrontendLiteral& InDefaultLiteral);
