Fix comment missed in last CL review
#rnx

diff --git a/Engine/Plugins/Runtime/Metasound/Source/MetasoundEditor/Public/MetasoundEditorSubsystem.h b/Engine/Plugins/Runtime/Metasound/Source/MetasoundEditor/Public/MetasoundEditorSubsystem.h
--- a/Engine/Plugins/Runtime/Metasound/Source/MetasoundEditor/Public/MetasoundEditorSubsystem.h
+++ b/Engine/Plugins/Runtime/Metasound/Source/MetasoundEditor/Public/MetasoundEditorSubsystem.h
@@ -34,8 +34,7 @@ public:
 		EMetaSoundBuilderResult& OutResult,
 		UPARAM(DisplayName = "Template SoundWave") const USoundWave* TemplateSoundWave = nullptr);
 
-	// @param NodeStyle - Initial style of given nodes, including position.  If not provided, sets to default data including position according to simple
-	// algorithm making a column of inputs, outputs, and any other nodes in between.
+	// Sets the visual location to InLocation of a given node InNode of a given builder's document.
 	UFUNCTION(BlueprintCallable, Category = "Audio|MetaSound|Builder|Editor", meta = (ExpandEnumAsExecs = "OutResult"))
 	void SetNodeLocation(
 		UPARAM(DisplayName = "Builder") UMetaSoundBuilderBase * InBuilder,
