Fix for overlapping input/output boxes introduced by earlier CL
#rnx

diff --git a/Engine/Plugins/Runtime/Metasound/Source/MetasoundFrontend/Private/MetasoundFrontendDocumentBuilder.cpp b/Engine/Plugins/Runtime/Metasound/Source/MetasoundFrontend/Private/MetasoundFrontendDocumentBuilder.cpp
--- a/Engine/Plugins/Runtime/Metasound/Source/MetasoundFrontend/Private/MetasoundFrontendDocumentBuilder.cpp
+++ b/Engine/Plugins/Runtime/Metasound/Source/MetasoundFrontend/Private/MetasoundFrontendDocumentBuilder.cpp
@@ -1706,6 +1706,14 @@ void FMetaSoundFrontendDocumentBuilder::InitNodeLocations()
 			{
 				Style.Display.Locations = { { FGuid::NewGuid(), NewLocation } };
 			}
+			// Initialize the position if the location hasn't been assigned yet.  This can happen
+			// if default interfaces were assigned to the given MetaSound but not placed with respect
+			// to one another.  In this case, node location initialization takes "priority" to avoid
+			// visual overlap.
+			else if (Style.Display.Locations.Num() == 1 && Style.Display.Locations.Contains(FGuid()))
+			{
+				Style.Display.Locations = { { FGuid::NewGuid(), NewLocation } };
+			}
 		}
 	}
 #endif // WITH_EDITORONLY_DATA
