Fix MetaSound search engine class query retaining deleted assets

diff --git a/Engine/Plugins/Runtime/Metasound/Source/MetasoundFrontend/Private/MetasoundFrontendQuery.cpp b/Engine/Plugins/Runtime/Metasound/Source/MetasoundFrontend/Private/MetasoundFrontendQuery.cpp
--- a/Engine/Plugins/Runtime/Metasound/Source/MetasoundFrontend/Private/MetasoundFrontendQuery.cpp
+++ b/Engine/Plugins/Runtime/Metasound/Source/MetasoundFrontend/Private/MetasoundFrontendQuery.cpp
@@ -1112,11 +1112,8 @@ namespace Metasound
 		{
 			if (const FFrontendQueryPartition* Partition = InSelection.Find(Key))
 			{
-				if (Partition->Num() > 0)
-				{
-					OutKeysModified.Add(Key);
-					OutSelection.FindOrAdd(Key) = *Partition;
-				}
+				OutKeysModified.Add(Key);
+				OutSelection.FindOrAdd(Key) = *Partition;
 			}
 		}
 	}
diff --git a/Engine/Plugins/Runtime/Metasound/Source/MetasoundFrontend/Private/MetasoundFrontendQuerySteps.cpp b/Engine/Plugins/Runtime/Metasound/Source/MetasoundFrontend/Private/MetasoundFrontendQuerySteps.cpp
--- a/Engine/Plugins/Runtime/Metasound/Source/MetasoundFrontend/Private/MetasoundFrontendQuerySteps.cpp
+++ b/Engine/Plugins/Runtime/Metasound/Source/MetasoundFrontend/Private/MetasoundFrontendQuerySteps.cpp
@@ -9,6 +9,7 @@
 #include "MetasoundFrontendQuery.h"
 #include "MetasoundFrontendRegistries.h"
 #include "MetasoundFrontendRegistryTransaction.h"
+#include "MetasoundLog.h"
 
 namespace Metasound
 {
@@ -66,21 +67,69 @@ namespace Metasound
 	void FReduceRegistrationEventsToCurrentStatus::Reduce(const FFrontendQueryKey& InKey, FFrontendQueryPartition& InOutEntries) const
 	{
 		using namespace Frontend;
+		if (InOutEntries.IsEmpty())
+		{
+			return;
+		}
 
-		// Get most recent transaction
 		const FFrontendQueryEntry* FinalEntry = Algo::MaxElementBy(InOutEntries, GetTransactionTimestamp);
-
-		// Check if most recent transaction is valid and not an unregister transaction.
+		// If last entry is registration, keep the last entry 
+		// since that's the most relevant up to date status
 		if (IsValidTransactionOfType(FNodeRegistryTransaction::ETransactionType::NodeRegistration, FinalEntry))
 		{
 			FFrontendQueryEntry Entry = *FinalEntry;
 			InOutEntries.Reset();
 			InOutEntries.Add(Entry);
+		} 
+		else if (IsValidTransactionOfType(FNodeRegistryTransaction::ETransactionType::NodeUnregistration, FinalEntry))
+		{
+			// Check that pairs of entries are Registration - Unregistration
+			auto CheckRegistrationTransactionPairs = [](FFrontendQueryPartition& InOutEntries, int32 EndIndexInclusive)
+			{
+				InOutEntries.Sort([](const FFrontendQueryEntry& A, const FFrontendQueryEntry& B)
+				{
+					return GetTransactionTimestamp(A) < GetTransactionTimestamp(B);
+				});
+				for (int32 PairIndex = 0; PairIndex < EndIndexInclusive; PairIndex += 2)
+				{
+					const FNodeRegistryTransaction& FirstEntry = InOutEntries[PairIndex].Value.Get<FNodeRegistryTransaction>();
+					const FNodeRegistryTransaction& SecondEntry = InOutEntries[PairIndex + 1].Value.Get<FNodeRegistryTransaction>();
+					const FNodeRegistryTransaction::ETransactionType FirstEntryTransactionType = FirstEntry.GetTransactionType();
+					const FNodeRegistryTransaction::ETransactionType SecondEntryTransactionType = SecondEntry.GetTransactionType();
+
+					if (!(FirstEntryTransactionType == FNodeRegistryTransaction::ETransactionType::NodeRegistration &&
+						SecondEntryTransactionType == FNodeRegistryTransaction::ETransactionType::NodeUnregistration))
+					{
+						UE_LOG(LogMetaSound, Warning, TEXT("Mismatched transaction entries with keys '%s' and '%s' and transaction types '%s' and '%s' in a search engine query.\
+							Transactions should be registration and unregistration pairs except for the last entry."),
+							*FirstEntry.GetNodeRegistryKey(),
+							*SecondEntry.GetNodeRegistryKey(),
+							*FNodeRegistryTransaction::LexToString(FirstEntryTransactionType),
+							*FNodeRegistryTransaction::LexToString(SecondEntryTransactionType));
+					}
+				}
+			};
+			// If odd number of entries, check previous pairs and keep the final entry
+			if (InOutEntries.Num() % 2 == 1)
+			{
+				FFrontendQueryEntry Entry = *FinalEntry;
+				CheckRegistrationTransactionPairs(InOutEntries, FMath::Max(0, InOutEntries.Num() - 2));
+				InOutEntries.Reset();
+				InOutEntries.Add(Entry);
+			}
+			// If even, check all pairs  
+			// and clear them all, since the current status is as if those hadn't happened
+			else
+			{
+				CheckRegistrationTransactionPairs(InOutEntries, InOutEntries.Num() - 1);
+				InOutEntries.Reset();
+			}
 		}
 		else
 		{
-			InOutEntries.Reset();
+			UE_LOG(LogMetaSound, Warning, TEXT("Missing FNodeRegistryTransaction::ETransactionType case in search engine FReduceRegistrationEventsToCurrentStatus::Reduce."))
 		}
+
 	}
 
 	FReduceRegistrationEventsToCurrentStatus::FTimeType FReduceRegistrationEventsToCurrentStatus::GetTransactionTimestamp(const FFrontendQueryEntry& InEntry)
@@ -118,10 +167,13 @@ namespace Metasound
 		{
 			const FNodeRegistryTransaction& Transaction = InValue.Get<FNodeRegistryTransaction>();
 			
-			bool bSuccess = FMetasoundFrontendRegistryContainer::Get()->FindFrontendClassFromRegistered(Transaction.GetNodeRegistryKey(), FrontendClass);
-			check(bSuccess);
+			if (Transaction.GetTransactionType() == Frontend::FNodeRegistryTransaction::ETransactionType::NodeRegistration)
+			{
+				// It's possible that the node is no longer registered (we're processing removals) 
+				// but that's okay because the returned default FrontendClass will be processed out later
+				FMetasoundFrontendRegistryContainer::Get()->FindFrontendClassFromRegistered(Transaction.GetNodeRegistryKey(), FrontendClass);
+			}
 		}
-
 		InValue.Set<FMetasoundFrontendClass>(MoveTemp(FrontendClass));
 	}
 
diff --git a/Engine/Plugins/Runtime/Metasound/Source/MetasoundFrontend/Public/MetasoundFrontendRegistries.h b/Engine/Plugins/Runtime/Metasound/Source/MetasoundFrontend/Public/MetasoundFrontendRegistries.h
--- a/Engine/Plugins/Runtime/Metasound/Source/MetasoundFrontend/Public/MetasoundFrontendRegistries.h
+++ b/Engine/Plugins/Runtime/Metasound/Source/MetasoundFrontend/Public/MetasoundFrontendRegistries.h
@@ -190,6 +190,29 @@ namespace Metasound
 				Invalid
 			};
 
+			static FString LexToString(ETransactionType InTransactionType)
+			{
+				using namespace Metasound;
+
+				switch (InTransactionType)
+				{
+					case ETransactionType::NodeRegistration:
+						return FString(TEXT("NodeRegistration"));
+
+					case ETransactionType::NodeUnregistration:
+						return FString(TEXT("NodeUnregistration"));
+
+					case ETransactionType::Invalid:
+						return FString(TEXT("Invalid"));
+
+					default:
+					{
+						checkNoEntry();
+					}
+				}
+				return FString(TEXT(""));
+			}
+
 			UE_DEPRECATED("5.1", "This constructor is deprecated. Use a different constructor for FNodeREegistryTransaction")
 			FNodeRegistryTransaction(ETransactionType InType, const FNodeRegistryKey& InKey, const FNodeClassInfo& InNodeClassInfo, FTimeType InTimestamp);
 			FNodeRegistryTransaction(ETransactionType InType, const FNodeClassInfo& InNodeClassInfo, FTimeType InTimestamp);
