Fixing audio pops when connecting/disconnecting audio in dynamic metasound

diff --git a/Engine/Plugins/Runtime/Metasound/Source/MetasoundFrontend/Private/MetasoundPrimitives.cpp b/Engine/Plugins/Runtime/Metasound/Source/MetasoundFrontend/Private/MetasoundPrimitives.cpp
--- a/Engine/Plugins/Runtime/Metasound/Source/MetasoundFrontend/Private/MetasoundPrimitives.cpp
+++ b/Engine/Plugins/Runtime/Metasound/Source/MetasoundFrontend/Private/MetasoundPrimitives.cpp
@@ -61,5 +61,9 @@ REGISTER_METASOUND_DATATYPE(float, "Float", ::Metasound::ELiteralType::Float)
 REGISTER_METASOUND_DATATYPE(FString, "String", ::Metasound::ELiteralType::String)
 REGISTER_METASOUND_DATATYPE(Metasound::FTrigger, "Trigger", ::Metasound::ELiteralType::Boolean)
 REGISTER_METASOUND_DATATYPE(Metasound::FTime, "Time", ::Metasound::ELiteralType::Float)
-REGISTER_METASOUND_DATATYPE(Metasound::FAudioBuffer, "Audio")
 REGISTER_METASOUND_DATATYPE(Metasound::FSendAddress, "Transmission:Address", ::Metasound::ELiteralType::String)
+
+// FAudioBuffer is declared and defined in the MetaSoundGraphCore module. Here we only
+// enqueue the registration command to avoid linker issues caused by the interplay
+// of the REGISTER_METASOUND_DATATYPE and DEFINE_METASOUND_DATA_TYPE macros.
+ENQUEUE_METASOUND_DATATYPE_REGISTRATION_COMMAND(Metasound::FAudioBuffer, "Audio");
diff --git a/Engine/Plugins/Runtime/Metasound/Source/MetasoundFrontend/Public/MetasoundAudioBuffer.h b/Engine/Plugins/Runtime/Metasound/Source/MetasoundFrontend/Public/MetasoundAudioBuffer.h
deleted file mode 100644
diff --git a/Engine/Plugins/Runtime/Metasound/Source/MetasoundFrontend/Public/MetasoundDataTypeRegistrationMacro.h b/Engine/Plugins/Runtime/Metasound/Source/MetasoundFrontend/Public/MetasoundDataTypeRegistrationMacro.h
--- a/Engine/Plugins/Runtime/Metasound/Source/MetasoundFrontend/Public/MetasoundDataTypeRegistrationMacro.h
+++ b/Engine/Plugins/Runtime/Metasound/Source/MetasoundFrontend/Public/MetasoundDataTypeRegistrationMacro.h
@@ -801,15 +801,6 @@ namespace Metasound
 	};
 }
 
-// This should be used to expose a datatype as a potential input or output for a metasound graph.
-// The first argument to the macro is the class to expose.
-// the second argument is the display name of that type in the Metasound editor.
-// Optionally, a Metasound::ELiteralType can be passed in to designate a preferred literal type-
-// For example, if Metasound::ELiteralType::Float is passed in, we will default to using a float parameter to create this datatype.
-// If no argument is passed in, we will infer a literal type to use.
-// If 
-// Metasound::ELiteralType::Invalid can be used to enforce that we don't provide space for a literal, in which case you should have a default constructor or a constructor that takes [const FOperatorSettings&] implemented.
-// If you pass in a preferred arg type, please make sure that the passed in datatype has a matching constructor, since we won't check this until runtime.
 
 #define CANNOT_REGISTER_METASOUND_DATA_TYPE_ASSERT_STRING(DataType) \
 "To register " #DataType " to be used as a Metasounds input or output type, it needs a default constructor or one of the following constructors must be implemented:  " \
@@ -828,8 +819,20 @@ namespace Metasound
 #DataType "(const ::Metasound::FOperatorSettings& InSettings, const Audio::IProxyDataPtr& InData),  or " \
 #DataType "(const ::Metasound::FOperatorSettings& InSettings, const TArray<Audio::IProxyDataPtr>& InProxyArray)."
 
-#define REGISTER_METASOUND_DATATYPE(DataType, DataTypeName, ...) \
-	DEFINE_METASOUND_DATA_TYPE(DataType, DataTypeName); \
+#define ENQUEUE_METASOUND_DATATYPE_REGISTRATION_COMMAND(DataType, DataTypeName, ...) \
 	static_assert(::Metasound::TMetasoundDataTypeRegistration<DataType>::bCanRegister, CANNOT_REGISTER_METASOUND_DATA_TYPE_ASSERT_STRING(DataType)); \
 	template<> const bool ::Metasound::TMetasoundDataTypeRegistration<DataType>::bSuccessfullyRegistered = ::FMetasoundFrontendRegistryContainer::Get()->EnqueueInitCommand([](){ ::Metasound::RegisterDataTypeWithFrontend<DataType, ##__VA_ARGS__>(); }); // This static bool is useful for debugging, but also is the only way the compiler will let us call this function outside of an expression.
 
+// This should be used to expose a datatype as a potential input or output for a metasound graph.
+// The first argument to the macro is the class to expose.
+// the second argument is the display name of that type in the Metasound editor.
+// Optionally, a Metasound::ELiteralType can be passed in to designate a preferred literal type-
+// For example, if Metasound::ELiteralType::Float is passed in, we will default to using a float parameter to create this datatype.
+// If no argument is passed in, we will infer a literal type to use.
+// If 
+// Metasound::ELiteralType::Invalid can be used to enforce that we don't provide space for a literal, in which case you should have a default constructor or a constructor that takes [const FOperatorSettings&] implemented.
+// If you pass in a preferred arg type, please make sure that the passed in datatype has a matching constructor, since we won't check this until runtime.
+#define REGISTER_METASOUND_DATATYPE(DataType, DataTypeName, ...) \
+	DEFINE_METASOUND_DATA_TYPE(DataType, DataTypeName); \
+	ENQUEUE_METASOUND_DATATYPE_REGISTRATION_COMMAND(DataType, DataTypeName, ##__VA_ARGS__)
+
diff --git a/Engine/Plugins/Runtime/Metasound/Source/MetasoundGraphCore/Private/MetasoundAudioBuffer.cpp b/Engine/Plugins/Runtime/Metasound/Source/MetasoundGraphCore/Private/MetasoundAudioBuffer.cpp
new file mode 100644
--- /dev/null
+++ b/Engine/Plugins/Runtime/Metasound/Source/MetasoundGraphCore/Private/MetasoundAudioBuffer.cpp
@@ -0,0 +1,7 @@
+// Copyright Epic Games, Inc. All Rights Reserved.
+
+#include "MetasoundAudioBuffer.h"
+#include "MetasoundDataReferenceMacro.h"
+
+DEFINE_METASOUND_DATA_TYPE(Metasound::FAudioBuffer, "Audio");
+
diff --git a/Engine/Plugins/Runtime/Metasound/Source/MetasoundGraphCore/Private/MetasoundDynamicGraphAlgo.cpp b/Engine/Plugins/Runtime/Metasound/Source/MetasoundGraphCore/Private/MetasoundDynamicGraphAlgo.cpp
--- a/Engine/Plugins/Runtime/Metasound/Source/MetasoundGraphCore/Private/MetasoundDynamicGraphAlgo.cpp
+++ b/Engine/Plugins/Runtime/Metasound/Source/MetasoundGraphCore/Private/MetasoundDynamicGraphAlgo.cpp
@@ -30,6 +30,7 @@ namespace Metasound
 				{
 					// Update the existing entry
 					InOutTable[EntryIndex].Function = InFunction;
+					InOutTable[EntryIndex].Operator = &InOperator;
 				}
 				else if (bEntryCurrentlyExists)
 				{
@@ -64,6 +65,13 @@ namespace Metasound
 					InOutTable.Insert(EntryType{InOperatorID, InOperator, InFunction}, InsertLocation);
 				}
 			}
+
+			void UpdateGraphRuntimeTables(const FOperatorID& InOperatorID, IOperator& InOperator, FDynamicGraphOperatorData& InOutGraphOperatorData)
+			{
+				UpdateTableEntry(InOutGraphOperatorData.OperatorOrder, InOperatorID, InOperator, InOutGraphOperatorData.ExecuteTable, InOperator.GetExecuteFunction());
+				UpdateTableEntry(InOutGraphOperatorData.OperatorOrder, InOperatorID, InOperator, InOutGraphOperatorData.PostExecuteTable, InOperator.GetPostExecuteFunction());
+				UpdateTableEntry(InOutGraphOperatorData.OperatorOrder, InOperatorID, InOperator, InOutGraphOperatorData.ResetTable, InOperator.GetResetFunction());
+			}
 		} // namespace DynamicGraphAlgoPrivate
 
 		// Apply updates to data references through all the operators by following connections described in the FOperatorInfo map.
@@ -102,9 +110,7 @@ namespace Metasound
 					Operator.BindOutputs(OpInfo->VertexData.GetOutputs());
 
 					// Update execute/postexecute/reset tables in case those have changed after rebinding.
-					UpdateTableEntry(InOutGraphOperatorData.OperatorOrder, Current.OperatorID, Operator, InOutGraphOperatorData.ExecuteTable, Operator.GetExecuteFunction());
-					UpdateTableEntry(InOutGraphOperatorData.OperatorOrder, Current.OperatorID, Operator, InOutGraphOperatorData.PostExecuteTable, Operator.GetPostExecuteFunction());
-					UpdateTableEntry(InOutGraphOperatorData.OperatorOrder, Current.OperatorID, Operator, InOutGraphOperatorData.ResetTable, Operator.GetResetFunction());
+					UpdateGraphRuntimeTables(Current.OperatorID, Operator, InOutGraphOperatorData);
 
 					// See if binding altered the outputs. 
 					OutputUpdates.Reset();
@@ -232,6 +238,55 @@ namespace Metasound
 				}
 			}
 		}
+
+		void RebindWrappedOperator(const FOperatorID& InOperatorID, FOperatorInfo& InOperatorInfo, FDynamicGraphOperatorData& InOutGraphOperatorData)
+		{
+			using namespace DirectedGraphAlgo;
+			using namespace DynamicGraphAlgoPrivate;
+
+			check(InOperatorInfo.Operator);
+
+			/* Bind and diff the graph's interface to determine if there is an update to any vertices */
+
+			// Cache current state of operators output data. 
+			FOutputVertexInterfaceData& OutputVertexData = InOperatorInfo.VertexData.GetOutputs();
+			TArray<FVertexDataState> InitialVertexDataState;
+			GetVertexInterfaceDataState(OutputVertexData, InitialVertexDataState);
+
+			// Bind the operator to trigger any updates. 
+			InOperatorInfo.Operator->BindInputs(InOperatorInfo.VertexData.GetInputs());
+			InOperatorInfo.Operator->BindOutputs(InOperatorInfo.VertexData.GetOutputs());
+
+			// Update any execution tables that need updating after wrapping
+			UpdateGraphRuntimeTables(InOperatorID, *(InOperatorInfo.Operator), InOutGraphOperatorData);
+
+			// Determine if there have been changes to `OutputVertexData`. 
+			TSortedVertexNameMap<FAnyDataReference> OutputsToUpdate;
+			CompareVertexInterfaceDataToPriorState(OutputVertexData, InitialVertexDataState, OutputsToUpdate);
+
+			// If there have been any changes to `OutputVertexData`, then these need to be propagated
+			// through the graph to route them to operator inputs and to handle any knock-on updates
+			// to other data references. 
+
+			// Update the graph data by propagating the updates through the inputs nodes. 
+			for (const TPair<FVertexName, FAnyDataReference>& OutputToUpdate : OutputsToUpdate)
+			{
+				const FVertexName& VertexName = OutputToUpdate.Get<0>();
+				if (const TArray<FGraphOperatorData::FVertexDestination>* Destinations  = InOperatorInfo.OutputConnections.Find(VertexName))
+				{
+					for (const FGraphOperatorData::FVertexDestination& Destination : *Destinations)
+					{
+						PropagateBindUpdate(Destination.OperatorID, Destination.VertexName, OutputToUpdate.Get<1>(), InOutGraphOperatorData);
+						
+					}
+				}
+			}
+
+			// Refresh output vertex interface data in case any graph output 
+			// nodes were updated when bind updates were propagated through 
+			// the graph.
+			UpdateOutputVertexData(InOutGraphOperatorData);
+		}
 		
 		void RebindGraphInputs(FInputVertexInterfaceData& InOutVertexData, FDynamicGraphOperatorData& InOutGraphOperatorData)
 		{
diff --git a/Engine/Plugins/Runtime/Metasound/Source/MetasoundGraphCore/Private/MetasoundDynamicGraphAlgo.h b/Engine/Plugins/Runtime/Metasound/Source/MetasoundGraphCore/Private/MetasoundDynamicGraphAlgo.h
--- a/Engine/Plugins/Runtime/Metasound/Source/MetasoundGraphCore/Private/MetasoundDynamicGraphAlgo.h
+++ b/Engine/Plugins/Runtime/Metasound/Source/MetasoundGraphCore/Private/MetasoundDynamicGraphAlgo.h
@@ -105,6 +105,9 @@ namespace Metasound
 		 */
 		void UpdateOutputVertexData(FDynamicGraphOperatorData& InOutGraphOperatorData);
 
+		/** Rebinds an operator which is wrapping another operator. */
+		void RebindWrappedOperator(const FOperatorID& InOperatorID, FOperatorInfo& InOperatorInfo, FDynamicGraphOperatorData& InGraphOperatorData);
+
 		/** Rebind the graph inputs, updating internal operator bindings as needed. */
 		void RebindGraphInputs(FInputVertexInterfaceData& InOutVertexData, FDynamicGraphOperatorData& InOutGraphOperatorData);
 
diff --git a/Engine/Plugins/Runtime/Metasound/Source/MetasoundGraphCore/Private/MetasoundDynamicOperator.cpp b/Engine/Plugins/Runtime/Metasound/Source/MetasoundGraphCore/Private/MetasoundDynamicOperator.cpp
--- a/Engine/Plugins/Runtime/Metasound/Source/MetasoundGraphCore/Private/MetasoundDynamicOperator.cpp
+++ b/Engine/Plugins/Runtime/Metasound/Source/MetasoundGraphCore/Private/MetasoundDynamicOperator.cpp
@@ -3,6 +3,7 @@
 #include "MetasoundDynamicOperator.h"
 
 #include "Containers/Array.h"
+#include "MetasoundDynamicOperatorAudioFade.h"
 #include "MetasoundNodeInterface.h"
 #include "MetasoundOperatorInterface.h"
 #include "MetasoundVertexData.h"
@@ -603,6 +604,58 @@ namespace Metasound
 
 			return Result;
 		}
+
+		FBeginAudioFadeTransform::FBeginAudioFadeTransform(FOperatorID InOperatorIDToFade, EAudioFadeType InFadeType, TArrayView<const FVertexName> InInputVerticesToFade, TArrayView<const FVertexName> InOutputVerticesToFade)
+		: OperatorIDToFade(InOperatorIDToFade)
+		, InitFadeState(InFadeType == EAudioFadeType::FadeIn ? FAudioFadeOperatorWrapper::EFadeState::FadingIn : FAudioFadeOperatorWrapper::EFadeState::FadingOut)
+		, InputVerticesToFade(InInputVerticesToFade)
+		, OutputVerticesToFade(InOutputVerticesToFade)
+		{
+		}
+
+		EDynamicOperatorTransformQueueAction FBeginAudioFadeTransform::FBeginAudioFadeTransform::Transform(FDynamicGraphOperatorData& InGraphOperatorData)
+		{
+		  	if (FOperatorInfo* OperatorInfo = InGraphOperatorData.OperatorMap.Find(OperatorIDToFade))
+			{
+				// Make wrapped operator
+				OperatorInfo->Operator = MakeUnique<FAudioFadeOperatorWrapper>(InitFadeState, InGraphOperatorData.OperatorSettings, OperatorInfo->VertexData.GetInputs(), MoveTemp(OperatorInfo->Operator), InputVerticesToFade, OutputVerticesToFade);
+
+				// Update data references in graph
+				RebindWrappedOperator(OperatorIDToFade, *OperatorInfo, InGraphOperatorData);
+			}
+			else
+			{
+				UE_LOG(LogMetaSound, Error, TEXT("Could not find operator with ID %d when applying audio fade out."), OperatorIDToFade);
+			}
+			
+			return EDynamicOperatorTransformQueueAction::Continue;
+		}
+
+		FEndAudioFadeTransform::FEndAudioFadeTransform(FOperatorID InOperatorIDToFade)
+		: OperatorIDToFade(InOperatorIDToFade)
+		{
+		}
+
+		EDynamicOperatorTransformQueueAction FEndAudioFadeTransform::Transform(FDynamicGraphOperatorData& InGraphOperatorData)
+		{
+		  	if (FOperatorInfo* OperatorInfo = InGraphOperatorData.OperatorMap.Find(OperatorIDToFade))
+			{
+				FAudioFadeOperatorWrapper* Wrapper = static_cast<FAudioFadeOperatorWrapper*>(OperatorInfo->Operator.Get());
+				check(Wrapper);
+
+				// Unwrap operator
+				OperatorInfo->Operator = Wrapper->ReleaseOperator();
+
+				// Update data references in graph
+				RebindWrappedOperator(OperatorIDToFade, *OperatorInfo, InGraphOperatorData);
+			}
+			else
+			{
+				UE_LOG(LogMetaSound, Error, TEXT("Could not find operator with ID %d when applying audio fade out."), OperatorIDToFade);
+			}
+			
+			return EDynamicOperatorTransformQueueAction::Continue;
+		}
 	} // namespace DynamicGraph
 }
 
diff --git a/Engine/Plugins/Runtime/Metasound/Source/MetasoundGraphCore/Private/MetasoundDynamicOperator.h b/Engine/Plugins/Runtime/Metasound/Source/MetasoundGraphCore/Private/MetasoundDynamicOperator.h
--- a/Engine/Plugins/Runtime/Metasound/Source/MetasoundGraphCore/Private/MetasoundDynamicOperator.h
+++ b/Engine/Plugins/Runtime/Metasound/Source/MetasoundGraphCore/Private/MetasoundDynamicOperator.h
@@ -5,6 +5,7 @@
 #include "Containers/Array.h"
 #include "Containers/SpscQueue.h"
 #include "MetasoundDynamicGraphAlgo.h"
+#include "MetasoundDynamicOperatorAudioFade.h"
 #include "MetasoundDynamicOperatorTransactor.h"
 #include "MetasoundGraphAlgoPrivate.h"
 #include "MetasoundNodeInterface.h"
@@ -213,5 +214,46 @@ namespace Metasound
 		private:
 			TArray<TUniquePtr<IDynamicOperatorTransform>> Transforms;
 		};
+
+		enum class EAudioFadeType : uint8
+		{
+			FadeIn, //< Fade from silent to full volume.
+			FadeOut //< Fade from full volume to silent.
+		};
+
+		/* Marks the beginning of an audio fade. 
+		 *
+		 * When scheduling fade transformations on a dynamic operator, an FBeginAudioFadeTransform 
+		 * must be matched with an FEndAudioTransform with an FExecuteFence between them. 
+		 *
+		 * The FBeginAudioFadeTransform sets up the graph to perform a set of audio
+		 * fades. 
+		 * The FExecuteFence forces the fade to occur before any additional transforms 
+		 * are performed.
+		 * The FEndAudioFadeTransform cleans up any temporary state that was needed
+		 * to perform the fade. 
+		 * */
+		class FBeginAudioFadeTransform : public IDynamicOperatorTransform
+		{
+		public:
+
+			FBeginAudioFadeTransform(FOperatorID InOperatorIDToFade, EAudioFadeType InFadeType, TArrayView<const FVertexName> InInputVerticesToFade, TArrayView<const FVertexName> InOutputVerticesToFade);
+			virtual EDynamicOperatorTransformQueueAction Transform(FDynamicGraphOperatorData& InGraphOperatorData) override;
+		private:
+			FOperatorID OperatorIDToFade;
+			FAudioFadeOperatorWrapper::EFadeState InitFadeState;
+			TArray<FVertexName> InputVerticesToFade;
+			TArray<FVertexName> OutputVerticesToFade;
+		};
+
+		/** Marks the end of an audio fade. */
+		class FEndAudioFadeTransform : public IDynamicOperatorTransform
+		{
+		public:
+			FEndAudioFadeTransform(FOperatorID InOperatorIDToFade);
+			virtual EDynamicOperatorTransformQueueAction Transform(FDynamicGraphOperatorData& InGraphOperatorData) override;
+		private:
+			FOperatorID OperatorIDToFade;
+		};
 	}
 }
diff --git a/Engine/Plugins/Runtime/Metasound/Source/MetasoundGraphCore/Private/MetasoundDynamicOperatorAudioFade.cpp b/Engine/Plugins/Runtime/Metasound/Source/MetasoundGraphCore/Private/MetasoundDynamicOperatorAudioFade.cpp
new file mode 100644
--- /dev/null
+++ b/Engine/Plugins/Runtime/Metasound/Source/MetasoundGraphCore/Private/MetasoundDynamicOperatorAudioFade.cpp
@@ -0,0 +1,220 @@
+// Copyright Epic Games, Inc. All Rights Reserved.
+
+#include "MetasoundDynamicOperatorAudioFade.h"
+
+#include "Containers/Array.h"
+#include "DSP/FloatArrayMath.h"
+#include "HAL/UnrealMemory.h"
+#include "MetasoundAudioBuffer.h"
+#include "MetasoundDataReference.h"
+#include "MetasoundExecutableOperator.h"
+#include "MetasoundVertex.h"
+#include "MetasoundVertexData.h"
+
+namespace Metasound::DynamicGraph
+{
+	FAudioFadeOperatorWrapper::FAudioFadeOperatorWrapper(EFadeState InInitialFadeState, const FOperatorSettings& InOperatorSettings, const FInputVertexInterfaceData& InInputData, TUniquePtr<IOperator> InOperator, TArrayView<const FVertexName> InInputVerticesToFade, TArrayView<const FVertexName> InOutputVerticesToFade)
+	: WrappedOperator(MoveTemp(InOperator))
+	, InputVerticesToFade(InInputVerticesToFade)
+	, OutputVerticesToFade(InOutputVerticesToFade)
+	, InitialFadeState(InInitialFadeState)
+	, FadeState(InInitialFadeState)
+	{
+		// allocate internal input buffers and data references
+		InitInputBuffers(InOperatorSettings, InInputData);
+
+		SetWrappedOperatorInputs(InInputData);
+
+		// allocate internal output buffers and data references
+		InitOutputBuffers(InOperatorSettings);
+	}
+
+	TUniquePtr<IOperator> FAudioFadeOperatorWrapper::ReleaseOperator() 
+	{
+		// Reset internal arrays so we don't hold onto data references 
+		InputVerticesToFade.Reset();
+		OutputVerticesToFade.Reset();
+
+		WrappedInputAudioBuffers.Reset();
+		FadedInputAudioBuffers.Reset();
+		WrappedOutputAudioBuffers.Reset();
+		FadedOutputAudioBuffers.Reset();
+
+		return WrappedOperator.ReleaseOperator();
+	}
+
+	void FAudioFadeOperatorWrapper::BindInputs(FInputVertexInterfaceData& InOutVertexData)
+	{
+		// Bind input audio buffers that will be faded
+		for (int i = 0; i < InputVerticesToFade.Num(); i++)
+		{
+			const FVertexName& VertexName = InputVerticesToFade[i];
+			InOutVertexData.BindVertex(VertexName, WrappedInputAudioBuffers[i]);
+		}
+
+		// Pass remaining vertex data to be bound with wrapped operator
+		SetWrappedOperatorInputs(InOutVertexData);
+	}
+
+	void FAudioFadeOperatorWrapper::BindOutputs(FOutputVertexInterfaceData& InOutVertexData)
+	{
+		// Bind wrapped operator
+		WrappedOperator.BindOutputs(InOutVertexData);
+
+		// Override any output bindings are related to faded outputs .
+		for (int i = 0; i < OutputVerticesToFade.Num(); i++)
+		{
+			const FVertexName& VertexName = OutputVerticesToFade[i];
+
+			// Update internal references to wrapped operator's audio buffers
+			if (const FAnyDataReference* DataRef = InOutVertexData.FindDataReference(VertexName))
+			{
+				WrappedOutputAudioBuffers[i] = DataRef->GetDataReadReference<FAudioBuffer>();
+			}
+			else
+			{
+				UE_LOG(LogMetaSound, Warning, TEXT("Failed to find output audio buffer '%s' when attempting to fade out audio buffer."), *VertexName.ToString());
+			}
+
+			// Override which output buffers are exposed from this operator.
+			InOutVertexData.BindVertex(VertexName, FadedOutputAudioBuffers[i]);
+		}
+	}
+
+	void FAudioFadeOperatorWrapper::Execute()
+	{
+		// Fade inputs
+		FadeBuffers(WrappedInputAudioBuffers, FadedInputAudioBuffers);
+
+		WrappedOperator.Execute();
+
+		// Fade outputs
+		FadeBuffers(WrappedOutputAudioBuffers, FadedOutputAudioBuffers);
+
+		// Update fade state
+		if (FadeState == EFadeState::FadingIn)
+		{
+			FadeState = EFadeState::FullVolume;
+		}
+		else if (FadeState == EFadeState::FadingOut)
+		{
+			FadeState = EFadeState::Silence;
+		}
+	}
+
+	void FAudioFadeOperatorWrapper::PostExecute()
+	{
+		WrappedOperator.PostExecute();
+		
+		if (FadeState == EFadeState::Silence)
+		{
+			// If it is silent, zero output buffers to avoid having to rezero output buffers every
+			// time the operator is executed 
+			for (TDataWriteReference<FAudioBuffer>& FadedBufferRef : FadedInputAudioBuffers)
+			{
+				FadedBufferRef->Zero();
+			}
+			for (TDataWriteReference<FAudioBuffer>& FadedBufferRef : FadedOutputAudioBuffers)
+			{
+				FadedBufferRef->Zero();
+			}
+		}
+	}
+
+	void FAudioFadeOperatorWrapper::Reset(const IOperator::FResetParams& InParams)
+	{
+		WrappedOperator.Reset(InParams);
+		FadeState = InitialFadeState;
+	}
+
+	void FAudioFadeOperatorWrapper::InitInputBuffers(const FOperatorSettings& InOperatorSettings, const FInputVertexInterfaceData& InInputData)
+	{
+		for (const FVertexName& VertexName : InputVerticesToFade)
+		{
+			FadedInputAudioBuffers.Add(TDataWriteReference<FAudioBuffer>::CreateNew(InOperatorSettings));
+			if (const FAnyDataReference* DataRef = InInputData.FindDataReference(VertexName))
+			{
+				WrappedInputAudioBuffers.Add(DataRef->GetDataReadReference<FAudioBuffer>());
+			}
+			else
+			{
+				WrappedInputAudioBuffers.Add(TDataReadReference<FAudioBuffer>::CreateNew(InOperatorSettings));
+			}
+		}
+	}
+
+	void FAudioFadeOperatorWrapper::InitOutputBuffers(const FOperatorSettings& InOperatorSettings)
+	{
+		FOutputVertexInterfaceData OutputData;
+		WrappedOperator.BindOutputs(OutputData);
+
+		for (const FVertexName& VertexName : OutputVerticesToFade)
+		{
+			FadedOutputAudioBuffers.Add(TDataWriteReference<FAudioBuffer>::CreateNew(InOperatorSettings));
+			if (const FAnyDataReference* DataRef = OutputData.FindDataReference(VertexName))
+			{
+				WrappedOutputAudioBuffers.Add(DataRef->GetDataReadReference<FAudioBuffer>());
+			}
+			else
+			{
+				UE_LOG(LogMetaSound, Warning, TEXT("Failed to find output audio buffer '%s' when attempting to fade audio buffer."), *VertexName.ToString());
+				WrappedOutputAudioBuffers.Add(TDataReadReference<FAudioBuffer>::CreateNew(InOperatorSettings));
+			}
+		}
+	}
+
+	void FAudioFadeOperatorWrapper::FadeBuffers(const TArray<TDataReadReference<FAudioBuffer>>& InSrc, const TArray<TDataWriteReference<FAudioBuffer>>& InDst)
+	{
+		if (FadeState == EFadeState::FadingIn)
+		{
+			for (int32 i = 0; i < InSrc.Num(); i++)
+			{
+				const FAudioBuffer& WrappedBuffer = *InSrc[i];
+				FAudioBuffer& OutputBuffer = *InDst[i];
+
+				FMemory::Memcpy(OutputBuffer.GetData(), WrappedBuffer.GetData(), WrappedBuffer.Num() * sizeof(float));
+				Audio::ArrayFade(OutputBuffer, 0.f, 1.f);
+			}
+		}
+		else if (FadeState == EFadeState::FadingOut)
+		{
+			for (int32 i = 0; i < InSrc.Num(); i++)
+			{
+				const FAudioBuffer& WrappedBuffer = *InSrc[i];
+				FAudioBuffer& OutputBuffer = *InDst[i];
+
+				FMemory::Memcpy(OutputBuffer.GetData(), WrappedBuffer.GetData(), WrappedBuffer.Num() * sizeof(float));
+				Audio::ArrayFade(OutputBuffer, 1.f, 0.f);
+			}
+		}
+		else if (FadeState == EFadeState::FullVolume)
+		{
+			for (int32 i = 0; i < InSrc.Num(); i++)
+			{
+				const FAudioBuffer& WrappedBuffer = *InSrc[i];
+				FAudioBuffer& OutputBuffer = *InDst[i];
+
+				FMemory::Memcpy(OutputBuffer.GetData(), WrappedBuffer.GetData(), WrappedBuffer.Num() * sizeof(float));
+			}
+		}
+	}
+
+	void FAudioFadeOperatorWrapper::SetWrappedOperatorInputs(const FInputVertexInterfaceData& InInputData)
+	{
+		// Create an internal copy so that wrapped data references don't escape 
+		// this operator. 
+		FInputVertexInterfaceData InputDataCopy = InInputData;
+
+		if (InputVerticesToFade.Num() > 0)
+		{
+			for (int i = 0; i < InputVerticesToFade.Num(); i++)
+			{
+				const FVertexName& VertexName = InputVerticesToFade[i];
+				InputDataCopy.SetVertex(VertexName, FadedInputAudioBuffers[i]);
+			}
+		}
+
+		WrappedOperator.BindInputs(InputDataCopy);
+	}
+}
+
diff --git a/Engine/Plugins/Runtime/Metasound/Source/MetasoundGraphCore/Private/MetasoundDynamicOperatorAudioFade.h b/Engine/Plugins/Runtime/Metasound/Source/MetasoundGraphCore/Private/MetasoundDynamicOperatorAudioFade.h
new file mode 100644
--- /dev/null
+++ b/Engine/Plugins/Runtime/Metasound/Source/MetasoundGraphCore/Private/MetasoundDynamicOperatorAudioFade.h
@@ -0,0 +1,68 @@
+// Copyright Epic Games, Inc. All Rights Reserved.
+
+#pragma once
+
+#include "Containers/Array.h"
+#include "MetasoundAudioBuffer.h"
+#include "MetasoundDataReference.h"
+#include "MetasoundExecutableOperator.h"
+#include "MetasoundVertex.h"
+#include "MetasoundVertexData.h"
+
+namespace Metasound::DynamicGraph
+{
+	// FAudioFadeOperatorWrapper wraps an existing operator and applies a fade in/out to 
+	// any designated audio buffers on the next execution. The fade will only occur
+	// over a single audio buffer and will be finalized during calls to PostExecute. 
+	//
+	// When this operator wraps another operator, it will also provide different 
+	// output data references for any faded output audio buffers. Users of this operator
+	// should be wary to rebind downstream operators in any graph using this operator
+	// wrapper. 
+	//
+	// The wrapped operator can be release if desired, though it leaves this operator
+	// in a zombie state.
+	class FAudioFadeOperatorWrapper : public TExecutableOperator<FAudioFadeOperatorWrapper>
+	{
+	public:
+		enum class EFadeState : uint8
+		{
+			Silence,
+			FadingOut,
+			FadingIn,
+			FullVolume
+		};
+
+		FAudioFadeOperatorWrapper(EFadeState InInitialFadeState, const FOperatorSettings& InOperatorSettings, const FInputVertexInterfaceData& InInputData, TUniquePtr<IOperator> InOperator, TArrayView<const FVertexName> InInputVerticesToFade, TArrayView<const FVertexName> InOutputVerticesToFade);
+
+		TUniquePtr<IOperator> ReleaseOperator();
+
+		virtual void BindInputs(FInputVertexInterfaceData& InOutVertexData) override;
+		virtual void BindOutputs(FOutputVertexInterfaceData& InOutVertexData) override;
+
+		void Execute();
+		void PostExecute();
+		void Reset(const IOperator::FResetParams& InParams);
+
+	private:
+
+		void InitInputBuffers(const FOperatorSettings& InOperatorSettings, const FInputVertexInterfaceData& InInputData);
+		void InitOutputBuffers(const FOperatorSettings& InOperatorSettings);
+
+		void FadeBuffers(const TArray<TDataReadReference<FAudioBuffer>>& InSrc, const TArray<TDataWriteReference<FAudioBuffer>>& InDst);
+		void SetWrappedOperatorInputs(const FInputVertexInterfaceData& InInputData);
+
+		FExecuter WrappedOperator;
+
+		TArray<FVertexName> InputVerticesToFade;
+		TArray<FVertexName> OutputVerticesToFade;
+
+		TArray<TDataReadReference<FAudioBuffer>> WrappedInputAudioBuffers;
+		TArray<TDataWriteReference<FAudioBuffer>> FadedInputAudioBuffers;
+		TArray<TDataReadReference<FAudioBuffer>> WrappedOutputAudioBuffers;
+		TArray<TDataWriteReference<FAudioBuffer>> FadedOutputAudioBuffers;
+
+		EFadeState InitialFadeState;
+		EFadeState FadeState;
+	};
+}
diff --git a/Engine/Plugins/Runtime/Metasound/Source/MetasoundGraphCore/Private/MetasoundDynamicOperatorTransactor.cpp b/Engine/Plugins/Runtime/Metasound/Source/MetasoundGraphCore/Private/MetasoundDynamicOperatorTransactor.cpp
--- a/Engine/Plugins/Runtime/Metasound/Source/MetasoundGraphCore/Private/MetasoundDynamicOperatorTransactor.cpp
+++ b/Engine/Plugins/Runtime/Metasound/Source/MetasoundGraphCore/Private/MetasoundDynamicOperatorTransactor.cpp
@@ -22,8 +22,9 @@ namespace Metasound
 		namespace DynamicOperatorTransactorPrivate
 		{
 			// Literal nodes always have output vertex with this name. 
-			static FLazyName LiteralNodeOutputVertexName("Value");
-
+			static const FLazyName LiteralNodeOutputVertexName("Value");
+			
+			// Sorts the graph and determines order of operator execution.
 			TArray<FOperatorID> DetermineOperatorOrder(const IGraph& InGraph)
 			{
 				/* determine new operator order. */
@@ -40,6 +41,18 @@ namespace Metasound
 
 				return OperatorOrder;
 			}
+
+			static const FLazyName AudioBufferDataTypeName("Audio");
+
+			FString GetDebugNodeNameString(const INode& InNode)
+			{
+				return FString::Printf(TEXT("%s_v%d.%d"), *InNode.GetMetadata().ClassName.GetFullName().ToString(), InNode.GetMetadata().MajorVersion, InNode.GetMetadata().MinorVersion);
+			}
+
+			FString GetDebugNodeNameString(const FGuid& InNodeID, const INode& InNode)
+			{
+				return FString::Printf(TEXT("%s:%s"), *InNodeID.ToString(), *GetDebugNodeNameString(InNode));
+			}
 		}
 
 		bool operator<(const FDynamicOperatorTransactor::FLiteralNodeID& InLHS, const FDynamicOperatorTransactor::FLiteralNodeID& InRHS)
@@ -107,12 +120,24 @@ namespace Metasound
 
 			if (const INode* Node = Graph.FindNode(InNodeID))
 			{
-				FOperatorID OperatorID = GetOperatorID(Node);
-				auto CreateRemoveNodeTransform = [&](const FOperatorSettings& InOperatorSettings, const FMetasoundEnvironment& InEnvironment) -> TUniquePtr<IDynamicOperatorTransform>
+				TArray<FVertexName> OutputsToFade;
+				const FOutputVertexInterface& OutputInterface = Node->GetVertexInterface().GetOutputInterface();
+				for (const FOutputDataVertex& OutputVertex : OutputInterface)
 				{
-					return MakeUnique<FRemoveOperator>(OperatorID);
-				};
-				EnqueueTransformOnOperatorQueues(CreateRemoveNodeTransform);
+					if (OutputVertex.DataTypeName == DynamicOperatorTransactorPrivate::AudioBufferDataTypeName)
+					{
+						OutputsToFade.Add(OutputVertex.VertexName);
+					}
+				}
+
+				if (OutputsToFade.Num())
+				{
+					EnqueueFadeAndRemoveOperatorTransform(*Node, OutputsToFade);
+				}
+				else
+				{
+					EnqueueRemoveOperatorTransform(*Node);
+				}
 
 				constexpr bool bRemoveDataEdgesWithNode = true;
 				bool bRemovedNodeFromGraph = Graph.RemoveNode(InNodeID, bRemoveDataEdgesWithNode);
@@ -189,41 +214,20 @@ namespace Metasound
 			}
 
 			/* enqueue an update. */
-			const FOperatorID FromOperatorID = GetOperatorID(FromNode);
-			const FOperatorID LiteralOperatorID = GetOperatorID(ReplacementLiteralNode);
-			const FOperatorID ToOperatorID = GetOperatorID(ToNode);
 			TArray<FOperatorID> OperatorOrder = DynamicOperatorTransactorPrivate::DetermineOperatorOrder(Graph);
 
-			auto CreateRemoveEdgeTransform = [&](const FOperatorSettings& InOperatorSettings, const FMetasoundEnvironment& InEnvironment) -> TUniquePtr<IDynamicOperatorTransform>
+			if (ToNode->GetVertexInterface().GetInputVertex(InToVertex).DataTypeName == DynamicOperatorTransactorPrivate::AudioBufferDataTypeName)
 			{
-				// Add the literal node.
-				TUniquePtr<IDynamicOperatorTransform> AddNodeTransform = CreateAddOperatorTransform(*ReplacementLiteralNode, InOperatorSettings, InEnvironment);
-
-				// Connect literal node to target node.
-				TUniquePtr<IDynamicOperatorTransform> ConnectOperatorsTransform = MakeUnique<FConnectOperators>(LiteralOperatorID, DynamicOperatorTransactorPrivate::LiteralNodeOutputVertexName, ToOperatorID, InToVertex);
-
-				// Reorder the node graph.
-				TUniquePtr<IDynamicOperatorTransform> SetOrderTransform = MakeUnique<FSetOperatorOrder>(OperatorOrder);
-
-				if (AddNodeTransform.IsValid() && ConnectOperatorsTransform.IsValid() && SetOrderTransform.IsValid())
-				{
-					// Create an atomic transform so all sub-transforms happen before next execution.
-					TArray<TUniquePtr<IDynamicOperatorTransform>> AtomicTransforms;
-
-					AtomicTransforms.Emplace(MoveTemp(AddNodeTransform));
-					AtomicTransforms.Emplace(MoveTemp(ConnectOperatorsTransform));
-					AtomicTransforms.Emplace(MoveTemp(SetOrderTransform));
+				// Handle audio edge removal with a fade out.
+				EnqueueFadeAndRemoveEdgeOperatorTransform(*FromNode, InFromVertex, *ToNode, InToVertex, *ReplacementLiteralNode, OperatorOrder);
+			}
+			else
+			{
+				// Immediately disconnect non-audio edges. 
+				EnqueueRemoveEdgeOperatorTransform(*FromNode, InFromVertex, *ToNode, InToVertex, *ReplacementLiteralNode, OperatorOrder);
+			}
 
-					return MakeUnique<FAtomicTransform>(MoveTemp(AtomicTransforms));
-				}
-				else
-				{
-					UE_LOG(LogMetaSound, Error, TEXT("Cannot remove edge from %s:%s to %s:%s because of failure to create all transforms needed to perform operatorn."), *InFromNodeID.ToString(), *InFromVertex.ToString(), *InToNodeID.ToString(), *InToVertex.ToString());
-					return TUniquePtr<IDynamicOperatorTransform>(nullptr);
-				}
-			};
 
-			EnqueueTransformOnOperatorQueues(CreateRemoveEdgeTransform);
 		}
 
 		void FDynamicOperatorTransactor::SetValue(const FGuid& InNodeID, const FVertexName& InVertex, TUniquePtr<INode> InLiteralNode)
@@ -369,18 +373,24 @@ namespace Metasound
 		void FDynamicOperatorTransactor::AddDataEdgeInternal(const INode& InFromNode, const FVertexName& InFromVertex, const FGuid& InToNodeID, const INode& InToNode, const FVertexName& InToVertex)
 		{
 			using namespace DirectedGraphAlgo;
+			using namespace DynamicOperatorTransactorPrivate;
 
-			/* Determine if there is an existing literal node connected to the node. */
+			const FInputDataVertex* InputVertex = InToNode.GetVertexInterface().GetInputInterface().Find(InToVertex);
+			if (nullptr == InputVertex)
+			{
+				UE_LOG(LogMetaSound, Error, TEXT("Cannot connect nodes because destination node %s does not contain input vertex %s"), *GetDebugNodeNameString(InToNodeID, InToNode), *InToVertex.ToString());
+				return;
+			}
+
+			/* Determine if there is an existing literal node connected to the node. 
+			 * Literal nodes are stored on the FDynamicOperatorTransactor and need to be 
+			 * disconnected and removed if they are no longer being used. 
+			 */
 			const FLiteralNodeID LiteralNodeKey{InToNodeID, InToVertex};
-			bool bIsAlreadyConnectedToALiteralNode = false;
-			FOperatorID PriorLiteralOperatorID;
-			if (const TUniquePtr<INode>* PriorLiteralNode = LiteralNodeMap.Find(LiteralNodeKey))
+			TUniquePtr<INode> PriorLiteralNode;
+			if (LiteralNodeMap.RemoveAndCopyValue(LiteralNodeKey, PriorLiteralNode))
 			{
-				bIsAlreadyConnectedToALiteralNode = true;
-				PriorLiteralOperatorID = GetOperatorID(PriorLiteralNode->Get());
-				
-				Graph.RemoveDataEdge(**PriorLiteralNode, DynamicOperatorTransactorPrivate::LiteralNodeOutputVertexName, InToNode, InToVertex);
-				LiteralNodeMap.Remove(LiteralNodeKey);
+				Graph.RemoveDataEdge(*PriorLiteralNode, DynamicOperatorTransactorPrivate::LiteralNodeOutputVertexName, InToNode, InToVertex);
 			}
 
 			/* add edge to internal graph. */
@@ -388,20 +398,54 @@ namespace Metasound
 
 			// Find order of operators after adding edge. 
 			TArray<FOperatorID> OperatorOrder = DynamicOperatorTransactorPrivate::DetermineOperatorOrder(Graph);
+			
+			if (InputVertex->DataTypeName == DynamicOperatorTransactorPrivate::AudioBufferDataTypeName)
+			{
+				// If edge is audio, then the connection needs to be faded
+				EnqueueFadeAndAddEdgeOperatorTransform(InFromNode, InFromVertex, InToNode, InToVertex, PriorLiteralNode.Get(), OperatorOrder);
+			}
+			else
+			{
+				// If the edge is not audio, then no fading is performed. 
+				EnqueueAddEdgeOperatorTransform(InFromNode, InFromVertex, InToNode, InToVertex, PriorLiteralNode.Get(), OperatorOrder);
+			}
+		}
+		
+		void FDynamicOperatorTransactor::EnqueueFadeAndRemoveOperatorTransform(const INode& InNode, TArrayView<const FVertexName> InOutputsToFade)
+		{
+			TArrayView<const FVertexName> InputsToFade; // We do not need to fade any inputs when removing a node.
+
+			EnqueueBeginFadeOperatorTransform(InNode, EAudioFadeType::FadeOut, InputsToFade, InOutputsToFade);
+
+			const FOperatorID OperatorID = DirectedGraphAlgo::GetOperatorID(InNode);
+
+			auto CreateEndFadeOutTransform = [&](const FOperatorSettings& InOperatorSettings, const FMetasoundEnvironment& InEnvironment) -> TUniquePtr<IDynamicOperatorTransform>
+			{
+				// We can skip the FEndAudioFadeTransform as an optimization here.
+				// FEndAudioFadeTransform removes the fading wrapper around the node,
+				// but since the node is being removed, we can remove the wrapper
+				// and the node with a single FRemoveOperator transform.
+				return MakeUnique<FRemoveOperator>(OperatorID);
+			};
+
+			EnqueueTransformOnOperatorQueues(CreateEndFadeOutTransform);
+		}
 
+		void FDynamicOperatorTransactor::EnqueueAddEdgeOperatorTransform(const INode& InFromNode, const FVertexName& InFromVertex, const INode& InToNode, const FVertexName& InToVertex, const INode* InPriorLiteralNode, const TArray<FOperatorID>& InNewOperatorOrder)
+		{
 			/* enqueue an update. */
-			FOperatorID FromOperatorID = GetOperatorID(InFromNode);
-			FOperatorID ToOperatorID = GetOperatorID(InToNode);
+			FOperatorID FromOperatorID = DirectedGraphAlgo::GetOperatorID(InFromNode);
+			FOperatorID ToOperatorID = DirectedGraphAlgo::GetOperatorID(InToNode);
 
 			auto CreateAddEdgeTransform = [&](const FOperatorSettings& InOperatorSettings, const FMetasoundEnvironment& InEnvironment) -> TUniquePtr<IDynamicOperatorTransform>
 			{
 				TArray<TUniquePtr<IDynamicOperatorTransform>> AtomicTransforms;
 
-				if (bIsAlreadyConnectedToALiteralNode)
+				if (InPriorLiteralNode)
 				{
-					AtomicTransforms.Add(MakeUnique<FRemoveOperator>(PriorLiteralOperatorID));
+					AtomicTransforms.Add(MakeUnique<FRemoveOperator>(DirectedGraphAlgo::GetOperatorID(InPriorLiteralNode)));
 				}
-				AtomicTransforms.Add(MakeUnique<FSetOperatorOrder>(OperatorOrder));
+				AtomicTransforms.Add(MakeUnique<FSetOperatorOrder>(InNewOperatorOrder));
 				AtomicTransforms.Add(MakeUnique<FConnectOperators>(FromOperatorID, InFromVertex, ToOperatorID, InToVertex));
 
 				return MakeUnique<FAtomicTransform>(MoveTemp(AtomicTransforms));
@@ -410,8 +454,159 @@ namespace Metasound
 			EnqueueTransformOnOperatorQueues(CreateAddEdgeTransform);
 		}
 
+		void FDynamicOperatorTransactor::EnqueueFadeAndAddEdgeOperatorTransform(const INode& InFromNode, const FVertexName& InFromVertex, const INode& InToNode, const FVertexName& InToVertex, const INode* InPriorLiteralNode, const TArray<FOperatorID>& InNewOperatorOrder)
+		{
+			FOperatorID FromOperatorID = DirectedGraphAlgo::GetOperatorID(InFromNode);
+			FOperatorID ToOperatorID = DirectedGraphAlgo::GetOperatorID(InToNode);
+			
+			// Fade inputs on the receiving node when adding an edge. We don't fade the source node's outputs
+			// because those outputs could also be connected to other nodes which we do not want to fade. 
+			TArrayView<const FVertexName> InputsToFade(&InToVertex, 1);
+			TArrayView<const FVertexName> OutputsToFade;
+
+			auto CreateBeginFadeAndAddEdgeTransform = [&](const FOperatorSettings& InOperatorSettings, const FMetasoundEnvironment& InEnvironment) -> TUniquePtr<IDynamicOperatorTransform>
+			{
+				TArray<TUniquePtr<IDynamicOperatorTransform>> AtomicTransforms;
+
+				if (InPriorLiteralNode)
+				{
+					// We assume that any FAudioBuffer created from a literal node is just silent audio. That means we can
+					// Remove that literal node without fading out the audio from the literal node connected to the destination node. 
+					AtomicTransforms.Add(MakeUnique<FRemoveOperator>(DirectedGraphAlgo::GetOperatorID(InPriorLiteralNode)));
+				}
+				
+				AtomicTransforms.Add(MakeUnique<FSetOperatorOrder>(InNewOperatorOrder));
+				AtomicTransforms.Add(MakeUnique<FConnectOperators>(FromOperatorID, InFromVertex, ToOperatorID, InToVertex));
+				AtomicTransforms.Add(MakeUnique<FBeginAudioFadeTransform>(ToOperatorID, EAudioFadeType::FadeIn, InputsToFade, OutputsToFade));
+				
+				// Fence must be the last transform since the fade must be performed
+				// before anything else happens in the graph. To apply a fade, the 
+				// graph must execute. This FExecuteFence transform ensures that the graph
+				// is executed before any additional transforms are applied. 
+				AtomicTransforms.Add(MakeUnique<FExecuteFence>());
+
+				return MakeUnique<FAtomicTransform>(MoveTemp(AtomicTransforms));
+			};
+
+			EnqueueTransformOnOperatorQueues(CreateBeginFadeAndAddEdgeTransform);
+
+			EnqueueEndFadeOperatorTransform(InToNode);
+		}
+
+		void FDynamicOperatorTransactor::EnqueueBeginFadeOperatorTransform(const INode& InNode, EAudioFadeType InFadeType, TArrayView<const FVertexName> InInputsToFade, TArrayView<const FVertexName> InOutputsToFade)
+		{
+			const FOperatorID OperatorID = DirectedGraphAlgo::GetOperatorID(InNode);
+
+			auto CreateBeginAudioFadeTransform = [&](const FOperatorSettings& InOperatorSettings, const FMetasoundEnvironment& InEnvironment) -> TUniquePtr<IDynamicOperatorTransform>
+			{
+				TArray<TUniquePtr<IDynamicOperatorTransform>> BeginAudioFadeAtomicTransforms;
+
+				BeginAudioFadeAtomicTransforms.Add(MakeUnique<FBeginAudioFadeTransform>(OperatorID, InFadeType, InInputsToFade, InOutputsToFade));
+				// Fence must be the last transform since the fade must be performed
+				// before anything else happens in the graph. To apply a fade, the 
+				// graph must execute. This FExecuteFence transform ensures that the graph
+				// is executed before any additional transforms are applied. 
+				BeginAudioFadeAtomicTransforms.Add(MakeUnique<FExecuteFence>());
+
+				return MakeUnique<FAtomicTransform>(MoveTemp(BeginAudioFadeAtomicTransforms));
+			};
+
+			EnqueueTransformOnOperatorQueues(CreateBeginAudioFadeTransform);
+		}
+
+		void FDynamicOperatorTransactor::EnqueueEndFadeOperatorTransform(const INode& InNode)
+		{
+			const FOperatorID OperatorID = DirectedGraphAlgo::GetOperatorID(InNode);
+
+			auto CreateEndAudioFadeTransform = [&](const FOperatorSettings& InOperatorSettings, const FMetasoundEnvironment& InEnvironment) -> TUniquePtr<IDynamicOperatorTransform>
+			{
+				return MakeUnique<FEndAudioFadeTransform>(OperatorID);
+			};
+
+			EnqueueTransformOnOperatorQueues(CreateEndAudioFadeTransform);
+		}
+
+		void FDynamicOperatorTransactor::EnqueueRemoveOperatorTransform(const INode& InNode)
+		{
+			FOperatorID OperatorID = DirectedGraphAlgo::GetOperatorID(InNode);
+
+			auto CreateRemoveNodeTransform = [&](const FOperatorSettings& InOperatorSettings, const FMetasoundEnvironment& InEnvironment) -> TUniquePtr<IDynamicOperatorTransform>
+			{
+				return MakeUnique<FRemoveOperator>(OperatorID);
+			};
+			EnqueueTransformOnOperatorQueues(CreateRemoveNodeTransform);
+		}
+
+		void FDynamicOperatorTransactor::EnqueueRemoveEdgeOperatorTransform(const INode& InFromNode, const FVertexName& InFromVertex, const INode& InToNode, const FVertexName& InToVertex, const INode& InReplacementLiteralNode, const TArray<FOperatorID>& InNewOperatorOrder)
+		{
+			using namespace DynamicOperatorTransactorPrivate;
+
+			const FOperatorID FromOperatorID = DirectedGraphAlgo::GetOperatorID(InFromNode);
+			const FOperatorID ToOperatorID = DirectedGraphAlgo::GetOperatorID(InToNode);
+			const FOperatorID LiteralOperatorID = DirectedGraphAlgo::GetOperatorID(InReplacementLiteralNode);
+
+			auto CreateRemoveEdgeTransform = [&](const FOperatorSettings& InOperatorSettings, const FMetasoundEnvironment& InEnvironment) -> TUniquePtr<IDynamicOperatorTransform>
+			{
+				// Add the literal node.
+				TUniquePtr<IDynamicOperatorTransform> AddNodeTransform = CreateAddOperatorTransform(InReplacementLiteralNode, InOperatorSettings, InEnvironment);
+
+				// Connect literal node to target node.
+				TUniquePtr<IDynamicOperatorTransform> ConnectOperatorsTransform = MakeUnique<FConnectOperators>(LiteralOperatorID, DynamicOperatorTransactorPrivate::LiteralNodeOutputVertexName, ToOperatorID, InToVertex);
+
+				// Reorder the node graph.
+				TUniquePtr<IDynamicOperatorTransform> SetOrderTransform = MakeUnique<FSetOperatorOrder>(InNewOperatorOrder);
+
+				if (AddNodeTransform.IsValid() && ConnectOperatorsTransform.IsValid() && SetOrderTransform.IsValid())
+				{
+					// Create an atomic transform so all sub-transforms happen before next execution.
+					TArray<TUniquePtr<IDynamicOperatorTransform>> AtomicTransforms;
+
+					AtomicTransforms.Emplace(MoveTemp(AddNodeTransform));
+					AtomicTransforms.Emplace(MoveTemp(ConnectOperatorsTransform));
+					AtomicTransforms.Emplace(MoveTemp(SetOrderTransform));
+
+					return MakeUnique<FAtomicTransform>(MoveTemp(AtomicTransforms));
+				}
+				else
+				{
+
+					UE_LOG(LogMetaSound, Error, TEXT("Cannot remove edge from %s:%s to %s:%s because of failure to create all transforms needed to perform operatorn."), *GetDebugNodeNameString(InFromNode), *InFromVertex.ToString(), *GetDebugNodeNameString(InToNode), *InToVertex.ToString());
+					return TUniquePtr<IDynamicOperatorTransform>(nullptr);
+				}
+			};
+
+			EnqueueTransformOnOperatorQueues(CreateRemoveEdgeTransform);
+		}
+
+		void FDynamicOperatorTransactor::EnqueueFadeAndRemoveEdgeOperatorTransform(const INode& InFromNode, const FVertexName& InFromVertex, const INode& InToNode, const FVertexName& InToVertex, const INode& InReplacementLiteralNode, const TArray<FOperatorID>& InNewOperatorOrder)
+		{
+			// Fade the input to the node getting disconnected rather than the output of the source node. The source node
+			// may be connected to other nodes and fading it's output would fade all the other connected nodes' inputs. 
+			TArrayView<const FVertexName> InputVerticesToFade(&InToVertex, 1);
+			TArrayView<const FVertexName> OutputVerticesToFade;
+
+			EnqueueBeginFadeOperatorTransform(InToNode, EAudioFadeType::FadeOut, InputVerticesToFade, OutputVerticesToFade);
+
+			// Replace input with literal. This assumes that the replacement audio
+			// buffer contains silent audio. The fade transform will get the input
+			// audio to silent which will then seamlessly be swapped with a silent
+			// audio buffer as a permanent connection. 
+			//
+			// If we ever find ourselves creating audio buffers with literals which 
+			// are anything other than silent buffers, we should rework this operation
+			// to do either a cross-fade, or an additional "fade in" to the new value. 
+			EnqueueRemoveEdgeOperatorTransform(InFromNode, InFromVertex, InToNode, InToVertex, InReplacementLiteralNode, InNewOperatorOrder);
+
+			// Remove fade operation. 
+			EnqueueEndFadeOperatorTransform(InToNode);
+		}
+
+
+
 		TUniquePtr<IDynamicOperatorTransform> FDynamicOperatorTransactor::CreateAddOperatorTransform(const INode& InNode, const FOperatorSettings& InOperatorSettings, const FMetasoundEnvironment& InEnvironment) const
 		{
+			using namespace DynamicOperatorTransactorPrivate;
+
 			const FOperatorID OperatorID = DirectedGraphAlgo::GetOperatorID(InNode);
 			FVertexInterfaceData InterfaceData(InNode.GetVertexInterface());
 			FBuildOperatorParams OperatorParams
@@ -429,7 +624,7 @@ namespace Metasound
 			{
 				if (Error.IsValid())
 				{
-					UE_LOG(LogMetaSound, Warning, TEXT("Encountered error while building operator for node %s. %s:%s"), *InNode.GetMetadata().ClassName.GetFullName().ToString(), *Error->GetErrorType().ToString(), *Error->GetErrorDescription().ToString());
+					UE_LOG(LogMetaSound, Warning, TEXT("Encountered error while building operator for node %s. %s:%s"), *GetDebugNodeNameString(InNode), *Error->GetErrorType().ToString(), *Error->GetErrorDescription().ToString());
 				}
 			}
 
diff --git a/Engine/Plugins/Runtime/Metasound/Source/MetasoundGraphCore/Public/MetasoundAudioBuffer.h b/Engine/Plugins/Runtime/Metasound/Source/MetasoundGraphCore/Public/MetasoundAudioBuffer.h
new file mode 100644
--- /dev/null
+++ b/Engine/Plugins/Runtime/Metasound/Source/MetasoundGraphCore/Public/MetasoundAudioBuffer.h
@@ -0,0 +1,164 @@
+// Copyright Epic Games, Inc. All Rights Reserved.
+
+#pragma once
+
+#include "DSP/AlignedBuffer.h"
+#include "MetasoundGraphCoreModule.h"
+#include "MetasoundDataReferenceMacro.h"
+#include "MetasoundLog.h"
+#include "MetasoundOperatorSettings.h"
+
+/** Define which determines whether to check that the size of the audio buffer has not changed since initialization */
+#define METASOUNDGRAPHCORE_CHECKAUDIONUM !(UE_BUILD_SHIPPING || UE_BUILD_TEST)
+
+namespace Metasound
+{
+	/**  FAudioBuffer
+	 *
+	 * FAudioBuffer is the default buffer for passing audio data between nodes.  It should not be resized.
+	 * It is compatible with functions accepting const references to FAlignedFloatBuffer (const FAlignedFloatBuffer&)
+	 * arguments via an implicit conversion operator which exposes the underlying FAlignedFloatBuffer container.
+	 *
+	 */
+	class FAudioBuffer
+	{
+		public:
+			/** Create an FAudioBuffer with a specific number of samples.
+			 *
+			 * @param InNumSamples - Number of samples in buffer.
+			 */
+			explicit FAudioBuffer(int32 InNumSamples)
+			{
+				Buffer.AddZeroed(InNumSamples);
+
+#if METASOUNDGRAPHCORE_CHECKAUDIONUM
+				// InitialNum should not change during the life of an FAudioBuffer
+				InitialNum = Buffer.Num();
+#endif
+			}
+
+			/**
+			 * This is the constructor used by the frontend.
+			 */
+			FAudioBuffer(const FOperatorSettings& InSettings)
+			{
+				Buffer.AddZeroed(InSettings.GetNumFramesPerBlock());
+
+#if METASOUNDGRAPHCORE_CHECKAUDIONUM
+				// InitialNum should not change during the life of an FAudioBuffer
+				InitialNum = Buffer.Num();
+#endif	
+			}
+
+			FAudioBuffer()
+			{
+#if METASOUNDGRAPHCORE_CHECKAUDIONUM
+				// InitialNum should not change during the life of an FAudioBuffer
+				InitialNum = Buffer.Num();
+#endif
+			}
+
+			/** Return a pointer to the audio float data. */
+			FORCEINLINE const float* GetData() const
+			{
+#if METASOUNDGRAPHCORE_CHECKAUDIONUM
+				UE_CLOG(InitialNum != Buffer.Num(), LogMetaSound, Error, TEXT("MetaSound audio buffer size change detected.  Audio Buffers should not be resized."));
+#endif				
+				return Buffer.GetData();
+			}
+
+			/** Return a pointer to the audio float data. */
+			FORCEINLINE float* GetData()
+			{
+#if METASOUNDGRAPHCORE_CHECKAUDIONUM
+				UE_CLOG(InitialNum != Buffer.Num(), LogMetaSound, Error, TEXT("MetaSound audio buffer size change detected.  Audio Buffers should not be resized."));
+#endif
+				return Buffer.GetData();
+			}
+
+			/** Return the number of samples in the audio buffer. */
+			FORCEINLINE int32 Num() const
+			{
+#if METASOUNDGRAPHCORE_CHECKAUDIONUM
+				UE_CLOG(InitialNum != Buffer.Num(), LogMetaSound, Error, TEXT("MetaSound audio buffer size change detected.  Audio Buffers should not be resized."));
+#endif
+				return Buffer.Num();
+			}
+
+			/** Implicit conversion to Audio::FAlignedFloatBuffer */
+			FORCEINLINE operator const Audio::FAlignedFloatBuffer& () const
+			{
+#if METASOUNDGRAPHCORE_CHECKAUDIONUM
+				UE_CLOG(InitialNum != Buffer.Num(), LogMetaSound, Error, TEXT("MetaSound audio buffer size change detected.  Audio Buffers should not be resized."));
+#endif
+return Buffer;
+			}
+
+			/** Implicit conversion to Audio::FAlignedFloatBuffer 
+			 *
+			 * WARNING: if the buffer is resized, it will cause errors.
+			 */
+			FORCEINLINE operator Audio::FAlignedFloatBuffer& ()
+			{
+#if METASOUNDGRAPHCORE_CHECKAUDIONUM
+				UE_CLOG(InitialNum != Buffer.Num(), LogMetaSound, Error, TEXT("MetaSound audio buffer size change detected.  Audio Buffers should not be resized."));
+#endif				
+				return Buffer;
+			}
+
+			/** Implicit conversion to Audio::FAlignedFloatBuffer */
+			FORCEINLINE operator TArrayView<const float> () const
+			{
+#if METASOUNDGRAPHCORE_CHECKAUDIONUM
+				UE_CLOG(InitialNum != Buffer.Num(), LogMetaSound, Error, TEXT("MetaSound audio buffer size change detected.  Audio Buffers should not be resized."));
+#endif
+return Buffer;
+			}
+
+			/** Implicit conversion to Audio::FAlignedFloatBuffer 
+			 *
+			 * WARNING: if the buffer is resized, it will cause errors.
+			 */
+			FORCEINLINE operator TArrayView<float> ()
+			{
+#if METASOUNDGRAPHCORE_CHECKAUDIONUM
+				UE_CLOG(InitialNum != Buffer.Num(), LogMetaSound, Error, TEXT("MetaSound audio buffer size change detected.  Audio Buffers should not be resized."));
+#endif				
+				return Buffer;
+			}
+
+			FORCEINLINE void Zero()
+			{
+				if (Buffer.Num() > 0)
+				{
+					FMemory::Memset(Buffer.GetData(), 0, sizeof(float) * Buffer.Num());
+				}
+			}
+
+		private:
+
+			Audio::FAlignedFloatBuffer Buffer;
+
+#if METASOUNDGRAPHCORE_CHECKAUDIONUM
+			int32 InitialNum;
+#endif
+	};
+
+	DECLARE_METASOUND_DATA_REFERENCE_TYPES(FAudioBuffer, METASOUNDGRAPHCORE_API, FAudioBufferTypeInfo, FAudioBufferReadRef, FAudioBufferWriteRef);
+
+
+	// This empty base class is used so that we can specialize various nodes (Send, Receive, etc.) for subclasses of IAudioDatatype.
+	class METASOUNDGRAPHCORE_API IAudioDataType
+	{
+		/**
+		 * Audio datatypes require the following member functions:
+		 * int32 GetNumChannels() const { return NumChannels; }
+		 *
+		 * int32 GetMaxNumChannels() const { return MaxNumChannels; }
+
+		 * const TArrayView<const FAudioBufferReadRef> GetBuffers() const { return ReadableBuffers; }
+
+		 * const TArrayView<const FAudioBufferWriteRef> GetBuffers() { return WritableBuffers; }
+		 */
+	};
+};
diff --git a/Engine/Plugins/Runtime/Metasound/Source/MetasoundGraphCore/Public/MetasoundDynamicOperatorTransactor.h b/Engine/Plugins/Runtime/Metasound/Source/MetasoundGraphCore/Public/MetasoundDynamicOperatorTransactor.h
--- a/Engine/Plugins/Runtime/Metasound/Source/MetasoundGraphCore/Public/MetasoundDynamicOperatorTransactor.h
+++ b/Engine/Plugins/Runtime/Metasound/Source/MetasoundGraphCore/Public/MetasoundDynamicOperatorTransactor.h
@@ -20,8 +20,10 @@ namespace Metasound
 
 	namespace DynamicGraph
 	{
+		// Forward declare
 		class IDynamicOperatorTransform;
 		class FDynamicOperator;
+		enum class EAudioFadeType : uint8;
 
 		using FLiteralAssignmentFunction = void(*)(const FOperatorSettings& InOperatorSettings, const FLiteral& InLiteral, const FAnyDataReference& OutDataRef);
 		using FReferenceCreationFunction = TOptional<FAnyDataReference>(*)(const FOperatorSettings& InSettings, const FLiteral& InLiteral, EDataReferenceAccessType InAccessType);
@@ -109,6 +111,17 @@ namespace Metasound
 			const FGraph& GetGraph() const;
 
 		private:
+			using FOperatorID = uintptr_t;
+
+			void EnqueueFadeAndRemoveOperatorTransform(const INode& InNode, TArrayView<const FVertexName> InOutputsToFade);
+			void EnqueueRemoveOperatorTransform(const INode& InNode);
+			void EnqueueBeginFadeOperatorTransform(const INode& InNode, EAudioFadeType InFadeType, TArrayView<const FVertexName> InInputsToFade, TArrayView<const FVertexName> InOutputsToFade);
+			void EnqueueEndFadeOperatorTransform(const INode& InNode);
+
+			void EnqueueRemoveEdgeOperatorTransform(const INode& InFromNode, const FVertexName& InFromVertex, const INode& InToNode, const FVertexName& InToVertex, const INode& InReplacementLiteralNode, const TArray<FOperatorID>& InNewOperatorOrder);
+			void EnqueueFadeAndRemoveEdgeOperatorTransform(const INode& InFromNode, const FVertexName& InFromVertex, const INode& InToNode, const FVertexName& InToVertex, const INode& InReplacementLiteralNode, const TArray<FOperatorID>& InNewOperatorOrder);
+			void EnqueueAddEdgeOperatorTransform(const INode& InFromNode, const FVertexName& InFromVertex, const INode& InToNode, const FVertexName& InToVertex, const INode* InPriorLiteralNode, const TArray<FOperatorID>& InNewOperatorOrder);
+			void EnqueueFadeAndAddEdgeOperatorTransform(const INode& InFromNode, const FVertexName& InFromVertex, const INode& InToNode, const FVertexName& InToVertex, const INode* InPriorLiteralNode, const TArray<FOperatorID>& InNewOperatorOrder);
 
 			void AddDataEdgeInternal(const INode& InFromNode, const FVertexName& InFromVertex, const FGuid& InToNodeID, const INode& InToNode, const FVertexName& InToVertex);
 
diff --git a/Engine/Plugins/Runtime/Metasound/Source/MetasoundGraphCore/Public/MetasoundExecutableOperator.h b/Engine/Plugins/Runtime/Metasound/Source/MetasoundGraphCore/Public/MetasoundExecutableOperator.h
--- a/Engine/Plugins/Runtime/Metasound/Source/MetasoundGraphCore/Public/MetasoundExecutableOperator.h
+++ b/Engine/Plugins/Runtime/Metasound/Source/MetasoundGraphCore/Public/MetasoundExecutableOperator.h
@@ -75,6 +75,65 @@ namespace Metasound
 				}
 			}
 		};
+
+		// Helper template to determine whether a member function is declared
+		// for a given template class.
+		template <typename U>
+		class TIsPostExecuteMethodDeclared 
+		{
+			private:
+				template<typename T, T> 
+				struct Helper;
+
+				template<typename T>
+				static uint8 Check(Helper<void(T::*)(), &T::PostExecute>*);
+
+				template<typename T> static uint16 Check(...);
+
+			public:
+				static constexpr bool Value = sizeof(Check<U>(0)) == sizeof(uint8);
+		};
+
+		// Helper to determine whether an PostExecute(...) function exists. 
+		//
+		// Note: This helper must be instantiated inside the functions and cannot
+		// be part of the `TExecutableOperator<OperatorType>` class definition as
+		// classes are not completely defined at the point in the compilation process
+		// when static constexpr class members are evaluated. 
+		template<typename OperatorType>
+		struct TPostExecuteFunctionAdapter
+		{
+			// Returns an IOperator::FPostExecuteFunction if the OperatorType has a
+			// PostExecute(...) class member function.
+			static IOperator::FPostExecuteFunction GetPostExecuteFunction() 
+			{
+				if (TIsPostExecuteMethodDeclared<OperatorType>::Value)
+				{
+					return &TPostExecuteFunctionAdapter::PostExecuteFunction;
+				}
+				else
+				{
+					return nullptr;
+				}
+			}
+
+		private:
+			static void PostExecuteFunction(IOperator* InOperator)
+			{
+				if constexpr (TIsPostExecuteMethodDeclared<OperatorType>::Value)
+				{
+					OperatorType* DerivedOperator = static_cast<OperatorType*>(InOperator);
+
+					check(nullptr != DerivedOperator);
+
+					DerivedOperator->PostExecute();
+				}
+				else
+				{
+					checkNoEntry();
+				}
+			}
+		};
 	}
 
 	// As a general rule, ExecutableDataTypes should be avoided whenever possible
@@ -142,7 +201,7 @@ namespace Metasound
 
 		virtual FPostExecuteFunction GetPostExecuteFunction() override
 		{
-			return nullptr;
+			return MetasoundExecutableOperatorPrivate::TPostExecuteFunctionAdapter<DerivedOperatorType>::GetPostExecuteFunction();
 		}
 
 	private:
