Fix for VaraibleNode test warnings in automation tests
#rnx

diff --git a/Engine/Plugins/Runtime/Metasound/Source/MetasoundFrontend/Private/MetasoundFrontendDataTypeRegistry.cpp b/Engine/Plugins/Runtime/Metasound/Source/MetasoundFrontend/Private/MetasoundFrontendDataTypeRegistry.cpp
--- a/Engine/Plugins/Runtime/Metasound/Source/MetasoundFrontend/Private/MetasoundFrontendDataTypeRegistry.cpp
+++ b/Engine/Plugins/Runtime/Metasound/Source/MetasoundFrontend/Private/MetasoundFrontendDataTypeRegistry.cpp
@@ -1,6 +1,7 @@
 // Copyright Epic Games, Inc. All Rights Reserved.
 #include "MetasoundFrontendDataTypeRegistry.h"
 
+#include "MetasoundLiteral.h"
 #include "MetasoundTrace.h"
 #include "Misc/ScopeLock.h"
 
@@ -1237,17 +1238,29 @@ namespace Metasound
 				// on the input vertex.
 				const FInputDataVertex& InputVertex = Binding.GetVertex();
 				EDataReferenceAccessType AccessType = VertexAccessTypeToDataReferenceAccessType(InputVertex.AccessType);
-				TOptional<FAnyDataReference> DataRef = DataTypeRegistry.CreateDataReference(InputVertex.DataTypeName, AccessType, InputVertex.GetDefaultLiteral(), InOperatorSettings);
 
-				if (DataRef.IsSet())
+				if (const IDataTypeRegistryEntry* Entry = DataTypeRegistry.FindDataTypeRegistryEntry(InputVertex.DataTypeName))
 				{
-					// Set as vertex data reference.
-					OutVertexData.SetVertex(InputVertex.VertexName, *DataRef);
+					TOptional<FAnyDataReference> DataRef = Entry->CreateDataReference(AccessType, InputVertex.GetDefaultLiteral(), InOperatorSettings);
+
+					if (DataRef.IsSet())
+					{
+						// Set as vertex data reference.
+						OutVertexData.SetVertex(InputVertex.VertexName, *DataRef);
+					}
+					else
+					{
+						const FDataTypeRegistryInfo& DataTypeInfo = Entry->GetDataTypeInfo();
+						if (DataTypeInfo.bIsParsable)
+						{
+							// All parsable inputs should have creatable defaults.
+							UE_LOG(LogMetaSound, Warning, TEXT("Failed to create default data reference for vertex %s of data type %s using constructor argument %s"), *InputVertex.VertexName.ToString(), *InputVertex.DataTypeName.ToString(), *::LexToString(InputVertex.GetDefaultLiteral()));
+						}
+					}
 				}
 				else
 				{
-					// All inputs should have creatable defaults.
-					UE_LOG(LogMetaSound, Warning, TEXT("Failed to create default data reference for vertex %s of data type %s"), *InputVertex.VertexName.ToString(), *InputVertex.DataTypeName.ToString());
+					UE_LOG(LogMetaSound, Warning, TEXT("Failed to create default data reference for vertex %s of data type %s because data type is not registered. Please ensure that the plugin which registers the data type is loaded."), *InputVertex.VertexName.ToString(), *InputVertex.DataTypeName.ToString());
 				}
 			}
 		}
