- Fix for missing detach in PostOutputFormatChange case where persistent builder was registered when system is not ready for this (changes could be made by controllers behind the back of the builder system, snuck in from controller mutation delegate notification change)
- Fix crash if attempting to Create<Patch/Source>PresetBuilder with null script interface ref
- Misc clean-up
#rnx

diff --git a/Engine/Plugins/Runtime/Metasound/Source/MetasoundEngine/Private/MetasoundBuilderSubsystem.cpp b/Engine/Plugins/Runtime/Metasound/Source/MetasoundEngine/Private/MetasoundBuilderSubsystem.cpp
--- a/Engine/Plugins/Runtime/Metasound/Source/MetasoundEngine/Private/MetasoundBuilderSubsystem.cpp
+++ b/Engine/Plugins/Runtime/Metasound/Source/MetasoundEngine/Private/MetasoundBuilderSubsystem.cpp
@@ -348,6 +348,56 @@
 	}
 }
 
+void UMetaSoundBuilderBase::ConvertFromPreset(EMetaSoundBuilderResult& OutResult)
+{
+	const bool bSuccess = Builder.ConvertFromPreset();
+	OutResult = bSuccess ? EMetaSoundBuilderResult::Succeeded : EMetaSoundBuilderResult::Failed;
+}
+
+void UMetaSoundBuilderBase::ConvertToPreset(const TScriptInterface<IMetaSoundDocumentInterface>& ReferencedNodeClass, EMetaSoundBuilderResult& OutResult)
+{
+	const IMetaSoundDocumentInterface* ReferencedInterface = ReferencedNodeClass.GetInterface();
+	if (!ReferencedInterface)
+	{
+		OutResult = EMetaSoundBuilderResult::Failed;
+		return;
+	}
+	
+	// Ensure the referenced node class isn't transient 
+	if (Cast<UMetaSoundBuilderDocument>(ReferencedInterface))
+	{
+		UE_LOG(LogMetaSound, Warning, TEXT("Transient document builders cannot be referenced when converting builder '%s' to a preset. Build the referenced node class an asset first or use an existing asset instead"), *GetName());
+		OutResult = EMetaSoundBuilderResult::Failed;
+		return;
+	}
+
+	// Ensure the referenced node class is a matching object type 
+	const UClass& BaseMetaSoundClass = ReferencedInterface->GetBaseMetaSoundUClass();
+	UObject* ReferencedObject = ReferencedNodeClass.GetObject();
+	if (!ReferencedObject || (ReferencedObject && !ReferencedObject->IsA(&BaseMetaSoundClass)))
+	{
+		UE_LOG(LogMetaSound, Warning, TEXT("The referenced node type must match the base MetaSound class when converting builder '%s' to a preset (ex. source preset must reference another source)"), *GetName());
+		OutResult = EMetaSoundBuilderResult::Failed;
+		return;
+	}
+
+	// Ensure the referenced node is registered
+	if (FMetasoundAssetBase* ReferencedMetaSoundAsset = Metasound::IMetasoundUObjectRegistry::Get().GetObjectAsAssetBase(ReferencedObject))
+	{
+		ReferencedMetaSoundAsset->RegisterGraphWithFrontend();
+	}
+
+	const FMetasoundFrontendDocument& ReferencedDocument = ReferencedInterface->GetDocument();
+	if (Builder.ConvertToPreset(ReferencedDocument))
+	{
+		OutResult = EMetaSoundBuilderResult::Succeeded;
+	}
+	else
+	{
+		OutResult = EMetaSoundBuilderResult::Failed;
+	}
+}
+
 void UMetaSoundBuilderBase::DisconnectNodes(const FMetaSoundBuilderNodeOutputHandle& NodeOutputHandle, const FMetaSoundBuilderNodeInputHandle& NodeInputHandle, EMetaSoundBuilderResult& OutResult)
 {
 	const bool bEdgeRemoved = Builder.RemoveEdge(FMetasoundFrontendEdge
@@ -615,6 +665,11 @@
 	return Builder.IsInterfaceDeclared(InterfaceName);
 }
 
+bool UMetaSoundBuilderBase::IsPreset() const
+{
+	return Builder.IsPreset();
+}
+
 bool UMetaSoundBuilderBase::NodesAreConnected(const FMetaSoundBuilderNodeOutputHandle& OutputHandle, const FMetaSoundBuilderNodeInputHandle& InputHandle) const
 {
 	const FMetasoundFrontendEdge Edge = { OutputHandle.NodeID, OutputHandle.VertexID, InputHandle.NodeID, InputHandle.VertexID };
@@ -631,17 +686,6 @@
 	return Builder.IsNodeOutputConnected(OutputHandle.NodeID, OutputHandle.VertexID);
 }
 
-bool UMetaSoundBuilderBase::IsPreset() const
-{
-	return Builder.IsPreset();
-}
-
-void UMetaSoundBuilderBase::ConvertFromPreset(EMetaSoundBuilderResult& OutResult)
-{
-	const bool bSuccess = Builder.ConvertFromPreset();
-	OutResult = bSuccess ? EMetaSoundBuilderResult::Succeeded : EMetaSoundBuilderResult::Failed;
-}
-
 void UMetaSoundBuilderBase::RemoveGraphInput(FName Name, EMetaSoundBuilderResult& OutResult)
 {
 	const bool bRemoved = Builder.RemoveGraphInput(Name);
@@ -809,46 +853,6 @@
 	}
 }
 
-void UMetaSoundBuilderBase::ConvertToPreset(const TScriptInterface<IMetaSoundDocumentInterface>& ReferencedNodeClass, EMetaSoundBuilderResult& OutResult)
-{
-	const IMetaSoundDocumentInterface* ReferencedInterface = ReferencedNodeClass.GetInterface();
-	check(ReferencedInterface);
-	
-	// Ensure the referenced node class isn't transient 
-	if (Cast<UMetaSoundBuilderDocument>(ReferencedInterface))
-	{
-		UE_LOG(LogMetaSound, Warning, TEXT("Transient document builders cannot be referenced when converting builder '%s' to a preset. Build the referenced node class an asset first or use an existing asset instead"), *GetName());
-		OutResult = EMetaSoundBuilderResult::Failed;
-		return;
-	}
-
-	// Ensure the referenced node class is a matching object type 
-	const UClass& BaseMetaSoundClass = ReferencedInterface->GetBaseMetaSoundUClass();
-	UObject* ReferencedObject = ReferencedNodeClass.GetObject();
-	if (!ReferencedObject || (ReferencedObject && !ReferencedObject->IsA(&BaseMetaSoundClass)))
-	{
-		UE_LOG(LogMetaSound, Warning, TEXT("The referenced node type must match the base MetaSound class when converting builder '%s' to a preset (ex. source preset must reference another source)"), *GetName());
-		OutResult = EMetaSoundBuilderResult::Failed;
-		return;
-	}
-
-	// Ensure the referenced node is registered
-	if (FMetasoundAssetBase* ReferencedMetaSoundAsset = Metasound::IMetasoundUObjectRegistry::Get().GetObjectAsAssetBase(ReferencedObject))
-	{
-		ReferencedMetaSoundAsset->RegisterGraphWithFrontend();
-	}
-
-	const FMetasoundFrontendDocument& ReferencedDocument = ReferencedInterface->GetDocument();
-	if (Builder.ConvertToPreset(ReferencedDocument))
-	{
-		OutResult = EMetaSoundBuilderResult::Succeeded;
-	}
-	else
-	{
-		OutResult = EMetaSoundBuilderResult::Failed;
-	}
-}
-
 void UMetaSoundBuilderBase::UpdateDependencyClassNames(const TMap<FMetasoundFrontendClassName, FMetasoundFrontendClassName>& OldToNewReferencedClassNames)
 {
 	Builder.UpdateDependencyClassNames(OldToNewReferencedClassNames);
@@ -1454,9 +1458,15 @@
 
 UMetaSoundPatchBuilder* UMetaSoundBuilderSubsystem::CreatePatchPresetBuilder(FName BuilderName, const TScriptInterface<IMetaSoundDocumentInterface>& ReferencedNodeClass, EMetaSoundBuilderResult& OutResult)
 {
-	UMetaSoundPatchBuilder& Builder = Metasound::Engine::BuilderSubsystemPrivate::CreateTransientBuilder<UMetaSoundPatchBuilder>(BuilderName);
-	Builder.ConvertToPreset(ReferencedNodeClass, OutResult);
-	return &Builder;
+	if (ReferencedNodeClass)
+	{
+		UMetaSoundPatchBuilder& Builder = Metasound::Engine::BuilderSubsystemPrivate::CreateTransientBuilder<UMetaSoundPatchBuilder>(BuilderName);
+		Builder.ConvertToPreset(ReferencedNodeClass, OutResult);
+		return &Builder;
+	}
+
+	OutResult = EMetaSoundBuilderResult::Failed;
+	return nullptr;
 }
 
 UMetaSoundBuilderBase& UMetaSoundBuilderSubsystem::CreatePresetBuilder(FName BuilderName, const TScriptInterface<IMetaSoundDocumentInterface>& ReferencedPatchClass, EMetaSoundBuilderResult& OutResult)
@@ -1479,9 +1489,15 @@
 
 UMetaSoundSourceBuilder* UMetaSoundBuilderSubsystem::CreateSourcePresetBuilder(FName BuilderName, const TScriptInterface<IMetaSoundDocumentInterface>& ReferencedNodeClass, EMetaSoundBuilderResult& OutResult)
 {
-	UMetaSoundSourceBuilder& Builder = Metasound::Engine::BuilderSubsystemPrivate::CreateTransientBuilder<UMetaSoundSourceBuilder>();
-	Builder.ConvertToPreset(ReferencedNodeClass, OutResult);
-	return &Builder;
+	if (ReferencedNodeClass)
+	{
+		UMetaSoundSourceBuilder& Builder = Metasound::Engine::BuilderSubsystemPrivate::CreateTransientBuilder<UMetaSoundSourceBuilder>();
+		Builder.ConvertToPreset(ReferencedNodeClass, OutResult);
+		return &Builder;
+	}
+
+	OutResult = EMetaSoundBuilderResult::Failed;
+	return nullptr;
 }
 
 UMetaSoundBuilderSubsystem& UMetaSoundBuilderSubsystem::GetChecked()
@@ -1559,7 +1575,7 @@
 	return FMetasoundFrontendLiteral { Param };
 }
 
-bool UMetaSoundBuilderSubsystem::DetachBuilderFromAsset(const FMetasoundFrontendClassName& InClassName)
+bool UMetaSoundBuilderSubsystem::DetachBuilderFromAsset(const FMetasoundFrontendClassName& InClassName) const
 {
 	return AssetBuilders.Remove(InClassName.GetFullName()) > 0;
 }
diff --git a/Engine/Plugins/Runtime/Metasound/Source/MetasoundEngine/Private/MetasoundSource.cpp b/Engine/Plugins/Runtime/Metasound/Source/MetasoundEngine/Private/MetasoundSource.cpp
--- a/Engine/Plugins/Runtime/Metasound/Source/MetasoundEngine/Private/MetasoundSource.cpp
+++ b/Engine/Plugins/Runtime/Metasound/Source/MetasoundEngine/Private/MetasoundSource.cpp
@@ -189,9 +189,18 @@ void UMetaSoundSource::PostEditChangeOutputFormat()
 {
 	using namespace Metasound::Frontend;
 
-	UMetaSoundSourceBuilder* SourceBuilder = UMetaSoundBuilderSubsystem::GetConstChecked().AttachSourceBuilderToAsset(this);
-	EMetaSoundBuilderResult Result;
-	SourceBuilder->SetFormat(OutputFormat, Result);
+	EMetaSoundBuilderResult Result = EMetaSoundBuilderResult::Failed;
+	{
+
+		const UMetaSoundBuilderSubsystem& BuilderSubsystem = UMetaSoundBuilderSubsystem::GetConstChecked();
+		UMetaSoundSourceBuilder* SourceBuilder = BuilderSubsystem.AttachSourceBuilderToAsset(this);
+		SourceBuilder->SetFormat(OutputFormat, Result);
+		// TODO: Once builders are notified of controller changes and can be safely persistent, this
+		// can be removed so builders can be shared and not have to be created for each change output
+		// format mutation transaction.
+		BuilderSubsystem.DetachBuilderFromAsset(GetDocument().RootGraph.Metadata.GetClassName());
+	}
+
 	if (Result == EMetaSoundBuilderResult::Succeeded)
 	{
 		// Update the data in this UMetaSoundSource to reflect what is in the metasound document.
diff --git a/Engine/Plugins/Runtime/Metasound/Source/MetasoundEngine/Public/MetasoundBuilderSubsystem.h b/Engine/Plugins/Runtime/Metasound/Source/MetasoundEngine/Public/MetasoundBuilderSubsystem.h
--- a/Engine/Plugins/Runtime/Metasound/Source/MetasoundEngine/Public/MetasoundBuilderSubsystem.h
+++ b/Engine/Plugins/Runtime/Metasound/Source/MetasoundEngine/Public/MetasoundBuilderSubsystem.h
@@ -553,13 +553,13 @@ public:
 
 	UMetaSoundBuilderBase& AttachBuilderToAssetChecked(UObject& InMetaSound) const;
 
-	// May move these to UFUNCTIONS in the future, but right now too many unguarded cases
-	// where a user could mutate these and potentially improperly reference transient assets.
+	// Moving toward becoming UFUNCTIONS in the future, but right now unguarded case
+	// where a user could improperly reference transient assets. Also, unsafe where
+	// legacy controller system are potentially mutating and not notifying delegates
+	// appropriately.
 	UMetaSoundPatchBuilder* AttachPatchBuilderToAsset(UMetaSoundPatch* InPatch) const;
 	UMetaSoundSourceBuilder* AttachSourceBuilderToAsset(UMetaSoundSource* InSource) const;
-
-	// Unregisters builder for MetaSound class asset.
-	bool DetachBuilderFromAsset(const FMetasoundFrontendClassName& InClassName);
+	bool DetachBuilderFromAsset(const FMetasoundFrontendClassName& InClassName) const;
 
 	UFUNCTION(BlueprintCallable, Category = "Audio|MetaSound|Builder",  meta = (ExpandEnumAsExecs = "OutResult"))
 	UPARAM(DisplayName = "Patch Builder") UMetaSoundPatchBuilder* CreatePatchBuilder(FName BuilderName, EMetaSoundBuilderResult& OutResult);
