- Guard against re-initialization of MetaSound doc caches double registering delegate notifiers (happens when clearing a graph)
- Fix copy/pasta in guard against attempting to add output twice (was checking existence of matching input name)
#rnx

diff --git a/Engine/Plugins/Runtime/Metasound/Source/MetasoundFrontend/Private/MetasoundFrontendDocumentBuilder.cpp b/Engine/Plugins/Runtime/Metasound/Source/MetasoundFrontend/Private/MetasoundFrontendDocumentBuilder.cpp
--- a/Engine/Plugins/Runtime/Metasound/Source/MetasoundFrontend/Private/MetasoundFrontendDocumentBuilder.cpp
+++ b/Engine/Plugins/Runtime/Metasound/Source/MetasoundFrontend/Private/MetasoundFrontendDocumentBuilder.cpp
@@ -156,8 +156,8 @@ namespace Metasound::Frontend
 					if (INDEX_NONE != RemoveIndex)
 					{
 						PairedInputs.Add(FVertexPair { InputsToRemove[RemoveIndex], InputsToAdd[AddIndex].Key });
-						InputsToRemove.RemoveAtSwap(RemoveIndex);
-						InputsToAdd.RemoveAtSwap(AddIndex);
+						InputsToRemove.RemoveAtSwap(RemoveIndex, 1, false);
+						InputsToAdd.RemoveAtSwap(AddIndex, 1, false);
 					}
 				}
 
@@ -219,40 +219,52 @@ namespace Metasound::Frontend
 
 			bool RemoveUnsupportedVertices(FMetaSoundFrontendDocumentBuilder& OutBuilder) const
 			{
-				if (!InputsToRemove.IsEmpty() || !OutputsToRemove.IsEmpty())
+				bool bDidEdit = false;
+
+				for (const TPair<FMetasoundFrontendClassInput, const FMetasoundFrontendInterface*>& Pair : InputsToAdd)
 				{
-					auto RemoveMemberNodes = [&OutBuilder](
-						const FMetasoundFrontendClassVertex& ToRemove,
-						TFunctionRef<const FMetasoundFrontendClassVertex*(FName)> FindVertexFunc,
-						TFunctionRef<const FMetasoundFrontendNode* (FName)> FindNodeFunc)
+					if (OutBuilder.RemoveGraphInput(Pair.Key.Name))
 					{
-						if (const FMetasoundFrontendClassVertex* ClassInput = FindVertexFunc(ToRemove.Name))
-						{
-							if (FMetasoundFrontendClassVertex::IsFunctionalEquivalent(*ClassInput, ToRemove))
-							{
-								if (const FMetasoundFrontendNode* Node = FindNodeFunc(ToRemove.Name))
-								{
-									const bool bNodeRemoved = OutBuilder.RemoveNode(Node->GetID());
-									checkf(bNodeRemoved, TEXT("Failed to remove member while attempting to modify MetaSound graph interfaces"));
-								}
-							}
-						}
-					};
+						UE_LOG(LogMetaSound, Warning, TEXT("Removed existing targeted input '%s' to avoid name collision/member data descrepancies while modifying interface(s). Desired edges may have been removed as a result."), *Pair.Key.Name.ToString());
+						bDidEdit = true;
+					}
+				}
+
+				for (const TPair<FMetasoundFrontendClassOutput, const FMetasoundFrontendInterface*>& Pair : OutputsToAdd)
+				{
+					if (OutBuilder.RemoveGraphOutput(Pair.Key.Name))
+					{
+						UE_LOG(LogMetaSound, Warning, TEXT("Removed existing targeted output '%s' to avoid name collision/member data descrepancies while modifying interface(s). Desired edges may have been removed as a result."), *Pair.Key.Name.ToString());
+						bDidEdit = true;
+					}
+				}
 
+				if (!InputsToRemove.IsEmpty() || !OutputsToRemove.IsEmpty())
+				{
 					// Remove unsupported inputs
 					for (const FMetasoundFrontendClassVertex& InputToRemove : InputsToRemove)
 					{
-						RemoveMemberNodes(InputToRemove,
-							[&OutBuilder](FName Name) { return OutBuilder.FindGraphInput(Name); },
-							[&OutBuilder](FName Name) { return OutBuilder.FindGraphInputNode(Name); });
+						if (OutBuilder.RemoveGraphInput(InputToRemove.Name))
+						{
+							bDidEdit = true;
+						}
+						else
+						{
+							UE_LOG(LogMetaSound, Warning, TEXT("Failed to remove existing input '%s', which was an expected member of a removed interface."), *InputToRemove.Name.ToString());
+						}
 					}
 
 					// Remove unsupported outputs
 					for (const FMetasoundFrontendClassVertex& OutputToRemove : OutputsToRemove)
 					{
-						RemoveMemberNodes(OutputToRemove,
-							[&OutBuilder](FName Name) { return OutBuilder.FindGraphOutput(Name); },
-							[&OutBuilder](FName Name) { return OutBuilder.FindGraphOutputNode(Name); });
+						if (OutBuilder.RemoveGraphOutput(OutputToRemove.Name))
+						{
+							bDidEdit = true;
+						}
+						else
+						{
+							UE_LOG(LogMetaSound, Warning, TEXT("Failed to remove existing output '%s', which was an expected member of a removed interface."), *OutputToRemove.Name.ToString());
+						}
 					}
 
 					return true;
@@ -894,10 +906,10 @@ const FMetasoundFrontendNode* FMetaSoundFrontendDocumentBuilder::AddGraphOutput(
 		UE_LOG(LogMetaSound, Error, TEXT("TypeName unset when attempting to add class output '%s'"), *InClassOutput.Name.ToString());
 		return nullptr;
 	}
-	else if (const FMetasoundFrontendClassInput* Input = DocumentCache->GetInterfaceCache().FindInput(InClassOutput.Name))
+	else if (const FMetasoundFrontendClassOutput* Output = DocumentCache->GetInterfaceCache().FindOutput(InClassOutput.Name))
 	{
 		UE_LOG(LogMetaSound, Error, TEXT("Attempting to add MetaSound graph output '%s' when output with name already exists"), *InClassOutput.Name.ToString());
-		return DocumentCache->GetNodeCache().FindNode(Input->NodeID);
+		return DocumentCache->GetNodeCache().FindNode(Output->NodeID);
 	}
 
 	auto FindRegistryClass = [&InClassOutput](FMetasoundFrontendClass& OutClass) -> bool
diff --git a/Engine/Plugins/Runtime/Metasound/Source/MetasoundFrontend/Private/MetasoundFrontendDocumentCache.cpp b/Engine/Plugins/Runtime/Metasound/Source/MetasoundFrontend/Private/MetasoundFrontendDocumentCache.cpp
--- a/Engine/Plugins/Runtime/Metasound/Source/MetasoundFrontend/Private/MetasoundFrontendDocumentCache.cpp
+++ b/Engine/Plugins/Runtime/Metasound/Source/MetasoundFrontend/Private/MetasoundFrontendDocumentCache.cpp
@@ -34,9 +34,13 @@ namespace Metasound::Frontend
 			IDToIndex.Add(Class.ID, Index);
 		}
 
-		OutDelegates.OnDependencyAdded.AddSP(this, &FDocumentCache::OnDependencyAdded);
-		OutDelegates.OnRemovingDependency.AddSP(this, &FDocumentCache::OnRemovingDependency);
-		OutDelegates.OnRenamingDependencyClass.AddSP(this, &FDocumentCache::OnRenamingDependencyClass);
+		OutDelegates.OnDependencyAdded.Remove(OnDependencyAddedHandle);
+		OutDelegates.OnRemovingDependency.Remove(OnRemovingDependencyHandle);
+		OutDelegates.OnRenamingDependencyClass.Remove(OnRenamingDependencyClassHandle);
+
+		OnDependencyAddedHandle = OutDelegates.OnDependencyAdded.AddSP(this, &FDocumentCache::OnDependencyAdded);
+		OnRemovingDependencyHandle = OutDelegates.OnRemovingDependency.AddSP(this, &FDocumentCache::OnRemovingDependency);
+		OnRenamingDependencyClassHandle = OutDelegates.OnRenamingDependencyClass.AddSP(this, &FDocumentCache::OnRenamingDependencyClass);
 
 		if (!EdgeCache.IsValid())
 		{
@@ -192,10 +196,15 @@ namespace Metasound::Frontend
 			OutputNameToIndex.Add(Outputs[Index].Name, Index);
 		};
 
-		OutDelegates.OnInputAdded.AddSP(this, &FDocumentGraphInterfaceCache::OnInputAdded);
-		OutDelegates.OnOutputAdded.AddSP(this, &FDocumentGraphInterfaceCache::OnOutputAdded);
-		OutDelegates.OnRemovingInput.AddSP(this, &FDocumentGraphInterfaceCache::OnRemovingInput);
-		OutDelegates.OnRemovingOutput.AddSP(this, &FDocumentGraphInterfaceCache::OnRemovingOutput);
+		OutDelegates.OnInputAdded.Remove(OnInputAddedHandle);
+		OutDelegates.OnOutputAdded.Remove(OnOutputAddedHandle);
+		OutDelegates.OnRemovingInput.Remove(OnRemovingInputHandle);
+		OutDelegates.OnRemovingOutput.Remove(OnRemovingOutputHandle);
+
+		OnInputAddedHandle = OutDelegates.OnInputAdded.AddSP(this, &FDocumentGraphInterfaceCache::OnInputAdded);
+		OnOutputAddedHandle = OutDelegates.OnOutputAdded.AddSP(this, &FDocumentGraphInterfaceCache::OnOutputAdded);
+		OnRemovingInputHandle = OutDelegates.OnRemovingInput.AddSP(this, &FDocumentGraphInterfaceCache::OnRemovingInput);
+		OnRemovingOutputHandle = OutDelegates.OnRemovingOutput.AddSP(this, &FDocumentGraphInterfaceCache::OnRemovingOutput);
 	}
 
 	void FDocumentGraphInterfaceCache::OnInputAdded(int32 NewIndex)
@@ -253,8 +262,11 @@ namespace Metasound::Frontend
 			ClassIDToNodeIndices.FindOrAdd(Node.ClassID).Add(Index);
 		}
 
-		OutDelegates.OnNodeAdded.AddSP(this, &FDocumentGraphNodeCache::OnNodeAdded);
-		OutDelegates.OnRemovingNode.AddSP(this, &FDocumentGraphNodeCache::OnRemovingNode);
+		OutDelegates.OnNodeAdded.Remove(OnNodeAddedHandle);
+		OnNodeAddedHandle = OutDelegates.OnNodeAdded.AddSP(this, &FDocumentGraphNodeCache::OnNodeAdded);
+
+		OutDelegates.OnRemovingNode.Remove(OnRemovingNodeHandle);
+		OnRemovingNodeHandle = OutDelegates.OnRemovingNode.AddSP(this, &FDocumentGraphNodeCache::OnRemovingNode);
 	}
 
 	bool FDocumentGraphNodeCache::ContainsNode(const FGuid& InNodeID) const
@@ -427,8 +439,11 @@ namespace Metasound::Frontend
 			InputToEdgeIndex.Add(Edge.GetToVertexHandle()) = Index;
 		}
 
-		OutDelegates.OnEdgeAdded.AddSP(this, &FDocumentGraphEdgeCache::OnEdgeAdded);
-		OutDelegates.OnRemovingEdge.AddSP(this, &FDocumentGraphEdgeCache::OnRemovingEdge);
+		OutDelegates.OnEdgeAdded.Remove(OnAddedHandle);
+		OnAddedHandle = OutDelegates.OnEdgeAdded.AddSP(this, &FDocumentGraphEdgeCache::OnEdgeAdded);
+
+		OutDelegates.OnRemovingEdge.Remove(OnRemovedHandle);
+		OnRemovedHandle = OutDelegates.OnRemovingEdge.AddSP(this, &FDocumentGraphEdgeCache::OnRemovingEdge);
 	}
 
 	bool FDocumentGraphEdgeCache::ContainsEdge(const FMetasoundFrontendEdge& InEdge) const
@@ -490,8 +505,8 @@ namespace Metasound::Frontend
 		const FMetasoundFrontendDocument& Document = Parent->GetDocument();
 		const FMetasoundFrontendGraph& Graph = Document.RootGraph.Graph;
 		const FMetasoundFrontendEdge& NewEdge = Graph.Edges[InNewIndex];
-		InputToEdgeIndex.Add(FMetasoundFrontendVertexHandle { NewEdge.ToNodeID, NewEdge.ToVertexID }, InNewIndex);
-		OutputToEdgeIndices.FindOrAdd({ NewEdge.FromNodeID, NewEdge.FromVertexID }).Add(InNewIndex);
+		InputToEdgeIndex.Add(NewEdge.GetToVertexHandle(), InNewIndex);
+		OutputToEdgeIndices.FindOrAdd(NewEdge.GetFromVertexHandle()).Add(InNewIndex);
 	}
 
 	void FDocumentGraphEdgeCache::OnRemovingEdge(int32 IndexToRemove)
@@ -499,12 +514,12 @@ namespace Metasound::Frontend
 		const FMetasoundFrontendDocument& Document = Parent->GetDocument();
 		const FMetasoundFrontendGraph& Graph = Document.RootGraph.Graph;
 		const FMetasoundFrontendEdge& EdgeToRemove = Graph.Edges[IndexToRemove];
-		InputToEdgeIndex.Remove(FMetasoundFrontendVertexHandle { EdgeToRemove.ToNodeID, EdgeToRemove.ToVertexID });
+		InputToEdgeIndex.Remove(EdgeToRemove.GetToVertexHandle());
 		TArray<int32>& Indices = OutputToEdgeIndices.FindChecked(EdgeToRemove.GetFromVertexHandle());
 		Indices.RemoveAllSwap([&IndexToRemove](const int32& Index) { return Index == IndexToRemove; }, false /* bAllowShrinking */);
 		if (Indices.IsEmpty())
 		{
-			OutputToEdgeIndices.Remove(FMetasoundFrontendVertexHandle { EdgeToRemove.FromNodeID, EdgeToRemove.FromVertexID });
+			OutputToEdgeIndices.Remove(EdgeToRemove.GetFromVertexHandle());
 		}
 	}
 } // namespace Metasound::Frontend
diff --git a/Engine/Plugins/Runtime/Metasound/Source/MetasoundFrontend/Private/MetasoundFrontendDocumentCache.h b/Engine/Plugins/Runtime/Metasound/Source/MetasoundFrontend/Private/MetasoundFrontendDocumentCache.h
--- a/Engine/Plugins/Runtime/Metasound/Source/MetasoundFrontend/Private/MetasoundFrontendDocumentCache.h
+++ b/Engine/Plugins/Runtime/Metasound/Source/MetasoundFrontend/Private/MetasoundFrontendDocumentCache.h
@@ -46,6 +46,9 @@ namespace Metasound::Frontend
 		TMap<FMetasoundFrontendVertexHandle, int32> InputToEdgeIndex;
 
 		TSharedPtr<IDocumentCache> Parent;
+
+		FDelegateHandle OnAddedHandle;
+		FDelegateHandle OnRemovedHandle;
 	};
 
 
@@ -85,6 +88,9 @@ namespace Metasound::Frontend
 		TSortedMap<FGuid, TArray<int32>> ClassIDToNodeIndices;
 
 		TSharedPtr<IDocumentCache> Parent;
+
+		FDelegateHandle OnNodeAddedHandle;
+		FDelegateHandle OnRemovingNodeHandle;
 	};
 
 
@@ -114,6 +120,11 @@ namespace Metasound::Frontend
 		TMap<FName, int32> OutputNameToIndex;
 
 		TSharedPtr<IDocumentCache> Parent;
+
+		FDelegateHandle OnInputAddedHandle;
+		FDelegateHandle OnOutputAddedHandle;
+		FDelegateHandle OnRemovingInputHandle;
+		FDelegateHandle OnRemovingOutputHandle;
 	};
 
 
@@ -153,6 +164,10 @@ namespace Metasound::Frontend
 		TSharedPtr<FDocumentGraphNodeCache> NodeCache;
 		TSharedPtr<FDocumentGraphInterfaceCache> InterfaceCache;
 
+		FDelegateHandle OnDependencyAddedHandle;
+		FDelegateHandle OnRemovingDependencyHandle;
+		FDelegateHandle OnRenamingDependencyClassHandle;
+
 		const FMetasoundFrontendDocument* Document = nullptr;
 	};
 } // namespace Metasound::Frontend
