Minor code cleanup fixing dynamic operator audio fade.

diff --git a/Engine/Plugins/Runtime/Metasound/Source/MetasoundGraphCore/Private/MetasoundDynamicOperatorAudioFade.h b/Engine/Plugins/Runtime/Metasound/Source/MetasoundGraphCore/Private/MetasoundDynamicOperatorAudioFade.h
--- a/Engine/Plugins/Runtime/Metasound/Source/MetasoundGraphCore/Private/MetasoundDynamicOperatorAudioFade.h
+++ b/Engine/Plugins/Runtime/Metasound/Source/MetasoundGraphCore/Private/MetasoundDynamicOperatorAudioFade.h
@@ -21,7 +21,7 @@ namespace Metasound::DynamicGraph
 	// wrapper. 
 	//
 	// The wrapped operator can be release if desired, though it leaves this operator
-	// in a zombie state.
+	// in a state where it will have no effect on any audio buffers.
 	class FAudioFadeOperatorWrapper : public TExecutableOperator<FAudioFadeOperatorWrapper>
 	{
 	public:
diff --git a/Engine/Plugins/Runtime/Metasound/Source/MetasoundGraphCore/Private/MetasoundDynamicOperatorTransactor.cpp b/Engine/Plugins/Runtime/Metasound/Source/MetasoundGraphCore/Private/MetasoundDynamicOperatorTransactor.cpp
--- a/Engine/Plugins/Runtime/Metasound/Source/MetasoundGraphCore/Private/MetasoundDynamicOperatorTransactor.cpp
+++ b/Engine/Plugins/Runtime/Metasound/Source/MetasoundGraphCore/Private/MetasoundDynamicOperatorTransactor.cpp
@@ -4,7 +4,9 @@
 
 #include "Containers/SpscQueue.h"
 #include "Containers/UnrealString.h"
+#include "MetasoundAudioBuffer.h"
 #include "MetasoundBuilderInterface.h"
+#include "MetasoundDataReference.h"
 #include "MetasoundDynamicOperator.h"
 #include "MetasoundGraph.h"
 #include "MetasoundGraphAlgo.h"
@@ -42,8 +44,6 @@ namespace Metasound
 				return OperatorOrder;
 			}
 
-			static const FLazyName AudioBufferDataTypeName("Audio");
-
 			FString GetDebugNodeNameString(const INode& InNode)
 			{
 				return FString::Printf(TEXT("%s_v%d.%d"), *InNode.GetMetadata().ClassName.GetFullName().ToString(), InNode.GetMetadata().MajorVersion, InNode.GetMetadata().MinorVersion);
@@ -124,7 +124,7 @@ namespace Metasound
 				const FOutputVertexInterface& OutputInterface = Node->GetVertexInterface().GetOutputInterface();
 				for (const FOutputDataVertex& OutputVertex : OutputInterface)
 				{
-					if (OutputVertex.DataTypeName == DynamicOperatorTransactorPrivate::AudioBufferDataTypeName)
+					if (OutputVertex.DataTypeName == GetMetasoundDataTypeName<FAudioBuffer>())
 					{
 						OutputsToFade.Add(OutputVertex.VertexName);
 					}
@@ -216,7 +216,7 @@ namespace Metasound
 			/* enqueue an update. */
 			TArray<FOperatorID> OperatorOrder = DynamicOperatorTransactorPrivate::DetermineOperatorOrder(Graph);
 
-			if (ToNode->GetVertexInterface().GetInputVertex(InToVertex).DataTypeName == DynamicOperatorTransactorPrivate::AudioBufferDataTypeName)
+			if (ToNode->GetVertexInterface().GetInputVertex(InToVertex).DataTypeName == GetMetasoundDataTypeName<FAudioBuffer>())
 			{
 				// Handle audio edge removal with a fade out.
 				EnqueueFadeAndRemoveEdgeOperatorTransform(*FromNode, InFromVertex, *ToNode, InToVertex, *ReplacementLiteralNode, OperatorOrder);
@@ -399,7 +399,7 @@ namespace Metasound
 			// Find order of operators after adding edge. 
 			TArray<FOperatorID> OperatorOrder = DynamicOperatorTransactorPrivate::DetermineOperatorOrder(Graph);
 			
-			if (InputVertex->DataTypeName == DynamicOperatorTransactorPrivate::AudioBufferDataTypeName)
+			if (InputVertex->DataTypeName == GetMetasoundDataTypeName<FAudioBuffer>())
 			{
 				// If edge is audio, then the connection needs to be faded
 				EnqueueFadeAndAddEdgeOperatorTransform(InFromNode, InFromVertex, InToNode, InToVertex, PriorLiteralNode.Get(), OperatorOrder);
