Add/clean-up missing or improper category identifiers on MetaSound node classes
- Make Wave Writer/Players follow alphabetical expectation to avoid most used implementations (mono/stereo) from being at bottom of node picker list
#rnx

diff --git a/Engine/Plugins/Runtime/AudioModulation/Source/AudioModulation/Private/MetaSoundNodes/GetModulatorValueNode.cpp b/Engine/Plugins/Runtime/AudioModulation/Source/AudioModulation/Private/MetaSoundNodes/GetModulatorValueNode.cpp
--- a/Engine/Plugins/Runtime/AudioModulation/Source/AudioModulation/Private/MetaSoundNodes/GetModulatorValueNode.cpp
+++ b/Engine/Plugins/Runtime/AudioModulation/Source/AudioModulation/Private/MetaSoundNodes/GetModulatorValueNode.cpp
@@ -49,7 +49,7 @@ namespace AudioModulation
 					AudioModulation::PluginAuthor,
 					AudioModulation::PluginNodeMissingPrompt,
 					GetDefaultInterface(),
-					{ },
+					{ AudioModulation::PluginNodeCategory },
 					{ },
 					{ }
 				};
diff --git a/Engine/Plugins/Runtime/AudioModulation/Source/AudioModulation/Private/MetaSoundNodes/MixModulatorsNode.cpp b/Engine/Plugins/Runtime/AudioModulation/Source/AudioModulation/Private/MetaSoundNodes/MixModulatorsNode.cpp
--- a/Engine/Plugins/Runtime/AudioModulation/Source/AudioModulation/Private/MetaSoundNodes/MixModulatorsNode.cpp
+++ b/Engine/Plugins/Runtime/AudioModulation/Source/AudioModulation/Private/MetaSoundNodes/MixModulatorsNode.cpp
@@ -52,7 +52,7 @@ namespace AudioModulation
 					AudioModulation::PluginAuthor,
 					AudioModulation::PluginNodeMissingPrompt,
 					GetDefaultInterface(),
-					{ },
+					{ AudioModulation::PluginNodeCategory },
 					{ },
 					{ }
 				};
diff --git a/Engine/Plugins/Runtime/AudioModulation/Source/AudioModulation/Private/SoundModulatorAsset.cpp b/Engine/Plugins/Runtime/AudioModulation/Source/AudioModulation/Private/SoundModulatorAsset.cpp
--- a/Engine/Plugins/Runtime/AudioModulation/Source/AudioModulation/Private/SoundModulatorAsset.cpp
+++ b/Engine/Plugins/Runtime/AudioModulation/Source/AudioModulation/Private/SoundModulatorAsset.cpp
@@ -7,6 +7,7 @@
 namespace AudioModulation
 {
 	const FString PluginAuthor = TEXT("Epic Games, Inc.");
+	const FText PluginNodeCategory = NSLOCTEXT("AudioModulation", "AudioModulationNodeCategory", "Modulation");
 	const FText PluginNodeMissingPrompt = NSLOCTEXT("AudioModulation", "DefaultMissingNodePrompt", "The node was likely removed, renamed, or the AudioModulation plugin is not loaded.");
 
 	FSoundModulatorAsset::FSoundModulatorAsset(const TSharedPtr<Audio::IProxyData>& InProxyPtr)
diff --git a/Engine/Plugins/Runtime/AudioModulation/Source/AudioModulation/Public/SoundModulatorAsset.h b/Engine/Plugins/Runtime/AudioModulation/Source/AudioModulation/Public/SoundModulatorAsset.h
--- a/Engine/Plugins/Runtime/AudioModulation/Source/AudioModulation/Public/SoundModulatorAsset.h
+++ b/Engine/Plugins/Runtime/AudioModulation/Source/AudioModulation/Public/SoundModulatorAsset.h
@@ -11,6 +11,7 @@
 namespace AudioModulation
 {
 	extern const FString AUDIOMODULATION_API PluginAuthor;
+	extern const FText AUDIOMODULATION_API PluginNodeCategory;
 	extern const FText AUDIOMODULATION_API PluginNodeMissingPrompt;
 
 	class AUDIOMODULATION_API FSoundModulatorAsset
diff --git a/Engine/Plugins/Runtime/Metasound/Source/MetasoundEngine/Private/MetasoundWavePlayerNode.cpp b/Engine/Plugins/Runtime/Metasound/Source/MetasoundEngine/Private/MetasoundWavePlayerNode.cpp
--- a/Engine/Plugins/Runtime/Metasound/Source/MetasoundEngine/Private/MetasoundWavePlayerNode.cpp
+++ b/Engine/Plugins/Runtime/Metasound/Source/MetasoundEngine/Private/MetasoundWavePlayerNode.cpp
@@ -1194,7 +1194,7 @@ namespace Metasound
 
 	struct FMonoAudioChannelConfigurationInfo
 	{
-		static FText GetNodeDisplayName() { return METASOUND_LOCTEXT("Metasound_WavePlayerMonoNodeDisplayName", "Wave Player (Mono)"); }
+		static FText GetNodeDisplayName() { return METASOUND_LOCTEXT("Metasound_WavePlayerMonoNodeDisplayName", "Wave Player (1.0, Mono)"); }
 		static FName GetVariantName() { return Metasound::EngineNodes::MonoVariant; }
 
 		static TArray<FOutputDataVertex> GetAudioOutputs()
@@ -1210,7 +1210,7 @@ namespace Metasound
 
 	struct FStereoAudioChannelConfigurationInfo
 	{
-		static FText GetNodeDisplayName() { return METASOUND_LOCTEXT("Metasound_WavePlayerStereoNodeDisplayName", "Wave Player (Stereo)"); }
+		static FText GetNodeDisplayName() { return METASOUND_LOCTEXT("Metasound_WavePlayerStereoNodeDisplayName", "Wave Player (2.0, Stereo)"); }
 		static FName GetVariantName() { return Metasound::EngineNodes::StereoVariant; }
 
 		static TArray<FOutputDataVertex> GetAudioOutputs()
@@ -1227,7 +1227,7 @@ namespace Metasound
 
 	struct FQuadAudioChannelConfigurationInfo
 	{
-		static FText GetNodeDisplayName() { return METASOUND_LOCTEXT("Metasound_WavePlayerQuadNodeDisplayName", "Wave Player (Quad)"); }
+		static FText GetNodeDisplayName() { return METASOUND_LOCTEXT("Metasound_WavePlayerQuadNodeDisplayName", "Wave Player (4.0, Quad)"); }
 		static FName GetVariantName() { return Metasound::EngineNodes::QuadVariant; }
 
 		static TArray<FOutputDataVertex> GetAudioOutputs()
@@ -1246,7 +1246,7 @@ namespace Metasound
 
 	struct FFiveDotOneAudioChannelConfigurationInfo
 	{
-		static FText GetNodeDisplayName() { return METASOUND_LOCTEXT("Metasound_WavePlayerFiveDotOneNodeDisplayName", "Wave Player (5.1)"); }
+		static FText GetNodeDisplayName() { return METASOUND_LOCTEXT("Metasound_WavePlayerFiveDotOneNodeDisplayName", "Wave Player (5.1, Surround)"); }
 		static FName GetVariantName() { return Metasound::EngineNodes::FiveDotOneVariant; }
 
 		static TArray<FOutputDataVertex> GetAudioOutputs()
@@ -1267,7 +1267,7 @@ namespace Metasound
 
 	struct FSevenDotOneAudioChannelConfigurationInfo
 	{
-		static FText GetNodeDisplayName() { return METASOUND_LOCTEXT("Metasound_WavePlayerSevenDotOneNodeDisplayName", "Wave Player (7.1)"); }
+		static FText GetNodeDisplayName() { return METASOUND_LOCTEXT("Metasound_WavePlayerSevenDotOneNodeDisplayName", "Wave Player (7.1, Surround)"); }
 		static FName GetVariantName() { return Metasound::EngineNodes::SevenDotOneVariant; }
 
 		static TArray<FOutputDataVertex> GetAudioOutputs()
diff --git a/Engine/Plugins/Runtime/Metasound/Source/MetasoundEngine/Private/MetasoundWaveTableEnvelopeNode.cpp b/Engine/Plugins/Runtime/Metasound/Source/MetasoundEngine/Private/MetasoundWaveTableEnvelopeNode.cpp
--- a/Engine/Plugins/Runtime/Metasound/Source/MetasoundEngine/Private/MetasoundWaveTableEnvelopeNode.cpp
+++ b/Engine/Plugins/Runtime/Metasound/Source/MetasoundEngine/Private/MetasoundWaveTableEnvelopeNode.cpp
@@ -84,8 +84,8 @@ namespace Metasound
 					PluginAuthor,
 					PluginNodeMissingPrompt,
 					GetDefaultInterface(),
-					{ NodeCategories::Envelopes },
-					{ METASOUND_LOCTEXT("WaveTableEnvelopeSynthesisKeyword", "Synthesis")},
+					{ NodeCategories::WaveTables },
+					{ NodeCategories::Generators, METASOUND_LOCTEXT("WaveTableEnvelopeSynthesisKeyword", "Synthesis")},
 					{ }
 				};
 
diff --git a/Engine/Plugins/Runtime/Metasound/Source/MetasoundEngine/Private/MetasoundWaveTableEvaluateNode.cpp b/Engine/Plugins/Runtime/Metasound/Source/MetasoundEngine/Private/MetasoundWaveTableEvaluateNode.cpp
--- a/Engine/Plugins/Runtime/Metasound/Source/MetasoundEngine/Private/MetasoundWaveTableEvaluateNode.cpp
+++ b/Engine/Plugins/Runtime/Metasound/Source/MetasoundEngine/Private/MetasoundWaveTableEvaluateNode.cpp
@@ -65,8 +65,8 @@ namespace Metasound
 					PluginAuthor,
 					PluginNodeMissingPrompt,
 					GetDefaultInterface(),
-					{ NodeCategories::Envelopes },
-					{ METASOUND_LOCTEXT("WaveTableEvaluateEnvelopeKeyword", "Envelope"), METASOUND_LOCTEXT("WaveTableEvaluateCurveKeyword", "Curve")},
+					{ NodeCategories::WaveTables },
+					{ NodeCategories::Envelopes, METASOUND_LOCTEXT("WaveTableEvaluateCurveKeyword", "Curve")},
 					{ }
 				};
 
diff --git a/Engine/Plugins/Runtime/Metasound/Source/MetasoundEngine/Private/MetasoundWaveTableGetNode.cpp b/Engine/Plugins/Runtime/Metasound/Source/MetasoundEngine/Private/MetasoundWaveTableGetNode.cpp
--- a/Engine/Plugins/Runtime/Metasound/Source/MetasoundEngine/Private/MetasoundWaveTableGetNode.cpp
+++ b/Engine/Plugins/Runtime/Metasound/Source/MetasoundEngine/Private/MetasoundWaveTableGetNode.cpp
@@ -8,6 +8,7 @@
 #include "MetasoundFacade.h"
 #include "MetasoundNodeRegistrationMacro.h"
 #include "MetasoundPrimitives.h"
+#include "MetasoundStandardNodesCategories.h"
 #include "MetasoundTrace.h"
 #include "MetasoundWaveTable.h"
 #include "WaveTable.h"
@@ -56,7 +57,7 @@ namespace Metasound
 					PluginAuthor,
 					PluginNodeMissingPrompt,
 					GetDefaultInterface(),
-					{ },
+					{ NodeCategories::WaveTables },
 					{ },
 					{ }
 				};
diff --git a/Engine/Plugins/Runtime/Metasound/Source/MetasoundEngine/Private/MetasoundWaveTableOscillatorNode.cpp b/Engine/Plugins/Runtime/Metasound/Source/MetasoundEngine/Private/MetasoundWaveTableOscillatorNode.cpp
--- a/Engine/Plugins/Runtime/Metasound/Source/MetasoundEngine/Private/MetasoundWaveTableOscillatorNode.cpp
+++ b/Engine/Plugins/Runtime/Metasound/Source/MetasoundEngine/Private/MetasoundWaveTableOscillatorNode.cpp
@@ -67,7 +67,7 @@ namespace Metasound
 					PluginAuthor,
 					PluginNodeMissingPrompt,
 					GetDefaultInterface(),
-					{ NodeCategories::Generators },
+					{ NodeCategories::WaveTables },
 					{ METASOUND_LOCTEXT("WaveTableOscillatorSynthesisKeyword", "Synthesis")},
 					{ }
 				};
diff --git a/Engine/Plugins/Runtime/Metasound/Source/MetasoundEngine/Private/MetasoundWaveTablePlayerNode.cpp b/Engine/Plugins/Runtime/Metasound/Source/MetasoundEngine/Private/MetasoundWaveTablePlayerNode.cpp
--- a/Engine/Plugins/Runtime/Metasound/Source/MetasoundEngine/Private/MetasoundWaveTablePlayerNode.cpp
+++ b/Engine/Plugins/Runtime/Metasound/Source/MetasoundEngine/Private/MetasoundWaveTablePlayerNode.cpp
@@ -92,8 +92,8 @@ namespace Metasound
 					PluginAuthor,
 					PluginNodeMissingPrompt,
 					GetDefaultInterface(),
-					{ NodeCategories::Generators },
-					{ METASOUND_LOCTEXT("WaveTablePlayerSynthesisKeyword", "Synthesis")},
+					{ NodeCategories::WaveTables },
+					{ NodeCategories::Generators, METASOUND_LOCTEXT("WaveTablePlayerSynthesisKeyword", "Synthesis")},
 					{ }
 				};
 
diff --git a/Engine/Plugins/Runtime/Metasound/Source/MetasoundStandardNodes/Private/MetasoundBPMToSecondsNode.cpp b/Engine/Plugins/Runtime/Metasound/Source/MetasoundStandardNodes/Private/MetasoundBPMToSecondsNode.cpp
--- a/Engine/Plugins/Runtime/Metasound/Source/MetasoundStandardNodes/Private/MetasoundBPMToSecondsNode.cpp
+++ b/Engine/Plugins/Runtime/Metasound/Source/MetasoundStandardNodes/Private/MetasoundBPMToSecondsNode.cpp
@@ -4,6 +4,7 @@
 
 #include "Internationalization/Text.h"
 #include "MetasoundExecutableOperator.h"
+#include "MetasoundFrontendNodesCategories.h"
 #include "MetasoundNodeRegistrationMacro.h"
 #include "MetasoundParamHelper.h"
 #include "MetasoundPrimitives.h"
@@ -166,6 +167,7 @@ namespace Metasound
 			Info.DisplayName = METASOUND_LOCTEXT("BPMToSecondsNode_DisplayName", "BPM To Seconds");
 			Info.Description = METASOUND_LOCTEXT("BPMToSecondsNode_Desc", "Calculates a beat time in seconds from the given BPM, beat multiplier and divisions of a whole note.");
 			Info.Author = PluginAuthor;
+			Info.CategoryHierarchy = { NodeCategories::Conversions };
 			Info.PromptIfMissing = PluginNodeMissingPrompt;
 			Info.DefaultInterface = GetVertexInterface();
 			Info.Keywords = { METASOUND_LOCTEXT("BPMNodeRhythmKeyword", "Rhythm"), METASOUND_LOCTEXT("BPMNodeTempoKeyword", "Tempo") };
diff --git a/Engine/Plugins/Runtime/Metasound/Source/MetasoundStandardNodes/Private/MetasoundEnvelopeFollowerNode.cpp b/Engine/Plugins/Runtime/Metasound/Source/MetasoundStandardNodes/Private/MetasoundEnvelopeFollowerNode.cpp
--- a/Engine/Plugins/Runtime/Metasound/Source/MetasoundStandardNodes/Private/MetasoundEnvelopeFollowerNode.cpp
+++ b/Engine/Plugins/Runtime/Metasound/Source/MetasoundStandardNodes/Private/MetasoundEnvelopeFollowerNode.cpp
@@ -3,6 +3,7 @@
 #include "MetasoundEnvelopeFollowerNode.h"
 
 #include "Algo/MaxElement.h"
+#include "DSP/EnvelopeFollower.h"
 #include "Internationalization/Text.h"
 #include "MetasoundAudioBuffer.h"
 #include "MetasoundEnvelopeFollowerTypes.h"
@@ -12,16 +13,17 @@
 #include "MetasoundNodeRegistrationMacro.h"
 #include "MetasoundDataTypeRegistrationMacro.h"
 #include "MetasoundOperatorSettings.h"
+#include "MetasoundParamHelper.h"
 #include "MetasoundPrimitives.h"
+#include "MetasoundStandardNodesCategories.h"
 #include "MetasoundStandardNodesNames.h"
 #include "MetasoundTrigger.h"
 #include "MetasoundTime.h"
 #include "MetasoundVertex.h"
-#include "MetasoundParamHelper.h"
-#include "DSP/EnvelopeFollower.h"
 
 #define LOCTEXT_NAMESPACE "MetasoundStandardNodes_EnvelopeFollower"
 
+
 namespace Metasound
 {
 	namespace EnvelopeFollowerVertexNames
@@ -224,6 +226,7 @@ namespace Metasound
 			Info.DisplayName = METASOUND_LOCTEXT("Metasound_EnvelopeFollowerDisplayName", "Envelope Follower");
 			Info.Description = METASOUND_LOCTEXT("Metasound_EnvelopeFollowerDescription", "Outputs an envelope from an input audio signal.");
 			Info.Author = PluginAuthor;
+			Info.CategoryHierarchy = { NodeCategories::Envelopes };
 			Info.PromptIfMissing = PluginNodeMissingPrompt;
 			Info.DefaultInterface = GetVertexInterface();
 
diff --git a/Engine/Plugins/Runtime/Metasound/Source/MetasoundStandardNodes/Private/MetasoundFlangerNode.cpp b/Engine/Plugins/Runtime/Metasound/Source/MetasoundStandardNodes/Private/MetasoundFlangerNode.cpp
--- a/Engine/Plugins/Runtime/Metasound/Source/MetasoundStandardNodes/Private/MetasoundFlangerNode.cpp
+++ b/Engine/Plugins/Runtime/Metasound/Source/MetasoundStandardNodes/Private/MetasoundFlangerNode.cpp
@@ -1,18 +1,19 @@
 // Copyright Epic Games, Inc. All Rights Reserved.
 
+#include "DSP/BufferVectorOperations.h"
+#include "DSP/Flanger.h"
 #include "Internationalization/Text.h"
-#include "MetasoundFacade.h"
 #include "MetasoundExecutableOperator.h"
+#include "MetasoundFacade.h"
 #include "MetasoundNodeRegistrationMacro.h"
+#include "MetasoundParamHelper.h"
 #include "MetasoundPrimitives.h"
-#include "MetasoundStandardNodesNames.h"
 #include "MetasoundStandardNodesCategories.h"
-#include "DSP/BufferVectorOperations.h"
-#include "DSP/Flanger.h"
-#include "MetasoundParamHelper.h"
+#include "MetasoundStandardNodesNames.h"
 
 #define LOCTEXT_NAMESPACE "MetasoundStandardNodes_FlangerNode"
 
+
 namespace Metasound
 {
 	namespace FlangerVertexNames
@@ -45,6 +46,7 @@ namespace Metasound
 				Info.DisplayName = NodeDisplayName;
 				Info.Description = NodeDescription;
 				Info.Author = PluginAuthor;
+				Info.CategoryHierarchy = { NodeCategories::Filters };
 				Info.PromptIfMissing = PluginNodeMissingPrompt;
 				Info.DefaultInterface = GetVertexInterface();
 
diff --git a/Engine/Plugins/Runtime/Metasound/Source/MetasoundStandardNodes/Private/MetasoundInterpToNode.cpp b/Engine/Plugins/Runtime/Metasound/Source/MetasoundStandardNodes/Private/MetasoundInterpToNode.cpp
--- a/Engine/Plugins/Runtime/Metasound/Source/MetasoundStandardNodes/Private/MetasoundInterpToNode.cpp
+++ b/Engine/Plugins/Runtime/Metasound/Source/MetasoundStandardNodes/Private/MetasoundInterpToNode.cpp
@@ -2,15 +2,16 @@
 
 #include "MetasoundFacade.h"
 
+#include "DSP/VolumeFader.h"
 #include "Internationalization/Text.h"
 #include "MetasoundExecutableOperator.h"
 #include "MetasoundNodeRegistrationMacro.h"
 #include "MetasoundParamHelper.h"
 #include "MetasoundPrimitives.h"
+#include "MetasoundStandardNodesCategories.h"
 #include "MetasoundStandardNodesNames.h"
 #include "MetasoundTrigger.h"
 #include "MetasoundTime.h"
-#include "DSP/VolumeFader.h"
 
 #define LOCTEXT_NAMESPACE "MetasoundStandardNodes_InterpNode"
 
@@ -176,6 +177,7 @@ namespace Metasound
 			Info.Description = METASOUND_LOCTEXT("Metasound_InterpNodeDescription", "Interpolates between the current value and a target value over the specified time.");
 			Info.Author = PluginAuthor;
 			Info.PromptIfMissing = PluginNodeMissingPrompt;
+			Info.CategoryHierarchy = { NodeCategories::Math };
 			Info.DefaultInterface = GetVertexInterface();
 			Info.Keywords.Add(METASOUND_LOCTEXT("LerpKeyword", "Lerp"));
 
diff --git a/Engine/Plugins/Runtime/Metasound/Source/MetasoundStandardNodes/Private/MetasoundRingModNode.cpp b/Engine/Plugins/Runtime/Metasound/Source/MetasoundStandardNodes/Private/MetasoundRingModNode.cpp
--- a/Engine/Plugins/Runtime/Metasound/Source/MetasoundStandardNodes/Private/MetasoundRingModNode.cpp
+++ b/Engine/Plugins/Runtime/Metasound/Source/MetasoundStandardNodes/Private/MetasoundRingModNode.cpp
@@ -7,9 +7,9 @@
 #include "MetasoundNodeRegistrationMacro.h"
 #include "MetasoundStandardNodesNames.h"
 #include "MetasoundAudioBuffer.h"
-#include "MetasoundStandardNodesCategories.h"
 #include "MetasoundFacade.h"
 #include "MetasoundParamHelper.h"
+#include "MetasoundStandardNodesCategories.h"
 
 #define LOCTEXT_NAMESPACE "MetasoundStandardNodes_MetasoundRingModNodes"
 
@@ -50,12 +50,12 @@ namespace Metasound
 					FNodeClassName { StandardNodes::Namespace, "RingMod", StandardNodes::AudioVariant },
 					1, // Major Version
 					0, // Minor Version
-					METASOUND_LOCTEXT("RingModDisplayName", "RingMod"),
+					METASOUND_LOCTEXT("RingModDisplayName2", "Ring Modulator"),
 					METASOUND_LOCTEXT("RingModDesc", "Modulates a carrier signal."),
 					PluginAuthor,
 					PluginNodeMissingPrompt,
 					NodeInterface,
-					{ NodeCategories::Functions},
+					{ NodeCategories::Filters },
 					{ },
 					FNodeDisplayStyle{}
 				};
diff --git a/Engine/Plugins/Runtime/Metasound/Source/MetasoundStandardNodes/Private/MetasoundStandardNodesCategories.cpp b/Engine/Plugins/Runtime/Metasound/Source/MetasoundStandardNodes/Private/MetasoundStandardNodesCategories.cpp
--- a/Engine/Plugins/Runtime/Metasound/Source/MetasoundStandardNodes/Private/MetasoundStandardNodesCategories.cpp
+++ b/Engine/Plugins/Runtime/Metasound/Source/MetasoundStandardNodes/Private/MetasoundStandardNodesCategories.cpp
@@ -2,26 +2,26 @@
 
 #include "MetasoundStandardNodesCategories.h"
 
+#include "MetasoundNodeRegistrationMacro.h"
+
 #define LOCTEXT_NAMESPACE "MetasoundStandardNodes"
 
-namespace Metasound
+namespace Metasound::NodeCategories
 {
-	namespace NodeCategories
-	{
-		const FText Debug = { METASOUND_LOCTEXT("Metasound_DebugCategory", "Debug") };
-		const FText Delays = { METASOUND_LOCTEXT("Metasound_EffectsCategory", "Delays") };
-		const FText Dynamics = { METASOUND_LOCTEXT("Metasound_DynamicsCategory", "Dynamics") };
-		const FText Envelopes = { METASOUND_LOCTEXT("Metasound_EnvelopesCategory", "Envelopes") };
-		const FText Filters = { METASOUND_LOCTEXT("Metasound_FiltersCategory", "Filters") };
-		const FText Generators = { METASOUND_LOCTEXT("Metasound_GeneratorsCategory", "Generators") };
-		const FText Io = { METASOUND_LOCTEXT("Metasound_IoCategory", "External IO") };
-		const FText Math = { METASOUND_LOCTEXT("Metasound_MathCategory", "Math") };
-		const FText Mix = { METASOUND_LOCTEXT("Metasound_MixCategory", "Mix") };
-		const FText Music = { METASOUND_LOCTEXT("Metasound_MusicCategory", "Music") };
-		const FText RandomUtils = { METASOUND_LOCTEXT("Metasound_RandomCategory", "Random") };
-		const FText Spatialization = { METASOUND_LOCTEXT("Metasound_SpatializationCategory", "Spatialization") };
-		const FText Trigger = { METASOUND_LOCTEXT("Metasound_TriggerCategory", "Triggers") };
-	}
-}
+	const FText Debug = { METASOUND_LOCTEXT("Metasound_DebugCategory", "Debug") };
+	const FText Delays = { METASOUND_LOCTEXT("Metasound_EffectsCategory", "Delays") };
+	const FText Dynamics = { METASOUND_LOCTEXT("Metasound_DynamicsCategory", "Dynamics") };
+	const FText Envelopes = { METASOUND_LOCTEXT("Metasound_EnvelopesCategory", "Envelopes") };
+	const FText Filters = { METASOUND_LOCTEXT("Metasound_FiltersCategory", "Filters") };
+	const FText Generators = { METASOUND_LOCTEXT("Metasound_GeneratorsCategory", "Generators") };
+	const FText Io = { METASOUND_LOCTEXT("Metasound_IoCategory", "External IO") };
+	const FText Math = { METASOUND_LOCTEXT("Metasound_MathCategory", "Math") };
+	const FText Mix = { METASOUND_LOCTEXT("Metasound_MixCategory", "Mix") };
+	const FText Music = { METASOUND_LOCTEXT("Metasound_MusicCategory", "Music") };
+	const FText RandomUtils = { METASOUND_LOCTEXT("Metasound_RandomCategory", "Random") };
+	const FText Spatialization = { METASOUND_LOCTEXT("Metasound_SpatializationCategory", "Spatialization") };
+	const FText Trigger = { METASOUND_LOCTEXT("Metasound_TriggerCategory", "Triggers") };
+	const FText WaveTables = { METASOUND_LOCTEXT("Metasound_WaveTableCategory", "WaveTables") };
+} // Metasound::NodeCategories
 
 #undef LOCTEXT_NAMESPACE
\ No newline at end of file
diff --git a/Engine/Plugins/Runtime/Metasound/Source/MetasoundStandardNodes/Private/MetasoundSwitchAndRampNode.cpp b/Engine/Plugins/Runtime/Metasound/Source/MetasoundStandardNodes/Private/MetasoundSwitchAndRampNode.cpp
--- a/Engine/Plugins/Runtime/Metasound/Source/MetasoundStandardNodes/Private/MetasoundSwitchAndRampNode.cpp
+++ b/Engine/Plugins/Runtime/Metasound/Source/MetasoundStandardNodes/Private/MetasoundSwitchAndRampNode.cpp
@@ -1,15 +1,15 @@
 // Copyright Epic Games, Inc. All Rights Reserved.
 // Background on switch-and-ramp: http://msp.ucsd.edu/techniques/v0.11/book-html/node63.html
 
+#include "DSP/BufferVectorOperations.h"
 #include "Internationalization/Text.h"
 #include "MetasoundFacade.h"
 #include "MetasoundExecutableOperator.h"
 #include "MetasoundNodeRegistrationMacro.h"
+#include "MetasoundParamHelper.h"
 #include "MetasoundPrimitives.h"
 #include "MetasoundStandardNodesNames.h"
 #include "MetasoundStandardNodesCategories.h"
-#include "DSP/BufferVectorOperations.h"
-#include "MetasoundParamHelper.h"
 
 #define LOCTEXT_NAMESPACE "MetasoundStandardNodes_SwitchAndRampNode"
 
@@ -43,6 +43,7 @@ namespace Metasound
 				Info.DisplayName = NodeDisplayName;
 				Info.Description = NodeDescription;
 				Info.Author = PluginAuthor;
+				Info.CategoryHierarchy = { NodeCategories::Filters };
 				Info.PromptIfMissing = PluginNodeMissingPrompt;
 				Info.DefaultInterface = GetVertexInterface();
 
diff --git a/Engine/Plugins/Runtime/Metasound/Source/MetasoundStandardNodes/Private/MetasoundWaveShaperNode.cpp b/Engine/Plugins/Runtime/Metasound/Source/MetasoundStandardNodes/Private/MetasoundWaveShaperNode.cpp
--- a/Engine/Plugins/Runtime/Metasound/Source/MetasoundStandardNodes/Private/MetasoundWaveShaperNode.cpp
+++ b/Engine/Plugins/Runtime/Metasound/Source/MetasoundStandardNodes/Private/MetasoundWaveShaperNode.cpp
@@ -1,5 +1,6 @@
 // Copyright Epic Games, Inc. All Rights Reserved.
 
+#include "DSP/WaveShaper.h"
 #include "Internationalization/Text.h"
 #include "MetasoundEnumRegistrationMacro.h"
 #include "MetasoundExecutableOperator.h"
@@ -7,12 +8,13 @@
 #include "MetasoundDataTypeRegistrationMacro.h"
 #include "MetasoundParamHelper.h"
 #include "MetasoundPrimitives.h"
+#include "MetasoundStandardNodesCategories.h"
 #include "MetasoundStandardNodesNames.h"
 #include "MetasoundAudioBuffer.h"
-#include "DSP/WaveShaper.h"
 
 #define LOCTEXT_NAMESPACE "MetasoundStandardNodes_WaveShaperNode"
 
+
 namespace Metasound
 {
 	namespace WaveShaperNode
@@ -178,6 +180,7 @@ namespace Metasound
 			Info.DisplayName = METASOUND_LOCTEXT("Metasound_WaveShaperDisplayName", "WaveShaper");
 			Info.Description = METASOUND_LOCTEXT("Metasound_WaveShaperNodeDescription", "Applies non-linear shaping to the audio input.");
 			Info.Author = PluginAuthor;
+			Info.CategoryHierarchy = { NodeCategories::Filters };
 			Info.PromptIfMissing = PluginNodeMissingPrompt;
 			Info.DefaultInterface = GetVertexInterface();
 
diff --git a/Engine/Plugins/Runtime/Metasound/Source/MetasoundStandardNodes/Private/MetasoundWaveWriterNode.cpp b/Engine/Plugins/Runtime/Metasound/Source/MetasoundStandardNodes/Private/MetasoundWaveWriterNode.cpp
--- a/Engine/Plugins/Runtime/Metasound/Source/MetasoundStandardNodes/Private/MetasoundWaveWriterNode.cpp
+++ b/Engine/Plugins/Runtime/Metasound/Source/MetasoundStandardNodes/Private/MetasoundWaveWriterNode.cpp
@@ -324,7 +324,7 @@ namespace Metasound
 			{
 				// For backwards compatibility with previous (mono) WaveWriters keep the node name the same.
 				FName OperatorName = TEXT("WaveWriter");
-				FText NodeDisplayName = METASOUND_LOCTEXT("Metasound_WaveWriterNodeMonoDisplayName", "Wave Writer (Mono)");
+				FText NodeDisplayName = METASOUND_LOCTEXT("Metasound_WaveWriterNodeMonoDisplayName", "Wave Writer (1-channel, Mono)");
 				const FText NodeDescription = METASOUND_LOCTEXT("Metasound_WaveWriterNodeMonoDescription", "Write a mono audio signal to disk");
 				FVertexInterface NodeInterface = DeclareVertexInterface();
 
@@ -335,7 +335,7 @@ namespace Metasound
 			auto CreateNodeClassMetadataStereo = []() -> FNodeClassMetadata
 			{
 				FName OperatorName = TEXT("Wave Writer (Stereo)");
-				FText NodeDisplayName = METASOUND_LOCTEXT("Metasound_WaveWriterNodeStereoDisplayName", "Wave Writer (Stereo)");
+				FText NodeDisplayName = METASOUND_LOCTEXT("Metasound_WaveWriterNodeStereoDisplayName", "Wave Writer (2-channel, Stereo)");
 				const FText NodeDescription = METASOUND_LOCTEXT("Metasound_WaveWriterNodeStereoDescription", "Write a stereo audio signal to disk");
 				FVertexInterface NodeInterface = DeclareVertexInterface();
 
diff --git a/Engine/Plugins/Runtime/Metasound/Source/MetasoundStandardNodes/Public/MetasoundStandardNodesCategories.h b/Engine/Plugins/Runtime/Metasound/Source/MetasoundStandardNodes/Public/MetasoundStandardNodesCategories.h
--- a/Engine/Plugins/Runtime/Metasound/Source/MetasoundStandardNodes/Public/MetasoundStandardNodesCategories.h
+++ b/Engine/Plugins/Runtime/Metasound/Source/MetasoundStandardNodes/Public/MetasoundStandardNodesCategories.h
@@ -2,26 +2,23 @@
 
 #pragma once
 
-#include "CoreMinimal.h"
-#include "MetasoundAutoConverterNode.h"
+#include "Internationalization/Text.h"
 
 
-namespace Metasound
+namespace Metasound::NodeCategories
 {
-	namespace NodeCategories
-	{
-		extern const FText METASOUNDSTANDARDNODES_API Debug;
-		extern const FText METASOUNDSTANDARDNODES_API Delays;
-		extern const FText METASOUNDSTANDARDNODES_API Dynamics;
-		extern const FText METASOUNDSTANDARDNODES_API Envelopes;
-		extern const FText METASOUNDSTANDARDNODES_API Filters;
-		extern const FText METASOUNDSTANDARDNODES_API Generators;
-		extern const FText METASOUNDSTANDARDNODES_API Io;
-		extern const FText METASOUNDSTANDARDNODES_API Math;
-		extern const FText METASOUNDSTANDARDNODES_API Mix;
-		extern const FText METASOUNDSTANDARDNODES_API Music;
-		extern const FText METASOUNDSTANDARDNODES_API RandomUtils;
-		extern const FText METASOUNDSTANDARDNODES_API Spatialization;
-		extern const FText METASOUNDSTANDARDNODES_API Trigger;
-	}
-}
+	extern const FText METASOUNDSTANDARDNODES_API Debug;
+	extern const FText METASOUNDSTANDARDNODES_API Delays;
+	extern const FText METASOUNDSTANDARDNODES_API Dynamics;
+	extern const FText METASOUNDSTANDARDNODES_API Envelopes;
+	extern const FText METASOUNDSTANDARDNODES_API Filters;
+	extern const FText METASOUNDSTANDARDNODES_API Generators;
+	extern const FText METASOUNDSTANDARDNODES_API Io;
+	extern const FText METASOUNDSTANDARDNODES_API Math;
+	extern const FText METASOUNDSTANDARDNODES_API Mix;
+	extern const FText METASOUNDSTANDARDNODES_API Music;
+	extern const FText METASOUNDSTANDARDNODES_API RandomUtils;
+	extern const FText METASOUNDSTANDARDNODES_API Spatialization;
+	extern const FText METASOUNDSTANDARDNODES_API Trigger;
+	extern const FText METASOUNDSTANDARDNODES_API WaveTables;
+} // Metasound::NodeCategories
