Fix unity: missing headers
#rnx

diff --git a/Engine/Plugins/Runtime/Metasound/Source/MetasoundFrontend/Private/MetasoundFrontendNodesCategories.cpp b/Engine/Plugins/Runtime/Metasound/Source/MetasoundFrontend/Private/MetasoundFrontendNodesCategories.cpp
--- a/Engine/Plugins/Runtime/Metasound/Source/MetasoundFrontend/Private/MetasoundFrontendNodesCategories.cpp
+++ b/Engine/Plugins/Runtime/Metasound/Source/MetasoundFrontend/Private/MetasoundFrontendNodesCategories.cpp
@@ -23,10 +23,10 @@ namespace Metasound
 		// Base category for all inputs
 		const FText Inputs = LOCTEXT("Metasound_InputsCategory", "Inputs");
 
-		// Base cateogry for all outputs
+		// Base category for all outputs
 		const FText Outputs = LOCTEXT("Metasound_OutputsCategory", "Outputs");
 
-		// Base cateogry for all variables 
+		// Base category for all variables 
 		const FText Variables = LOCTEXT("Metasound_VariablesCategory", "Variables");
 	}
 }
diff --git a/Engine/Plugins/Runtime/Metasound/Source/MetasoundStandardNodes/Private/MetasoundConversionNodes.cpp b/Engine/Plugins/Runtime/Metasound/Source/MetasoundStandardNodes/Private/MetasoundConversionNodes.cpp
--- a/Engine/Plugins/Runtime/Metasound/Source/MetasoundStandardNodes/Private/MetasoundConversionNodes.cpp
+++ b/Engine/Plugins/Runtime/Metasound/Source/MetasoundStandardNodes/Private/MetasoundConversionNodes.cpp
@@ -4,6 +4,7 @@
 #include "MetasoundNodeRegistrationMacro.h"
 #include "MetasoundFacade.h"
 #include "MetasoundExecutableOperator.h"
+#include "MetasoundFrontendNodesCategories.h"
 #include "MetasoundPrimitives.h"
 #include "MetasoundStandardNodesNames.h"
 #include "MetasoundTime.h"
diff --git a/Engine/Plugins/Runtime/Metasound/Source/MetasoundStandardNodes/Private/MetasoundSemitoneToFreqMultiplierNode.cpp b/Engine/Plugins/Runtime/Metasound/Source/MetasoundStandardNodes/Private/MetasoundSemitoneToFreqMultiplierNode.cpp
--- a/Engine/Plugins/Runtime/Metasound/Source/MetasoundStandardNodes/Private/MetasoundSemitoneToFreqMultiplierNode.cpp
+++ b/Engine/Plugins/Runtime/Metasound/Source/MetasoundStandardNodes/Private/MetasoundSemitoneToFreqMultiplierNode.cpp
@@ -1,8 +1,9 @@
 // Copyright Epic Games, Inc. All Rights Reserved.
 
 #include "Internationalization/Text.h"
-#include "MetasoundFacade.h"
 #include "MetasoundExecutableOperator.h"
+#include "MetasoundFacade.h"
+#include "MetasoundFrontendNodesCategories.h"
 #include "MetasoundNodeRegistrationMacro.h"
 #include "MetasoundPrimitives.h"
 #include "MetasoundStandardNodesNames.h"
