non-unity fix #2
#rnx

diff --git a/Engine/Plugins/Runtime/Metasound/Source/MetasoundStandardNodes/Public/MetasoundRandomNode.h b/Engine/Plugins/Runtime/Metasound/Source/MetasoundStandardNodes/Public/MetasoundRandomNode.h
--- a/Engine/Plugins/Runtime/Metasound/Source/MetasoundStandardNodes/Public/MetasoundRandomNode.h
+++ b/Engine/Plugins/Runtime/Metasound/Source/MetasoundStandardNodes/Public/MetasoundRandomNode.h
@@ -3,15 +3,16 @@
 #pragma once
 
 #include "CoreMinimal.h"
-#include "MetasoundFacade.h"
+#include "Internationalization/Text.h"
 #include "MetasoundExecutableOperator.h"
+#include "MetasoundFacade.h"
+#include "MetasoundNodeRegistrationMacro.h"
 #include "MetasoundParamHelper.h"
 #include "MetasoundPrimitives.h"
 #include "MetasoundStandardNodesNames.h"
+#include "MetasoundTime.h"
 #include "MetasoundTrigger.h"
-#include "Internationalization/Text.h"
 #include "MetasoundStandardNodesCategories.h"
-#include "MetasoundTime.h"
 #include "MetasoundVertex.h"
 
 #define LOCTEXT_NAMESPACE "MetasoundStandardNodes"
