IWYU!
#rnx

diff --git a/Engine/Plugins/Runtime/Metasound/Source/MetasoundStandardNodes/Private/MetasoundFreqMultiplierToSemitoneNode.cpp b/Engine/Plugins/Runtime/Metasound/Source/MetasoundStandardNodes/Private/MetasoundFreqMultiplierToSemitoneNode.cpp
--- a/Engine/Plugins/Runtime/Metasound/Source/MetasoundStandardNodes/Private/MetasoundFreqMultiplierToSemitoneNode.cpp
+++ b/Engine/Plugins/Runtime/Metasound/Source/MetasoundStandardNodes/Private/MetasoundFreqMultiplierToSemitoneNode.cpp
@@ -1,14 +1,15 @@
 // Copyright Epic Games, Inc. All Rights Reserved.
 
 #include "Internationalization/Text.h"
-#include "MetasoundFacade.h"
 #include "MetasoundExecutableOperator.h"
+#include "MetasoundFacade.h"
+#include "MetasoundFrontendNodesCategories.h"
 #include "MetasoundNodeRegistrationMacro.h"
+#include "MetasoundParamHelper.h"
 #include "MetasoundPrimitives.h"
 #include "MetasoundStandardNodesNames.h"
 #include "MetasoundStandardNodesCategories.h"
 #include "DSP/Dsp.h"
-#include "MetasoundParamHelper.h"
 
 #define LOCTEXT_NAMESPACE "MetasoundStandardNodes_FreqMultiplierToSemitone"
 
