Fixing ADSR invalid memory access

diff --git a/Engine/Plugins/Runtime/Metasound/Source/MetasoundStandardNodes/Private/MetasoundADSREnvelopeNode.cpp b/Engine/Plugins/Runtime/Metasound/Source/MetasoundStandardNodes/Private/MetasoundADSREnvelopeNode.cpp
--- a/Engine/Plugins/Runtime/Metasound/Source/MetasoundStandardNodes/Private/MetasoundADSREnvelopeNode.cpp
+++ b/Engine/Plugins/Runtime/Metasound/Source/MetasoundStandardNodes/Private/MetasoundADSREnvelopeNode.cpp
@@ -25,10 +25,10 @@ namespace Metasound
 		METASOUND_PARAM(InputAttackTrigger, "Trigger Attack", "Trigger to start the attack phase of the envelope generator.");
 		METASOUND_PARAM(InputReleaseTrigger, "Trigger Release", "Trigger to start the release phase of the envelope generator.");
 
-		METASOUND_PARAM(InputAttackTime, "Attack Time", "Attack time of the envelope.");
-		METASOUND_PARAM(InputDecayTime, "Decay Time", "Decay time of the envelope.");
-		METASOUND_PARAM(InputSustainLevel, "Sustain Level", "The sustain level.");
-		METASOUND_PARAM(InputReleaseTime, "Release Time", "Release time of the envelope.");
+		METASOUND_PARAM(InputAttackTime, "Attack Time", "Attack time (seconds) of the envelope. Time to reach 1.0 in the envelope output.");
+		METASOUND_PARAM(InputDecayTime, "Decay Time", "Decay time (seconds) of the envelope. Time to reach the sustain level in the envelope output.");
+		METASOUND_PARAM(InputSustainLevel, "Sustain Level", "The sustain level (between 0.0 and 1.0).");
+		METASOUND_PARAM(InputReleaseTime, "Release Time", "Release time (seconds) of the envelope. Time to reach 0.0 in the envelope output.");
 
 		METASOUND_PARAM(InputAttackCurve, "Attack Curve", "The exponential curve factor of the attack. 1.0 = linear growth, < 1.0 logorithmic growth, > 1.0 exponential growth.");
 		METASOUND_PARAM(InputDecayCurve, "Decay Curve", "The exponential curve factor of the decay. 1.0 = linear decay, < 1.0 exponential decay, > 1.0 logorithmic decay.");
@@ -203,7 +203,7 @@ namespace Metasound
 						{
 							InState.CurrentSampleIndex = INDEX_NONE;
 							OutEnvelopeValue = 0.0f;
-							OutOnDoneFrames.Add(0);
+							OutOnDoneFrames.Add(EndFrame);
 						}
 					}
 				}
@@ -224,6 +224,7 @@ namespace Metasound
 
 			static void GetNextEnvelopeOutput(FEnvState& InState, int32 StartFrame, int32 EndFrame, TArray<int32>& OutOnDecayFrames, TArray<int32>& OutOnSustainFrames, TArray<int32>& OutOnDoneFrames, FAudioBuffer& OutEnvelopeValue)
 			{
+				OutEnvelopeValue.Zero();
 				// If we are not active zero the buffer and early exit
 				if (InState.CurrentSampleIndex == INDEX_NONE)
 				{
@@ -343,12 +344,6 @@ namespace Metasound
 									// Envelope is done
 									if (InState.EnvEase.IsDone())
 									{
-										// Zero out the rest of the envelope
-										int32 NumSamplesLeft = EndFrame - i - 1;
-										if (NumSamplesLeft > 0)
-										{
-											FMemory::Memzero(&OutEnvPtr[i + 1], sizeof(float) * NumSamplesLeft);
-										}
 										InState.CurrentSampleIndex = INDEX_NONE;
 										OutOnDoneFrames.Add(i);
 										break;
@@ -656,7 +651,7 @@ namespace Metasound
 			float ReleaseTimeSeconds = ReleaseTime->GetSeconds();
 			EnvState.AttackSampleCount = FMath::Max(1, static_cast<int32>(SampleRate * FMath::Max(0.0f, AttackTimeSeconds)));
 			EnvState.DecaySampleCount = FMath::Max(0, static_cast<int32>(SampleRate * FMath::Max(0.0f, DecayTimeSeconds)));
-			EnvState.SustainLevel = FMath::Max(0.0f, *SustainLevel);
+			EnvState.SustainLevel = FMath::Clamp(*SustainLevel, 0.0f, 1.0f);
 			EnvState.ReleaseSampleCount = FMath::Max(0, static_cast<int32>(SampleRate * FMath::Max(0.0f, ReleaseTimeSeconds)));
 			EnvState.AttackCurveFactor = FMath::Max(KINDA_SMALL_NUMBER, *AttackCurveFactor);
 			EnvState.DecayCurveFactor = FMath::Max(KINDA_SMALL_NUMBER, *DecayCurveFactor);
@@ -707,7 +702,7 @@ namespace Metasound
 			OnDecayTrigger->AdvanceBlock();
 			OnSustainTrigger->AdvanceBlock();
 			OnDone->AdvanceBlock();
- 
+ 				
 			FTriggerIterator TriggerIter(*TriggerAttackIn, *TriggerReleaseIn, NumFramesPerBlock);
 
 			FTriggerInfo NextTrigger = TriggerIter.NextTrigger(EnvState.IsTriggered());
