Fix for grain delay node NaNs with invalid inputs

diff --git a/Engine/Plugins/Runtime/Metasound/Source/MetasoundStandardNodes/Private/MetasoundGrainDelayNode.cpp b/Engine/Plugins/Runtime/Metasound/Source/MetasoundStandardNodes/Private/MetasoundGrainDelayNode.cpp
--- a/Engine/Plugins/Runtime/Metasound/Source/MetasoundStandardNodes/Private/MetasoundGrainDelayNode.cpp
+++ b/Engine/Plugins/Runtime/Metasound/Source/MetasoundStandardNodes/Private/MetasoundGrainDelayNode.cpp
@@ -286,7 +286,7 @@ namespace Metasound
 		const int32 LastAudioFrame = OperatorSettings.GetNumFramesPerBlock() - 1;
 		const int32 NoTrigger = OperatorSettings.GetNumFramesPerBlock() << 1;
 
-		GrainDelayProcessor.SetMaxGrains(*GrainMaxCount);
+		GrainDelayProcessor.SetMaxGrains(FMath::Clamp(*GrainMaxCount, 1, 1000));
 		
 		// Update the envelope before rendering the audio block
 		if (PreviousGrainEnvelopeType != *GrainDelayEnvelope)
@@ -309,8 +309,16 @@ namespace Metasound
 			PreviousGrainDelayMsec = CurrentDelayLengthMsec;
 		}
 
-		GrainDelayProcessor.SetFeedbackAmount(*GrainDelayFeedbackAmount);
+		const float GrainDelayFeedbackAmountClamped = FMath::Clamp(*GrainDelayFeedbackAmount, 0.0f, 1.0f);
+		GrainDelayProcessor.SetFeedbackAmount(GrainDelayFeedbackAmountClamped);
 		
+		const float NewDurationSeconds = GrainDelayProcessor.GetGrainDurationClamped(*GrainDuration + FMath::RandRange(-0.5f * *GrainDurationDelta, 0.5f * *GrainDurationDelta));
+
+		const float NewDelay = GrainDelayProcessor.GetGrainDelayClamped(CurrentDelayLengthMsec + FMath::RandRange(-0.5f * *GrainDelayDelta, 0.5f * *GrainDelayDelta));
+
+		const float ClampedPitchShiftDelta = GrainDelayProcessor.GetGrainPitchShiftClamped(*GrainPitchShift + FMath::RandRange(-0.5f * *GrainPitchShiftDelta, 0.5f * *GrainPitchShiftDelta));
+		const float NewPitchShiftRatioOffset = GrainDelayProcessor.GetGrainPitchShiftFrameRatio(ClampedPitchShiftDelta);
+
 		// Loop until we either hit the end of the frame or the next trigger
 		while (NextAudioFrame < LastAudioFrame)
 		{
@@ -344,22 +352,8 @@ namespace Metasound
 			{
 				++SpawnTrigIndex;
 
-				const float Delta = *GrainDurationDelta;
-				const float NewDuration = GrainDelayProcessor.GetGrainDurationClamped(*GrainDuration + FMath::RandRange(-0.5f * Delta, 0.5f * Delta));
-				if(FMath::IsNearlyZero(NewDuration))
-				{
-					// avoid spawning nodes of zero duration
-					continue;
-				}
-
-				const float NewDelay = GrainDelayProcessor.GetGrainDelayClamped(*GrainDelay + FMath::RandRange(-0.5f * (*GrainDelayDelta), 0.5f * (*GrainDelayDelta)));
-
-				// Clamp the pitch shift delta before determining the random range
-				const float ClampedPitchShiftDelta = GrainDelayProcessor.GetGrainPitchShiftClamped(*GrainPitchShiftDelta);
-				const float NewPitchShiftRatioOffset = GrainDelayProcessor.GetGrainPitchShiftFrameRatio(FMath::RandRange(-0.5f * ClampedPitchShiftDelta, 0.5f * ClampedPitchShiftDelta));
-
 				// Spawn a new grain. Note the grain manager handles spawning logic and rate limiting.
-				GrainDelayProcessor.SpawnGrain(NewDelay, NewDuration, NewPitchShiftRatioOffset);
+				GrainDelayProcessor.SpawnGrain(NewDelay, NewDurationSeconds, NewPitchShiftRatioOffset);
 			}
 		}
 	}
