Remove duplicate/copy options from right click menu for MetaSound outputs
#rnx

diff --git a/Engine/Plugins/Runtime/Metasound/Source/MetasoundEditor/Private/MetasoundEditorGraphNode.cpp b/Engine/Plugins/Runtime/Metasound/Source/MetasoundEditor/Private/MetasoundEditorGraphNode.cpp
--- a/Engine/Plugins/Runtime/Metasound/Source/MetasoundEditor/Private/MetasoundEditorGraphNode.cpp
+++ b/Engine/Plugins/Runtime/Metasound/Source/MetasoundEditor/Private/MetasoundEditorGraphNode.cpp
@@ -588,6 +588,11 @@ bool UMetasoundEditorGraphMemberNode::ClampFloatLiteral(const UMetasoundEditorGr
 	return bClampedFloatLiteral;
 }
 
+bool UMetasoundEditorGraphOutputNode::CanDuplicateNode() const
+{
+	return false;
+}
+
 void UMetasoundEditorGraphOutputNode::PinDefaultValueChanged(UEdGraphPin* InPin)
 {
 	using namespace Metasound::Editor;
diff --git a/Engine/Plugins/Runtime/Metasound/Source/MetasoundEditor/Public/MetasoundEditorGraphNode.h b/Engine/Plugins/Runtime/Metasound/Source/MetasoundEditor/Public/MetasoundEditorGraphNode.h
--- a/Engine/Plugins/Runtime/Metasound/Source/MetasoundEditor/Public/MetasoundEditorGraphNode.h
+++ b/Engine/Plugins/Runtime/Metasound/Source/MetasoundEditor/Public/MetasoundEditorGraphNode.h
@@ -163,6 +163,9 @@ public:
 	virtual FGuid GetNodeID() const override;
 	virtual UMetasoundEditorGraphMember* GetMember() const override;
 
+	// Can't duplicate, only one node per graph.
+	virtual bool CanDuplicateNode() const override;
+
 	// Disallow deleting outputs as they require being connected to some
 	// part of the graph by the Frontend Graph Builder (which is enforced
 	// even when the Editor Graph Node does not have a visible input by
