Fix for adding nodes before auditioning dynamic metasound

diff --git a/Engine/Plugins/Runtime/Metasound/Source/MetasoundGraphCore/Private/MetasoundDynamicOperator.cpp b/Engine/Plugins/Runtime/Metasound/Source/MetasoundGraphCore/Private/MetasoundDynamicOperator.cpp
--- a/Engine/Plugins/Runtime/Metasound/Source/MetasoundGraphCore/Private/MetasoundDynamicOperator.cpp
+++ b/Engine/Plugins/Runtime/Metasound/Source/MetasoundGraphCore/Private/MetasoundDynamicOperator.cpp
@@ -254,6 +254,17 @@ namespace Metasound
 			{
 				IOperator* Operator = OperatorInfo.Operator.Get();
 
+				if (FOperatorInfo* ExistingInfo = InGraphOperatorData.OperatorMap.Find(OperatorID))
+				{
+					// The options here are not good. The prior operator will be 
+					// removed and replaced with this new operator.
+					// Another option would be to leave the existing operator unchanged. 
+					// Neither option is satisfactory.
+					UE_LOG(LogMetaSound, Warning, TEXT("Overriding existing operator with the same operator ID %d. Duplicate operator IDs will lead to undefined behavior. Remove existing operators before adding a new one with the same ID"), OperatorID);
+
+					FRemoveOperator(OperatorID).Transform(InGraphOperatorData);
+				}
+
 				InGraphOperatorData.OperatorOrder.Add(OperatorID);
 				InGraphOperatorData.OperatorMap.Add(OperatorID, MoveTemp(OperatorInfo));
 
diff --git a/Engine/Plugins/Runtime/Metasound/Source/MetasoundGraphCore/Private/MetasoundDynamicOperatorTransactor.cpp b/Engine/Plugins/Runtime/Metasound/Source/MetasoundGraphCore/Private/MetasoundDynamicOperatorTransactor.cpp
--- a/Engine/Plugins/Runtime/Metasound/Source/MetasoundGraphCore/Private/MetasoundDynamicOperatorTransactor.cpp
+++ b/Engine/Plugins/Runtime/Metasound/Source/MetasoundGraphCore/Private/MetasoundDynamicOperatorTransactor.cpp
@@ -86,6 +86,22 @@ namespace Metasound
 			TSharedRef<TSpscQueue<TUniquePtr<IDynamicOperatorTransform>>> Queue = MakeShared<TSpscQueue<TUniquePtr<IDynamicOperatorTransform>>>();
 			OperatorInfos.Add(FDynamicOperatorInfo{InOperatorSettings, InEnvironment, Queue});
 
+			// Unconnected nodes are intentionally skipped by the operator builder
+			// in order to reduce perf. In a dynamic operator, these nodes may be
+			// connected in the future. We queue them up to be added here so that
+			// they exist on the dynamic operator in the case they are connected at
+			// a later time.
+			TArray<TPair<FGuid, const INode*>> UnconnectedNodes;
+			if (Graph.FindUnconnectedNodes(UnconnectedNodes) > 0)
+			{
+				for (const TPair<FGuid, const INode*>& GuidAndNode : UnconnectedNodes)
+				{
+					// Only add to this queue because we do not know at what point
+					// the other queues and dynamic operators were created.
+					Queue->Enqueue(CreateAddOperatorTransform(*GuidAndNode.Value, InOperatorSettings, InEnvironment));
+				}
+			}
+
 			return Queue;
 		}
 
@@ -104,12 +120,7 @@ namespace Metasound
 
 			Graph.AddNode(InNodeID, NodePtr);
 
-			auto CreateAddNodeTransform = [&](const FOperatorSettings& InOperatorSettings, const FMetasoundEnvironment& InEnvironment) -> TUniquePtr<IDynamicOperatorTransform>
-			{
-				return CreateAddOperatorTransform(*NodePtr, InOperatorSettings, InEnvironment);
-			};
-
-			EnqueueTransformOnOperatorQueues(CreateAddNodeTransform);
+			EnqueueAddOperatorTransform(*NodePtr);
 		}
 
 		void FDynamicOperatorTransactor::RemoveNode(const FGuid& InNodeID)
@@ -411,6 +422,16 @@ namespace Metasound
 			}
 		}
 		
+		void FDynamicOperatorTransactor::EnqueueAddOperatorTransform(const INode& InNode)
+		{
+			auto CreateAddNodeTransform = [&](const FOperatorSettings& InOperatorSettings, const FMetasoundEnvironment& InEnvironment) -> TUniquePtr<IDynamicOperatorTransform>
+			{
+				return CreateAddOperatorTransform(InNode, InOperatorSettings, InEnvironment);
+			};
+
+			EnqueueTransformOnOperatorQueues(CreateAddNodeTransform);
+		}
+
 		void FDynamicOperatorTransactor::EnqueueFadeAndRemoveOperatorTransform(const INode& InNode, TArrayView<const FVertexName> InOutputsToFade)
 		{
 			TArrayView<const FVertexName> InputsToFade; // We do not need to fade any inputs when removing a node.
diff --git a/Engine/Plugins/Runtime/Metasound/Source/MetasoundGraphCore/Private/MetasoundGraph.cpp b/Engine/Plugins/Runtime/Metasound/Source/MetasoundGraphCore/Private/MetasoundGraph.cpp
--- a/Engine/Plugins/Runtime/Metasound/Source/MetasoundGraphCore/Private/MetasoundGraph.cpp
+++ b/Engine/Plugins/Runtime/Metasound/Source/MetasoundGraphCore/Private/MetasoundGraph.cpp
@@ -1,7 +1,12 @@
 // Copyright Epic Games, Inc. All Rights Reserved.
 #include "MetasoundGraph.h"
 
+#include "Algo/BinarySearch.h"
+#include "Algo/Transform.h"
+#include "Containers/Array.h"
+#include "Containers/Map.h"
 #include "MetasoundOperatorBuilder.h"
+#include "Templates/SharedPointer.h"
 
 namespace Metasound
 {
@@ -183,6 +188,62 @@ namespace Metasound
 		return nullptr;
 	}
 
+	int32 FGraph::FindUnconnectedNodes(TArray<TPair<FGuid, const INode*>>& OutUnconnectedNodes) const
+	{
+		TArray<const INode*> ConnectedNodes;
+
+		auto AddNodeIfUnique = [&](const INode* InNode)
+		{
+			// Find index >= InNode
+			int32 Index = Algo::LowerBound(ConnectedNodes, InNode);
+
+			if (Index < ConnectedNodes.Num())
+			{
+				if (ConnectedNodes[Index] != InNode)
+				{
+					ConnectedNodes.Insert(InNode, Index);
+				}
+			}
+			else
+			{
+				ConnectedNodes.Add(InNode);
+			}
+		};
+
+		// Find nodes with an edge.
+		for (const FDataEdge& Edge : Edges)
+		{
+			AddNodeIfUnique(Edge.From.Node);
+			AddNodeIfUnique(Edge.To.Node);
+		}
+
+		// Find input and output nodes.
+		for (const TPair<FNodeDataVertexKey, FInputDataDestination>& Pair : InputDestinations)
+		{
+			AddNodeIfUnique(Pair.Value.Node);
+		}
+		for (const TPair<FNodeDataVertexKey, FOutputDataSource>& Pair : OutputSources)
+		{
+			AddNodeIfUnique(Pair.Value.Node);
+		}
+
+		// Filter all nodes by the connected nodes to get the unconnected nodes.
+		auto IsNodeUnconnected = [&](const TPair<FGuid, TSharedPtr<const INode>>& InGuidAndNode)
+		{
+			const INode* NodePtr = InGuidAndNode.Value.Get();
+			return (nullptr != NodePtr) && (INDEX_NONE == Algo::BinarySearch(ConnectedNodes, NodePtr));
+		};
+
+		auto RemoveSharedPtr = [&](const TPair<FGuid, TSharedPtr<const INode>>& InGuidAndNode) -> TPair<FGuid, const INode*>
+		{
+			return TPair<FGuid, const INode*>(InGuidAndNode.Key, InGuidAndNode.Value.Get());
+		};
+
+		Algo::TransformIf(Nodes, OutUnconnectedNodes, IsNodeUnconnected, RemoveSharedPtr);
+
+		return OutUnconnectedNodes.Num();
+	}
+
 	bool FGraph::RemoveNode(const FGuid& InNodeID, bool bInRemoveDataEdgesWithNode)
 	{
 		if (bInRemoveDataEdgesWithNode)
diff --git a/Engine/Plugins/Runtime/Metasound/Source/MetasoundGraphCore/Public/MetasoundDynamicOperatorTransactor.h b/Engine/Plugins/Runtime/Metasound/Source/MetasoundGraphCore/Public/MetasoundDynamicOperatorTransactor.h
--- a/Engine/Plugins/Runtime/Metasound/Source/MetasoundGraphCore/Public/MetasoundDynamicOperatorTransactor.h
+++ b/Engine/Plugins/Runtime/Metasound/Source/MetasoundGraphCore/Public/MetasoundDynamicOperatorTransactor.h
@@ -113,6 +113,7 @@ namespace Metasound
 		private:
 			using FOperatorID = uintptr_t;
 
+			void EnqueueAddOperatorTransform(const INode& InNode);
 			void EnqueueFadeAndRemoveOperatorTransform(const INode& InNode, TArrayView<const FVertexName> InOutputsToFade);
 			void EnqueueRemoveOperatorTransform(const INode& InNode);
 			void EnqueueBeginFadeOperatorTransform(const INode& InNode, EAudioFadeType InFadeType, TArrayView<const FVertexName> InInputsToFade, TArrayView<const FVertexName> InOutputsToFade);
diff --git a/Engine/Plugins/Runtime/Metasound/Source/MetasoundGraphCore/Public/MetasoundGraph.h b/Engine/Plugins/Runtime/Metasound/Source/MetasoundGraphCore/Public/MetasoundGraph.h
--- a/Engine/Plugins/Runtime/Metasound/Source/MetasoundGraphCore/Public/MetasoundGraph.h
+++ b/Engine/Plugins/Runtime/Metasound/Source/MetasoundGraphCore/Public/MetasoundGraph.h
@@ -110,6 +110,14 @@ namespace Metasound
 			 */
 			const INode* FindNode(const FGuid& InNodeID) const;
 
+			/** Populates an array with all nodes which exist in the graph but do not have
+			 * any connections.
+			 * 
+			 * @param OutUnconnectedNodes - Array to populate.
+			 * @return Number of unconnected nodes found.
+			 */
+			int32 FindUnconnectedNodes(TArray<TPair<FGuid, const INode*>>& OutUnconnectedNodes) const;
+
 			/** Removes node from graph.
 			 *
 			 * @param InNodeID - ID of node to remove.
