[MetaSoundEditor] Added Oscilloscope and Vectorscope to analyzers tab + code adjustments to support dynamic widget changes when switching to a different MetaSound output layout.

diff --git a/Engine/Plugins/Runtime/AudioWidgets/Source/AudioWidgets/Private/AudioOscilloscope.cpp b/Engine/Plugins/Runtime/AudioWidgets/Source/AudioWidgets/Private/AudioOscilloscope.cpp
--- a/Engine/Plugins/Runtime/AudioWidgets/Source/AudioWidgets/Private/AudioOscilloscope.cpp
+++ b/Engine/Plugins/Runtime/AudioWidgets/Source/AudioWidgets/Private/AudioOscilloscope.cpp
@@ -14,20 +14,45 @@ namespace AudioWidgets
 		const EAudioPanelLayoutType InPanelLayoutType)
 		: OscilloscopePanelStyle(FAudioOscilloscopePanelStyle::GetDefault())
 	{
-		// Init audio bus
+		CreateAudioBus(InNumChannels);
+		CreateDataProvider(InWorld, InTimeWindowMs, InMaxTimeWindowMs, InAnalysisPeriodMs, InPanelLayoutType);
+		CreateOscilloscopeWidget(InNumChannels, InPanelLayoutType);
+	}
+
+	void FAudioOscilloscope::CreateAudioBus(const uint32 InNumChannels)
+	{
 		AudioBus = TStrongObjectPtr(NewObject<UAudioBus>());
 		AudioBus->AudioBusChannels = AudioBusUtils::ConvertIntToEAudioBusChannels(InNumChannels);
+	}
+
+	void FAudioOscilloscope::CreateDataProvider(UWorld* InWorld,
+		const float InTimeWindowMs,
+		const float InMaxTimeWindowMs,
+		const float InAnalysisPeriodMs,
+		const EAudioPanelLayoutType InPanelLayoutType)
+	{
+		check(AudioBus);
 
-		// Init data provider
-		constexpr uint32 NumChannelsToProvide = 1;
+		const uint32 NumChannelsToProvide = (InPanelLayoutType == EAudioPanelLayoutType::Advanced) ? 1 : AudioBus->GetNumChannels(); // Advanced mode waveform display is based on channel selection
 		AudioSamplesDataProvider = MakeShared<FWaveformAudioSamplesDataProvider>(InWorld, AudioBus.Get(), NumChannelsToProvide, InTimeWindowMs, InMaxTimeWindowMs, InAnalysisPeriodMs);
+	}
+
+	void FAudioOscilloscope::CreateOscilloscopeWidget(const uint32 InNumChannels, const EAudioPanelLayoutType InPanelLayoutType)
+	{
+		check(AudioSamplesDataProvider);
 
-		// Init widget 
 		const FFixedSampledSequenceView SequenceView = AudioSamplesDataProvider->GetDataView();
 
-		OscilloscopePanelWidget = SNew(SAudioOscilloscopePanelWidget, SequenceView, InNumChannels)
-		.PanelLayoutType(InPanelLayoutType)
-		.PanelStyle(&OscilloscopePanelStyle);
+		if (!OscilloscopePanelWidget.IsValid())
+		{
+			OscilloscopePanelWidget = SNew(SAudioOscilloscopePanelWidget, SequenceView, InNumChannels)
+			.PanelLayoutType(InPanelLayoutType)
+			.PanelStyle(&OscilloscopePanelStyle);
+		}
+		else
+		{
+			OscilloscopePanelWidget->BuildWidget(SequenceView, InNumChannels);
+		}
 
 		// Interconnect data provider and widget
 		AudioSamplesDataProvider->OnDataViewGenerated.AddSP(OscilloscopePanelWidget.Get(), &SAudioOscilloscopePanelWidget::ReceiveSequenceView);
diff --git a/Engine/Plugins/Runtime/AudioWidgets/Source/AudioWidgets/Private/AudioVectorscope.cpp b/Engine/Plugins/Runtime/AudioWidgets/Source/AudioWidgets/Private/AudioVectorscope.cpp
--- a/Engine/Plugins/Runtime/AudioWidgets/Source/AudioWidgets/Private/AudioVectorscope.cpp
+++ b/Engine/Plugins/Runtime/AudioWidgets/Source/AudioWidgets/Private/AudioVectorscope.cpp
@@ -15,19 +15,40 @@ namespace AudioWidgets
 		const EAudioPanelLayoutType InPanelLayoutType)
 		: VectorscopePanelStyle(FAudioVectorscopePanelStyle::GetDefault())
 	{
-		// Init audio bus
+		CreateAudioBus(InNumChannels);
+		CreateDataProvider(InWorld, InTimeWindowMs, InMaxTimeWindowMs, InAnalysisPeriodMs);
+		CreateVectorscopeWidget(InPanelLayoutType);
+	}
+
+	void FAudioVectorscope::CreateAudioBus(const uint32 InNumChannels)
+	{
 		AudioBus = TStrongObjectPtr(NewObject<UAudioBus>());
 		AudioBus->AudioBusChannels = AudioBusUtils::ConvertIntToEAudioBusChannels(InNumChannels);
+	}
+
+	void FAudioVectorscope::CreateDataProvider(UWorld* InWorld, const float InTimeWindowMs,	const float InMaxTimeWindowMs, const float InAnalysisPeriodMs)
+	{
+		check(AudioBus);
 
-		// Init data provider
 		AudioSamplesDataProvider = MakeShared<FWaveformAudioSamplesDataProvider>(InWorld, AudioBus.Get(), AudioBus->GetNumChannels(), InTimeWindowMs, InMaxTimeWindowMs, InAnalysisPeriodMs);
+	}
+
+	void FAudioVectorscope::CreateVectorscopeWidget(const EAudioPanelLayoutType InPanelLayoutType)
+	{
+		check(AudioSamplesDataProvider);
 
-		// Init widget 
 		const FFixedSampledSequenceView SequenceView = AudioSamplesDataProvider->GetDataView();
 
-		VectorscopePanelWidget = SNew(SAudioVectorscopePanelWidget, SequenceView)
-		.PanelLayoutType(InPanelLayoutType)
-		.PanelStyle(&VectorscopePanelStyle);
+		if (!VectorscopePanelWidget.IsValid())
+		{
+			VectorscopePanelWidget = SNew(SAudioVectorscopePanelWidget, SequenceView)
+				.PanelLayoutType(InPanelLayoutType)
+				.PanelStyle(&VectorscopePanelStyle);
+		}
+		else
+		{
+			VectorscopePanelWidget->BuildWidget(SequenceView);
+		}
 
 		// Interconnect data provider and widget
 		AudioSamplesDataProvider->OnDataViewGenerated.AddSP(VectorscopePanelWidget.Get(), &SAudioVectorscopePanelWidget::ReceiveSequenceView);
diff --git a/Engine/Plugins/Runtime/AudioWidgets/Source/AudioWidgets/Private/SAudioOscilloscopePanelWidget.cpp b/Engine/Plugins/Runtime/AudioWidgets/Source/AudioWidgets/Private/SAudioOscilloscopePanelWidget.cpp
--- a/Engine/Plugins/Runtime/AudioWidgets/Source/AudioWidgets/Private/SAudioOscilloscopePanelWidget.cpp
+++ b/Engine/Plugins/Runtime/AudioWidgets/Source/AudioWidgets/Private/SAudioOscilloscopePanelWidget.cpp
@@ -71,27 +71,33 @@ namespace AudioOscilloscopePanelWidgetPrivate
 
 void SAudioOscilloscopePanelWidget::Construct(const FArguments& InArgs, const FFixedSampledSequenceView& InData, const int32 InNumChannels)
 {
-	NumChannels = InNumChannels;
-
-	DataView = InData;
-
 	PanelLayoutType = InArgs._PanelLayoutType;
 
-	SequenceRulerDisplayUnit    = InArgs._SequenceRulerDisplayUnit;
-	ValueGridOverlayDisplayUnit = InArgs._YAxisLabelsUnit.Get();
-	bHideSequenceRuler          = InArgs._HideSequenceRuler.Get();
-	bHideSequenceGrid           = InArgs._HideSequenceGrid.Get();
-	bHideValueGrid              = InArgs._HideValueGrid.Get();
-	bHideTriggerThresholdLine   = InArgs._HideTriggerThresholdLine.Get();
+	SequenceRulerDisplayUnit      = InArgs._SequenceRulerDisplayUnit;
+	ValueGridOverlayDisplayUnit   = InArgs._YAxisLabelsUnit.Get();
+	ValueGridMaxDivisionParameter = InArgs._ValueGridMaxDivisionParameter;
+	bHideSequenceRuler            = InArgs._HideSequenceRuler.Get();
+	bHideSequenceGrid             = InArgs._HideSequenceGrid.Get();
+	bHideValueGrid                = InArgs._HideValueGrid.Get();
+	bHideTriggerThresholdLine     = InArgs._HideTriggerThresholdLine.Get();
 
 	check(InArgs._PanelStyle);
 	PanelStyle = InArgs._PanelStyle;
 
+	BuildWidget(InData, InNumChannels);
+}
+
+void SAudioOscilloscopePanelWidget::BuildWidget(const FFixedSampledSequenceView& InData, const int32 InNumChannels)
+{
+	NumChannels = InNumChannels;
+
+	DataView = InData;
+
 	CreateGridData(PanelStyle->TimeRulerStyle);
 
 	CreateSequenceRuler(SequenceGridData.ToSharedRef(), PanelStyle->TimeRulerStyle);
 	CreateBackground(PanelStyle->WaveViewerStyle);
-	CreateValueGridOverlay(InArgs._ValueGridMaxDivisionParameter, SampledSequenceValueGridOverlay::EGridDivideMode::MidSplit, ValueGridOverlayDisplayUnit, PanelStyle->ValueGridStyle);
+	CreateValueGridOverlay(ValueGridMaxDivisionParameter, SampledSequenceValueGridOverlay::EGridDivideMode::MidSplit, ValueGridOverlayDisplayUnit, PanelStyle->ValueGridStyle);
 	CreateSequenceViewer(SequenceGridData.ToSharedRef(), DataView, PanelStyle->WaveViewerStyle);
 	CreateTriggerThresholdLine(PanelStyle->TriggerThresholdLineStyle);
 
diff --git a/Engine/Plugins/Runtime/AudioWidgets/Source/AudioWidgets/Private/SAudioVectorscopePanelWidget.cpp b/Engine/Plugins/Runtime/AudioWidgets/Source/AudioWidgets/Private/SAudioVectorscopePanelWidget.cpp
--- a/Engine/Plugins/Runtime/AudioWidgets/Source/AudioWidgets/Private/SAudioVectorscopePanelWidget.cpp
+++ b/Engine/Plugins/Runtime/AudioWidgets/Source/AudioWidgets/Private/SAudioVectorscopePanelWidget.cpp
@@ -20,23 +20,29 @@
 
 void SAudioVectorscopePanelWidget::Construct(const FArguments& InArgs, const FFixedSampledSequenceView& InData)
 {
-	DataView = InData;
-
 	PanelLayoutType = InArgs._PanelLayoutType;
 
+	ValueGridMaxDivisionParameter = InArgs._ValueGridMaxDivisionParameter;
 	bHideValueGrid = InArgs._HideGrid.Get();
 
 	check(InArgs._PanelStyle);
 	PanelStyle = InArgs._PanelStyle;
 
+	BuildWidget(InData);
+}
+
+void SAudioVectorscopePanelWidget::BuildWidget(const FFixedSampledSequenceView& InData)
+{
+	DataView = InData;
+
 	CreateBackground(PanelStyle->VectorViewerStyle);
 
-	ValueGridOverlayXAxis = CreateValueGridOverlay(InArgs._ValueGridMaxDivisionParameter,
+	ValueGridOverlayXAxis = CreateValueGridOverlay(ValueGridMaxDivisionParameter,
 		SampledSequenceValueGridOverlay::EGridDivideMode::MidSplit,
 		PanelStyle->ValueGridStyle,
 		SampledSequenceDrawingUtils::ESampledSequenceDrawOrientation::Horizontal);
 
-	ValueGridOverlayYAxis = CreateValueGridOverlay(InArgs._ValueGridMaxDivisionParameter,
+	ValueGridOverlayYAxis = CreateValueGridOverlay(ValueGridMaxDivisionParameter,
 		SampledSequenceValueGridOverlay::EGridDivideMode::MidSplit,
 		PanelStyle->ValueGridStyle,
 		SampledSequenceDrawingUtils::ESampledSequenceDrawOrientation::Vertical);
diff --git a/Engine/Plugins/Runtime/AudioWidgets/Source/AudioWidgets/Public/AudioOscilloscope.h b/Engine/Plugins/Runtime/AudioWidgets/Source/AudioWidgets/Public/AudioOscilloscope.h
--- a/Engine/Plugins/Runtime/AudioWidgets/Source/AudioWidgets/Public/AudioOscilloscope.h
+++ b/Engine/Plugins/Runtime/AudioWidgets/Source/AudioWidgets/Public/AudioOscilloscope.h
@@ -25,6 +25,16 @@ namespace AudioWidgets
 			const float InAnalysisPeriodMs, 
 			const EAudioPanelLayoutType InPanelLayoutType);
 
+		void CreateAudioBus(const uint32 InNumChannels);
+
+		void CreateDataProvider(UWorld* InWorld,
+			const float InTimeWindowMs,
+			const float InMaxTimeWindowMs,
+			const float InAnalysisPeriodMs,
+			const EAudioPanelLayoutType InPanelLayoutType);
+
+		void CreateOscilloscopeWidget(const uint32 InNumChannels, const EAudioPanelLayoutType InPanelLayoutType);
+
 		void StartProcessing();
 		void StopProcessing();
 
diff --git a/Engine/Plugins/Runtime/AudioWidgets/Source/AudioWidgets/Public/AudioVectorscope.h b/Engine/Plugins/Runtime/AudioWidgets/Source/AudioWidgets/Public/AudioVectorscope.h
--- a/Engine/Plugins/Runtime/AudioWidgets/Source/AudioWidgets/Public/AudioVectorscope.h
+++ b/Engine/Plugins/Runtime/AudioWidgets/Source/AudioWidgets/Public/AudioVectorscope.h
@@ -24,6 +24,10 @@ namespace AudioWidgets
 			const float InAnalysisPeriodMs, 
 			const EAudioPanelLayoutType InPanelLayoutType);
 
+		void CreateAudioBus(const uint32 InNumChannels);
+		void CreateDataProvider(UWorld* InWorld, const float InTimeWindowMs, const float InMaxTimeWindowMs, const float InAnalysisPeriodMs);
+		void CreateVectorscopeWidget(const EAudioPanelLayoutType InPanelLayoutType);
+
 		void StartProcessing();
 		void StopProcessing();
 
diff --git a/Engine/Plugins/Runtime/AudioWidgets/Source/AudioWidgets/Public/SAudioOscilloscopePanelWidget.h b/Engine/Plugins/Runtime/AudioWidgets/Source/AudioWidgets/Public/SAudioOscilloscopePanelWidget.h
--- a/Engine/Plugins/Runtime/AudioWidgets/Source/AudioWidgets/Public/SAudioOscilloscopePanelWidget.h
+++ b/Engine/Plugins/Runtime/AudioWidgets/Source/AudioWidgets/Public/SAudioOscilloscopePanelWidget.h
@@ -91,6 +91,8 @@ public:
 
 	void Construct(const FArguments& InArgs, const FFixedSampledSequenceView& InData, const int32 InNumChannels);
 
+	void BuildWidget(const FFixedSampledSequenceView& InData, const int32 InNumChannels);
+
 	// SWidget methods overrides
 	virtual void Tick(const FGeometry& AllottedGeometry, const double InCurrentTime, const float InDeltaTime) override;
 
@@ -179,6 +181,8 @@ private:
 	EXAxisLabelsUnit SequenceRulerDisplayUnit;
 	EYAxisLabelsUnit ValueGridOverlayDisplayUnit;
 
+	uint32 ValueGridMaxDivisionParameter = 2;
+
 	FFixedSampledSequenceView DataView;
 	SampledSequenceDrawingUtils::FSampledSequenceDrawingParams DrawingParams;
 
diff --git a/Engine/Plugins/Runtime/AudioWidgets/Source/AudioWidgets/Public/SAudioVectorscopePanelWidget.h b/Engine/Plugins/Runtime/AudioWidgets/Source/AudioWidgets/Public/SAudioVectorscopePanelWidget.h
--- a/Engine/Plugins/Runtime/AudioWidgets/Source/AudioWidgets/Public/SAudioVectorscopePanelWidget.h
+++ b/Engine/Plugins/Runtime/AudioWidgets/Source/AudioWidgets/Public/SAudioVectorscopePanelWidget.h
@@ -45,6 +45,8 @@ public:
 
 	void Construct(const FArguments& InArgs, const FFixedSampledSequenceView& InData);
 
+	void BuildWidget(const FFixedSampledSequenceView& InData);
+
 	// IFixedSampledSequenceViewReceiver Interface
 	virtual void ReceiveSequenceView(const FFixedSampledSequenceView InData, const uint32 FirstSampleIndex) override;
 
@@ -92,6 +94,8 @@ private:
 
 	EAudioPanelLayoutType PanelLayoutType = EAudioPanelLayoutType::Basic;
 
+	uint32 ValueGridMaxDivisionParameter = 2;
+
 	bool bIsInputWidgetTransacting = false;
 
 	float TimeWindowValue = 0.0f;
diff --git a/Engine/Plugins/Runtime/Metasound/Source/MetasoundEditor/Private/MetasoundEditor.cpp b/Engine/Plugins/Runtime/Metasound/Source/MetasoundEditor/Private/MetasoundEditor.cpp
--- a/Engine/Plugins/Runtime/Metasound/Source/MetasoundEditor/Private/MetasoundEditor.cpp
+++ b/Engine/Plugins/Runtime/Metasound/Source/MetasoundEditor/Private/MetasoundEditor.cpp
@@ -5,6 +5,9 @@
 #include "Algo/AnyOf.h"
 #include "AudioDevice.h"
 #include "AudioMeterStyle.h"
+#include "AudioOscilloscope.h"
+#include "AudioVectorscope.h"
+#include "AudioWidgetsEnums.h"
 #include "Components/AudioComponent.h"
 #include "DetailLayoutBuilder.h"
 #include "EdGraph/EdGraphNode.h"
@@ -67,6 +70,7 @@
 #include "Widgets/Docking/SDockTab.h"
 #include "Widgets/Images/SImage.h"
 #include "Widgets/Input/SButton.h"
+#include "Widgets/Layout/SSplitter.h"
 #include "Widgets/Notifications/SNotificationList.h"
 #include "Widgets/SBoxPanel.h"
 #include "Widgets/SOverlay.h"
@@ -677,7 +681,7 @@ namespace Metasound
 
 		TSharedPtr<SWidget> FEditor::BuildAnalyzerWidget() const
 		{
-			if (!OutputMeter.IsValid())
+			if (!OutputMeter.IsValid() || !OutputOscilloscope.IsValid() || !OutputVectorscope.IsValid())
 			{
 				return SNullWidget::NullWidget->AsShared();
 			}
@@ -697,13 +701,28 @@ namespace Metasound
 			]
 			+ SOverlay::Slot()
 			[
-				SNew(SVerticalBox)
-				+ SVerticalBox::Slot()
-				.FillHeight(1.0f)
-				.HAlign(HAlign_Center)
-				.VAlign(VAlign_Fill)
+				SNew(SSplitter)
+				.Orientation(Orient_Vertical)
+				+ SSplitter::Slot()
+				.Value(0.6f)
 				[
-					OutputMeter->GetWidget()
+					SNew(SVerticalBox)
+					+ SVerticalBox::Slot()
+					.HAlign(HAlign_Center)
+					.VAlign(VAlign_Fill)
+					[
+						OutputMeter->GetWidget()
+					]
+				]
+				+ SSplitter::Slot()
+				.Value(0.2f)
+				[
+					OutputOscilloscope->GetPanelWidget()
+				]
+				+ SSplitter::Slot()
+				.Value(0.2f)
+				[
+					OutputVectorscope->GetPanelWidget()
 				]
 			];
 		}
@@ -1150,12 +1169,61 @@ namespace Metasound
 					if (ensure(EditorWorld))
 					{
 						OutputMeter->Init(MetaSoundSource->NumChannels, *EditorWorld);
+
+						const uint32 MetaSoundNumChannels = static_cast<uint32>(MetaSoundSource->NumChannels);
+
+						// Init Oscilloscope
+						constexpr float OscilloscopeTimeWindowMs     = 10.0f;
+						constexpr float OscilloscopeMaxTimeWindowMs  = 10.0f;
+						constexpr float OscilloscopeAnalysisPeriodMs = 10.0f;
+						constexpr EAudioPanelLayoutType OscilloscopePanelLayoutType = EAudioPanelLayoutType::Basic;
+
+						if (!OutputOscilloscope.IsValid())
+						{
+							OutputOscilloscope = MakeShared<AudioWidgets::FAudioOscilloscope>(EditorWorld,
+								MetaSoundNumChannels,
+								OscilloscopeTimeWindowMs,
+								OscilloscopeMaxTimeWindowMs,
+								OscilloscopeAnalysisPeriodMs,
+								OscilloscopePanelLayoutType);
+						}
+						else
+						{
+							OutputOscilloscope->CreateAudioBus(MetaSoundNumChannels);
+							OutputOscilloscope->CreateDataProvider(EditorWorld,	OscilloscopeTimeWindowMs, OscilloscopeMaxTimeWindowMs, OscilloscopeAnalysisPeriodMs, OscilloscopePanelLayoutType);
+							OutputOscilloscope->CreateOscilloscopeWidget(MetaSoundNumChannels, EAudioPanelLayoutType::Basic);
+						}
+
+						// Init Vectorscope
+						constexpr float VectorscopeTimeWindowMs     = 30.0f;
+						constexpr float VectorscopeMaxTimeWindowMs  = 30.0f;
+						constexpr float VectorscopeAnalysisPeriodMs = 10.0f;
+						constexpr EAudioPanelLayoutType VectorscopePanelLayoutType = EAudioPanelLayoutType::Basic;
+
+						if (!OutputVectorscope.IsValid())
+						{
+							OutputVectorscope = MakeShared<AudioWidgets::FAudioVectorscope>(EditorWorld,
+								MetaSoundNumChannels,
+								VectorscopeTimeWindowMs,
+								VectorscopeMaxTimeWindowMs,
+								VectorscopeAnalysisPeriodMs,
+								VectorscopePanelLayoutType);
+						}
+						else
+						{
+							OutputVectorscope->CreateAudioBus(MetaSoundNumChannels);
+							OutputVectorscope->CreateDataProvider(EditorWorld, VectorscopeTimeWindowMs, VectorscopeMaxTimeWindowMs, VectorscopeAnalysisPeriodMs);
+							OutputVectorscope->CreateVectorscopeWidget(VectorscopePanelLayoutType);
+						}
+
 					}
 				}
 			}
 			else
 			{
 				OutputMeter.Reset();
+				OutputOscilloscope.Reset();
+				OutputVectorscope.Reset();
 			}
 		}
 
@@ -1165,6 +1233,9 @@ namespace Metasound
 			{
 				OutputMeter->Teardown();
 			}
+
+			OutputOscilloscope->StopProcessing();
+			OutputVectorscope->StopProcessing();
 		}
 
 		void FEditor::ExtendToolbar()
@@ -1545,7 +1616,17 @@ namespace Metasound
 					{
 						PreviewComp->SetAudioBusSendPostEffect(AudioBus, 1.0f);
 					}
-				
+
+					if (UAudioBus* AudioBus = OutputOscilloscope->GetAudioBus())
+					{
+						PreviewComp->SetAudioBusSendPostEffect(AudioBus, 1.0f);
+					}
+
+					if (UAudioBus* AudioBus = OutputVectorscope->GetAudioBus())
+					{
+						PreviewComp->SetAudioBusSendPostEffect(AudioBus, 1.0f);
+					}
+
 					FMetasoundAssetBase* MetasoundAsset = IMetasoundUObjectRegistry::Get().GetObjectAsAssetBase(Metasound);
 					check(MetasoundAsset);
 
@@ -1642,6 +1723,16 @@ namespace Metasound
 						OutputMeterWidget->bIsActiveTimerRegistered = true;
 					}
 				}
+
+				if (OutputOscilloscope.IsValid())
+				{
+					OutputOscilloscope->StartProcessing();
+				}
+
+				if (OutputVectorscope.IsValid())
+				{
+					OutputVectorscope->StartProcessing();
+				}
 			}
 		}
 
diff --git a/Engine/Plugins/Runtime/Metasound/Source/MetasoundEditor/Private/MetasoundEditor.h b/Engine/Plugins/Runtime/Metasound/Source/MetasoundEditor/Private/MetasoundEditor.h
--- a/Engine/Plugins/Runtime/Metasound/Source/MetasoundEditor/Private/MetasoundEditor.h
+++ b/Engine/Plugins/Runtime/Metasound/Source/MetasoundEditor/Private/MetasoundEditor.h
@@ -36,6 +36,12 @@
 
 
 // Forward Declarations
+namespace AudioWidgets
+{
+	class FAudioOscilloscope;
+	class FAudioVectorscope;
+}
+
 class FTabManager;
 class SDockableTab;
 class SGraphEditor;
@@ -420,6 +426,12 @@ namespace Metasound
 			/** Meter used in the analyzer tab for auditioning preview output. */
 			TSharedPtr<AudioWidgets::FAudioMeter> OutputMeter;
 
+			/** Oscilloscope used in the analyzer tab for visualizing preview output. */
+			TSharedPtr<AudioWidgets::FAudioOscilloscope> OutputOscilloscope;
+
+			/** Vectorscope used in the analyzer tab for visualizing preview output. */
+			TSharedPtr<AudioWidgets::FAudioVectorscope> OutputVectorscope;
+
 			/** Palette of Node types */
 			TSharedPtr<SMetasoundPalette> Palette;
 
