Fix pops when swapping connections in graph

diff --git a/Engine/Plugins/Runtime/Metasound/Source/MetasoundGraphCore/Private/MetasoundDynamicGraphAlgo.cpp b/Engine/Plugins/Runtime/Metasound/Source/MetasoundGraphCore/Private/MetasoundDynamicGraphAlgo.cpp
--- a/Engine/Plugins/Runtime/Metasound/Source/MetasoundGraphCore/Private/MetasoundDynamicGraphAlgo.cpp
+++ b/Engine/Plugins/Runtime/Metasound/Source/MetasoundGraphCore/Private/MetasoundDynamicGraphAlgo.cpp
@@ -17,20 +17,21 @@ namespace Metasound
 		namespace DynamicGraphAlgoPrivate
 		{
 			template<typename EntryType, typename FunctionType>
-			void UpdateTableEntry(const TArray<FOperatorID>& InOperatorOrder, const FOperatorID& InOperatorID, IOperator& InOperator, TArray<EntryType>& InOutTable, FunctionType InFunction)
+			void UpdateTableEntry(const TArray<FOperatorID>& InOperatorOrder, const FOperatorID& InOperatorID, IOperator* InOperator, TArray<EntryType>& InOutTable, FunctionType InFunction)
 			{
 				using namespace DirectedGraphAlgo;
 
 				int32 EntryIndex = InOutTable.IndexOfByPredicate([&](const EntryType& InEntry) { return InEntry.OperatorID == InOperatorID; });
 
-				bool bEntryShouldExist = (nullptr != InFunction);
+				int32 OperatorIndex = InOperatorOrder.Find(InOperatorID);
+				bool bEntryShouldExist = (nullptr != InOperator) && (nullptr != InFunction) && (INDEX_NONE != OperatorIndex); // Remove operators if they have no function for this table, or if they or not in the operator execution order. 
 				bool bEntryCurrentlyExists = (EntryIndex != INDEX_NONE);
 
 				if (bEntryCurrentlyExists && bEntryShouldExist)
 				{
 					// Update the existing entry
 					InOutTable[EntryIndex].Function = InFunction;
-					InOutTable[EntryIndex].Operator = &InOperator;
+					InOutTable[EntryIndex].Operator = InOperator;
 				}
 				else if (bEntryCurrentlyExists)
 				{
@@ -39,14 +40,8 @@ namespace Metasound
 				}
 				else if (bEntryShouldExist)
 				{
+					check(InOperator);
 					// Add new entry
-					int32 OperatorIndex = InOperatorOrder.Find(InOperatorID);
-					if (INDEX_NONE == OperatorIndex)
-					{
-						UE_LOG(LogMetaSound, Error, TEXT("Failed to update operator stack. OperatorID not found in operator order array."));
-						return;
-					}
-
 					int32 PreceedingOperatorIndex = OperatorIndex - 1;
 					int32 InsertLocation = 0; // Default to inserting at the beginning
 					while (PreceedingOperatorIndex >= 0)
@@ -62,18 +57,21 @@ namespace Metasound
 						}
 					}
 
-					InOutTable.Insert(EntryType{InOperatorID, InOperator, InFunction}, InsertLocation);
+					InOutTable.Insert(EntryType{InOperatorID, *InOperator, InFunction}, InsertLocation);
 				}
 			}
 
-			void UpdateGraphRuntimeTables(const FOperatorID& InOperatorID, IOperator& InOperator, FDynamicGraphOperatorData& InOutGraphOperatorData)
-			{
-				UpdateTableEntry(InOutGraphOperatorData.OperatorOrder, InOperatorID, InOperator, InOutGraphOperatorData.ExecuteTable, InOperator.GetExecuteFunction());
-				UpdateTableEntry(InOutGraphOperatorData.OperatorOrder, InOperatorID, InOperator, InOutGraphOperatorData.PostExecuteTable, InOperator.GetPostExecuteFunction());
-				UpdateTableEntry(InOutGraphOperatorData.OperatorOrder, InOperatorID, InOperator, InOutGraphOperatorData.ResetTable, InOperator.GetResetFunction());
-			}
 		} // namespace DynamicGraphAlgoPrivate
 
+		void UpdateGraphRuntimeTableEntries(const FOperatorID& InOperatorID, IOperator* InOperator, FDynamicGraphOperatorData& InOutGraphOperatorData)
+		{
+			using namespace DynamicGraphAlgoPrivate;
+
+			UpdateTableEntry(InOutGraphOperatorData.OperatorOrder, InOperatorID, InOperator, InOutGraphOperatorData.ExecuteTable, InOperator ? InOperator->GetExecuteFunction() : nullptr);
+			UpdateTableEntry(InOutGraphOperatorData.OperatorOrder, InOperatorID, InOperator, InOutGraphOperatorData.PostExecuteTable, InOperator ? InOperator->GetPostExecuteFunction() : nullptr);
+			UpdateTableEntry(InOutGraphOperatorData.OperatorOrder, InOperatorID, InOperator, InOutGraphOperatorData.ResetTable, InOperator ? InOperator->GetResetFunction() : nullptr);
+		}
+
 		// Apply updates to data references through all the operators by following connections described in the FOperatorInfo map.
 		void PropagateBindUpdate(FOperatorID InInitialOperatorID, const FVertexName& InVertexName, const FAnyDataReference& InNewReference, FDynamicGraphOperatorData& InOutGraphOperatorData)
 		{
@@ -110,7 +108,7 @@ namespace Metasound
 					Operator.BindOutputs(OpInfo->VertexData.GetOutputs());
 
 					// Update execute/postexecute/reset tables in case those have changed after rebinding.
-					UpdateGraphRuntimeTables(Current.OperatorID, Operator, InOutGraphOperatorData);
+					UpdateGraphRuntimeTableEntries(Current.OperatorID, &Operator, InOutGraphOperatorData);
 
 					// See if binding altered the outputs. 
 					OutputUpdates.Reset();
@@ -258,7 +256,7 @@ namespace Metasound
 			InOperatorInfo.Operator->BindOutputs(InOperatorInfo.VertexData.GetOutputs());
 
 			// Update any execution tables that need updating after wrapping
-			UpdateGraphRuntimeTables(InOperatorID, *(InOperatorInfo.Operator), InOutGraphOperatorData);
+			UpdateGraphRuntimeTableEntries(InOperatorID, InOperatorInfo.Operator.Get(), InOutGraphOperatorData);
 
 			// Determine if there have been changes to `OutputVertexData`. 
 			TSortedVertexNameMap<FAnyDataReference> OutputsToUpdate;
diff --git a/Engine/Plugins/Runtime/Metasound/Source/MetasoundGraphCore/Private/MetasoundDynamicGraphAlgo.h b/Engine/Plugins/Runtime/Metasound/Source/MetasoundGraphCore/Private/MetasoundDynamicGraphAlgo.h
--- a/Engine/Plugins/Runtime/Metasound/Source/MetasoundGraphCore/Private/MetasoundDynamicGraphAlgo.h
+++ b/Engine/Plugins/Runtime/Metasound/Source/MetasoundGraphCore/Private/MetasoundDynamicGraphAlgo.h
@@ -87,6 +87,21 @@ namespace Metasound
 			TArray<FResetEntry> ResetTable;
 		};
 
+		/** Update the runtime table entries.
+		 * 
+		 * This updates the function tables in the FDynamicGraphOperatorData for a given operator. For an operator to exist in the function tables
+		 * it must provide a function to be called from GetExecutionFunction(), GetPostExecutionFunction() and/or GetResetFunction() and the operator's ID
+		 * must also be present in the FDynamicGraphOperatorData's operator order array.  If either of these are not true, then the operator will 
+		 * not have an entry in a specific execution table. 
+		 * 
+		 * Note: This function does not update the order of the entry in the execution table. 
+		 * 
+		 * @param InOperatorID - ID of operator to update.
+		 * @param InOperator - Pointer to operator.
+		 * @param InOutGraphOperatorData - Structure containing the operator order and runtime tables.
+		 */
+		void UpdateGraphRuntimeTableEntries(const FOperatorID& InOperatorID, IOperator* InOperator, FDynamicGraphOperatorData& InOutGraphOperatorData);
+
 		/** Propagate FVertexInterfaceData updates through the operators in the Dynamic Graph Operator Data. 
 		 *
 		 * A change to an operator's input may result in a change to the operator's output. The updates to the
diff --git a/Engine/Plugins/Runtime/Metasound/Source/MetasoundGraphCore/Private/MetasoundDynamicOperator.cpp b/Engine/Plugins/Runtime/Metasound/Source/MetasoundGraphCore/Private/MetasoundDynamicOperator.cpp
--- a/Engine/Plugins/Runtime/Metasound/Source/MetasoundGraphCore/Private/MetasoundDynamicOperator.cpp
+++ b/Engine/Plugins/Runtime/Metasound/Source/MetasoundGraphCore/Private/MetasoundDynamicOperator.cpp
@@ -226,15 +226,26 @@ namespace Metasound
 		{
 		}
 
-		EDynamicOperatorTransformQueueAction FSetOperatorOrder::Transform(FDynamicGraphOperatorData& InGraphOperatorData)
+		EDynamicOperatorTransformQueueAction FSetOperatorOrder::Transform(FDynamicGraphOperatorData& InOutGraphOperatorData)
 		{
 			METASOUND_TRACE_CPUPROFILER_EVENT_SCOPE(Metasound::DynamicOperator::SetOperatorOrder)
 
 			using namespace DynamicOperatorPrivate;
 
-			FTableSorter::SortTable(Order, InGraphOperatorData.ExecuteTable);
-			FTableSorter::SortTable(Order, InGraphOperatorData.PostExecuteTable);
-			FTableSorter::SortTable(Order, InGraphOperatorData.ResetTable);
+			InOutGraphOperatorData.OperatorOrder = Order;
+
+			// Update the entry into the various runtime tables dependent upon when the operator exists
+			// in the operator order array.
+			for (const TPair<FOperatorID, FOperatorInfo>& OpIDAndInfo : InOutGraphOperatorData.OperatorMap)
+			{
+				IOperator* Operator = OpIDAndInfo.Value.Operator.Get();
+				UpdateGraphRuntimeTableEntries(OpIDAndInfo.Key, Operator, InOutGraphOperatorData);
+			}
+
+			// Sort operator tables to be in the correct order.
+			FTableSorter::SortTable(Order, InOutGraphOperatorData.ExecuteTable);
+			FTableSorter::SortTable(Order, InOutGraphOperatorData.PostExecuteTable);
+			FTableSorter::SortTable(Order, InOutGraphOperatorData.ResetTable);
 
 			return EDynamicOperatorTransformQueueAction::Continue;
 		}
diff --git a/Engine/Plugins/Runtime/Metasound/Source/MetasoundGraphCore/Private/MetasoundDynamicOperatorTransactor.cpp b/Engine/Plugins/Runtime/Metasound/Source/MetasoundGraphCore/Private/MetasoundDynamicOperatorTransactor.cpp
--- a/Engine/Plugins/Runtime/Metasound/Source/MetasoundGraphCore/Private/MetasoundDynamicOperatorTransactor.cpp
+++ b/Engine/Plugins/Runtime/Metasound/Source/MetasoundGraphCore/Private/MetasoundDynamicOperatorTransactor.cpp
@@ -207,16 +207,16 @@ namespace Metasound
 				return;
 			}
 
-			// Store literal node associated with the target of the literal value.
-			LiteralNodeMap.Add(FLiteralNodeID{InToNodeID, InToVertex}, MoveTemp(InReplacementLiteralNode));
 			/* remove edge from internal graph. */
 			bool bSuccess = Graph.RemoveDataEdge(*FromNode, InFromVertex, *ToNode, InToVertex);
 			if (!bSuccess)
 			{
-				UE_LOG(LogMetaSound, Warning, TEXT("Failed to remove edge from %s:%s to %s:%s on internal graph."), *InFromNodeID.ToString(), *InFromVertex.ToString(), *InToNodeID.ToString(), *InToVertex.ToString());
-				// No need to early exit here because edge likely did not exist.
+				UE_LOG(LogMetaSound, Error, TEXT("Failed to remove edge from %s:%s to %s:%s on internal graph."), *InFromNodeID.ToString(), *InFromVertex.ToString(), *InToNodeID.ToString(), *InToVertex.ToString());
+				return;
 			}
 
+			// Store literal node associated with the target of the literal value.
+			LiteralNodeMap.Add(FLiteralNodeID{ InToNodeID, InToVertex }, MoveTemp(InReplacementLiteralNode));
 			bSuccess = Graph.AddDataEdge(*ReplacementLiteralNode, DynamicOperatorTransactorPrivate::LiteralNodeOutputVertexName, *ToNode, InToVertex);
 			if (!bSuccess)
 			{
@@ -269,7 +269,7 @@ namespace Metasound
 
 			AddDataEdgeInternal(*LiteralNode, DynamicOperatorTransactorPrivate::LiteralNodeOutputVertexName, InNodeID, *Node, InVertex);
 
-			// Add literal node before calling "AddDataEdgeInternal" so that AddDataEdgeInternal can check if there is an existing literal node.
+			// Add literal node after calling "AddDataEdgeInternal" so that AddDataEdgeInternal can check if there is a prior existing literal node.
 			LiteralNodeMap.Add(FLiteralNodeID{InNodeID, InVertex}, MoveTemp(InLiteralNode));
 		}
 
