Fixing and clarifying midi note quantizer tooltip

diff --git a/Engine/Plugins/Runtime/Metasound/Source/MetasoundStandardNodes/Private/MetasoundMidiNoteQuantizerNode.cpp b/Engine/Plugins/Runtime/Metasound/Source/MetasoundStandardNodes/Private/MetasoundMidiNoteQuantizerNode.cpp
--- a/Engine/Plugins/Runtime/Metasound/Source/MetasoundStandardNodes/Private/MetasoundMidiNoteQuantizerNode.cpp
+++ b/Engine/Plugins/Runtime/Metasound/Source/MetasoundStandardNodes/Private/MetasoundMidiNoteQuantizerNode.cpp
@@ -27,7 +27,7 @@ namespace Metasound
 	{
 		// inputs
 		METASOUND_PARAM(ParamNoteIn, "Note In", "Midi Note to quantize");
-		METASOUND_PARAM(ParamRootNote, "Root Note", "Midi note to treat as the Root (0 = C, 1 = D, etc.).  Octave does not matter. Values < 0 will clamp to 0");
+		METASOUND_PARAM(ParamRootNote, "Root Note", "Midi note to treat as the Root (e.g. where 0.0 = C, 1.0 = Db/C#, etc). Values are clamped to positive values.");
 		METASOUND_PARAM(ParamScaleDegrees, "Scale Degrees", "Set of notes in ascending order, represeting half steps starting at 0.0f meaning the Root Note. Scale degrees should be the notes in a scale that not including the octave. Can be the output of the Scale To Note Array node.");
 		METASOUND_PARAM(ParamScaleRange, "Scale Range In", "The number of semitones in the scale. E.g. a regular diatonic scale will be 12 semitones. Exotic scales could be something else.");
 
