Fix for bindinputs failuer on metasound send node.

diff --git a/Engine/Plugins/Runtime/Metasound/Source/MetasoundFrontend/Public/MetasoundSendNode.h b/Engine/Plugins/Runtime/Metasound/Source/MetasoundFrontend/Public/MetasoundSendNode.h
--- a/Engine/Plugins/Runtime/Metasound/Source/MetasoundFrontend/Public/MetasoundSendNode.h
+++ b/Engine/Plugins/Runtime/Metasound/Source/MetasoundFrontend/Public/MetasoundSendNode.h
@@ -109,29 +109,13 @@ namespace Metasound
 				{
 					using namespace SendVertexNames; 
 					InOutVertexData.BindReadVertex<FSendAddress>(METASOUND_GET_PARAM_NAME(AddressInput), SendAddress);
-					InOutVertexData.BindReadVertex<TDataType>(GetSendInputName(), TDataReadReference<TDataType>(InputData));
+					InOutVertexData.BindReadVertex<TDataType>(GetSendInputName(), InputData);
 				}
 
 				virtual void BindOutputs(FOutputVertexInterfaceData& InOutVertexData) override
 				{
 				}
 
-				virtual FDataReferenceCollection GetInputs() const override
-				{
-					// This should never be called. Bind(...) is called instead. This method
-					// exists as a stop-gap until the API can be deprecated and removed.
-					checkNoEntry();
-					return {};
-				}
-
-				virtual FDataReferenceCollection GetOutputs() const override
-				{
-					// This should never be called. Bind(...) is called instead. This method
-					// exists as a stop-gap until the API can be deprecated and removed.
-					checkNoEntry();
-					return {};
-				}
-
 				void Execute()
 				{
 					if (*SendAddress != CachedSendAddress)
@@ -214,35 +198,18 @@ namespace Metasound
 
 			TSendNode(const FNodeInitData& InInitData)
 				: FNode(InInitData.InstanceName, InInitData.InstanceID, GetNodeInfo())
-				, Interface(DeclareVertexInterface())
 				, Factory(MakeOperatorFactoryRef<FSendOperatorFactory>())
 			{
 			}
 
 			virtual ~TSendNode() = default;
 
-			virtual const FVertexInterface& GetVertexInterface() const override
-			{
-				return Interface;
-			}
-
-			virtual bool SetVertexInterface(const FVertexInterface& InInterface) override
-			{
-				return Interface == InInterface;
-			}
-
-			virtual bool IsVertexInterfaceSupported(const FVertexInterface& InInterface) const override
-			{
-				return Interface == InInterface;
-			}
-
 			virtual FOperatorFactorySharedRef GetDefaultOperatorFactory() const override
 			{
 				return Factory;
 			}
 
 		private:
-			FVertexInterface Interface;
 			FOperatorFactorySharedRef Factory;
 	};
 }
