Fix pop when removing node after several nodes hav e been added and connected and disconnected.

diff --git a/Engine/Plugins/Runtime/Metasound/Source/MetasoundGraphCore/Private/MetasoundDynamicOperatorTransactor.cpp b/Engine/Plugins/Runtime/Metasound/Source/MetasoundGraphCore/Private/MetasoundDynamicOperatorTransactor.cpp
--- a/Engine/Plugins/Runtime/Metasound/Source/MetasoundGraphCore/Private/MetasoundDynamicOperatorTransactor.cpp
+++ b/Engine/Plugins/Runtime/Metasound/Source/MetasoundGraphCore/Private/MetasoundDynamicOperatorTransactor.cpp
@@ -571,8 +571,8 @@ namespace Metasound
 				// Add the literal node.
 				TUniquePtr<IDynamicOperatorTransform> AddNodeTransform = CreateAddOperatorTransform(InReplacementLiteralNode, InOperatorSettings, InEnvironment);
 
-				// Connect literal node to target node.
-				TUniquePtr<IDynamicOperatorTransform> ConnectOperatorsTransform = MakeUnique<FConnectOperators>(LiteralOperatorID, DynamicOperatorTransactorPrivate::LiteralNodeOutputVertexName, ToOperatorID, InToVertex);
+				// Swap prior connection with new connections.
+				TUniquePtr<IDynamicOperatorTransform> ConnectOperatorsTransform = MakeUnique<FSwapOperatorConnection>(FromOperatorID, InFromVertex, LiteralOperatorID, DynamicOperatorTransactorPrivate::LiteralNodeOutputVertexName, ToOperatorID, InToVertex);
 
 				// Reorder the node graph.
 				TUniquePtr<IDynamicOperatorTransform> SetOrderTransform = MakeUnique<FSetOperatorOrder>(InNewOperatorOrder);
