Warn when failing to paste a node that has a missing class from frontend class registry rather than ensuring
#rnx

diff --git a/Engine/Plugins/Runtime/Metasound/Source/MetasoundEditor/Private/MetasoundEditor.cpp b/Engine/Plugins/Runtime/Metasound/Source/MetasoundEditor/Private/MetasoundEditor.cpp
--- a/Engine/Plugins/Runtime/Metasound/Source/MetasoundEditor/Private/MetasoundEditor.cpp
+++ b/Engine/Plugins/Runtime/Metasound/Source/MetasoundEditor/Private/MetasoundEditor.cpp
@@ -2522,10 +2522,19 @@ namespace Metasound
 						}
 						else
 						{
-							FNodeHandle NewHandle = FGraphBuilder::AddNodeHandle(*Metasound, *ExternalNode);
-							if (!NewHandle->IsValid())
+							FMetasoundFrontendClass ExternalClass;
+							if (ISearchEngine::Get().FindClassWithHighestVersion(ClassName, ExternalClass))
+							{
+								FNodeHandle NewHandle = FGraphBuilder::AddNodeHandle(*Metasound, *ExternalNode);
+								if (!NewHandle->IsValid())
+								{
+									NodesToRemove.Add(GraphNode);
+								}
+							}
+							else
 							{
 								NodesToRemove.Add(GraphNode);
+								UE_LOG(LogMetaSound, Warning, TEXT("Cannot add pasted node with class '%s': Node class not found"), *ClassName.ToString());
 							}
 						}
 					}
