- Remove --> RemoveSwap delegates to allow for handlers to not duplicate forwarded messages where applicable
- Refactor remove node handler for transactor to not do swap piecemeal, which was resulting in dropped connections
- Add test for adding/removing nodes while auditioning (excercises add/remove edges by extension)
#rnx

diff --git a/Engine/Plugins/Runtime/Metasound/Source/MetasoundEngine/Private/MetasoundBuilderSubsystem.cpp b/Engine/Plugins/Runtime/Metasound/Source/MetasoundEngine/Private/MetasoundBuilderSubsystem.cpp
--- a/Engine/Plugins/Runtime/Metasound/Source/MetasoundEngine/Private/MetasoundBuilderSubsystem.cpp
+++ b/Engine/Plugins/Runtime/Metasound/Source/MetasoundEngine/Private/MetasoundBuilderSubsystem.cpp
@@ -1036,7 +1036,7 @@ bool UMetaSoundSourceBuilder::GetLiveUpdatesEnabled() const
 void UMetaSoundSourceBuilder::InitDelegates(Metasound::Frontend::FDocumentModifyDelegates& OutDocumentDelegates) const
 {
 	OutDocumentDelegates.EdgeDelegates.OnEdgeAdded.AddUObject(this, &UMetaSoundSourceBuilder::OnEdgeAdded);
-	OutDocumentDelegates.EdgeDelegates.OnRemovingEdge.AddUObject(this, &UMetaSoundSourceBuilder::OnRemovingEdge);
+	OutDocumentDelegates.EdgeDelegates.OnRemoveSwappingEdge.AddUObject(this, &UMetaSoundSourceBuilder::OnRemoveSwappingEdge);
 
 	OutDocumentDelegates.InterfaceDelegates.OnInputAdded.AddUObject(this, &UMetaSoundSourceBuilder::OnInputAdded);
 	OutDocumentDelegates.InterfaceDelegates.OnOutputAdded.AddUObject(this, &UMetaSoundSourceBuilder::OnOutputAdded);
@@ -1045,7 +1045,7 @@ void UMetaSoundSourceBuilder::InitDelegates(Metasound::Frontend::FDocumentModify
 
 	OutDocumentDelegates.NodeDelegates.OnNodeAdded.AddUObject(this, &UMetaSoundSourceBuilder::OnNodeAdded);
 	OutDocumentDelegates.NodeDelegates.OnNodeInputLiteralSet.AddUObject(this, &UMetaSoundSourceBuilder::OnNodeInputLiteralSet);
-	OutDocumentDelegates.NodeDelegates.OnRemovingNode.AddUObject(this, &UMetaSoundSourceBuilder::OnRemovingNode);
+	OutDocumentDelegates.NodeDelegates.OnRemoveSwappingNode.AddUObject(this, &UMetaSoundSourceBuilder::OnRemoveSwappingNode);
 	OutDocumentDelegates.NodeDelegates.OnRemovingNodeInputLiteral.AddUObject(this, &UMetaSoundSourceBuilder::OnRemovingNodeInputLiteral);
 }
 
@@ -1156,7 +1156,8 @@ void UMetaSoundSourceBuilder::OnNodeInputLiteralSet(int32 NodeIndex, int32 Verte
 	const FMetasoundFrontendNode& Node = GraphClass.Graph.Nodes[NodeIndex];
 	const FMetasoundFrontendVertex& Input = Node.Interface.Inputs[VertexIndex];
 
-	// Only send the literal down if not connected, as the graph core layer will disconnect.
+	// Only send the literal down if not connected, as the graph core layer
+	// will disconnect if a new literal is sent and edge already exists.
 	if (!Builder.IsNodeInputConnected(Node.GetID(), Input.VertexID))
 	{
 		ExecuteAuditionableTransaction([this, &Node, &Input, &LiteralIndex](FDynamicOperatorTransactor& Transactor)
@@ -1189,12 +1190,12 @@ void UMetaSoundSourceBuilder::OnOutputAdded(int32 OutputIndex) const
 	});
 }
 
-void UMetaSoundSourceBuilder::OnRemovingEdge(int32 EdgeIndex) const
+void UMetaSoundSourceBuilder::OnRemoveSwappingEdge(int32 SwapIndex, int32 LastIndex) const
 {
 	using namespace Metasound::DynamicGraph;
 
 	const FMetasoundFrontendDocument& Doc = Builder.GetDocument();
-	const FMetasoundFrontendEdge& EdgeBeingRemoved = Doc.RootGraph.Graph.Edges[EdgeIndex];
+	const FMetasoundFrontendEdge& EdgeBeingRemoved = Doc.RootGraph.Graph.Edges[SwapIndex];
 	ExecuteAuditionableTransaction([this, EdgeBeingRemoved](FDynamicOperatorTransactor& Transactor)
 	{
 		using namespace Metasound;
@@ -1240,16 +1241,20 @@ void UMetaSoundSourceBuilder::OnRemovingInput(int32 InputIndex) const
 	});
 }
 
-void UMetaSoundSourceBuilder::OnRemovingNode(int32 NodeIndex) const
+void UMetaSoundSourceBuilder::OnRemoveSwappingNode(int32 SwapIndex, int32 LastIndex) const
 {
 	using namespace Metasound::DynamicGraph;
-	ExecuteAuditionableTransaction([this, NodeIndex](FDynamicOperatorTransactor& Transactor)
+
+	// Last index will just be re-added, so this aspect of the swap is ignored by transactor
+	// (i.e. no sense removing and re-adding the node that is swapped from the end as this
+	// would potentially disconnect that node in the runtime graph model).
+	ExecuteAuditionableTransaction([this, SwapIndex](FDynamicOperatorTransactor& Transactor)
 	{
 		using namespace Metasound::Frontend;
 
 		const FMetasoundFrontendDocument& Doc = Builder.GetDocument();
 		const FMetasoundFrontendGraphClass& GraphClass = Doc.RootGraph;
-		const FMetasoundFrontendNode& NodeBeingRemoved = GraphClass.Graph.Nodes[NodeIndex];
+		const FMetasoundFrontendNode& NodeBeingRemoved = GraphClass.Graph.Nodes[SwapIndex];
 		const FGuid& NodeID = NodeBeingRemoved.GetID();
 		Transactor.RemoveNode(NodeID);
 		return true;
@@ -1260,26 +1265,34 @@ void UMetaSoundSourceBuilder::OnRemovingNodeInputLiteral(int32 NodeIndex, int32
 {
 	using namespace Metasound::DynamicGraph;
 
-	ExecuteAuditionableTransaction([this, &NodeIndex, &VertexIndex, &LiteralIndex](FDynamicOperatorTransactor& Transactor)
+	const FMetasoundFrontendGraphClass& GraphClass = Builder.GetDocument().RootGraph;
+	const FMetasoundFrontendNode& Node = GraphClass.Graph.Nodes[NodeIndex];
+	const FMetasoundFrontendVertex& Input = Node.Interface.Inputs[VertexIndex];
+
+	// Only send the literal down if not connected, as the graph core layer will disconnect.
+	if (!Builder.IsNodeInputConnected(Node.GetID(), Input.VertexID))
 	{
-		using namespace Metasound;
-		using namespace Metasound::Engine;
-		using namespace Metasound::Frontend;
+		ExecuteAuditionableTransaction([this, &NodeIndex, &VertexIndex, &LiteralIndex](FDynamicOperatorTransactor& Transactor)
+		{
+			using namespace Metasound;
+			using namespace Metasound::Engine;
+			using namespace Metasound::Frontend;
 
-		const TArray<FMetasoundFrontendNode>& Nodes = Builder.GetDocument().RootGraph.Graph.Nodes;
-		const FMetasoundFrontendNode& Node = Nodes[NodeIndex];
-		const FMetasoundFrontendVertex& Input = Node.Interface.Inputs[VertexIndex];
+			const TArray<FMetasoundFrontendNode>& Nodes = Builder.GetDocument().RootGraph.Graph.Nodes;
+			const FMetasoundFrontendNode& Node = Nodes[NodeIndex];
+			const FMetasoundFrontendVertex& Input = Node.Interface.Inputs[VertexIndex];
 
-		const FMetasoundFrontendLiteral* InputDefault = Builder.GetNodeInputClassDefault(Node.GetID(), Input.VertexID);
-		if (ensureAlwaysMsgf(InputDefault, TEXT("Could not dynamically assign default literal from class definition upon removing input literal: document's dependency entry invalid and has no default assigned")))
-		{
-			TUniquePtr<INode> LiteralNode = BuilderSubsystemPrivate::CreateDynamicNodeFromFrontendLiteral(Input.TypeName, *InputDefault);
-			Transactor.SetValue(Node.GetID(), Input.Name, MoveTemp(LiteralNode));
-			return true;
-		}
+			const FMetasoundFrontendLiteral* InputDefault = Builder.GetNodeInputClassDefault(Node.GetID(), Input.VertexID);
+			if (ensureAlwaysMsgf(InputDefault, TEXT("Could not dynamically assign default literal from class definition upon removing input '%s' literal: document's dependency entry invalid and has no default assigned"), *Input.Name.ToString()))
+			{
+				TUniquePtr<INode> LiteralNode = BuilderSubsystemPrivate::CreateDynamicNodeFromFrontendLiteral(Input.TypeName, *InputDefault);
+				Transactor.SetValue(Node.GetID(), Input.Name, MoveTemp(LiteralNode));
+				return true;
+			}
 
-		return false;
-	});
+			return false;
+		});
+	}
 }
 
 void UMetaSoundSourceBuilder::OnRemovingOutput(int32 OutputIndex) const
diff --git a/Engine/Plugins/Runtime/Metasound/Source/MetasoundEngine/Public/MetasoundBuilderSubsystem.h b/Engine/Plugins/Runtime/Metasound/Source/MetasoundEngine/Public/MetasoundBuilderSubsystem.h
--- a/Engine/Plugins/Runtime/Metasound/Source/MetasoundEngine/Public/MetasoundBuilderSubsystem.h
+++ b/Engine/Plugins/Runtime/Metasound/Source/MetasoundEngine/Public/MetasoundBuilderSubsystem.h
@@ -516,9 +516,9 @@ private:
 	void OnNodeAdded(int32 NodeIndex) const;
 	void OnNodeInputLiteralSet(int32 NodeIndex, int32 VertexIndex, int32 LiteralIndex) const;
 	void OnOutputAdded(int32 OutputIndex) const;
-	void OnRemovingEdge(int32 EdgeIndex) const;
+	void OnRemoveSwappingEdge(int32 SwapIndex, int32 LastIndex) const;
 	void OnRemovingInput(int32 InputIndex) const;
-	void OnRemovingNode(int32 NodeIndex) const;
+	void OnRemoveSwappingNode(int32 SwapIndex, int32 LastIndex) const;
 	void OnRemovingNodeInputLiteral(int32 NodeIndex, int32 VertexIndex, int32 LiteralIndex) const;
 	void OnRemovingOutput(int32 OutputIndex) const;
 
diff --git a/Engine/Plugins/Runtime/Metasound/Source/MetasoundEngineTest/Private/EngineTestMetasoundSource.cpp b/Engine/Plugins/Runtime/Metasound/Source/MetasoundEngineTest/Private/EngineTestMetasoundSource.cpp
--- a/Engine/Plugins/Runtime/Metasound/Source/MetasoundEngineTest/Private/EngineTestMetasoundSource.cpp
+++ b/Engine/Plugins/Runtime/Metasound/Source/MetasoundEngineTest/Private/EngineTestMetasoundSource.cpp
@@ -334,6 +334,33 @@ bool FAudioComponentRemoveFromRootLatentCommand::Update()
 	return false;
 }
 
+DEFINE_LATENT_AUTOMATION_COMMAND_THREE_PARAMETER(FMetaSoundSourceBuilderRemoveNodesLatentCommand, FAutomationTestBase&, Test, UMetaSoundSourceBuilder*, Builder, FMetaSoundNodeHandle, Node);
+
+bool FMetaSoundSourceBuilderRemoveNodesLatentCommand::Update()
+{
+	if (Builder)
+	{
+		EMetaSoundBuilderResult Result;
+		Builder->RemoveNode(Node, Result);
+		Test.AddErrorIfFalse(Result == EMetaSoundBuilderResult::Succeeded, "Failed to remove node from MetaSound graph");
+	}
+
+	return true;
+}
+
+DEFINE_LATENT_AUTOMATION_COMMAND_FOUR_PARAMETER(FMetaSoundSourceBuilderConnectNodesLatentCommand, FAutomationTestBase&, Test, UMetaSoundSourceBuilder*, Builder, FMetaSoundBuilderNodeOutputHandle, Output, FMetaSoundBuilderNodeInputHandle, Input);
+
+bool FMetaSoundSourceBuilderConnectNodesLatentCommand::Update()
+{
+	if (Builder)
+	{
+		EMetaSoundBuilderResult Result;
+		Builder->ConnectNodes(Output, Input, Result);
+		Test.AddErrorIfFalse(Result == EMetaSoundBuilderResult::Succeeded, TEXT("Failed to connect MetaSound nodes"));
+	}
+
+	return true;
+}
 
 // This test creates a MetaSound from the legacy controller document editing system and attempts to play it directly.
 IMPLEMENT_SIMPLE_AUTOMATION_TEST(FAudioMetasoundSourceTest, "Audio.Metasound.PlayMetasoundSource", EAutomationTestFlags::EditorContext | EAutomationTestFlags::EngineFilter)
@@ -568,4 +595,85 @@ bool FAudioMetasoundSourceBuilderMutateInterface::RunTest(const FString& Paramet
 
 	return true;
 }
+
+// This test creates a MetaSound source from a SourceBuilder, then adds and connects saw oscillator nodes setting their input freq to a new chromatic tone which cascades up, and finally removes all of the added nodes.
+IMPLEMENT_SIMPLE_AUTOMATION_TEST(FAudioMetasoundSourceBuilderAddRemoveNodes, "Audio.Metasound.AddRemoveNodes", EAutomationTestFlags::EditorContext | EAutomationTestFlags::EngineFilter)
+bool FAudioMetasoundSourceBuilderAddRemoveNodes::RunTest(const FString& Parameters)
+{
+	using namespace Audio;
+	using namespace EngineTestMetaSoundPatchBuilderPrivate;
+	using namespace EngineTestMetaSoundSourcePrivate;
+	using namespace Metasound::Frontend;
+
+	constexpr EMetaSoundOutputAudioFormat OutputFormat = EMetaSoundOutputAudioFormat::Mono;
+	constexpr bool bIsOneShot = false;
+	FInitTestBuilderSourceOutput Output;
+	UMetaSoundSourceBuilder& Builder = CreateSourceBuilder(*this, EMetaSoundOutputAudioFormat::Mono, bIsOneShot, Output);
+	Builder.AddToRoot();
+
+	EMetaSoundBuilderResult Result;
+
+		if (UAudioComponent* AudioComponent = CreateTestComponent(*this))
+		{
+			constexpr bool bEnableLiveUpdate = true;
+			ADD_LATENT_AUTOMATION_COMMAND(FMetaSoundSourceBuilderAuditionLatentCommand(&Builder, AudioComponent, bEnableLiveUpdate));
+
+			const FName OscType = "Sine";
+			const TArray<float> ChromaticFreqs
+			{
+				293.66,
+				311.13,
+				329.63,
+				349.23,
+				369.99,
+				392.00,
+				415.30,
+				440.00
+			};
+
+		TArray<FMetaSoundNodeHandle> GenNodes;
+		for (int32 i = 0; i < 8; ++i)
+		{
+			FMetaSoundNodeHandle OscNode = Builder.AddNodeByClassName({ "UE", OscType, "Audio" }, 1, Result);
+			AddErrorIfFalse(Result == EMetaSoundBuilderResult::Succeeded, "Failed to add osc node to graph");
+
+			if (Result == EMetaSoundBuilderResult::Succeeded)
+			{
+				GenNodes.Add(MoveTemp(OscNode));
+			}
+		}
+
+		for (int32 i = 0; i < GenNodes.Num(); ++i)
+		{
+			FMetaSoundNodeHandle& OscNode = GenNodes[i];
+			EMetaSoundBuilderResult ConnectResult;
+			const FMetaSoundBuilderNodeOutputHandle OscNodeAudioOutput = Builder.FindNodeOutputByName(OscNode, "Audio", ConnectResult);
+			AddErrorIfFalse(ConnectResult == EMetaSoundBuilderResult::Succeeded && OscNodeAudioOutput.IsSet(), TEXT("Failed to find oscillator node output 'Audio'"));
+
+			FMetaSoundBuilderNodeInputHandle InputHandle = Builder.FindNodeInputByName(OscNode, "Frequency", ConnectResult);
+			AddErrorIfFalse(ConnectResult == EMetaSoundBuilderResult::Succeeded && InputHandle.IsSet(), TEXT("Failed to find oscillator node input 'Frequency'"));
+
+			FMetasoundFrontendLiteral Literal;
+			Literal.Set(ChromaticFreqs[i]);
+			Builder.SetNodeInputDefault(InputHandle, Literal, ConnectResult);
+			AddErrorIfFalse(ConnectResult == EMetaSoundBuilderResult::Succeeded && OscNodeAudioOutput.IsSet(), TEXT("Failed to find oscillator node output 'Audio'"));
+
+			Builder.FindGraphOutputNode(SourceOneShotInterface::Outputs::OnFinished, Result);
+			AddErrorIfFalse(Result == EMetaSoundBuilderResult::Failed, "Failed to remove 'OnFinished' output to MetaSound using RemoveInterface Builder API call");
+
+			ADD_LATENT_AUTOMATION_COMMAND(FMetaSoundSourceBuilderConnectNodesLatentCommand(*this, &Builder, OscNodeAudioOutput, Output.AudioOutNodeInputs.Last()));
+			ADD_LATENT_AUTOMATION_COMMAND(FEngineWaitLatentCommand(0.125f));
+			ADD_LATENT_AUTOMATION_COMMAND(FMetaSoundSourceBuilderRemoveNodesLatentCommand(*this, &Builder, GenNodes[i]));
+			ADD_LATENT_AUTOMATION_COMMAND(FEngineWaitLatentCommand(0.125f));
+		}
+
+		ADD_LATENT_AUTOMATION_COMMAND(FEngineWaitLatentCommand(1.f));
+		ADD_LATENT_AUTOMATION_COMMAND(FAudioComponentStopLatentCommand(AudioComponent));
+		ADD_LATENT_AUTOMATION_COMMAND(FEngineWaitLatentCommand(0.1f));
+		ADD_LATENT_AUTOMATION_COMMAND(FBuilderRemoveFromRootLatentCommand(&Builder));
+		ADD_LATENT_AUTOMATION_COMMAND(FAudioComponentRemoveFromRootLatentCommand(AudioComponent));
+	}
+
+	return true;
+}
 #endif // WITH_DEV_AUTOMATION_TESTS
diff --git a/Engine/Plugins/Runtime/Metasound/Source/MetasoundFrontend/Private/MetasoundFrontendDocumentBuilder.cpp b/Engine/Plugins/Runtime/Metasound/Source/MetasoundFrontend/Private/MetasoundFrontendDocumentBuilder.cpp
--- a/Engine/Plugins/Runtime/Metasound/Source/MetasoundFrontend/Private/MetasoundFrontendDocumentBuilder.cpp
+++ b/Engine/Plugins/Runtime/Metasound/Source/MetasoundFrontend/Private/MetasoundFrontendDocumentBuilder.cpp
@@ -1852,16 +1852,8 @@ bool FMetaSoundFrontendDocumentBuilder::RemoveDependency(const FGuid& InClassID)
 		}
 
 		const int32 LastIndex = Dependencies.Num() - 1;
-		DocumentDelegates->OnRemovingDependency.Broadcast(Index);
-		if (Index != LastIndex)
-		{
-			DocumentDelegates->OnRemovingDependency.Broadcast(LastIndex);
-		}
+		DocumentDelegates->OnRemoveSwappingDependency.Broadcast(Index, LastIndex);
 		Dependencies.RemoveAtSwap(Index, 1, false);
-		if (Index != LastIndex)
-		{
-			DocumentDelegates->OnDependencyAdded.Broadcast(Index);
-		}
 	}
 
 	return true;
@@ -1890,16 +1882,8 @@ bool FMetaSoundFrontendDocumentBuilder::RemoveDependency(EMetasoundFrontendClass
 		}
 
 		const int32 LastIndex = Dependencies.Num() - 1;
-		DocumentDelegates->OnRemovingDependency.Broadcast(Index);
-		if (Index != LastIndex)
-		{
-			DocumentDelegates->OnRemovingDependency.Broadcast(LastIndex);
-		}
+		DocumentDelegates->OnRemoveSwappingDependency.Broadcast(Index, LastIndex);
 		Dependencies.RemoveAtSwap(Index, 1, false);
-		if (Index != LastIndex)
-		{
-			DocumentDelegates->OnDependencyAdded.Broadcast(Index);
-		}
 	}
 
 	return true;
@@ -1917,16 +1901,8 @@ bool FMetaSoundFrontendDocumentBuilder::RemoveEdge(const FMetasoundFrontendEdge&
 	{
 		const int32 Index = *IndexPtr;
 		const int32 LastIndex = Edges.Num() - 1;
-		DocumentDelegates->EdgeDelegates.OnRemovingEdge.Broadcast(Index);
-		if (Index != LastIndex)
-		{
-			DocumentDelegates->EdgeDelegates.OnRemovingEdge.Broadcast(LastIndex);
-		}
+		DocumentDelegates->EdgeDelegates.OnRemoveSwappingEdge.Broadcast(Index, LastIndex);
 		Edges.RemoveAtSwap(Index, 1, false);
-		if (Index != LastIndex)
-		{
-			DocumentDelegates->EdgeDelegates.OnEdgeAdded.Broadcast(Index);
-		}
 	}
 
 	return false;
@@ -2060,17 +2036,9 @@ bool FMetaSoundFrontendDocumentBuilder::RemoveEdgesFromNodeOutput(const FGuid& I
 		for (int32 Index : IndicesCopy)
 		{
 			const int32 LastIndex = Graph.Edges.Num() - 1;
-			DocumentDelegates->EdgeDelegates.OnRemovingEdge.Broadcast(Index);
-			if (Index != LastIndex)
-			{
-				DocumentDelegates->EdgeDelegates.OnRemovingEdge.Broadcast(LastIndex);
-			}
+			DocumentDelegates->EdgeDelegates.OnRemoveSwappingEdge.Broadcast(Index, LastIndex);
 			constexpr bool bAllowShrinking = false;
 			Graph.Edges.RemoveAtSwap(Index, 1, bAllowShrinking);
-			if (Index != LastIndex)
-			{
-				DocumentDelegates->EdgeDelegates.OnEdgeAdded.Broadcast(Index);
-			}
 		}
 
 		return true;
@@ -2090,16 +2058,8 @@ bool FMetaSoundFrontendDocumentBuilder::RemoveEdgeToNodeInput(const FGuid& InNod
 		FMetasoundFrontendGraph& Graph = Document.RootGraph.Graph;
 		const int32 Index = *IndexPtr; // Copy off indices as the pointer may be modified when notifying the cache below
 		const int32 LastIndex = Graph.Edges.Num() - 1;
-		DocumentDelegates->EdgeDelegates.OnRemovingEdge.Broadcast(Index);
-		if (Index != LastIndex)
-		{
-			DocumentDelegates->EdgeDelegates.OnRemovingEdge.Broadcast(LastIndex);
-		}
+		DocumentDelegates->EdgeDelegates.OnRemoveSwappingEdge.Broadcast(Index, LastIndex);
 		Graph.Edges.RemoveAtSwap(Index, 1, false);
-		if (Index != LastIndex)
-		{
-			DocumentDelegates->EdgeDelegates.OnEdgeAdded.Broadcast(Index);
-		}
 
 		return true;
 	}
@@ -2256,18 +2216,9 @@ bool FMetaSoundFrontendDocumentBuilder::RemoveNode(const FGuid& InNodeID)
 			RemoveEdgesFromNodeOutput(InNodeID, Vertex.VertexID);
 		}
 
-		DocumentDelegates->NodeDelegates.OnRemovingNode.Broadcast(Index);
-
 		const int32 LastIndex = Nodes.Num() - 1;
-		if (Index != LastIndex)
-		{
-			DocumentDelegates->NodeDelegates.OnRemovingNode.Broadcast(LastIndex);
-		}
+		DocumentDelegates->NodeDelegates.OnRemoveSwappingNode.Broadcast(Index, LastIndex);
 		Nodes.RemoveAtSwap(Index, 1, false);
-		if (Index != LastIndex)
-		{
-			DocumentDelegates->NodeDelegates.OnNodeAdded.Broadcast(Index);
-		}
 
 #if WITH_EDITORONLY_DATA
 		Document.Metadata.ModifyContext.AddNodeIDModified(InNodeID);
diff --git a/Engine/Plugins/Runtime/Metasound/Source/MetasoundFrontend/Private/MetasoundFrontendDocumentCache.cpp b/Engine/Plugins/Runtime/Metasound/Source/MetasoundFrontend/Private/MetasoundFrontendDocumentCache.cpp
--- a/Engine/Plugins/Runtime/Metasound/Source/MetasoundFrontend/Private/MetasoundFrontendDocumentCache.cpp
+++ b/Engine/Plugins/Runtime/Metasound/Source/MetasoundFrontend/Private/MetasoundFrontendDocumentCache.cpp
@@ -11,6 +11,63 @@
 
 namespace Metasound::Frontend
 {
+	namespace DocumentCachePrivate
+	{
+		// Modifies cache map of an attempted array index swap removal, where the
+		// given cache map is a map of keys to indices.
+		template <typename TMapType, typename TKeyType>
+		void RemoveSwapMapIndexChecked(const TKeyType& InSwapKey, const TKeyType& InLastKey, TMapType& OutMap)
+		{
+			// Remove swap item
+			int32 SwapIndex = INDEX_NONE;
+			const bool bRemoved = OutMap.RemoveAndCopyValue(InSwapKey, SwapIndex);
+			checkf(bRemoved, TEXT("Failed to remove/swap expected MetaSound Document Cache key/value pair"));
+
+			// Swap last item index if not the same entry
+			if (InLastKey != InSwapKey)
+			{
+				OutMap.FindChecked(InLastKey) = SwapIndex;
+			}
+		}
+
+		// Modifies cache map of an attempted array index swap removal, where the
+		// given cache map is a map of keys to arrays of indices.
+		template <typename TMapArrayType, typename TKeyType>
+		void RemoveSwapMapArrayIndexChecked(
+			const TKeyType& InSwapKey,
+			int32 SwapIndex,
+			const TKeyType& InLastKey,
+			int32 LastIndex,
+			TMapArrayType& OutMap
+		)
+		{
+			constexpr bool bAllowShrinking = false;
+
+			// Remove swap item
+			{
+				TArray<int32>& Indices = OutMap.FindChecked(InSwapKey);
+				Indices.RemoveAllSwap([&SwapIndex](const int32& Index)
+				{
+					return Index == SwapIndex;
+				}, bAllowShrinking);
+				if (Indices.IsEmpty())
+				{
+					OutMap.Remove(InSwapKey);
+				}
+			}
+
+			// Swap last item index if not the same entry
+			if (InSwapKey != InLastKey && SwapIndex != LastIndex)
+			{
+				TArray<int32>& Indices = OutMap.FindChecked(InLastKey);
+				Algo::ForEachIf(Indices,
+					[&LastIndex](const int32& Index) { return Index == LastIndex; },
+					[&SwapIndex](int32& Index) { Index = SwapIndex; }
+				);
+			}
+		}
+	}
+
 	FDocumentCache::FDocumentCache(const FMetasoundFrontendDocument& InDocument)
 		: Document(&InDocument)
 	{
@@ -34,12 +91,12 @@ namespace Metasound::Frontend
 			IDToIndex.Add(Class.ID, Index);
 		}
 
-		OutDelegates.OnDependencyAdded.Remove(OnDependencyAddedHandle);
-		OutDelegates.OnRemovingDependency.Remove(OnRemovingDependencyHandle);
+		OutDelegates.OnDependencyAdded.Remove(OnAddedHandle);
+		OutDelegates.OnRemoveSwappingDependency.Remove(OnRemoveSwappingHandle);
 		OutDelegates.OnRenamingDependencyClass.Remove(OnRenamingDependencyClassHandle);
 
-		OnDependencyAddedHandle = OutDelegates.OnDependencyAdded.AddSP(this, &FDocumentCache::OnDependencyAdded);
-		OnRemovingDependencyHandle = OutDelegates.OnRemovingDependency.AddSP(this, &FDocumentCache::OnRemovingDependency);
+		OnAddedHandle = OutDelegates.OnDependencyAdded.AddSP(this, &FDocumentCache::OnDependencyAdded);
+		OnRemoveSwappingHandle = OutDelegates.OnRemoveSwappingDependency.AddSP(this, &FDocumentCache::OnRemoveSwappingDependency);
 		OnRenamingDependencyClassHandle = OutDelegates.OnRenamingDependencyClass.AddSP(this, &FDocumentCache::OnRenamingDependencyClass);
 
 		if (!EdgeCache.IsValid())
@@ -127,14 +184,18 @@ namespace Metasound::Frontend
 		IDToIndex.Add(Dependency.ID, NewIndex);
 	}
 
-	void FDocumentCache::OnRemovingDependency(int32 IndexBeingRemoved)
+	void FDocumentCache::OnRemoveSwappingDependency(int32 SwapIndex, int32 LastIndex)
 	{
 		using namespace Metasound::Frontend;
 
-		const FMetasoundFrontendClass& DependencyBeingRemoved = GetDocument().Dependencies[IndexBeingRemoved];
-		const FNodeRegistryKey Key = NodeRegistryKey::CreateKey(DependencyBeingRemoved.Metadata);
-		KeyToIndex.Remove(Key);
-		IDToIndex.Remove(DependencyBeingRemoved.ID);
+		const TArray<FMetasoundFrontendClass>& Dependencies = GetDocument().Dependencies;
+		const FMetasoundFrontendClass& SwapDependency = Dependencies[SwapIndex];
+		const FMetasoundFrontendClass& LastDependency = Dependencies[LastIndex];
+
+		const FNodeRegistryKey SwapKey = NodeRegistryKey::CreateKey(SwapDependency.Metadata);
+		const FNodeRegistryKey LastKey = NodeRegistryKey::CreateKey(LastDependency.Metadata);
+		DocumentCachePrivate::RemoveSwapMapIndexChecked(SwapKey, LastKey, KeyToIndex);
+		DocumentCachePrivate::RemoveSwapMapIndexChecked(SwapDependency.ID, LastDependency.ID, IDToIndex);
 	}
 
 	void FDocumentCache::OnRenamingDependencyClass(const int32 IndexBeingRenamed, const FMetasoundFrontendClassName& NewName)
@@ -262,11 +323,11 @@ namespace Metasound::Frontend
 			ClassIDToNodeIndices.FindOrAdd(Node.ClassID).Add(Index);
 		}
 
-		OutDelegates.OnNodeAdded.Remove(OnNodeAddedHandle);
-		OnNodeAddedHandle = OutDelegates.OnNodeAdded.AddSP(this, &FDocumentGraphNodeCache::OnNodeAdded);
+		OutDelegates.OnNodeAdded.Remove(OnAddedHandle);
+		OnAddedHandle = OutDelegates.OnNodeAdded.AddSP(this, &FDocumentGraphNodeCache::OnNodeAdded);
 
-		OutDelegates.OnRemovingNode.Remove(OnRemovingNodeHandle);
-		OnRemovingNodeHandle = OutDelegates.OnRemovingNode.AddSP(this, &FDocumentGraphNodeCache::OnRemovingNode);
+		OutDelegates.OnRemoveSwappingNode.Remove(OnRemoveSwappingHandle);
+		OnRemoveSwappingHandle = OutDelegates.OnRemoveSwappingNode.AddSP(this, &FDocumentGraphNodeCache::OnRemoveSwappingNode);
 	}
 
 	bool FDocumentGraphNodeCache::ContainsNode(const FGuid& InNodeID) const
@@ -404,21 +465,17 @@ namespace Metasound::Frontend
 		ClassIDToNodeIndices.FindOrAdd(Node.ClassID).Add(InNewIndex);
 	}
 
-	void FDocumentGraphNodeCache::OnRemovingNode(int32 IndexBeingRemoved)
+	void FDocumentGraphNodeCache::OnRemoveSwappingNode(int32 SwapIndex, int32 LastIndex)
 	{
 		using namespace Metasound::Frontend;
 
-		// If not found, transaction was missed
 		const FMetasoundFrontendDocument& Document = Parent->GetDocument();
-		const FMetasoundFrontendNode& NodeBeingRemoved = Document.RootGraph.Graph.Nodes[IndexBeingRemoved];
+		const TArray<FMetasoundFrontendNode>& Nodes = Document.RootGraph.Graph.Nodes;
+		const FMetasoundFrontendNode& SwapNode = Nodes[SwapIndex];
+		const FMetasoundFrontendNode& LastNode = Nodes[LastIndex];
 
-		IDToIndex.Remove(NodeBeingRemoved.GetID());
-		TArray<int32>& NodeIndices = ClassIDToNodeIndices.FindChecked(NodeBeingRemoved.ClassID);
-		NodeIndices.Remove(IndexBeingRemoved);
-		if (NodeIndices.IsEmpty())
-		{
-			ClassIDToNodeIndices.Remove(NodeBeingRemoved.ClassID);
-		}
+		DocumentCachePrivate::RemoveSwapMapArrayIndexChecked(SwapNode.ClassID, SwapIndex, LastNode.ClassID, LastIndex, ClassIDToNodeIndices);
+		DocumentCachePrivate::RemoveSwapMapIndexChecked(SwapNode.GetID(), LastNode.GetID(), IDToIndex);
 	}
 
 	FDocumentGraphEdgeCache::FDocumentGraphEdgeCache(TSharedRef<IDocumentCache> ParentCache)
@@ -442,8 +499,8 @@ namespace Metasound::Frontend
 		OutDelegates.OnEdgeAdded.Remove(OnAddedHandle);
 		OnAddedHandle = OutDelegates.OnEdgeAdded.AddSP(this, &FDocumentGraphEdgeCache::OnEdgeAdded);
 
-		OutDelegates.OnRemovingEdge.Remove(OnRemovedHandle);
-		OnRemovedHandle = OutDelegates.OnRemovingEdge.AddSP(this, &FDocumentGraphEdgeCache::OnRemovingEdge);
+		OutDelegates.OnRemoveSwappingEdge.Remove(OnRemoveSwappingHandle);
+		OnRemoveSwappingHandle = OutDelegates.OnRemoveSwappingEdge.AddSP(this, &FDocumentGraphEdgeCache::OnRemoveSwappingEdge);
 	}
 
 	bool FDocumentGraphEdgeCache::ContainsEdge(const FMetasoundFrontendEdge& InEdge) const
@@ -509,17 +566,14 @@ namespace Metasound::Frontend
 		OutputToEdgeIndices.FindOrAdd(NewEdge.GetFromVertexHandle()).Add(InNewIndex);
 	}
 
-	void FDocumentGraphEdgeCache::OnRemovingEdge(int32 IndexToRemove)
+	void FDocumentGraphEdgeCache::OnRemoveSwappingEdge(int32 SwapIndex, int32 LastIndex)
 	{
 		const FMetasoundFrontendDocument& Document = Parent->GetDocument();
-		const FMetasoundFrontendGraph& Graph = Document.RootGraph.Graph;
-		const FMetasoundFrontendEdge& EdgeToRemove = Graph.Edges[IndexToRemove];
-		InputToEdgeIndex.Remove(EdgeToRemove.GetToVertexHandle());
-		TArray<int32>& Indices = OutputToEdgeIndices.FindChecked(EdgeToRemove.GetFromVertexHandle());
-		Indices.RemoveAllSwap([&IndexToRemove](const int32& Index) { return Index == IndexToRemove; }, false /* bAllowShrinking */);
-		if (Indices.IsEmpty())
-		{
-			OutputToEdgeIndices.Remove(EdgeToRemove.GetFromVertexHandle());
-		}
+		const TArray<FMetasoundFrontendEdge>& Edges = Document.RootGraph.Graph.Edges;
+
+		const FMetasoundFrontendEdge& SwapEdge = Edges[SwapIndex];
+		const FMetasoundFrontendEdge& LastEdge = Edges[LastIndex];
+		DocumentCachePrivate::RemoveSwapMapIndexChecked(SwapEdge.GetToVertexHandle(), LastEdge.GetToVertexHandle(), InputToEdgeIndex);
+		DocumentCachePrivate::RemoveSwapMapArrayIndexChecked(SwapEdge.GetFromVertexHandle(), SwapIndex, LastEdge.GetFromVertexHandle(), LastIndex, OutputToEdgeIndices);
 	}
 } // namespace Metasound::Frontend
diff --git a/Engine/Plugins/Runtime/Metasound/Source/MetasoundFrontend/Private/MetasoundFrontendDocumentCache.h b/Engine/Plugins/Runtime/Metasound/Source/MetasoundFrontend/Private/MetasoundFrontendDocumentCache.h
--- a/Engine/Plugins/Runtime/Metasound/Source/MetasoundFrontend/Private/MetasoundFrontendDocumentCache.h
+++ b/Engine/Plugins/Runtime/Metasound/Source/MetasoundFrontend/Private/MetasoundFrontendDocumentCache.h
@@ -37,7 +37,7 @@ namespace Metasound::Frontend
 
 	private:
 		void OnEdgeAdded(int32 InNewIndex);
-		void OnRemovingEdge(int32 IndexBeingRemoved);
+		void OnRemoveSwappingEdge(int32 SwapIndex, int32 LastIndex);
 
 		// Cache of outputs NodeId/VertexId pairs to associated edge indices
 		TMap<FMetasoundFrontendVertexHandle, TArray<int32>> OutputToEdgeIndices;
@@ -48,7 +48,7 @@ namespace Metasound::Frontend
 		TSharedPtr<IDocumentCache> Parent;
 
 		FDelegateHandle OnAddedHandle;
-		FDelegateHandle OnRemovedHandle;
+		FDelegateHandle OnRemoveSwappingHandle;
 	};
 
 
@@ -79,7 +79,7 @@ namespace Metasound::Frontend
 
 	private:
 		void OnNodeAdded(int32 NewIndex);
-		void OnRemovingNode(int32 IndexBeingRemoved);
+		void OnRemoveSwappingNode(int32 IndexBeingRemoved, int32 LastIndex);
 
 		// Cache of NodeId to array index of node
 		TSortedMap<FGuid, int32> IDToIndex;
@@ -89,8 +89,8 @@ namespace Metasound::Frontend
 
 		TSharedPtr<IDocumentCache> Parent;
 
-		FDelegateHandle OnNodeAddedHandle;
-		FDelegateHandle OnRemovingNodeHandle;
+		FDelegateHandle OnAddedHandle;
+		FDelegateHandle OnRemoveSwappingHandle;
 	};
 
 
@@ -151,7 +151,7 @@ namespace Metasound::Frontend
 
 	private:
 		void OnDependencyAdded(int32 InNewIndex);
-		void OnRemovingDependency(int32 IndexBeingRemoved);
+		void OnRemoveSwappingDependency(int32 SwapIndex, int32 LastIndex);
 		void OnRenamingDependencyClass(const int32 IndexBeingRenamed, const FMetasoundFrontendClassName& NewName);
 
 		// Cache of dependency (Class) ID to corresponding class dependency index
@@ -164,8 +164,8 @@ namespace Metasound::Frontend
 		TSharedPtr<FDocumentGraphNodeCache> NodeCache;
 		TSharedPtr<FDocumentGraphInterfaceCache> InterfaceCache;
 
-		FDelegateHandle OnDependencyAddedHandle;
-		FDelegateHandle OnRemovingDependencyHandle;
+		FDelegateHandle OnAddedHandle;
+		FDelegateHandle OnRemoveSwappingHandle;
 		FDelegateHandle OnRenamingDependencyClassHandle;
 
 		const FMetasoundFrontendDocument* Document = nullptr;
diff --git a/Engine/Plugins/Runtime/Metasound/Source/MetasoundFrontend/Public/MetasoundFrontendDocumentModifyDelegates.h b/Engine/Plugins/Runtime/Metasound/Source/MetasoundFrontend/Public/MetasoundFrontendDocumentModifyDelegates.h
--- a/Engine/Plugins/Runtime/Metasound/Source/MetasoundFrontend/Public/MetasoundFrontendDocumentModifyDelegates.h
+++ b/Engine/Plugins/Runtime/Metasound/Source/MetasoundFrontend/Public/MetasoundFrontendDocumentModifyDelegates.h
@@ -8,6 +8,7 @@
 
 DECLARE_MULTICAST_DELEGATE_OneParam(FOnMetaSoundFrontendDocumentMutateArray, int32 /* Index */);
 DECLARE_MULTICAST_DELEGATE_OneParam(FOnMetaSoundFrontendDocumentMutateInterfaceArray, const FMetasoundFrontendInterface& /* Interface */);
+DECLARE_MULTICAST_DELEGATE_TwoParams(FOnMetaSoundFrontendDocumentRemoveSwappingArray, int32 /* Index */, int32 /* LastIndex */);
 DECLARE_MULTICAST_DELEGATE_TwoParams(FOnMetaSoundFrontendDocumentRenameClass, const int32 /* Index */, const FMetasoundFrontendClassName& /* NewName */);
 DECLARE_MULTICAST_DELEGATE_ThreeParams(FOnMetaSoundFrontendDocumentMutateNodeInputLiteralArray, int32 /* NodeIndex */, int32 /* VertexIndex */, int32 /* LiteralIndex */);
 
@@ -29,7 +30,7 @@ namespace Metasound::Frontend
 	struct METASOUNDFRONTEND_API FNodeModifyDelegates
 	{
 		FOnMetaSoundFrontendDocumentMutateArray OnNodeAdded;
-		FOnMetaSoundFrontendDocumentMutateArray OnRemovingNode;
+		FOnMetaSoundFrontendDocumentRemoveSwappingArray OnRemoveSwappingNode;
 
 		FOnMetaSoundFrontendDocumentMutateNodeInputLiteralArray OnNodeInputLiteralSet;
 		FOnMetaSoundFrontendDocumentMutateNodeInputLiteralArray OnRemovingNodeInputLiteral;
@@ -38,13 +39,13 @@ namespace Metasound::Frontend
 	struct METASOUNDFRONTEND_API FEdgeModifyDelegates
 	{
 		FOnMetaSoundFrontendDocumentMutateArray OnEdgeAdded;
-		FOnMetaSoundFrontendDocumentMutateArray OnRemovingEdge;
+		FOnMetaSoundFrontendDocumentRemoveSwappingArray OnRemoveSwappingEdge;
 	};
 
 	struct METASOUNDFRONTEND_API FDocumentModifyDelegates : TSharedFromThis<FDocumentModifyDelegates>
 	{
 		FOnMetaSoundFrontendDocumentMutateArray OnDependencyAdded;
-		FOnMetaSoundFrontendDocumentMutateArray OnRemovingDependency;
+		FOnMetaSoundFrontendDocumentRemoveSwappingArray OnRemoveSwappingDependency;
 		FOnMetaSoundFrontendDocumentRenameClass OnRenamingDependencyClass;
 
 		FInterfaceModifyDelegates InterfaceDelegates;
