Fix for crash when closing MetaSound editor with patch loaded

diff --git a/Engine/Plugins/Runtime/Metasound/Source/MetasoundEditor/Private/MetasoundEditor.cpp b/Engine/Plugins/Runtime/Metasound/Source/MetasoundEditor/Private/MetasoundEditor.cpp
--- a/Engine/Plugins/Runtime/Metasound/Source/MetasoundEditor/Private/MetasoundEditor.cpp
+++ b/Engine/Plugins/Runtime/Metasound/Source/MetasoundEditor/Private/MetasoundEditor.cpp
@@ -1234,8 +1234,15 @@ namespace Metasound
 				OutputMeter->Teardown();
 			}
 
-			OutputOscilloscope->StopProcessing();
-			OutputVectorscope->StopProcessing();
+			if (OutputOscilloscope.IsValid())
+			{
+				OutputOscilloscope->StopProcessing();
+			}
+
+			if (OutputVectorscope.IsValid())
+			{
+				OutputVectorscope->StopProcessing();
+			}
 		}
 
 		void FEditor::ExtendToolbar()
