Clean-up reporting when gathering referenced assets to avoid warnings of asset registry not being loaded when its not necessary to be.
#rnx

diff --git a/Engine/Plugins/Runtime/Metasound/Source/MetasoundEngine/Private/MetasoundAssetSubsystem.cpp b/Engine/Plugins/Runtime/Metasound/Source/MetasoundEngine/Private/MetasoundAssetSubsystem.cpp
--- a/Engine/Plugins/Runtime/Metasound/Source/MetasoundEngine/Private/MetasoundAssetSubsystem.cpp
+++ b/Engine/Plugins/Runtime/Metasound/Source/MetasoundEngine/Private/MetasoundAssetSubsystem.cpp
@@ -330,11 +330,6 @@ TSet<UMetaSoundAssetSubsystem::FAssetInfo> UMetaSoundAssetSubsystem::GetReferenc
 	METASOUND_TRACE_CPUPROFILER_EVENT_SCOPE(UMetaSoundAssetSubsystem::GetReferencedAssetClasses);
 	using namespace Metasound::Frontend;
 
-	if (!bIsInitialAssetScanComplete)
-	{
-		UE_LOG(LogMetaSound, Warning, TEXT("Attempt to get registered dependent assets for %s before asset scan is complete may result in missed dependencies"), *InAssetBase.GetOwningAssetName());
-	}
-
 	TSet<FAssetInfo> OutAssetInfos;
 	const FMetasoundFrontendDocument& Document = InAssetBase.GetDocumentChecked();
 	for (const FMetasoundFrontendClass& Class : Document.Dependencies)
@@ -344,6 +339,37 @@ TSet<UMetaSoundAssetSubsystem::FAssetInfo> UMetaSoundAssetSubsystem::GetReferenc
 		{
 			OutAssetInfos.Add(FAssetInfo{Key, *ObjectPath});
 		}
+		else
+		{
+			const FMetasoundFrontendRegistryContainer* Registry = FMetasoundFrontendRegistryContainer::Get();
+			check(Registry);
+			const bool bIsRegistered = Registry->IsNodeRegistered(Key);
+
+			bool bReportFail = false;
+			if (bIsRegistered)
+			{
+				if (!Registry->IsNodeNative(Key))
+				{
+					bReportFail = true;
+				}
+			}
+			else
+			{
+				bReportFail = true;
+			}
+
+			if (bReportFail)
+			{
+				if (bIsInitialAssetScanComplete)
+				{
+					UE_LOG(LogMetaSound, Warning, TEXT("MetaSound Node Class with registry key '%s' not registered when gathering referenced asset classes from '%s': Retrieving all asset classes may not be comprehensive."), *Key, *InAssetBase.GetOwningAssetName());
+				}
+				else
+				{
+					UE_LOG(LogMetaSound, Warning, TEXT("Attempt to get registered dependent asset with key '%s' from MetaSound asset '%s' before asset scan has completed: Asset class cannot be provided"), *Key, *InAssetBase.GetOwningAssetName());
+				}
+			}
+		}
 	}
 	return MoveTemp(OutAssetInfos);
 }
