- Fix for auditioning the same sound multiple times not respecting live updates on additional calls
- Upgrade Live Node Update automation tests to LiveUpdateMultipleMetaSoundSources for more comprehensive coverage
#rnx

diff --git a/Engine/Plugins/Runtime/Metasound/Source/MetasoundEngine/Private/MetasoundBuilderSubsystem.cpp b/Engine/Plugins/Runtime/Metasound/Source/MetasoundEngine/Private/MetasoundBuilderSubsystem.cpp
--- a/Engine/Plugins/Runtime/Metasound/Source/MetasoundEngine/Private/MetasoundBuilderSubsystem.cpp
+++ b/Engine/Plugins/Runtime/Metasound/Source/MetasoundEngine/Private/MetasoundBuilderSubsystem.cpp
@@ -953,7 +953,7 @@ void UMetaSoundSourceBuilder::Audition(UObject* Parent, UAudioComponent* AudioCo
 	AudioComponent->SetSound(AuditionSound.Get());
 
 	const bool bEnableDynamicGenerators = bLiveUpdatesEnabled;
-	AuditionSound->EnableDynamicGenerators(bEnableDynamicGenerators);
+	AuditionSound->SetDynamicGeneratorEnabled(bEnableDynamicGenerators);
 
 	if (CreateGenerator.IsBound())
 	{
diff --git a/Engine/Plugins/Runtime/Metasound/Source/MetasoundEngine/Private/MetasoundSource.cpp b/Engine/Plugins/Runtime/Metasound/Source/MetasoundEngine/Private/MetasoundSource.cpp
--- a/Engine/Plugins/Runtime/Metasound/Source/MetasoundEngine/Private/MetasoundSource.cpp
+++ b/Engine/Plugins/Runtime/Metasound/Source/MetasoundEngine/Private/MetasoundSource.cpp
@@ -191,10 +191,12 @@ void UMetaSoundSource::PostEditChangeOutputFormat()
 
 	EMetaSoundBuilderResult Result = EMetaSoundBuilderResult::Failed;
 	{
-
 		const UMetaSoundBuilderSubsystem& BuilderSubsystem = UMetaSoundBuilderSubsystem::GetConstChecked();
+
 		UMetaSoundSourceBuilder* SourceBuilder = BuilderSubsystem.AttachSourceBuilderToAsset(this);
+		check(SourceBuilder);
 		SourceBuilder->SetFormat(OutputFormat, Result);
+
 		// TODO: Once builders are notified of controller changes and can be safely persistent, this
 		// can be removed so builders can be shared and not have to be created for each change output
 		// format mutation transaction.
@@ -1063,21 +1065,24 @@ Metasound::SourcePrivate::FParameterRouter& UMetaSoundSource::GetParameterRouter
 	return Router;
 }
 
-TSharedPtr<Metasound::DynamicGraph::FDynamicOperatorTransactor> UMetaSoundSource::EnableDynamicGenerators(bool bInIsEnabled)
+TSharedPtr<Metasound::DynamicGraph::FDynamicOperatorTransactor> UMetaSoundSource::SetDynamicGeneratorEnabled(bool bInIsEnabled)
 {
 	using namespace Metasound;
 	using namespace Metasound::DynamicGraph;
 
-	if (bInIsEnabled && (!DynamicTransactor.IsValid()))
+	if (bInIsEnabled)
 	{
-		TSharedPtr<FGraph> CurrentGraph = GetRuntimeData().Graph;
-		if (CurrentGraph.IsValid())
-		{
-			DynamicTransactor = MakeShared<FDynamicOperatorTransactor>(*CurrentGraph);
-		}
-		else
+		if (!DynamicTransactor.IsValid())
 		{
-			DynamicTransactor = MakeShared<FDynamicOperatorTransactor>();
+			TSharedPtr<FGraph> CurrentGraph = GetRuntimeData().Graph;
+			if (CurrentGraph.IsValid())
+			{
+				DynamicTransactor = MakeShared<FDynamicOperatorTransactor>(*CurrentGraph);
+			}
+			else
+			{
+				DynamicTransactor = MakeShared<FDynamicOperatorTransactor>();
+			}
 		}
 	}
 	else
diff --git a/Engine/Plugins/Runtime/Metasound/Source/MetasoundEngine/Public/MetasoundSource.h b/Engine/Plugins/Runtime/Metasound/Source/MetasoundEngine/Public/MetasoundSource.h
--- a/Engine/Plugins/Runtime/Metasound/Source/MetasoundEngine/Public/MetasoundSource.h
+++ b/Engine/Plugins/Runtime/Metasound/Source/MetasoundEngine/Public/MetasoundSource.h
@@ -253,19 +253,19 @@ private:
 	void TrackGenerator(uint64 Id, TSharedPtr<Metasound::FMetasoundGenerator> Generator);
 	void ForgetGenerator(ISoundGeneratorPtr Generator);
 
-	/** Enable dynamic generator. 
+	/** Enable/disable dynamic generator.
 	 *
 	 * Once a dynamic generator is enabled, all changes to the MetaSound should be applied to the
 	 * FDynamicOperatorTransactor in order to keep parity between the document and active graph.
 	 *
 	 * Note: Disabling the dynamic generator will sever the communication between any active generators
-	 * even if the dynamic generator is re-enabled. 
+	 * even if the dynamic generator is re-enabled during the lifetime of the active generators
 	 */
-	TSharedPtr<Metasound::DynamicGraph::FDynamicOperatorTransactor> EnableDynamicGenerators(bool bInIsEnabled);
+	TSharedPtr<Metasound::DynamicGraph::FDynamicOperatorTransactor> SetDynamicGeneratorEnabled(bool bInIsEnabled);
 
 	/** Get dynamic transactor
 	 *
-	 * If dynamic generators are enabled, this will return a valid pointer to a dynamic transactor. 
+	 * If dynamic generators are enabled, this will return a valid pointer to a dynamic transactor.
 	 * Changes to this transactor will be forwarded to any active Dynamic MetaSound Generators.
 	 */
 	TSharedPtr<Metasound::DynamicGraph::FDynamicOperatorTransactor> GetDynamicGeneratorTransactor() const;
diff --git a/Engine/Plugins/Runtime/Metasound/Source/MetasoundEngineTest/Private/EngineTestMetasoundSource.cpp b/Engine/Plugins/Runtime/Metasound/Source/MetasoundEngineTest/Private/EngineTestMetasoundSource.cpp
--- a/Engine/Plugins/Runtime/Metasound/Source/MetasoundEngineTest/Private/EngineTestMetasoundSource.cpp
+++ b/Engine/Plugins/Runtime/Metasound/Source/MetasoundEngineTest/Private/EngineTestMetasoundSource.cpp
@@ -286,6 +286,19 @@ bool FAudioComponentStopLatentCommand::Update()
 	return false;
 }
 
+DEFINE_LATENT_AUTOMATION_COMMAND_TWO_PARAMETER(FMetaSoundSourceLatentSetParamsCommand, UAudioComponent*, AudioComponent, TArray<FAudioParameter>, Params);
+
+bool FMetaSoundSourceLatentSetParamsCommand::Update()
+{
+	if (AudioComponent)
+	{
+		AudioComponent->SetParameters(MoveTemp(Params));
+		return true;
+	}
+
+	return false;
+}
+
 DEFINE_LATENT_AUTOMATION_COMMAND_THREE_PARAMETER(FMetaSoundSourceBuilderAuditionLatentCommand, UMetaSoundSourceBuilder*, Builder, UAudioComponent*, AudioComponent, bool, bEnableLiveUpdates);
 
 bool FMetaSoundSourceBuilderAuditionLatentCommand::Update()
@@ -449,8 +462,17 @@ bool FAudioMetasoundSourceBuilderTestSpamAudition::RunTest(const FString& Parame
 	return false;
 }
 
-// This test creates a MetaSound source from a SourceBuilder, adds a simple sine tone generator with a connected graph input frequency, attempts to audition it, and then switches to a new tri tone generator during playback.
-IMPLEMENT_SIMPLE_AUTOMATION_TEST(FAudioMetasoundSourceBuilderLiveUpdateNode, "Audio.Metasound.LiveUpdateNodeMetaSoundSource", EAutomationTestFlags::EditorContext | EAutomationTestFlags::EngineFilter)
+// This test exercises auditioning multiple sources, both at the same time and attempting to audition live changes after stopping and restarting an audio component by:
+// 1. Creates a MetaSound source from a SourceBuilder
+// 2. Adds a simple sine tone generator with a connected graph input frequency
+// 3. Auditions the builder via two AudioComponents
+// 4. Sets the parameter input for frequency to a different value for the second component
+// 5. Disconnects the parameter/graph input and sets the literal on the builder to a new frequency, which is observed on both components
+// 6. Stops first the second component then the first.
+// 7. Restarts audition of the original AudioComponent (which should continue after all the prior changes)
+// 8. Swaps to use a tri tone generator
+// 9. Stops original component, completing the test
+IMPLEMENT_SIMPLE_AUTOMATION_TEST(FAudioMetasoundSourceBuilderLiveUpdateNode, "Audio.Metasound.LiveUpdateMultipleMetaSoundSources", EAutomationTestFlags::EditorContext | EAutomationTestFlags::EngineFilter)
 bool FAudioMetasoundSourceBuilderLiveUpdateNode::RunTest(const FString& Parameters)
 {
 	using namespace EngineTestMetaSoundSourcePrivate;
@@ -458,27 +480,61 @@ bool FAudioMetasoundSourceBuilderLiveUpdateNode::RunTest(const FString& Paramete
 	using namespace Metasound::Engine;
 	using namespace Metasound::Frontend;
 
+	FMetaSoundBuilderNodeInputHandle GenInputNodeFreq;
 	FMetaSoundBuilderNodeInputHandle MonoOutNodeInput;
-	UMetaSoundSourceBuilder& Builder = CreateMonoSourceSinGenBuilder(*this, nullptr, &MonoOutNodeInput);
+	UMetaSoundSourceBuilder& Builder = CreateMonoSourceSinGenBuilder(*this, &GenInputNodeFreq, &MonoOutNodeInput, 440.f);
 	Builder.AddToRoot();
 
 	if (UAudioComponent* AudioComponent = CreateTestComponent(*this))
 	{
 		constexpr bool bEnableLiveUpdate = true;
 		ADD_LATENT_AUTOMATION_COMMAND(FMetaSoundSourceBuilderAuditionLatentCommand(&Builder, AudioComponent, bEnableLiveUpdate));
-		ADD_LATENT_AUTOMATION_COMMAND(FEngineWaitLatentCommand(2.f));
+		ADD_LATENT_AUTOMATION_COMMAND(FEngineWaitLatentCommand(0.5f));
+
+		// Send commands to more than one component. By setting param to 220.f on the second, listener can hear two
+		// operators function on separate runtime graphs/generators
+		if (UAudioComponent* AudioComponent2 = CreateTestComponent(*this))
+		{
+			ADD_LATENT_AUTOMATION_COMMAND(FMetaSoundSourceBuilderAuditionLatentCommand(&Builder, AudioComponent2, bEnableLiveUpdate));
+			ADD_LATENT_AUTOMATION_COMMAND(FEngineWaitLatentCommand(0.5f));
+
+			TArray<FAudioParameter> Params { FAudioParameter { "Frequency", 220.0f } };
+			ADD_LATENT_AUTOMATION_COMMAND(FMetaSoundSourceLatentSetParamsCommand(AudioComponent2, Params));
+			ADD_LATENT_AUTOMATION_COMMAND(FEngineWaitLatentCommand(2.f));
+
+			// Set a literal to ensure live updates still update as expected (nothing should happen until the next step when the edge
+			// is disconnected, at which point we should hear 880Hz)
+			FMetasoundFrontendLiteral NewFreq;
+			NewFreq.Set(880.f);
+			ADD_LATENT_AUTOMATION_COMMAND(FMetaSoundSourceBuilderSetLiteralLatentCommand(*this, &Builder, GenInputNodeFreq, NewFreq));
+			ADD_LATENT_AUTOMATION_COMMAND(FEngineWaitLatentCommand(0.5f));
+
+			// Remove the edge to the graph input, at which point listener should hear the literal value of 880 via both test components.
+			ADD_LATENT_AUTOMATION_COMMAND(FMetaSoundSourceBuilderDisconnectInputLatentCommand(*this, &Builder, GenInputNodeFreq));
+			ADD_LATENT_AUTOMATION_COMMAND(FEngineWaitLatentCommand(1.f));
+
+			// Nuke the second component from orbit, the rest of the test is working with just the OG component.
+			ADD_LATENT_AUTOMATION_COMMAND(FAudioComponentStopLatentCommand(AudioComponent2));
+			ADD_LATENT_AUTOMATION_COMMAND(FAudioComponentRemoveFromRootLatentCommand(AudioComponent2));
+		}
+		else
+		{
+			return false;
+		}
+		
+		// Stop and hear silence
+		ADD_LATENT_AUTOMATION_COMMAND(FAudioComponentStopLatentCommand(AudioComponent));
+		ADD_LATENT_AUTOMATION_COMMAND(FEngineWaitLatentCommand(1.f));
+
+		// Restart audition to ensure it restarts as expected
+		ADD_LATENT_AUTOMATION_COMMAND(FMetaSoundSourceBuilderAuditionLatentCommand(&Builder, AudioComponent, bEnableLiveUpdate));
+		ADD_LATENT_AUTOMATION_COMMAND(FEngineWaitLatentCommand(1.f));
 
 		// Disconnect graph audio output from existing sinosc output and connect to added triosc
 		ADD_LATENT_AUTOMATION_COMMAND(FMetaSoundSourceBuilderCreateAndConnectTriGeneratorNodeLatentCommand(*this, &Builder, MonoOutNodeInput));
+		ADD_LATENT_AUTOMATION_COMMAND(FEngineWaitLatentCommand(1.f));
 
-		FName DataTypeName;
-		FMetasoundFrontendLiteral NewValue;
-		NewValue.Set(FMetasoundFrontendLiteral::FDefault{ });
-		ADD_LATENT_AUTOMATION_COMMAND(FMetaSoundSourceBuilderSetLiteralLatentCommand(*this, &Builder, MonoOutNodeInput, NewValue));
-
-		ADD_LATENT_AUTOMATION_COMMAND(FEngineWaitLatentCommand(2.f));
 		ADD_LATENT_AUTOMATION_COMMAND(FAudioComponentStopLatentCommand(AudioComponent));
-		ADD_LATENT_AUTOMATION_COMMAND(FEngineWaitLatentCommand(0.5f));
 		ADD_LATENT_AUTOMATION_COMMAND(FAudioComponentRemoveFromRootLatentCommand(AudioComponent));
 		ADD_LATENT_AUTOMATION_COMMAND(FBuilderRemoveFromRootLatentCommand(&Builder));
 
