Fix comment with unacceptable word
#rnx

diff --git a/Engine/Plugins/Runtime/Metasound/Source/MetasoundEngineTest/Private/EngineTestMetasoundSource.cpp b/Engine/Plugins/Runtime/Metasound/Source/MetasoundEngineTest/Private/EngineTestMetasoundSource.cpp
--- a/Engine/Plugins/Runtime/Metasound/Source/MetasoundEngineTest/Private/EngineTestMetasoundSource.cpp
+++ b/Engine/Plugins/Runtime/Metasound/Source/MetasoundEngineTest/Private/EngineTestMetasoundSource.cpp
@@ -513,7 +513,7 @@ bool FAudioMetasoundSourceBuilderLiveUpdateNode::RunTest(const FString& Paramete
 			ADD_LATENT_AUTOMATION_COMMAND(FMetaSoundSourceBuilderDisconnectInputLatentCommand(*this, &Builder, GenInputNodeFreq));
 			ADD_LATENT_AUTOMATION_COMMAND(FEngineWaitLatentCommand(1.f));
 
-			// Nuke the second component from orbit, the rest of the test is working with just the OG component.
+			// Remove the second component, as the rest of the test is working with just the OG component.
 			ADD_LATENT_AUTOMATION_COMMAND(FAudioComponentStopLatentCommand(AudioComponent2));
 			ADD_LATENT_AUTOMATION_COMMAND(FAudioComponentRemoveFromRootLatentCommand(AudioComponent2));
 		}
