Add tags for InitParams & RebuildGraph MS transforms for more perf analysis datapoints
#rnx

diff --git a/Engine/Plugins/Runtime/Metasound/Source/MetasoundEngine/Private/MetasoundSource.cpp b/Engine/Plugins/Runtime/Metasound/Source/MetasoundEngine/Private/MetasoundSource.cpp
--- a/Engine/Plugins/Runtime/Metasound/Source/MetasoundEngine/Private/MetasoundSource.cpp
+++ b/Engine/Plugins/Runtime/Metasound/Source/MetasoundEngine/Private/MetasoundSource.cpp
@@ -338,6 +338,7 @@ void UMetaSoundSource::InitParameters(TArray<FAudioParameter>& ParametersToInit,
 	using namespace Metasound::SourcePrivate;
 
 	METASOUND_LLM_SCOPE;
+	METASOUND_TRACE_CPUPROFILER_EVENT_SCOPE(UMetaSoundSource::InitParameters);
 
 	// Have to call register vs a simple get as the source may have yet to start playing/has not been registered
 	// via InitResources. If it has, this call is fast and returns the already cached RuntimeData.
diff --git a/Engine/Plugins/Runtime/Metasound/Source/MetasoundFrontend/Private/MetasoundFrontendTransform.cpp b/Engine/Plugins/Runtime/Metasound/Source/MetasoundFrontend/Private/MetasoundFrontendTransform.cpp
--- a/Engine/Plugins/Runtime/Metasound/Source/MetasoundFrontend/Private/MetasoundFrontendTransform.cpp
+++ b/Engine/Plugins/Runtime/Metasound/Source/MetasoundFrontend/Private/MetasoundFrontendTransform.cpp
@@ -701,6 +701,8 @@ namespace Metasound
 
 		bool FRebuildPresetRootGraph::Transform(FDocumentHandle InDocument) const
 		{
+			METASOUND_TRACE_CPUPROFILER_EVENT_SCOPE(Metasound::Frontend::FRebuildPresetRootGraph::Transform);
+
 			FGraphHandle PresetGraphHandle = InDocument->GetRootGraph();
 			if (!ensure(PresetGraphHandle->IsValid()))
 			{
