Fix ensure when adding a MetaSound input/output node with an invalid data type using the Builder API

diff --git a/Engine/Plugins/Runtime/Metasound/Source/MetasoundFrontend/Private/MetasoundFrontendDocumentBuilder.cpp b/Engine/Plugins/Runtime/Metasound/Source/MetasoundFrontend/Private/MetasoundFrontendDocumentBuilder.cpp
--- a/Engine/Plugins/Runtime/Metasound/Source/MetasoundFrontend/Private/MetasoundFrontendDocumentBuilder.cpp
+++ b/Engine/Plugins/Runtime/Metasound/Source/MetasoundFrontend/Private/MetasoundFrontendDocumentBuilder.cpp
@@ -830,6 +830,11 @@ const FMetasoundFrontendNode* FMetaSoundFrontendDocumentBuilder::AddGraphInput(c
 		check(OutputNode);
 		return OutputNode;
 	}
+	else if (!IDataTypeRegistry::Get().IsRegistered(InClassInput.TypeName))
+	{
+		UE_LOG(LogMetaSound, Error, TEXT("Cannot add MetaSound graph input '%s' with unregistered TypeName '%s'"), *InClassInput.Name.ToString(), *InClassInput.TypeName.ToString());
+		return nullptr;
+	}
 
 	auto FindRegistryClass = [&InClassInput](FMetasoundFrontendClass& OutClass) -> bool
 	{
@@ -911,6 +916,11 @@ const FMetasoundFrontendNode* FMetaSoundFrontendDocumentBuilder::AddGraphOutput(
 		UE_LOG(LogMetaSound, Error, TEXT("Attempting to add MetaSound graph output '%s' when output with name already exists"), *InClassOutput.Name.ToString());
 		return DocumentCache->GetNodeCache().FindNode(Output->NodeID);
 	}
+	else if (!IDataTypeRegistry::Get().IsRegistered(InClassOutput.TypeName))
+	{
+		UE_LOG(LogMetaSound, Error, TEXT("Cannot add MetaSound graph output '%s' with unregistered TypeName '%s'"), *InClassOutput.Name.ToString(), *InClassOutput.TypeName.ToString());
+		return nullptr;
+	}
 
 	auto FindRegistryClass = [&InClassOutput](FMetasoundFrontendClass& OutClass) -> bool
 	{
