Fix for MetaSound live update failing to create & forward input/output (and other currently unsupported types like variable) nodes to transactor
#rnx

diff --git a/Engine/Plugins/Runtime/Metasound/Source/MetasoundEngine/Private/MetasoundBuilderSubsystem.cpp b/Engine/Plugins/Runtime/Metasound/Source/MetasoundEngine/Private/MetasoundBuilderSubsystem.cpp
--- a/Engine/Plugins/Runtime/Metasound/Source/MetasoundEngine/Private/MetasoundBuilderSubsystem.cpp
+++ b/Engine/Plugins/Runtime/Metasound/Source/MetasoundEngine/Private/MetasoundBuilderSubsystem.cpp
@@ -1110,41 +1110,84 @@ void UMetaSoundSourceBuilder::OnNodeAdded(int32 NodeIndex) const
 		const FMetasoundFrontendDocument& Doc = Builder.GetDocument();
 		const FMetasoundFrontendGraphClass& GraphClass = Doc.RootGraph;
 		const FMetasoundFrontendNode& AddedNode = GraphClass.Graph.Nodes[NodeIndex];
-		const Metasound::FNodeInitData InitData { AddedNode.Name, AddedNode.GetID() };
 
 		const FMetasoundFrontendClass* NodeClass = Builder.FindDependency(AddedNode.ClassID);
 		checkf(NodeClass, TEXT("Node successfully added to graph but document is missing associated dependency"));
 
-		auto MatchesVertexNodeID = [&InitData](const FMetasoundFrontendClassVertex& Vertex) { return Vertex.NodeID == InitData.InstanceID; };
+		const FNodeRegistryKey& ClassKey = NodeRegistryKey::CreateKey(NodeClass->Metadata);
+		FMetasoundFrontendRegistryContainer& NodeRegistry = *FMetasoundFrontendRegistryContainer::Get();
+		IDataTypeRegistry& DataTypeRegistry = IDataTypeRegistry::Get();
+
+		TUniquePtr<INode> NewNode;
+
 		switch (NodeClass->Metadata.GetType())
 		{
-			case EMetasoundFrontendClassType::Input:
-			case EMetasoundFrontendClassType::Output:
+			case EMetasoundFrontendClassType::VariableDeferredAccessor:
+			case EMetasoundFrontendClassType::VariableAccessor:
+			case EMetasoundFrontendClassType::VariableMutator:
 			case EMetasoundFrontendClassType::External:
+			case EMetasoundFrontendClassType::Graph:
 			{
-			const FNodeRegistryKey& ClassKey = NodeRegistryKey::CreateKey(NodeClass->Metadata);
-				FMetasoundFrontendRegistryContainer& Registry = *FMetasoundFrontendRegistryContainer::Get();
-				TUniquePtr<INode> NewNode = Registry.CreateNode(ClassKey, InitData);
-				Transactor.AddNode(InitData.InstanceID, MoveTemp(NewNode));
-				return true;
+				const Metasound::FNodeInitData InitData { AddedNode.Name, AddedNode.GetID() };
+				NewNode = NodeRegistry.CreateNode(ClassKey, InitData);
+			}
+			break;
+
+			case EMetasoundFrontendClassType::Input:
+			{
+				const FName& DataTypeName = NodeClass->Metadata.GetClassName().Name;
+				const FMetasoundFrontendVertex& InputVertex = AddedNode.Interface.Inputs.Last();
+				const FMetasoundFrontendLiteral& DefaultLiteral = NodeClass->Interface.Inputs.Last().DefaultLiteral;
+				FInputNodeConstructorParams InitData
+				{
+					AddedNode.Name,
+					AddedNode.GetID(),
+					InputVertex.Name,
+					DefaultLiteral.ToLiteral(DataTypeName)
+				};
+
+				NewNode = DataTypeRegistry.CreateInputNode(DataTypeName, MoveTemp(InitData));
 			}
 			break;
 
-			case EMetasoundFrontendClassType::Graph:
-			case EMetasoundFrontendClassType::Invalid:
-			case EMetasoundFrontendClassType::Literal:
-			case EMetasoundFrontendClassType::Template:
 			case EMetasoundFrontendClassType::Variable:
-			case EMetasoundFrontendClassType::VariableAccessor:
-			case EMetasoundFrontendClassType::VariableDeferredAccessor:
-			case EMetasoundFrontendClassType::VariableMutator:
-			default:
 			{
-				static_assert(static_cast<int32>(EMetasoundFrontendClassType::Invalid) == 10, "Possible missing EMetasoundFrontendClassType case coverage");
-				UE_LOG(LogMetaSound, Error, TEXT("Builder '%s' Audition Error: Addition of node was not propagated to active runtime graph: Runtime manipualtion of node ClassType unsupported."), *GetName());
-				return false;
+				const FName& DataTypeName = NodeClass->Metadata.GetClassName().Name;
+				FDefaultLiteralNodeConstructorParams InitData { AddedNode.Name, AddedNode.GetID(), DataTypeRegistry.CreateDefaultLiteral(DataTypeName)};
+				NewNode = DataTypeRegistry.CreateVariableNode(DataTypeName, MoveTemp(InitData));
 			}
+			break;
+
+			case EMetasoundFrontendClassType::Literal:
+			{
+				const FName& DataTypeName = NodeClass->Metadata.GetClassName().Name;
+				FDefaultLiteralNodeConstructorParams InitData { AddedNode.Name, AddedNode.GetID(), DataTypeRegistry.CreateDefaultLiteral(DataTypeName)};
+				NewNode = DataTypeRegistry.CreateLiteralNode(DataTypeName, MoveTemp(InitData));
+			}
+			break;
+
+			case EMetasoundFrontendClassType::Output:
+			{
+				const FName& DataTypeName = NodeClass->Metadata.GetClassName().Name;
+				const FMetasoundFrontendVertex& OutputVertex = AddedNode.Interface.Outputs.Last();
+				FDefaultNamedVertexNodeConstructorParams InitData { AddedNode.Name, AddedNode.GetID(), OutputVertex.Name };
+				NewNode = DataTypeRegistry.CreateOutputNode(DataTypeName, MoveTemp(InitData));
+			}
+			break;
+
+			case EMetasoundFrontendClassType::Template:
+			default:
+			static_assert(static_cast<int32>(EMetasoundFrontendClassType::Invalid) == 10, "Possible missed EMetasoundFrontendClassType case coverage");
+		};
+
+		if (!NewNode.IsValid())
+		{
+			UE_LOG(LogMetaSound, Error, TEXT("Builder '%s' failed to create and forward added node '%s' to live update transactor."), *GetName(), *AddedNode.Name.ToString());
+			return false;
 		}
+
+		Transactor.AddNode(AddedNode.GetID(), MoveTemp(NewNode));
+		return true;
 	});
 }
 
