Remove creation of redundant document handle
#rnx

diff --git a/Engine/Plugins/Runtime/Metasound/Source/MetasoundFrontend/Private/MetasoundFrontendTransform.cpp b/Engine/Plugins/Runtime/Metasound/Source/MetasoundFrontend/Private/MetasoundFrontendTransform.cpp
--- a/Engine/Plugins/Runtime/Metasound/Source/MetasoundFrontend/Private/MetasoundFrontendTransform.cpp
+++ b/Engine/Plugins/Runtime/Metasound/Source/MetasoundFrontend/Private/MetasoundFrontendTransform.cpp
@@ -728,8 +728,8 @@ namespace Metasound
 			// Determine the inputs and outputs needed in the wrapping graph. Also
 			// cache any exiting literals that have been set on the wrapping graph.
 			TSet<FName> InputsInheritingDefault;
-			TArray<FMetasoundFrontendClassInput> ClassInputs = GenerateRequiredClassInputs(PresetGraphHandle, InputsInheritingDefault);
-			TArray<FMetasoundFrontendClassOutput> ClassOutputs = GenerateRequiredClassOutputs(PresetGraphHandle);
+			TArray<FMetasoundFrontendClassInput> ClassInputs = GenerateRequiredClassInputs(InDocument, PresetGraphHandle, InputsInheritingDefault);
+			TArray<FMetasoundFrontendClassOutput> ClassOutputs = GenerateRequiredClassOutputs(InDocument, PresetGraphHandle);
 
 			FGuid PresetNodeID;
 			PresetGraphHandle->IterateConstNodes([InPresetNodeID = &PresetNodeID](FConstNodeHandle PresetNodeHandle)
@@ -873,7 +873,7 @@ namespace Metasound
 #endif // WITH_EDITOR
 		}
 
-		TArray<FMetasoundFrontendClassInput> FRebuildPresetRootGraph::GenerateRequiredClassInputs(const FConstGraphHandle& InPresetGraph, TSet<FName>& OutInputsInheritingDefault) const
+		TArray<FMetasoundFrontendClassInput> FRebuildPresetRootGraph::GenerateRequiredClassInputs(FDocumentHandle& InDocumentHandle, const FConstGraphHandle& InPresetGraph, TSet<FName>& OutInputsInheritingDefault) const
 		{
 			TArray<FMetasoundFrontendClassInput> ClassInputs;
 
@@ -896,8 +896,7 @@ namespace Metasound
 					ClassInput.Metadata.SetDescription(InputNode->GetDescription());
 					ClassInput.Metadata.SetDisplayName(Input->GetMetadata().GetDisplayName());
 #endif // WITH_EDITOR
-					FConstDocumentHandle DocumentHandle = InPresetGraph->GetOwningDocument();
-					FDocumentAccessPtr DocumentPtr = DocumentHandle->GetDocumentPtr();
+					FDocumentAccessPtr DocumentPtr = InDocumentHandle->GetDocumentPtr();
 					const FMetasoundFrontendDocument* Document = DocumentPtr.Get();
 					check(Document);
 					ClassInput.VertexID = FDocumentIDGenerator::Get().CreateVertexID(*Document);;
@@ -952,7 +951,7 @@ namespace Metasound
 			return ClassInputs;
 		}
 
-		TArray<FMetasoundFrontendClassOutput> FRebuildPresetRootGraph::GenerateRequiredClassOutputs(const FConstGraphHandle& InPresetGraph) const
+		TArray<FMetasoundFrontendClassOutput> FRebuildPresetRootGraph::GenerateRequiredClassOutputs(FDocumentHandle& InDocumentHandle, const FConstGraphHandle& InPresetGraph) const
 		{
 			TArray<FMetasoundFrontendClassOutput> ClassOutputs;
 
@@ -976,8 +975,7 @@ namespace Metasound
 					ClassOutput.Metadata.SetDisplayName(Output->GetMetadata().GetDisplayName());
 #endif // WITH_EDITOR
 
-					FConstDocumentHandle DocumentHandle = InPresetGraph->GetOwningDocument();
-					FDocumentAccessPtr DocumentPtr = DocumentHandle->GetDocumentPtr();
+					FDocumentAccessPtr DocumentPtr = InDocumentHandle->GetDocumentPtr();
 					const FMetasoundFrontendDocument* Document = DocumentPtr.Get();
 					check(Document);
 					ClassOutput.VertexID = FDocumentIDGenerator::Get().CreateVertexID(*Document);
diff --git a/Engine/Plugins/Runtime/Metasound/Source/MetasoundFrontend/Public/MetasoundFrontendTransform.h b/Engine/Plugins/Runtime/Metasound/Source/MetasoundFrontend/Public/MetasoundFrontendTransform.h
--- a/Engine/Plugins/Runtime/Metasound/Source/MetasoundFrontend/Public/MetasoundFrontendTransform.h
+++ b/Engine/Plugins/Runtime/Metasound/Source/MetasoundFrontend/Public/MetasoundFrontendTransform.h
@@ -176,10 +176,10 @@ namespace Metasound
 			// Get the class inputs needed for this preset. Input literals set on 
 			// the preset graph will be used if they are set and are marked as inheriting
 			// the default from the referenced graph.
-			TArray<FMetasoundFrontendClassInput> GenerateRequiredClassInputs(const FConstGraphHandle& InParentGraph, TSet<FName>& OutInputsInheritingDefault) const;
+			TArray<FMetasoundFrontendClassInput> GenerateRequiredClassInputs(FDocumentHandle& InDocumentHandle, const FConstGraphHandle& InParentGraph, TSet<FName>& OutInputsInheritingDefault) const;
 
 			// Get the class Outputs needed for this preset.
-			TArray<FMetasoundFrontendClassOutput> GenerateRequiredClassOutputs(const FConstGraphHandle& InParentGraph) const;
+			TArray<FMetasoundFrontendClassOutput> GenerateRequiredClassOutputs(FDocumentHandle& InDocumentHandle, const FConstGraphHandle& InParentGraph) const;
 
 			// Add inputs to parent graph and connect to wrapped graph node.
 			void AddAndConnectInputs(const TArray<FMetasoundFrontendClassInput>& InClassInputs, FGraphHandle& InParentGraphHandle, FNodeHandle& InReferencedNode) const;
