- Fix for crash when attempting to add graph input to MetaSound while auditioning builder from BP
- Fix for ensure if DataType provided to add graph input/output MetaSound Builder is invalid
#rnx

diff --git a/Engine/Plugins/Runtime/Metasound/Source/MetasoundEngine/Private/MetasoundBuilderSubsystem.cpp b/Engine/Plugins/Runtime/Metasound/Source/MetasoundEngine/Private/MetasoundBuilderSubsystem.cpp
--- a/Engine/Plugins/Runtime/Metasound/Source/MetasoundEngine/Private/MetasoundBuilderSubsystem.cpp
+++ b/Engine/Plugins/Runtime/Metasound/Source/MetasoundEngine/Private/MetasoundBuilderSubsystem.cpp
@@ -69,32 +69,41 @@ namespace Metasound::Engine
 
 FMetaSoundBuilderNodeOutputHandle UMetaSoundBuilderBase::AddGraphInputNode(FName Name, FName DataType, FMetasoundFrontendLiteral DefaultValue, EMetaSoundBuilderResult& OutResult, bool bIsConstructorInput)
 {
+	using namespace Metasound::Frontend;
+
 	FMetaSoundBuilderNodeOutputHandle NewHandle;
 
-	const FMetasoundFrontendNode* Node = Builder.FindGraphInputNode(Name);
-	if (Node)
+	if (IDataTypeRegistry::Get().FindDataTypeRegistryEntry(DataType) == nullptr)
 	{
-		UE_LOG(LogMetaSound, Warning, TEXT("AddGraphInputNode Failed: Input Node already exists with name '%s'; returning handle to existing node which may or may not match requested DataType '%s'"), *Name.ToString(), *DataType.ToString());
+		UE_LOG(LogMetaSound, Error, TEXT("AddGraphInputNode Failed on builder '%s' when attempting to add '%s': '%s' is not a registered DataType"), *GetName(), *Name.ToString(), *DataType.ToString());
 	}
 	else
 	{
-		FMetasoundFrontendClassInput Description;
-		Description.Name = Name;
-		Description.TypeName = DataType;
-		Description.NodeID = FGuid::NewGuid();
-		Description.VertexID = FGuid::NewGuid();
-		Description.DefaultLiteral = static_cast<FMetasoundFrontendLiteral>(DefaultValue);
-		Description.AccessType = bIsConstructorInput ? EMetasoundFrontendVertexAccessType::Value : EMetasoundFrontendVertexAccessType::Reference;
-		Node = Builder.AddGraphInput(Description);
-	}
+		const FMetasoundFrontendNode* Node = Builder.FindGraphInputNode(Name);
+		if (Node)
+		{
+			UE_LOG(LogMetaSound, Warning, TEXT("AddGraphInputNode Failed: Input Node already exists with name '%s'; returning handle to existing node which may or may not match requested DataType '%s'"), *Name.ToString(), *DataType.ToString());
+		}
+		else
+		{
+			FMetasoundFrontendClassInput Description;
+			Description.Name = Name;
+			Description.TypeName = DataType;
+			Description.NodeID = FGuid::NewGuid();
+			Description.VertexID = FGuid::NewGuid();
+			Description.DefaultLiteral = static_cast<FMetasoundFrontendLiteral>(DefaultValue);
+			Description.AccessType = bIsConstructorInput ? EMetasoundFrontendVertexAccessType::Value : EMetasoundFrontendVertexAccessType::Reference;
+			Node = Builder.AddGraphInput(Description);
+		}
 
-	if (Node)
-	{
-		const TArray<FMetasoundFrontendVertex>& Outputs = Node->Interface.Outputs;
-		checkf(!Outputs.IsEmpty(), TEXT("Node should be initialized and have one output."));
+		if (Node)
+		{
+			const TArray<FMetasoundFrontendVertex>& Outputs = Node->Interface.Outputs;
+			checkf(!Outputs.IsEmpty(), TEXT("Node should be initialized and have one output."));
 
-		NewHandle.NodeID = Node->GetID();
-		NewHandle.VertexID = Outputs.Last().VertexID;
+			NewHandle.NodeID = Node->GetID();
+			NewHandle.VertexID = Outputs.Last().VertexID;
+		}
 	}
 
 	OutResult = NewHandle.IsSet() ? EMetaSoundBuilderResult::Succeeded : EMetaSoundBuilderResult::Failed;
@@ -103,33 +112,43 @@ FMetaSoundBuilderNodeOutputHandle UMetaSoundBuilderBase::AddGraphInputNode(FName
 
 FMetaSoundBuilderNodeInputHandle UMetaSoundBuilderBase::AddGraphOutputNode(FName Name, FName DataType, FMetasoundFrontendLiteral DefaultValue, EMetaSoundBuilderResult& OutResult, bool bIsConstructorOutput)
 {
+	using namespace Metasound::Frontend;
+
 	FMetaSoundBuilderNodeInputHandle NewHandle;
-	const FMetasoundFrontendNode* Node = Builder.FindGraphOutputNode(Name);
-	if (Node)
+
+	if (IDataTypeRegistry::Get().FindDataTypeRegistryEntry(DataType) == nullptr)
 	{
-		UE_LOG(LogMetaSound, Warning, TEXT("AddGraphOutputNode Failed: Output Node already exists with name '%s'; returning handle to existing node which may or may not match requested DataType '%s'"), *Name.ToString(), *DataType.ToString());
+		UE_LOG(LogMetaSound, Error, TEXT("AddGraphOutputNode Failed on builder '%s' when attempting to add '%s': '%s' is not a registered DataType"), *GetName(), *Name.ToString(), *DataType.ToString());
 	}
 	else
 	{
-		FMetasoundFrontendClassOutput Description;
-		Description.Name = Name;
-		Description.TypeName = DataType;
-		Description.NodeID = FGuid::NewGuid();
-		Description.VertexID = FGuid::NewGuid();
-		Description.AccessType = bIsConstructorOutput ? EMetasoundFrontendVertexAccessType::Value : EMetasoundFrontendVertexAccessType::Reference;
-		Node = Builder.AddGraphOutput(Description);
-	}
-
-	if (Node)
-	{
-		const TArray<FMetasoundFrontendVertex>& Inputs = Node->Interface.Inputs;
-		checkf(!Inputs.IsEmpty(), TEXT("Node should be initialized and have one input."));
+		const FMetasoundFrontendNode* Node = Builder.FindGraphOutputNode(Name);
+		if (Node)
+		{
+			UE_LOG(LogMetaSound, Warning, TEXT("AddGraphOutputNode Failed: Output Node already exists with name '%s'; returning handle to existing node which may or may not match requested DataType '%s'"), *Name.ToString(), *DataType.ToString());
+		}
+		else
+		{
+			FMetasoundFrontendClassOutput Description;
+			Description.Name = Name;
+			Description.TypeName = DataType;
+			Description.NodeID = FGuid::NewGuid();
+			Description.VertexID = FGuid::NewGuid();
+			Description.AccessType = bIsConstructorOutput ? EMetasoundFrontendVertexAccessType::Value : EMetasoundFrontendVertexAccessType::Reference;
+			Node = Builder.AddGraphOutput(Description);
+		}
 
-		const FGuid& VertexID = Inputs.Last().VertexID;
-		if (Builder.SetNodeInputDefault(Node->GetID(), VertexID, DefaultValue))
+		if (Node)
 		{
-			NewHandle.NodeID = Node->GetID();
-			NewHandle.VertexID = VertexID;
+			const TArray<FMetasoundFrontendVertex>& Inputs = Node->Interface.Inputs;
+			checkf(!Inputs.IsEmpty(), TEXT("Node should be initialized and have one input."));
+
+			const FGuid& VertexID = Inputs.Last().VertexID;
+			if (Builder.SetNodeInputDefault(Node->GetID(), VertexID, DefaultValue))
+			{
+				NewHandle.NodeID = Node->GetID();
+				NewHandle.VertexID = VertexID;
+			}
 		}
 	}
 
@@ -1070,13 +1089,15 @@ void UMetaSoundSourceBuilder::OnEdgeAdded(int32 EdgeIndex) const
 	});
 }
 
-TOptional<Metasound::FAnyDataReference> UMetaSoundSourceBuilder::CreateDataReference(const Metasound::FOperatorSettings& InOperatorSettings, const Metasound::FLiteral& InLiteral, Metasound::EDataReferenceAccessType AccessType)
+TOptional<Metasound::FAnyDataReference> UMetaSoundSourceBuilder::CreateDataReference(
+	const Metasound::FOperatorSettings& InOperatorSettings,
+	FName DataType,
+	const Metasound::FLiteral& InLiteral,
+	Metasound::EDataReferenceAccessType AccessType)
 {
 	using namespace Metasound;
 	using namespace Metasound::Frontend;
 
-	// TODO: move to TFunctionRef
-	const FName DataType = GetMetasoundDataTypeName<float>();
 	return IDataTypeRegistry::Get().CreateDataReference(DataType, AccessType, InLiteral, InOperatorSettings);
 };
 
@@ -1093,7 +1114,7 @@ void UMetaSoundSourceBuilder::OnInputAdded(int32 InputIndex) const
 		const FMetasoundFrontendGraphClass& GraphClass = Doc.RootGraph;
 		const FMetasoundFrontendClassInput& NewInput = GraphClass.Interface.Inputs[InputIndex];
 		const FLiteral NewInputLiteral = NewInput.DefaultLiteral.ToLiteral(NewInput.TypeName);
-		Transactor.AddInputDataDestination(NewInput.NodeID, NewInput.Name, NewInputLiteral, &UMetaSoundSourceBuilder::CreateDataReference);
+		Transactor.AddInputDataDestination(NewInput.NodeID, NewInput.Name, NewInput.TypeName, NewInputLiteral, &UMetaSoundSourceBuilder::CreateDataReference);
 		return true;
 	});
 }
diff --git a/Engine/Plugins/Runtime/Metasound/Source/MetasoundEngine/Public/MetasoundBuilderSubsystem.h b/Engine/Plugins/Runtime/Metasound/Source/MetasoundEngine/Public/MetasoundBuilderSubsystem.h
--- a/Engine/Plugins/Runtime/Metasound/Source/MetasoundEngine/Public/MetasoundBuilderSubsystem.h
+++ b/Engine/Plugins/Runtime/Metasound/Source/MetasoundEngine/Public/MetasoundBuilderSubsystem.h
@@ -508,7 +508,7 @@ public:
 	virtual void InitFrontendBuilder() override;
 
 private:
-	static TOptional<Metasound::FAnyDataReference> CreateDataReference(const Metasound::FOperatorSettings& InOperatorSettings, const Metasound::FLiteral& InLiteral, Metasound::EDataReferenceAccessType AccessType);
+	static TOptional<Metasound::FAnyDataReference> CreateDataReference(const Metasound::FOperatorSettings& InOperatorSettings, FName DataType, const Metasound::FLiteral& InLiteral, Metasound::EDataReferenceAccessType AccessType);
 
 	void InitDelegates(Metasound::Frontend::FDocumentModifyDelegates& OutDocumentDelegates) const;
 	void OnEdgeAdded(int32 EdgeIndex) const;
diff --git a/Engine/Plugins/Runtime/Metasound/Source/MetasoundGraphCore/Private/MetasoundDynamicOperatorTransactor.cpp b/Engine/Plugins/Runtime/Metasound/Source/MetasoundGraphCore/Private/MetasoundDynamicOperatorTransactor.cpp
--- a/Engine/Plugins/Runtime/Metasound/Source/MetasoundGraphCore/Private/MetasoundDynamicOperatorTransactor.cpp
+++ b/Engine/Plugins/Runtime/Metasound/Source/MetasoundGraphCore/Private/MetasoundDynamicOperatorTransactor.cpp
@@ -280,7 +280,7 @@ namespace Metasound
 		 * @param InVertexName - Key for input vertex on InNode.
 		 *
 		 */
-		void FDynamicOperatorTransactor::AddInputDataDestination(const FGuid& InNodeID, const FVertexName& InVertexName, const FLiteral& InDefaultLiteral, FReferenceCreationFunction InFunc)
+		void FDynamicOperatorTransactor::AddInputDataDestination(const FGuid& InNodeID, const FVertexName& InVertexName, const FName InDataType, const FLiteral& InDefaultLiteral, FReferenceCreationFunction InFunc)
 		{
 			METASOUND_TRACE_CPUPROFILER_EVENT_SCOPE(Metasound::FDynamicOperatorTransactor::AddInputDataDestination);
 
@@ -305,7 +305,7 @@ namespace Metasound
 
 			auto CreateAddInputTransform = [&](const FOperatorSettings& InOperatorSettings, const FMetasoundEnvironment& InEnvironment) -> TUniquePtr<IDynamicOperatorTransform>
 			{
-				TOptional<FAnyDataReference> NewDataReference = InFunc(InOperatorSettings, InDefaultLiteral, ReferenceAccessType);
+				TOptional<FAnyDataReference> NewDataReference = InFunc(InOperatorSettings, InDataType, InDefaultLiteral, ReferenceAccessType);
 				if (NewDataReference.IsSet())
 				{
 					return MakeUnique<FAddInput>(OperatorID, InVertexName, *NewDataReference);
diff --git a/Engine/Plugins/Runtime/Metasound/Source/MetasoundGraphCore/Public/MetasoundDynamicOperatorTransactor.h b/Engine/Plugins/Runtime/Metasound/Source/MetasoundGraphCore/Public/MetasoundDynamicOperatorTransactor.h
--- a/Engine/Plugins/Runtime/Metasound/Source/MetasoundGraphCore/Public/MetasoundDynamicOperatorTransactor.h
+++ b/Engine/Plugins/Runtime/Metasound/Source/MetasoundGraphCore/Public/MetasoundDynamicOperatorTransactor.h
@@ -26,7 +26,7 @@ namespace Metasound
 		enum class EAudioFadeType : uint8;
 
 		using FLiteralAssignmentFunction = void(*)(const FOperatorSettings& InOperatorSettings, const FLiteral& InLiteral, const FAnyDataReference& OutDataRef);
-		using FReferenceCreationFunction = TOptional<FAnyDataReference>(*)(const FOperatorSettings& InSettings, const FLiteral& InLiteral, EDataReferenceAccessType InAccessType);
+		using FReferenceCreationFunction = TOptional<FAnyDataReference>(*)(const FOperatorSettings& InSettings, FName DataType, const FLiteral& InLiteral, EDataReferenceAccessType InAccessType);
 		using FOnInputVertexUpdated = TFunction<void(const FVertexName&, const FInputVertexInterfaceData&)>;
 		using FOnOutputVertexUpdated = TFunction<void(const FVertexName&, const FOutputVertexInterfaceData&)>;
 
@@ -91,7 +91,7 @@ namespace Metasound
 			 * @param InNode - Node which receives the data.
 			 * @param InVertexName - Key for input vertex on InNode.
 			 */
-			void AddInputDataDestination(const FGuid& InNode, const FVertexName& InVertexName, const FLiteral& InDefaultLiteral, FReferenceCreationFunction InFunc);
+			void AddInputDataDestination(const FGuid& InNode, const FVertexName& InVertexName, FName InDataType, const FLiteral& InDefaultLiteral, FReferenceCreationFunction InFunc);
 
 			/** Remove an exposed input from the graph. */
 			void RemoveInputDataDestination(const FVertexName& InVertexName);
