Audio unit test build fixes

diff --git a/Engine/Plugins/Runtime/Metasound/Source/MetasoundEngine/Private/MetasoundBuilderSubsystem.cpp b/Engine/Plugins/Runtime/Metasound/Source/MetasoundEngine/Private/MetasoundBuilderSubsystem.cpp
--- a/Engine/Plugins/Runtime/Metasound/Source/MetasoundEngine/Private/MetasoundBuilderSubsystem.cpp
+++ b/Engine/Plugins/Runtime/Metasound/Source/MetasoundEngine/Private/MetasoundBuilderSubsystem.cpp
@@ -1114,7 +1114,7 @@ void UMetaSoundSourceBuilder::OnInputAdded(int32 InputIndex) const
 		const FMetasoundFrontendGraphClass& GraphClass = Doc.RootGraph;
 		const FMetasoundFrontendClassInput& NewInput = GraphClass.Interface.Inputs[InputIndex];
 		const FLiteral NewInputLiteral = NewInput.DefaultLiteral.ToLiteral(NewInput.TypeName);
-		Transactor.AddInputDataDestination(NewInput.NodeID, NewInput.Name, NewInput.TypeName, NewInputLiteral, &UMetaSoundSourceBuilder::CreateDataReference);
+		Transactor.AddInputDataDestination(NewInput.NodeID, NewInput.Name, NewInputLiteral, &UMetaSoundSourceBuilder::CreateDataReference);
 		return true;
 	});
 }
diff --git a/Engine/Plugins/Runtime/Metasound/Source/MetasoundGraphCore/Private/MetasoundDynamicOperatorTransactor.cpp b/Engine/Plugins/Runtime/Metasound/Source/MetasoundGraphCore/Private/MetasoundDynamicOperatorTransactor.cpp
--- a/Engine/Plugins/Runtime/Metasound/Source/MetasoundGraphCore/Private/MetasoundDynamicOperatorTransactor.cpp
+++ b/Engine/Plugins/Runtime/Metasound/Source/MetasoundGraphCore/Private/MetasoundDynamicOperatorTransactor.cpp
@@ -12,6 +12,8 @@
 #include "MetasoundGraphAlgo.h"
 #include "MetasoundGraphAlgoPrivate.h"
 #include "MetasoundNodeInterface.h"
+#include "MetasoundOperatorBuilder.h"
+#include "MetasoundOperatorBuilderSettings.h"
 #include "MetasoundOperatorInterface.h"
 #include "MetasoundTrace.h"
 #include "Templates/SharedPointer.h"
@@ -23,6 +25,19 @@ namespace Metasound
 	{
 		namespace DynamicOperatorTransactorPrivate
 		{
+			// Return operator builder settings appropriate for building subgraphs
+			// of a dynamic operaotr.
+			FOperatorBuilderSettings GetOperatorBuilderSettings()
+			{
+				FOperatorBuilderSettings Settings = FOperatorBuilderSettings::GetDefaultSettings();
+
+				// Subgraphs must be rebindable to support connecting and disconnecting
+				// data references to subgraphs. 
+				Settings.bEnableOperatorRebind = true;
+
+				return Settings;
+			}
+
 			// Literal nodes always have output vertex with this name. 
 			static const FLazyName LiteralNodeOutputVertexName("Value");
 			
@@ -72,12 +87,14 @@ namespace Metasound
 		}
 
 		FDynamicOperatorTransactor::FDynamicOperatorTransactor(const FGraph& InGraph)
-		: Graph(InGraph)
+		: OperatorBuilder(DynamicOperatorTransactorPrivate::GetOperatorBuilderSettings())
+		, Graph(InGraph)
 		{
 		}
 
 		FDynamicOperatorTransactor::FDynamicOperatorTransactor()
-		: Graph(TEXT(""), FGuid())
+		: OperatorBuilder(DynamicOperatorTransactorPrivate::GetOperatorBuilderSettings())
+		, Graph(TEXT(""), FGuid())
 		{
 		}
 
@@ -280,7 +297,7 @@ namespace Metasound
 		 * @param InVertexName - Key for input vertex on InNode.
 		 *
 		 */
-		void FDynamicOperatorTransactor::AddInputDataDestination(const FGuid& InNodeID, const FVertexName& InVertexName, const FName InDataType, const FLiteral& InDefaultLiteral, FReferenceCreationFunction InFunc)
+		void FDynamicOperatorTransactor::AddInputDataDestination(const FGuid& InNodeID, const FVertexName& InVertexName, const FLiteral& InDefaultLiteral, FReferenceCreationFunction InFunc)
 		{
 			METASOUND_TRACE_CPUPROFILER_EVENT_SCOPE(Metasound::FDynamicOperatorTransactor::AddInputDataDestination);
 
@@ -305,7 +322,7 @@ namespace Metasound
 
 			auto CreateAddInputTransform = [&](const FOperatorSettings& InOperatorSettings, const FMetasoundEnvironment& InEnvironment) -> TUniquePtr<IDynamicOperatorTransform>
 			{
-				TOptional<FAnyDataReference> NewDataReference = InFunc(InOperatorSettings, InDataType, InDefaultLiteral, ReferenceAccessType);
+				TOptional<FAnyDataReference> NewDataReference = InFunc(InOperatorSettings, InputVertex.DataTypeName, InDefaultLiteral, ReferenceAccessType);
 				if (NewDataReference.IsSet())
 				{
 					return MakeUnique<FAddInput>(OperatorID, InVertexName, *NewDataReference);
@@ -635,7 +652,8 @@ namespace Metasound
 				InNode,
 				InOperatorSettings,
 				InterfaceData.GetInputs(),
-				InEnvironment
+				InEnvironment,
+				&OperatorBuilder // Supply an operator builder set to build rebindable inputs to ensure that subgraphs have their data references updated. 
 			};
 
 			FBuildResults Results;
diff --git a/Engine/Plugins/Runtime/Metasound/Source/MetasoundGraphCore/Public/MetasoundDynamicOperatorTransactor.h b/Engine/Plugins/Runtime/Metasound/Source/MetasoundGraphCore/Public/MetasoundDynamicOperatorTransactor.h
--- a/Engine/Plugins/Runtime/Metasound/Source/MetasoundGraphCore/Public/MetasoundDynamicOperatorTransactor.h
+++ b/Engine/Plugins/Runtime/Metasound/Source/MetasoundGraphCore/Public/MetasoundDynamicOperatorTransactor.h
@@ -6,6 +6,7 @@
 #include "Containers/UnrealString.h"
 #include "MetasoundBuilderInterface.h"
 #include "MetasoundGraph.h"
+#include "MetasoundOperatorBuilder.h"
 #include "MetasoundOperatorInterface.h"
 #include "MetasoundNodeInterface.h"
 #include "Templates/Function.h"
@@ -91,7 +92,7 @@ namespace Metasound
 			 * @param InNode - Node which receives the data.
 			 * @param InVertexName - Key for input vertex on InNode.
 			 */
-			void AddInputDataDestination(const FGuid& InNode, const FVertexName& InVertexName, FName InDataType, const FLiteral& InDefaultLiteral, FReferenceCreationFunction InFunc);
+			void AddInputDataDestination(const FGuid& InNode, const FVertexName& InVertexName, const FLiteral& InDefaultLiteral, FReferenceCreationFunction InFunc);
 
 			/** Remove an exposed input from the graph. */
 			void RemoveInputDataDestination(const FVertexName& InVertexName);
@@ -132,6 +133,7 @@ namespace Metasound
 
 			void EnqueueTransformOnOperatorQueues(FCreateTransformFunctionRef InFunc);
 
+			FOperatorBuilder OperatorBuilder;
 			FGraph Graph;
 
 			struct FDynamicOperatorInfo
