Refactor MetaSound search engine query registration event reduction

diff --git a/Engine/Plugins/Runtime/Metasound/Source/MetasoundFrontend/Private/MetasoundFrontendQuerySteps.cpp b/Engine/Plugins/Runtime/Metasound/Source/MetasoundFrontend/Private/MetasoundFrontendQuerySteps.cpp
--- a/Engine/Plugins/Runtime/Metasound/Source/MetasoundFrontend/Private/MetasoundFrontendQuerySteps.cpp
+++ b/Engine/Plugins/Runtime/Metasound/Source/MetasoundFrontend/Private/MetasoundFrontendQuerySteps.cpp
@@ -67,69 +67,49 @@ namespace Metasound
 	void FReduceRegistrationEventsToCurrentStatus::Reduce(const FFrontendQueryKey& InKey, FFrontendQueryPartition& InOutEntries) const
 	{
 		using namespace Frontend;
-		if (InOutEntries.IsEmpty())
-		{
-			return;
-		}
 
-		const FFrontendQueryEntry* FinalEntry = Algo::MaxElementBy(InOutEntries, GetTransactionTimestamp);
-		// If last entry is registration, keep the last entry 
-		// since that's the most relevant up to date status
-		if (IsValidTransactionOfType(FNodeRegistryTransaction::ETransactionType::NodeRegistration, FinalEntry))
+		InOutEntries.Sort([](const FFrontendQueryEntry& A, const FFrontendQueryEntry& B)
 		{
-			FFrontendQueryEntry Entry = *FinalEntry;
-			InOutEntries.Reset();
-			InOutEntries.Add(Entry);
-		} 
-		else if (IsValidTransactionOfType(FNodeRegistryTransaction::ETransactionType::NodeUnregistration, FinalEntry))
+			return GetTransactionTimestamp(A) < GetTransactionTimestamp(B);
+		});
+
+		// Registration - Unregistration pairs result in no net change, 
+		// so reduce to most recent transaction that is not a pair
+		int32 MostRecentNonPairedIndex = INDEX_NONE;
+		int32 Index = 0;
+		while (Index < InOutEntries.Num())
 		{
-			// Check that pairs of entries are Registration - Unregistration
-			auto CheckRegistrationTransactionPairs = [](FFrontendQueryPartition& InOutEntries, int32 EndIndexInclusive)
+			// If last entry, cannot be a pair 
+			if (Index == InOutEntries.Num() - 1)
 			{
-				InOutEntries.Sort([](const FFrontendQueryEntry& A, const FFrontendQueryEntry& B)
-				{
-					return GetTransactionTimestamp(A) < GetTransactionTimestamp(B);
-				});
-				for (int32 PairIndex = 0; PairIndex < EndIndexInclusive; PairIndex += 2)
-				{
-					const FNodeRegistryTransaction& FirstEntry = InOutEntries[PairIndex].Value.Get<FNodeRegistryTransaction>();
-					const FNodeRegistryTransaction& SecondEntry = InOutEntries[PairIndex + 1].Value.Get<FNodeRegistryTransaction>();
-					const FNodeRegistryTransaction::ETransactionType FirstEntryTransactionType = FirstEntry.GetTransactionType();
-					const FNodeRegistryTransaction::ETransactionType SecondEntryTransactionType = SecondEntry.GetTransactionType();
-
-					if (!(FirstEntryTransactionType == FNodeRegistryTransaction::ETransactionType::NodeRegistration &&
-						SecondEntryTransactionType == FNodeRegistryTransaction::ETransactionType::NodeUnregistration))
-					{
-						UE_LOG(LogMetaSound, Warning, TEXT("Mismatched transaction entries with keys '%s' and '%s' and transaction types '%s' and '%s' in a search engine query.\
-							Transactions should be registration and unregistration pairs except for the last entry."),
-							*FirstEntry.GetNodeRegistryKey(),
-							*SecondEntry.GetNodeRegistryKey(),
-							*FNodeRegistryTransaction::LexToString(FirstEntryTransactionType),
-							*FNodeRegistryTransaction::LexToString(SecondEntryTransactionType));
-					}
-				}
-			};
-			// If odd number of entries, check previous pairs and keep the final entry
-			if (InOutEntries.Num() % 2 == 1)
+				MostRecentNonPairedIndex = Index;
+				break;
+			}
+			// Skip if valid pair
+			if (InOutEntries[Index].Value.Get<FNodeRegistryTransaction>().GetTransactionType() == FNodeRegistryTransaction::ETransactionType::NodeRegistration &&
+				InOutEntries[Index + 1].Value.Get<FNodeRegistryTransaction>().GetTransactionType() == FNodeRegistryTransaction::ETransactionType::NodeUnregistration)
 			{
-				FFrontendQueryEntry Entry = *FinalEntry;
-				CheckRegistrationTransactionPairs(InOutEntries, FMath::Max(0, InOutEntries.Num() - 2));
-				InOutEntries.Reset();
-				InOutEntries.Add(Entry);
+				Index += 2;
 			}
-			// If even, check all pairs  
-			// and clear them all, since the current status is as if those hadn't happened
+			// If not valid pair, this is the current most recent non paired index
 			else
 			{
-				CheckRegistrationTransactionPairs(InOutEntries, InOutEntries.Num() - 1);
-				InOutEntries.Reset();
+				MostRecentNonPairedIndex = Index;
+				Index++;
 			}
 		}
+
+		// Empty or entries were all pairs 
+		if (MostRecentNonPairedIndex == INDEX_NONE)
+		{
+			InOutEntries.Reset();
+		}
 		else
 		{
-			UE_LOG(LogMetaSound, Warning, TEXT("Missing FNodeRegistryTransaction::ETransactionType case in search engine FReduceRegistrationEventsToCurrentStatus::Reduce."))
+			const FFrontendQueryEntry Entry = InOutEntries[MostRecentNonPairedIndex];
+			InOutEntries.Reset();
+			InOutEntries.Add(Entry);
 		}
-
 	}
 
 	FReduceRegistrationEventsToCurrentStatus::FTimeType FReduceRegistrationEventsToCurrentStatus::GetTransactionTimestamp(const FFrontendQueryEntry& InEntry)
