Move MetaSound AutoUpdate and PreprocessDocument from runtime to cook using deterministic ID generation
- remove old cvars for editor/runtime deterministic ID generation

diff --git a/Engine/Plugins/Runtime/Metasound/Source/MetasoundEngine/Private/MetasoundEngineAsset.h b/Engine/Plugins/Runtime/Metasound/Source/MetasoundEngine/Private/MetasoundEngineAsset.h
--- a/Engine/Plugins/Runtime/Metasound/Source/MetasoundEngine/Private/MetasoundEngineAsset.h
+++ b/Engine/Plugins/Runtime/Metasound/Source/MetasoundEngine/Private/MetasoundEngineAsset.h
@@ -5,6 +5,7 @@
 #include "Metasound.h"
 #include "MetasoundAssetManager.h"
 #include "MetasoundBuilderSubsystem.h"
+#include "MetasoundFrontendDocumentIdGenerator.h"
 #include "MetasoundUObjectRegistry.h"
 #include "Serialization/Archive.h"
 
@@ -98,9 +99,19 @@ namespace Metasound
 
 			if (UMetasoundEditorGraphBase* MetaSoundGraph = Cast<UMetasoundEditorGraphBase>(InMetaSound.GetGraph()))
 			{
-				// Cooked data must be deterministic, so do not call register graph as this can
-				// initiate an auto-update and/or local registry data cache and modify serialized data.
-				if (!InSaveContext.IsCooking())
+				if (InSaveContext.IsCooking())
+				{
+					// Use deterministic ID generation so more can be done at cook rather than runtime
+					if (MetaSoundEnableCookDeterministicIDGeneration != 0)
+					{
+						{
+							constexpr bool bIsDeterministic = true;
+							FDocumentIDGenerator::FScopeDeterminism DeterminismScope = FDocumentIDGenerator::FScopeDeterminism(bIsDeterministic);
+							InMetaSound.CookGraph();
+						}
+					}
+				}
+				else
 				{
 					MetaSoundGraph->RegisterGraphWithFrontend();
 					MetaSoundGraph->GetModifyContext().SetForceRefreshViews();
diff --git a/Engine/Plugins/Runtime/Metasound/Source/MetasoundEngine/Private/MetasoundSource.cpp b/Engine/Plugins/Runtime/Metasound/Source/MetasoundEngine/Private/MetasoundSource.cpp
--- a/Engine/Plugins/Runtime/Metasound/Source/MetasoundEngine/Private/MetasoundSource.cpp
+++ b/Engine/Plugins/Runtime/Metasound/Source/MetasoundEngine/Private/MetasoundSource.cpp
@@ -20,6 +20,7 @@
 #include "MetasoundFrontendController.h"
 #include "MetasoundFrontendDataTypeRegistry.h"
 #include "MetasoundFrontendDocumentBuilder.h"
+#include "MetasoundFrontendDocumentIdGenerator.h"
 #include "MetasoundFrontendQuery.h"
 #include "MetasoundFrontendQuerySteps.h"
 #include "MetasoundFrontendTransform.h"
@@ -53,6 +54,15 @@ namespace Metasound
 		{
 			Frontend::FMetaSoundAssetRegistrationOptions RegOptions;
 			RegOptions.bForceReregister = false;
+#if !WITH_EDITOR 
+			if (Frontend::MetaSoundEnableCookDeterministicIDGeneration != 0)
+			{
+				// When without editor, don't AutoUpdate or PreprocessDocument at runtime. This should happen at cook or asset save.
+				// When with editor, those are needed because sounds are not necessarily saved before previewing
+				RegOptions.bAutoUpdate = false;
+				RegOptions.bPreprocessDocument = false;
+			}
+#endif // !WITH_EDITOR
 			if (const UMetaSoundSettings* Settings = GetDefault<UMetaSoundSettings>())
 			{
 				RegOptions.bAutoUpdateLogWarningOnDroppedConnection = Settings->bAutoUpdateLogWarningOnDroppedConnection;
diff --git a/Engine/Plugins/Runtime/Metasound/Source/MetasoundFrontend/Private/MetasoundAssetBase.cpp b/Engine/Plugins/Runtime/Metasound/Source/MetasoundFrontend/Private/MetasoundAssetBase.cpp
--- a/Engine/Plugins/Runtime/Metasound/Source/MetasoundFrontend/Private/MetasoundAssetBase.cpp
+++ b/Engine/Plugins/Runtime/Metasound/Source/MetasoundFrontend/Private/MetasoundAssetBase.cpp
@@ -238,12 +238,20 @@ void FMetasoundAssetBase::RegisterGraphWithFrontend(Metasound::Frontend::FMetaSo
 	{
 		AssetName = OwningAsset->GetName();
 	}
+	
+	TSharedPtr<FMetasoundFrontendDocument> DocumentToRegister;
+	if (InRegistrationOptions.bPreprocessDocument)
+	{
+		DocumentToRegister = PreprocessDocument();
+	}
+	else
+	{
+		DocumentToRegister = MakeShared<FMetasoundFrontendDocument>(GetDocumentChecked());
+	}
 
-	TSharedPtr<FMetasoundFrontendDocument> PreprocessedDocument = PreprocessDocument();
-	CacheRuntimeData(*PreprocessedDocument);
+	CacheRuntimeData(*DocumentToRegister);
 	FNodeClassInfo AssetClassInfo = GetAssetClassInfo();
-
-	TUniquePtr<INodeRegistryEntry> RegistryEntry = MakeUnique<AssetBasePrivate::FNodeRegistryEntry>(AssetName, PreprocessedDocument, AssetClassInfo.AssetPath);
+	TUniquePtr<INodeRegistryEntry> RegistryEntry = MakeUnique<AssetBasePrivate::FNodeRegistryEntry>(AssetName, DocumentToRegister, AssetClassInfo.AssetPath);
 
 	UnregisterGraphWithFrontend();
 	RegistryKey = FMetasoundFrontendRegistryContainer::Get()->RegisterNode(MoveTemp(RegistryEntry));
@@ -268,6 +276,73 @@ void FMetasoundAssetBase::RegisterGraphWithFrontend(Metasound::Frontend::FMetaSo
 	}
 }
 
+void FMetasoundAssetBase::CookGraph()
+{
+	using namespace Metasound;
+	using namespace Metasound::Frontend;
+
+	METASOUND_TRACE_CPUPROFILER_EVENT_SCOPE(MetaSoundAssetBase::CookGraph);
+	if (IsRegistered())
+	{
+		return;
+	}
+
+	RegisterAssetDependenciesForCook();
+	IMetaSoundAssetManager::GetChecked().AddOrUpdateAsset(*GetOwningAsset());
+
+	// Auto update must be done after all referenced asset classes are registered
+	const bool bDidUpdate = AutoUpdate(/*bAutoUpdateLogWarningOnDroppedConnection=*/true);
+#if WITH_EDITOR
+	if (bDidUpdate)
+	{
+		GetModifyContext().SetForceRefreshViews();
+	}
+#endif // WITH_EDITOR
+
+#if WITH_EDITOR
+	// Must be completed after auto-update to ensure all non-transient referenced dependency data is up-to-date (ex.
+	// class version), which is required for most accurately caching current registry metadata.
+	CacheRegistryMetadata();
+#endif // WITH_EDITOR
+
+	FString AssetName;
+	const UObject* OwningAsset = GetOwningAsset();
+	if (ensure(OwningAsset))
+	{
+		AssetName = OwningAsset->GetName();
+	}
+	
+	TSharedPtr<FMetasoundFrontendDocument> DocumentToRegister = PreprocessDocument();
+	TSharedPtr<FMetasoundFrontendDocument> PreprocessedDocumentCopy = MakeShared<FMetasoundFrontendDocument>(*DocumentToRegister);
+	// Replace document with preprocessed document
+	SetDocument(*PreprocessedDocumentCopy, /*bMarkDirty=*/false);
+
+	FNodeClassInfo AssetClassInfo = GetAssetClassInfo();
+	TUniquePtr<INodeRegistryEntry> RegistryEntry = MakeUnique<AssetBasePrivate::FNodeRegistryEntry>(AssetName, DocumentToRegister, AssetClassInfo.AssetPath);
+
+	UnregisterGraphWithFrontend();
+	RegistryKey = FMetasoundFrontendRegistryContainer::Get()->RegisterNode(MoveTemp(RegistryEntry));
+
+	if (NodeRegistryKey::IsValid(RegistryKey) && FMetasoundFrontendRegistryContainer::Get()->IsNodeRegistered(RegistryKey))
+	{
+#if WITH_EDITORONLY_DATA
+		UpdateAssetRegistry();
+#endif // WITH_EDITORONLY_DATA
+	}
+	else
+	{
+		FString ClassName;
+		if (OwningAsset)
+		{
+			if (UClass* Class = OwningAsset->GetClass())
+			{
+				ClassName = Class->GetName();
+			}
+		}
+		UE_LOG(LogMetaSound, Error, TEXT("While cooking graph, registration failed for MetaSound node class '%s' of UObject class '%s'"), *AssetName, *ClassName);
+	}
+}
+
 void FMetasoundAssetBase::UnregisterGraphWithFrontend()
 {
 	using namespace Metasound::Frontend;
@@ -310,11 +385,14 @@ bool FMetasoundAssetBase::IsInterfaceDeclared(const FMetasoundFrontendVersion& I
 	return GetDocumentChecked().Interfaces.Contains(InVersion);
 }
 
-void FMetasoundAssetBase::SetDocument(const FMetasoundFrontendDocument& InDocument)
+void FMetasoundAssetBase::SetDocument(const FMetasoundFrontendDocument& InDocument, bool bMarkDirty)
 {
 	FMetasoundFrontendDocument& Document = GetDocumentChecked();
 	Document = InDocument;
-	MarkMetasoundDocumentDirty();
+	if (bMarkDirty)
+	{
+		MarkMetasoundDocumentDirty();
+	}
 }
 
 void FMetasoundAssetBase::AddDefaultInterfaces()
@@ -332,25 +410,24 @@ void FMetasoundAssetBase::AddDefaultInterfaces()
 	FModifyRootGraphInterfaces({ }, InitVersions).Transform(GetDocumentChecked());
 }
 
-void FMetasoundAssetBase::SetDocument(FMetasoundFrontendDocument&& InDocument)
+void FMetasoundAssetBase::SetDocument(FMetasoundFrontendDocument&& InDocument, bool bMarkDirty)
 {
 	FMetasoundFrontendDocument& Document = GetDocumentChecked();
 	Document = MoveTemp(InDocument);
-	MarkMetasoundDocumentDirty();
+	if (bMarkDirty)
+	{
+		MarkMetasoundDocumentDirty();
+	}
 }
 
 bool FMetasoundAssetBase::VersionAsset()
 {
 	using namespace Metasound;
 	using namespace Metasound::Frontend;
+
 	METASOUND_TRACE_CPUPROFILER_EVENT_SCOPE(MetaSoundAssetBase::VersionAsset);
-#if WITH_EDITOR
-	const bool bIsDeterministic = MetaSoundEnableDeterministicIDGenerationInEditorCVar != 0 && !IsRunningCookCommandlet();
-	FDocumentIDGenerator::FScopeDeterminism DeterminismScope(bIsDeterministic);
-#else
-	const bool bIsDeterministic = MetaSoundEnableRuntimeDeterministicIDGeneration != 0 && !IsRunningCookCommandlet();
+	constexpr bool bIsDeterministic = true;
 	FDocumentIDGenerator::FScopeDeterminism DeterminismScope(bIsDeterministic);
-#endif // WITH_EDITOR
 
 	FName AssetName;
 	FString AssetPath;
@@ -843,16 +920,26 @@ void FMetasoundAssetBase::RegisterAssetDependencies(const Metasound::Frontend::F
 	}
 }
 
+void FMetasoundAssetBase::RegisterAssetDependenciesForCook()
+{
+	using namespace Metasound::Frontend;
+
+	IMetaSoundAssetManager& AssetManager = IMetaSoundAssetManager::GetChecked();
+	TArray<FMetasoundAssetBase*> References = GetReferencedAssets();
+	for (FMetasoundAssetBase* Reference : References)
+	{
+		if (!Reference->IsRegistered())
+		{
+			// TODO: Check for infinite recursion and error if so
+			AssetManager.AddOrUpdateAsset(*(Reference->GetOwningAsset()));
+			Reference->CookGraph();
+		}
+	}
+}
+
 bool FMetasoundAssetBase::AutoUpdate(bool bInLogWarningsOnDroppedConnection)
 {
 	using namespace Metasound::Frontend;
-#if WITH_EDITOR
-	const bool bIsDeterministic = MetaSoundEnableDeterministicIDGenerationInEditorCVar != 0 && !IsRunningCookCommandlet();
-	FDocumentIDGenerator::FScopeDeterminism DeterminismScope(bIsDeterministic);
-#else
-	const bool bIsDeterministic = MetaSoundEnableRuntimeDeterministicIDGeneration != 0 && !IsRunningCookCommandlet();
-	FDocumentIDGenerator::FScopeDeterminism DeterminismScope(bIsDeterministic);
-#endif // WITH_EDITOR
 
 	FString OwningAssetName = GetOwningAssetName();
 	const bool bAutoUpdated = FAutoUpdateRootGraph(MoveTemp(OwningAssetName), bInLogWarningsOnDroppedConnection).Transform(GetDocumentHandle());
diff --git a/Engine/Plugins/Runtime/Metasound/Source/MetasoundFrontend/Private/MetasoundFrontendDocument.cpp b/Engine/Plugins/Runtime/Metasound/Source/MetasoundFrontend/Private/MetasoundFrontendDocument.cpp
--- a/Engine/Plugins/Runtime/Metasound/Source/MetasoundFrontend/Private/MetasoundFrontendDocument.cpp
+++ b/Engine/Plugins/Runtime/Metasound/Source/MetasoundFrontend/Private/MetasoundFrontendDocument.cpp
@@ -426,13 +426,7 @@ FMetasoundFrontendClassInterface FMetasoundFrontendClassInterface::GenerateClass
 			ClassInput.Name = InputVertex.VertexName;
 			ClassInput.TypeName = InputVertex.DataTypeName;
 			ClassInput.AccessType = Metasound::DocumentPrivate::CoreVertexAccessTypeToFrontendVertexAccessType(InputVertex.AccessType);
-#if WITH_EDITOR
-			const bool bIsDeterministic = MetaSoundEnableDeterministicIDGenerationInEditorCVar != 0 && !IsRunningCookCommandlet();
-			ClassInput.VertexID = bIsDeterministic ? FClassIDGenerator::Get().CreateInputID(ClassInput) : FGuid::NewGuid();
-#else 
-			const bool bIsDeterministic = MetaSoundEnableRuntimeDeterministicIDGeneration != 0 && !IsRunningCookCommandlet();
-			ClassInput.VertexID = bIsDeterministic ? FClassIDGenerator::Get().CreateInputID(ClassInput) : FGuid::NewGuid();
-#endif // WITH_EDITOR
+			ClassInput.VertexID = FClassIDGenerator::Get().CreateInputID(ClassInput);
 
 #if WITH_EDITOR
 			const FDataVertexMetadata& VertexMetadata = InputVertex.Metadata;
@@ -488,13 +482,7 @@ FMetasoundFrontendClassInterface FMetasoundFrontendClassInterface::GenerateClass
 			ClassOutput.Name = OutputVertex.VertexName;
 			ClassOutput.TypeName = OutputVertex.DataTypeName;
 			ClassOutput.AccessType = Metasound::DocumentPrivate::CoreVertexAccessTypeToFrontendVertexAccessType(OutputVertex.AccessType);
-#if WITH_EDITOR
-			const bool bIsDeterministic = MetaSoundEnableDeterministicIDGenerationInEditorCVar != 0 && !IsRunningCookCommandlet();
-			ClassOutput.VertexID = bIsDeterministic ? FClassIDGenerator::Get().CreateOutputID(ClassOutput) : FGuid::NewGuid();
-#else 
-			const bool bIsDeterministic = MetaSoundEnableRuntimeDeterministicIDGeneration != 0 && !IsRunningCookCommandlet();
-			ClassOutput.VertexID = bIsDeterministic ? FClassIDGenerator::Get().CreateOutputID(ClassOutput) : FGuid::NewGuid();
-#endif // WITH_EDITOR
+			ClassOutput.VertexID = FClassIDGenerator::Get().CreateOutputID(ClassOutput);
 #if WITH_EDITOR
 			const FDataVertexMetadata& VertexMetadata = OutputVertex.Metadata;
 
@@ -760,13 +748,7 @@ FMetasoundFrontendClassInput::FMetasoundFrontendClassInput(const Audio::FParamet
 	Name = InInput.InitValue.ParamName;
 	DefaultLiteral = FMetasoundFrontendLiteral(InInput.InitValue);
 	TypeName = Metasound::DocumentPrivate::ResolveMemberDataType(InInput.DataType, InInput.InitValue.ParamType);
-#if WITH_EDITOR
-	const bool bIsDeterministic = MetaSoundEnableDeterministicIDGenerationInEditorCVar != 0 && !IsRunningCookCommandlet();
-	VertexID = bIsDeterministic ? FClassIDGenerator::Get().CreateInputID(InInput) : FGuid::NewGuid();
-#else 
-	const bool bIsDeterministic = MetaSoundEnableRuntimeDeterministicIDGeneration != 0 && !IsRunningCookCommandlet();
-	VertexID = bIsDeterministic ? FClassIDGenerator::Get().CreateInputID(InInput) : FGuid::NewGuid();
-#endif // WITH_EDITOR
+	VertexID = FClassIDGenerator::Get().CreateInputID(InInput);
 
 #if WITH_EDITOR
 	// Interfaces should never serialize text to avoid desync between
@@ -794,13 +776,7 @@ FMetasoundFrontendClassOutput::FMetasoundFrontendClassOutput(const Audio::FParam
 
 	Name = Output.ParamName;
 	TypeName = Metasound::DocumentPrivate::ResolveMemberDataType(Output.DataType, Output.ParamType);
-#if WITH_EDITOR
-	const bool bIsDeterministic = MetaSoundEnableDeterministicIDGenerationInEditorCVar != 0 && !IsRunningCookCommandlet();
-	VertexID = bIsDeterministic ? FClassIDGenerator::Get().CreateOutputID(Output) : FGuid::NewGuid();
-#else 
-	const bool bIsDeterministic = MetaSoundEnableRuntimeDeterministicIDGeneration != 0 && !IsRunningCookCommandlet();
-	VertexID = bIsDeterministic ? FClassIDGenerator::Get().CreateOutputID(Output) : FGuid::NewGuid();
-#endif // WITH_EDITOR
+	VertexID = FClassIDGenerator::Get().CreateOutputID(Output);
 
 #if WITH_EDITOR
 	// Interfaces should never serialize text to avoid desync between
diff --git a/Engine/Plugins/Runtime/Metasound/Source/MetasoundFrontend/Private/MetasoundFrontendDocumentIdGenerator.cpp b/Engine/Plugins/Runtime/Metasound/Source/MetasoundFrontend/Private/MetasoundFrontendDocumentIdGenerator.cpp
--- a/Engine/Plugins/Runtime/Metasound/Source/MetasoundFrontend/Private/MetasoundFrontendDocumentIdGenerator.cpp
+++ b/Engine/Plugins/Runtime/Metasound/Source/MetasoundFrontend/Private/MetasoundFrontendDocumentIdGenerator.cpp
@@ -6,19 +6,11 @@ namespace Metasound
 {
 	namespace Frontend
 	{
-		int32 MetaSoundEnableDeterministicIDGenerationInEditorCVar = 0;
-		FAutoConsoleVariableRef CVarMetaSoundEnableDeterministicIDGenerationInEditor(
-			TEXT("au.MetaSound.EnableDeterministicIDGenerationInEditorCVar"),
-			MetaSoundEnableDeterministicIDGenerationInEditorCVar,
-			TEXT("Enable deterministic ID generation in MetaSounds for AutoUpdate while in editor (requires editor restart for effect) \n")
-			TEXT("0: Disabled (default), !0: Enabled"),
-			ECVF_Default);
-
-		int32 MetaSoundEnableRuntimeDeterministicIDGeneration = 0;
-		FAutoConsoleVariableRef CVarMetaSoundEnableRuntimeDeterministicIDGeneration(
-			TEXT("au.MetaSound.EnableRuntimeDeterministicIDGeneration"),
-			MetaSoundEnableRuntimeDeterministicIDGeneration,
-			TEXT("Enable deterministic ID generation in MetaSounds for AutoUpdate at runtime (non editor) \n")
+		int32 MetaSoundEnableCookDeterministicIDGeneration = 0;
+		FAutoConsoleVariableRef CVarMetaSoundEnableCookDeterministicIDGeneration(
+			TEXT("au.MetaSound.EnableCookDeterministicIDGeneration"),
+			MetaSoundEnableCookDeterministicIDGeneration,
+			TEXT("Enable moving MetaSound registration operations like AutoUpdate and PreprocessDocument from runtime to cook using deterministic ID generation \n")
 			TEXT("0: Disabled (default), !0: Enabled"),
 			ECVF_Default);
 
@@ -122,7 +114,7 @@ namespace Metasound
 			return CreateNamespacedIDFromString(ParameterInterfaceOutputNamespaceGuid, NameToHash);
 		}
 
-		FGuid FClassIDGenerator::CreateNamespacedIDFromString(const FGuid NamespaceGuid, FString StringToHash) const
+		FGuid FClassIDGenerator::CreateNamespacedIDFromString(const FGuid NamespaceGuid, const FString& StringToHash) const
 		{
 			FSHA1 Hasher;
 			Hasher.Update(reinterpret_cast<const uint8*>(&NamespaceGuid), sizeof(FGuid));
diff --git a/Engine/Plugins/Runtime/Metasound/Source/MetasoundFrontend/Public/MetasoundAssetBase.h b/Engine/Plugins/Runtime/Metasound/Source/MetasoundFrontend/Public/MetasoundAssetBase.h
--- a/Engine/Plugins/Runtime/Metasound/Source/MetasoundFrontend/Public/MetasoundAssetBase.h
+++ b/Engine/Plugins/Runtime/Metasound/Source/MetasoundFrontend/Public/MetasoundAssetBase.h
@@ -75,6 +75,9 @@ public:
 	// Unregisters the root graph of the given asset with the MetaSound Frontend.
 	void UnregisterGraphWithFrontend();
 
+	// Prepares graph for cook, including autoupdating and preprocessing graph and registering node, similar to RegisterGraphWithFrontend. 
+	void CookGraph();
+
 	// Sets/overwrites the root class metadata
 	UE_DEPRECATED(5.3, "Directly setting graph class Metadata is no longer be supported. Use the FMetaSoundFrontendDocumentBuilder to modify class data.")
 	virtual void SetMetadata(FMetasoundFrontendClassMetadata& InMetadata);
@@ -132,8 +135,8 @@ public:
 
 	// Overwrites the existing document. If the document's interface is not supported,
 	// the FMetasoundAssetBase be while queried for a new one using `GetPreferredInterface`.
-	void SetDocument(const FMetasoundFrontendDocument& InDocument);
-	void SetDocument(FMetasoundFrontendDocument&& InDocument);
+	void SetDocument(const FMetasoundFrontendDocument& InDocument, bool bMarkDirty = true);
+	void SetDocument(FMetasoundFrontendDocument&& InDocument, bool bMarkDirty = true);
 
 	FMetasoundFrontendDocument& GetDocumentChecked();
 	const FMetasoundFrontendDocument& GetDocumentChecked() const;
@@ -221,6 +224,7 @@ protected:
 
 	bool AutoUpdate(bool bInLogWarningsOnDroppedConnection);
 	void RegisterAssetDependencies(const Metasound::Frontend::FMetaSoundAssetRegistrationOptions& InRegistrationOptions);
+	void RegisterAssetDependenciesForCook();
 	TSharedPtr<FMetasoundFrontendDocument> PreprocessDocument();
 private:
 #if WITH_EDITORONLY_DATA
diff --git a/Engine/Plugins/Runtime/Metasound/Source/MetasoundFrontend/Public/MetasoundAssetManager.h b/Engine/Plugins/Runtime/Metasound/Source/MetasoundFrontend/Public/MetasoundAssetManager.h
--- a/Engine/Plugins/Runtime/Metasound/Source/MetasoundFrontend/Public/MetasoundAssetManager.h
+++ b/Engine/Plugins/Runtime/Metasound/Source/MetasoundFrontend/Public/MetasoundAssetManager.h
@@ -55,6 +55,9 @@ namespace Metasound
 			// If true, warnings will be logged if updating a node results in existing connections being discarded.
 			bool bAutoUpdateLogWarningOnDroppedConnection = false;
 
+			// If true, the preprocessed document will be registered 
+			bool bPreprocessDocument = true;
+
 #if WITH_EDITOR
 			// Attempt to rebuild referenced classes (only run if class not registered or set to force re-register)
 			bool bRebuildReferencedAssetClasses = true;
diff --git a/Engine/Plugins/Runtime/Metasound/Source/MetasoundFrontend/Public/MetasoundFrontendDocumentIdGenerator.h b/Engine/Plugins/Runtime/Metasound/Source/MetasoundFrontend/Public/MetasoundFrontendDocumentIdGenerator.h
--- a/Engine/Plugins/Runtime/Metasound/Source/MetasoundFrontend/Public/MetasoundFrontendDocumentIdGenerator.h
+++ b/Engine/Plugins/Runtime/Metasound/Source/MetasoundFrontend/Public/MetasoundFrontendDocumentIdGenerator.h
@@ -9,8 +9,7 @@ namespace Metasound
 {
 	namespace Frontend
 	{
-		extern METASOUNDFRONTEND_API int32 MetaSoundEnableDeterministicIDGenerationInEditorCVar;
-		extern METASOUNDFRONTEND_API int32 MetaSoundEnableRuntimeDeterministicIDGeneration;
+		extern METASOUNDFRONTEND_API int32 MetaSoundEnableCookDeterministicIDGeneration;
 
 		/*** 
 		 * For generating IDs using a given document. 
@@ -87,7 +86,7 @@ namespace Metasound
 			FGuid CreateOutputID(const FMetasoundFrontendClassOutput& Output) const;
 			FGuid CreateOutputID(const Audio::FParameterInterface::FOutput& Output) const;
 
-			FGuid CreateNamespacedIDFromString(const FGuid NamespaceGuid, FString StringToHash) const;
+			FGuid CreateNamespacedIDFromString(const FGuid NamespaceGuid, const FString& StringToHash) const;
 		};
 	}
 }
