[Metasound] NodeTestGraphBuilder: Add method for adding a constructor input which has a type that is not directly supported by FLiteral

diff --git a/Engine/Plugins/Runtime/Metasound/Source/MetasoundEngineTest/Public/NodeTestGraphBuilder.h b/Engine/Plugins/Runtime/Metasound/Source/MetasoundEngineTest/Public/NodeTestGraphBuilder.h
--- a/Engine/Plugins/Runtime/Metasound/Source/MetasoundEngineTest/Public/NodeTestGraphBuilder.h
+++ b/Engine/Plugins/Runtime/Metasound/Source/MetasoundEngineTest/Public/NodeTestGraphBuilder.h
@@ -101,6 +101,43 @@ namespace Metasound::Test
 			return InputNode;
 		}
 
+		/**
+		 * Adds a constructor input node to the graph and connects it as an input to the NodeToConnect.
+		 * This version handles input data types that don't have direct support in FLiteral.
+		 *
+		 * @tparam InputDataType The data type of the input vertex
+		 * @tparam LiteralDataType The data type of the literal used to set the input vertex
+		 * @param NodeToConnect The node to connect the newly added constructor input node to
+		 * @param InputName The name of the input pin to connect the node to, and the name to use for the new node.
+		 * @param Value The value for the node to hold. The type of the value will determine the type for the node.
+		 * @param NodeName Optional name to use for the created constructor input node
+		 * @return NodeHandle of the added node, can be invalid if failed to add
+		 */
+		template<typename InputDataType, typename LiteralDataType>
+		Frontend::FNodeHandle AddAndConnectConstructorInput(
+			const Frontend::FNodeHandle& NodeToConnect,
+			const FName& InputName,
+			const LiteralDataType& Value,
+			const FName& NodeName = NAME_None) const
+		{
+			check(RootGraph->IsValid());
+
+			const FName NameToUse = NodeName.IsNone() ? InputName : NodeName;
+			
+			FMetasoundFrontendClassInput Input;
+			Input.Name = NameToUse;
+			Input.TypeName = GetMetasoundDataTypeName<InputDataType>();
+			Input.VertexID = FGuid::NewGuid();
+			Input.AccessType = EMetasoundFrontendVertexAccessType::Value;
+			Input.DefaultLiteral.Set(Value);
+			
+			Frontend::FNodeHandle InputNode = RootGraph->AddInputVertex(Input);
+
+			ConnectNodes(InputNode, NameToUse, NodeToConnect, InputName);
+			
+			return InputNode;
+		}
+
 		/** 
 		* Adds data reference input nodes for each input on the NodeHandle and connects them to the inputs of the NodeHandle
 		*
