[Metasound] Add the ability to get Metasound rendering perf numbers (timing) at runtime.

diff --git a/Engine/Plugins/Runtime/Metasound/Source/MetasoundEngine/Private/MetasoundGeneratorHandle.cpp b/Engine/Plugins/Runtime/Metasound/Source/MetasoundEngine/Private/MetasoundGeneratorHandle.cpp
--- a/Engine/Plugins/Runtime/Metasound/Source/MetasoundEngine/Private/MetasoundGeneratorHandle.cpp
+++ b/Engine/Plugins/Runtime/Metasound/Source/MetasoundEngine/Private/MetasoundGeneratorHandle.cpp
@@ -378,6 +378,47 @@ bool UMetasoundGeneratorHandle::WatchOutput(
 	return true;
 }
 
+void UMetasoundGeneratorHandle::EnableRuntimeRenderTiming(bool Enable)
+{
+	bRuntimeRenderTimingShouldBeEnabled = Enable;
+
+	if (!IsValid())
+	{
+		return;
+	}
+
+	if (!CachedMetasoundSource.IsValid())
+	{
+		CacheMetasoundSource();
+	}
+
+	TSharedPtr<Metasound::FMetasoundGenerator> Generator = PinGenerator();
+	if (Generator)
+	{
+		Generator->EnableRuntimeRenderTiming(Enable);
+	}
+}
+
+double UMetasoundGeneratorHandle::GetCPUCoreUtilization()
+{
+	if (!IsValid())
+	{
+		return 0.0;
+	}
+
+	if (!CachedMetasoundSource.IsValid())
+	{
+		CacheMetasoundSource();
+	}
+
+	TSharedPtr<Metasound::FMetasoundGenerator> Generator = PinGenerator();
+	if (!Generator)
+	{
+		return 0.0;
+	}
+	return Generator->GetCPUCoreUtilization();
+}
+
 void UMetasoundGeneratorHandle::RegisterPassthroughAnalyzerForType(FName TypeName, FName AnalyzerName, FName OutputName)
 {
 	check(!PassthroughAnalyzers.Contains(TypeName));
@@ -432,6 +473,8 @@ void UMetasoundGeneratorHandle::OnSourceCreatedAGenerator(uint64 InAudioComponen
 				}
 			}
 			
+			InGenerator->EnableRuntimeRenderTiming(bRuntimeRenderTimingShouldBeEnabled);
+
 			OnGeneratorHandleAttached.Broadcast();
 			// If anyone has told us they are interested in being notified when a generator's 
 			// graph has changed go ahead and set that up now...
diff --git a/Engine/Plugins/Runtime/Metasound/Source/MetasoundEngine/Public/MetasoundGeneratorHandle.h b/Engine/Plugins/Runtime/Metasound/Source/MetasoundEngine/Public/MetasoundGeneratorHandle.h
--- a/Engine/Plugins/Runtime/Metasound/Source/MetasoundEngine/Public/MetasoundGeneratorHandle.h
+++ b/Engine/Plugins/Runtime/Metasound/Source/MetasoundEngine/Public/MetasoundGeneratorHandle.h
@@ -88,6 +88,22 @@ public:
 		FName AnalyzerName = NAME_None,
 		FName AnalyzerOutputName = NAME_None);
 
+	/**
+	* Enable the profiling of the MetaSound render for this playing instance. You
+	* must call this before calling "GetRuntimeCPUCoreUtilization" (below) or you will just 
+	* get 0.0 back for core utilization.
+	*/
+	UFUNCTION(BlueprintCallable, Category="MetaSoundPerf", meta=(AdvancedDisplay = "2"))
+	void EnableRuntimeRenderTiming(bool Enable);
+
+	/**
+	* Get the CPU usage as "fraction of real time" used to render this metasound. 
+	* NOTE: The MetasoundSource asset MUST have had its EnableRenderTiming function called
+	* before the metasound is started!
+	*/
+	UFUNCTION(BlueprintCallable, Category="MetaSoundPerf", meta=(AdvancedDisplay = "2"))
+	double GetCPUCoreUtilization();
+
 	/**
 	 * Map a type name to a passthrough analyzer name to use as a default for UMetasoundOutputSubsystem::WatchOutput()
 	 *
@@ -201,4 +217,6 @@ private:
 	
 	TSharedRef<TSpscQueue<FOutputPayload>> ChangedOutputs = MakeShared<TSpscQueue<FOutputPayload>>();
 	FDelegateHandle OutputChangedDelegateHandle;
+
+	bool bRuntimeRenderTimingShouldBeEnabled{ false };
 };
diff --git a/Engine/Plugins/Runtime/Metasound/Source/MetasoundGenerator/Private/MetasoundGenerator.cpp b/Engine/Plugins/Runtime/Metasound/Source/MetasoundGenerator/Private/MetasoundGenerator.cpp
--- a/Engine/Plugins/Runtime/Metasound/Source/MetasoundGenerator/Private/MetasoundGenerator.cpp
+++ b/Engine/Plugins/Runtime/Metasound/Source/MetasoundGenerator/Private/MetasoundGenerator.cpp
@@ -80,7 +80,6 @@ namespace Metasound
 		}
 #endif // if ENABLE_METASOUNDGENERATOR_INVALID_SAMPLE_VALUE_LOGGING
 
-#if ENABLE_METASOUND_GENERATOR_RENDER_TIMING
 		struct FRenderTimer
 		{
 			FRenderTimer(const FOperatorSettings& InSettings, double InAnalysisDuration)
@@ -96,16 +95,12 @@ namespace Metasound
 				SecondsOfAudioProducedPerBlock = static_cast<double>(InSettings.GetNumFramesPerBlock()) / FMath::Max(1., static_cast<double>(InSettings.GetSampleRate()));
 			}
 
-			double GetCPUCoreUtilization() const
+			double UpdateCPUCoreUtilization()
 			{
-				return CPUCoreUtilization;
-			}
-
-			void UpdateCPUCoreUtilization(double InCPUSecondsToRenderBlock)
-			{
-				if (InCPUSecondsToRenderBlock > 0.0)
+				double CPUSecondsToRenderBlock = FPlatformTime::ToSeconds64(AccumulatedCycles);
+				if (CPUSecondsToRenderBlock > 0.0)
 				{
-					double NewCPUUtil = InCPUSecondsToRenderBlock / SecondsOfAudioProducedPerBlock;
+					double NewCPUUtil = CPUSecondsToRenderBlock / SecondsOfAudioProducedPerBlock;
 					if (CPUCoreUtilization >= 0.0)
 					{
 						CPUCoreUtilization = SmoothingAlpha * NewCPUUtil + (1. - SmoothingAlpha) * CPUCoreUtilization;
@@ -115,9 +110,17 @@ namespace Metasound
 						CPUCoreUtilization = NewCPUUtil;
 					}
 				}
+				AccumulatedCycles = 0;
+				return CPUCoreUtilization;
+			}
+
+			FORCEINLINE void AccumulateCycles(uint64 Cycles)
+			{
+				AccumulatedCycles += Cycles;
 			}
 
 		private:
+			uint64 AccumulatedCycles = 0;
 			double CPUCoreUtilization = -1.0;
 			double SmoothingAlpha = 1.0;
 			double SecondsOfAudioProducedPerBlock = 0.0;
@@ -125,8 +128,8 @@ namespace Metasound
 
 		struct FBlockRenderScope
 		{
-			FBlockRenderScope(FRenderTimer& InTimer)
-			: Timer(&InTimer)
+			FBlockRenderScope(FRenderTimer* InTimer)
+			: Timer(InTimer)
 			{
 				StartCycle = FPlatformTime::Cycles64();
 			}
@@ -134,18 +137,17 @@ namespace Metasound
 			~FBlockRenderScope()
 			{
 				uint64 EndCycle = FPlatformTime::Cycles64();
-				if (EndCycle > StartCycle)
+				if (Timer)
 				{
-					Timer->UpdateCPUCoreUtilization(FPlatformTime::ToSeconds64(EndCycle - StartCycle));
+					Timer->AccumulateCycles(EndCycle - StartCycle);
 				}
 			}
+			
 		private:
 			uint64 StartCycle = 0;
 			FRenderTimer* Timer = nullptr;
 		};
 
-#endif // if ENABLE_METASOUND_GENERATOR_RENDER_TIMING
-
 #ifndef ENABLE_METASOUND_CONSTANT_OUTPUT_VERTEX_ERROR_LOG
 #define ENABLE_METASOUND_CONSTANT_OUTPUT_VERTEX_ERROR_LOG DO_CHECK
 #endif
@@ -269,8 +271,11 @@ namespace Metasound
 		, NumFramesPerExecute(InOperatorSettings.GetNumFramesPerBlock())
 		, NumSamplesPerExecute(0)
 		, OnFinishedTriggerRef(FTriggerWriteRef::CreateNew(InOperatorSettings))
+		, RenderTime(0.0)
 #if ENABLE_METASOUND_GENERATOR_RENDER_TIMING
-		, RenderTimer(MakeUnique<MetasoundGeneratorPrivate::FRenderTimer>(InOperatorSettings, 1. /* AnalysisPeriod */))
+		, bDoRuntimeRenderTiming(true)
+#else
+		, bDoRuntimeRenderTiming(false)
 #endif // if ENABLE_METASOUND_GENERATOR_RENDER_TIMING
 
 	{
@@ -536,6 +541,8 @@ PRAGMA_ENABLE_DEPRECATION_WARNINGS
 
 		UpdateGraphIfPending();
 
+		HandleRenderTimingEnableDisable();
+
 		// Output silent audio if we're still building a graph
 		if (bIsWaitingForFirstGraph)
 		{
@@ -574,58 +581,68 @@ PRAGMA_ENABLE_DEPRECATION_WARNINGS
 
 		while (NumSamplesRemaining > 0)
 		{
-			ApplyPendingUpdatesToInputs();
-
+			// Create a scoped timed section for the bulk of the processing...
 			{
-#if ENABLE_METASOUND_GENERATOR_RENDER_TIMING
-				// Time how long the root executer takes.
-				MetasoundGeneratorPrivate::FBlockRenderScope BlockRenderScope(*RenderTimer);
-#endif // if ENABLE_METASOUND_GENERATOR_RENDER_TIMING
+				MetasoundGeneratorPrivate::FBlockRenderScope RuntimelBlockRenderScope(RenderTimer.Get());
+
+				ApplyPendingUpdatesToInputs();
 
 				// Call metasound graph operator.
 				RootExecuter.Execute();
-			}
+				
+				// Check if generated finished during this execute call
+				if (*OnFinishedTriggerRef)
+				{
+					FinishSample = ((*OnFinishedTriggerRef)[0] * NumChannels);
+				}
 
-			if (GraphAnalyzer.IsValid())
-			{
-				GraphAnalyzer->Execute();
-			}
+				// Interleave audio because ISoundGenerator interface expects interleaved audio.
+				InterleaveGeneratedAudio();
 
-			// Execute the output analyzers
-			for (const TUniquePtr<Frontend::IVertexAnalyzer>& Analyzer : OutputAnalyzers)
-			{
-				Analyzer->Execute();
-			}
 
-			// Check if generated finished during this execute call
-			if (*OnFinishedTriggerRef)
-			{
-				FinishSample = ((*OnFinishedTriggerRef)[0] * NumChannels);
-			}
+				// Add audio generated during graph execution to the output buffer.
+				int32 ThisLoopNumSamplesWritten = FillWithBuffer(InterleavedAudioBuffer, &OutAudio[NumSamplesWritten], NumSamplesRemaining);
 
-			// Interleave audio because ISoundGenerator interface expects interleaved audio.
-			InterleaveGeneratedAudio();
+				NumSamplesRemaining -= ThisLoopNumSamplesWritten;
+				NumSamplesWritten += ThisLoopNumSamplesWritten;
 
+				// If not all the samples were written, then we have to save the 
+				// additional samples to the overflow buffer.
+				if (ThisLoopNumSamplesWritten < InterleavedAudioBuffer.Num())
+				{
+					int32 OverflowCount = InterleavedAudioBuffer.Num() - ThisLoopNumSamplesWritten;
 
-			// Add audio generated during graph execution to the output buffer.
-			int32 ThisLoopNumSamplesWritten = FillWithBuffer(InterleavedAudioBuffer, &OutAudio[NumSamplesWritten], NumSamplesRemaining);
+					OverflowBuffer.Reset();
+					OverflowBuffer.AddUninitialized(OverflowCount);
 
-			NumSamplesRemaining -= ThisLoopNumSamplesWritten;
-			NumSamplesWritten += ThisLoopNumSamplesWritten;
+					FMemory::Memcpy(OverflowBuffer.GetData(), &InterleavedAudioBuffer.GetData()[ThisLoopNumSamplesWritten], OverflowCount * sizeof(float));
+				}
 
-			// If not all the samples were written, then we have to save the 
-			// additional samples to the overflow buffer.
-			if (ThisLoopNumSamplesWritten < InterleavedAudioBuffer.Num())
-			{
-				int32 OverflowCount = InterleavedAudioBuffer.Num() - ThisLoopNumSamplesWritten;
+				// Execute the output analyzers
+				for (const TUniquePtr<Frontend::IVertexAnalyzer>& Analyzer : OutputAnalyzers)
+				{
+					Analyzer->Execute();
+				}
+			}
+
+			// Don't time the graph analyzer. It is only used for graph visualization.
 
-				OverflowBuffer.Reset();
-				OverflowBuffer.AddUninitialized(OverflowCount);
+			if (GraphAnalyzer.IsValid())
+			{
+				GraphAnalyzer->Execute();
+			}
 
-				FMemory::Memcpy(OverflowBuffer.GetData(), &InterleavedAudioBuffer.GetData()[ThisLoopNumSamplesWritten], OverflowCount * sizeof(float));
+			// Create a scoped timed section for the post processing...
+			{
+				MetasoundGeneratorPrivate::FBlockRenderScope RuntimelBlockRenderScope(RenderTimer.Get());
+				RootExecuter.PostExecute();
 			}
 
-			RootExecuter.PostExecute();
+			// Update timer if there is one...
+			if (RenderTimer)
+			{
+				RenderTime = RenderTimer->UpdateCPUCoreUtilization();
+			}
 		}
 
 		return NumSamplesWritten;
@@ -641,12 +658,10 @@ PRAGMA_ENABLE_DEPRECATION_WARNINGS
 		return bIsFinished;
 	}
 
-#if ENABLE_METASOUND_GENERATOR_RENDER_TIMING
 	double FMetasoundGenerator::GetCPUCoreUtilization() const
 	{
-		return RenderTimer->GetCPUCoreUtilization();
+		return RenderTime;
 	}
-#endif // if ENABLE_METASOUND_GENERATOR_RENDER_TIMING
 
 	int32 FMetasoundGenerator::FillWithBuffer(const Audio::FAlignedFloatBuffer& InBuffer, float* OutAudio, int32 MaxNumOutputSamples)
 	{
@@ -756,6 +771,19 @@ PRAGMA_ENABLE_DEPRECATION_WARNINGS
 		}
 	}
 
+	void FMetasoundGenerator::HandleRenderTimingEnableDisable()
+	{
+		if (bDoRuntimeRenderTiming && !RenderTimer)
+		{
+			RenderTimer = MakeUnique<MetasoundGeneratorPrivate::FRenderTimer>(OperatorSettings, 1. /* AnalysisPeriod */);
+		}
+		else if (!bDoRuntimeRenderTiming && RenderTimer)
+		{
+			RenderTimer = nullptr;
+			RenderTime = 0.0;
+		}
+	}
+
 	FMetasoundConstGraphGenerator::FMetasoundConstGraphGenerator(FMetasoundGeneratorInitParams&& InInitParams)
 	: FMetasoundGenerator(InInitParams.OperatorSettings)
 	{
diff --git a/Engine/Plugins/Runtime/Metasound/Source/MetasoundGenerator/Public/MetasoundGenerator.h b/Engine/Plugins/Runtime/Metasound/Source/MetasoundGenerator/Public/MetasoundGenerator.h
--- a/Engine/Plugins/Runtime/Metasound/Source/MetasoundGenerator/Public/MetasoundGenerator.h
+++ b/Engine/Plugins/Runtime/Metasound/Source/MetasoundGenerator/Public/MetasoundGenerator.h
@@ -30,9 +30,7 @@ namespace Metasound
 
 	namespace MetasoundGeneratorPrivate
 	{
-#if ENABLE_METASOUND_GENERATOR_RENDER_TIMING
 		struct FRenderTimer;
-#endif // if ENABLE_METASOUND_GENERATOR_RENDER_TIMING
 
 		struct FParameterSetter
 		{
@@ -212,12 +210,15 @@ namespace Metasound
 		bool IsFinished() const override;
 		//~ End FSoundGenerator
 
-#if ENABLE_METASOUND_GENERATOR_RENDER_TIMING
+		/** Enables the performance timing of the metasound rendering process. You
+		 * must call this before "GetCPUCoreUtilization" or the results will
+		 * always be 0.0.
+		 */
+		void EnableRuntimeRenderTiming(bool Enable) { bDoRuntimeRenderTiming = Enable; }
+
 		/** Fraction of a single CPU core used to render audio on a scale of 0.0 to 1.0 */
 		double GetCPUCoreUtilization() const;
 
-#endif // if ENABLE_METASOUND_GENERATOR_RENDER_TIMING
-		
 		// Called when a new graph has been "compiled" and set up as this generator's graph.
 		// Note: We don't allow direct assignment to the FOnSetGraph delegate
 		// because we want to give the Delegate an initial immediate callback if the generator 
@@ -269,6 +270,7 @@ namespace Metasound
 		
 		void ApplyPendingUpdatesToInputs();
 
+		void HandleRenderTimingEnableDisable();
 
 		bool bIsGraphBuilding;
 		bool bIsFinishTriggered;
@@ -329,9 +331,9 @@ namespace Metasound
 
 		FOnSetGraph OnSetGraph;
 
-#if ENABLE_METASOUND_GENERATOR_RENDER_TIMING
+		double RenderTime;
+		bool bDoRuntimeRenderTiming;
 		TUniquePtr<MetasoundGeneratorPrivate::FRenderTimer> RenderTimer;
-#endif // if ENABLE_METASOUND_GENERATOR_RENDER_TIMING
 	};
 
 	/** FMetasoundConstGraphGenerator generates audio from a given metasound IOperator
