Audio bus writer underrun warning fix.

#tests AudioQA AQ_AudioBus_MetaSounds map, EngineTest test list "All Tests" group "Audio" search "audio bus"

diff --git a/Engine/Plugins/Runtime/Metasound/Source/MetasoundEngine/Private/MetasoundAudioBusWriterNode.cpp b/Engine/Plugins/Runtime/Metasound/Source/MetasoundEngine/Private/MetasoundAudioBusWriterNode.cpp
--- a/Engine/Plugins/Runtime/Metasound/Source/MetasoundEngine/Private/MetasoundAudioBusWriterNode.cpp
+++ b/Engine/Plugins/Runtime/Metasound/Source/MetasoundEngine/Private/MetasoundAudioBusWriterNode.cpp
@@ -132,10 +132,12 @@ namespace Metasound
 						AudioBusSubsystem->StartAudioBus(AudioBusKey, AudioBusChannels, false);
 
 						InterleavedBuffer.Reset();
-						InterleavedBuffer.AddZeroed(BlockSizeFrames * AudioBusChannels);
+						InterleavedBuffer.AddZeroed(NumBlocksToNumSamples(1));
 
 						// Create a bus patch input with enough room for the number of samples we expect and some buffering
 						AudioBusPatchInput = AudioBusSubsystem->AddPatchInputForAudioBus(AudioBusKey, BlockSizeFrames, AudioBusChannels);
+
+						ConnectionState = EConnectionState::Disconnected;
 					}
 				}
 			}
@@ -164,7 +166,6 @@ namespace Metasound
 			BlockSizeFrames = InParams.OperatorSettings.GetNumFramesPerBlock();
 
 			CreatePatchInput();
-			bFirstBlock = true;
 		}
 
 
@@ -256,13 +257,34 @@ namespace Metasound
 				}
 			}
 
-			if (bFirstBlock)
+			if (ConnectionState != EConnectionState::Connected)
 			{
-				bFirstBlock = false;
+				int32 InitialNumBlocks = 1;
 				if (AudioMixerOutputFrames != BlockSizeFrames)
 				{
-					// Ensure there will be enough samples in the patch input to support the maximum metasound executions the mixer requires to fill its output frames after the next push.
-					AudioBusPatchInput.PushAudio(nullptr, (FMath::DivideAndRoundUp(FMath::Max(AudioMixerOutputFrames, BlockSizeFrames), FMath::Min(AudioMixerOutputFrames, BlockSizeFrames)) - 1) * BlockSizeFrames * AudioBusChannels);
+					InitialNumBlocks = FMath::DivideAndRoundUp(FMath::Max(AudioMixerOutputFrames, BlockSizeFrames), FMath::Min(AudioMixerOutputFrames, BlockSizeFrames));
+				}
+
+				if (ConnectionState == EConnectionState::Disconnected)
+				{
+					if (InitialNumBlocks > 1)
+					{
+						// Ensure there will be enough samples in the patch input to support the maximum metasound executions the mixer requires to fill its output frames after the next push.
+						AudioBusPatchInput.PushAudio(nullptr, NumBlocksToNumSamples(InitialNumBlocks - 1));
+					}
+
+					ConnectionState = EConnectionState::ConnectionPending;
+				}
+				else
+				{
+					// Determine if the pending connection has been established, by detecting if samples have been consumed.
+					if (AudioBusPatchInput.GetNumSamplesAvailable() == NumBlocksToNumSamples(InitialNumBlocks))
+					{
+						UE_LOG(LogMetaSound, Verbose, TEXT("Writer node executed before mixer patch connection established."));
+						return;
+					}
+
+					ConnectionState = EConnectionState::Connected;
 				}
 			}
 
@@ -275,6 +297,11 @@ namespace Metasound
 		}
 
 	private:
+		int32 NumBlocksToNumSamples(int32 NumBlocks) const
+		{
+			return NumBlocks * BlockSizeFrames * AudioBusChannels;
+		}
+
 		FAudioBusAssetReadRef AudioBusAsset;
 		TArray<FAudioBufferReadRef> AudioInputs;
 
@@ -286,7 +313,13 @@ namespace Metasound
 		uint32 AudioBusChannels = INDEX_NONE;
 		uint32 AudioBusId = 0;
 		int32 BlockSizeFrames = 0;
-		bool bFirstBlock = true;
+		enum class EConnectionState : uint8
+		{
+			Disconnected,
+			ConnectionPending,
+			Connected
+		}
+		ConnectionState = EConnectionState::Disconnected;
 	};
 
 	template<uint32 NumChannels>
diff --git a/Engine/Source/Runtime/AudioMixer/Private/AudioMixerDevice.cpp b/Engine/Source/Runtime/AudioMixer/Private/AudioMixerDevice.cpp
--- a/Engine/Source/Runtime/AudioMixer/Private/AudioMixerDevice.cpp
+++ b/Engine/Source/Runtime/AudioMixer/Private/AudioMixerDevice.cpp
@@ -83,7 +83,7 @@ FAutoConsoleVariableRef CVarDebugGeneratorFreq(
 	TEXT("0: Not Disabled, 1: SinTone, 2: WhiteNoise"),
 	ECVF_Default);
 
-static int32 AudioMixerPatchBufferBlocks = 3;
+static int32 AudioMixerPatchBufferBlocks = 2;
 FAutoConsoleVariableRef CVarAudioMixerPatchBufferBlocks(
 	TEXT("au.PatchBufferBlocks"),
 	AudioMixerPatchBufferBlocks,
@@ -2721,8 +2721,13 @@ namespace Audio
 
 	FPatchOutputStrongPtr FMixerDevice::MakePatch(int32 InFrames, int32 InChannels, float InGain) const
 	{
+		// Assume the mixer will consume SourceManager->GetNumOutputFrames() per iteration and an input patch will generate InFrames per iteration.
+		// An input patch must have adequate space to contain as many frames as the mixer might consume, as well as as many as might be pushed to the patch.
+		// This should be twice the ceiling of the ratio of the larger number of frames to the smaller number, times InFrames.
+		// An output patch must have adequate space to contain as many frames as the mixer might generate, as well as as many as might be consumed from the patch.
+		// This should be the same number.
 		int32 MaxSizeFrames = FMath::Max(InFrames, SourceManager->GetNumOutputFrames()), MinSizeFrames = FMath::Min(InFrames, SourceManager->GetNumOutputFrames());
-		return MakeShared<Audio::FPatchOutput, ESPMode::ThreadSafe>(AudioMixerPatchBufferBlocks * MinSizeFrames * FMath::DivideAndRoundUp(MaxSizeFrames, MinSizeFrames) * InChannels, InGain);
+		return MakeShared<Audio::FPatchOutput, ESPMode::ThreadSafe>(AudioMixerPatchBufferBlocks * InFrames * FMath::DivideAndRoundUp(MaxSizeFrames, MinSizeFrames) * InChannels, InGain);
 	}
 
 	FPatchOutputStrongPtr FMixerDevice::AddPatchForSubmix(uint32 InObjectId, float InPatchGain)
diff --git a/Engine/Source/Runtime/SignalProcessing/Private/MultithreadedPatching.cpp b/Engine/Source/Runtime/SignalProcessing/Private/MultithreadedPatching.cpp
--- a/Engine/Source/Runtime/SignalProcessing/Private/MultithreadedPatching.cpp
+++ b/Engine/Source/Runtime/SignalProcessing/Private/MultithreadedPatching.cpp
@@ -298,6 +298,17 @@ namespace Audio
 		return SamplesPushed;
 	}
 
+	/** Returns the current number of samples buffered in this input. */
+	int32 FPatchInput::GetNumSamplesAvailable() const
+	{
+		if (!OutputHandle.IsValid())
+		{
+			return 0;
+		}
+
+		return OutputHandle->GetNumSamplesAvailable();
+	}
+
 	void FPatchInput::SetGain(float InGain)
 	{
 		if (!OutputHandle.IsValid())
diff --git a/Engine/Source/Runtime/SignalProcessing/Public/DSP/MultithreadedPatching.h b/Engine/Source/Runtime/SignalProcessing/Public/DSP/MultithreadedPatching.h
--- a/Engine/Source/Runtime/SignalProcessing/Public/DSP/MultithreadedPatching.h
+++ b/Engine/Source/Runtime/SignalProcessing/Public/DSP/MultithreadedPatching.h
@@ -111,6 +111,9 @@ namespace Audio
 		 */
 		SIGNALPROCESSING_API int32 PushAudio(const float* InBuffer, int32 NumSamples);
 
+		/** Returns the current number of samples buffered in this input. */
+		SIGNALPROCESSING_API int32 GetNumSamplesAvailable() const;
+
 		SIGNALPROCESSING_API void SetGain(float InGain);
 
 		/** Returns false if this output was removed, either because someone called FPatchMixer::RemoveTap with this FPatchInput, or the FPatchMixer was destroyed. */
