Various minor MetaSound Builder BP function fixups
- Default MajorVersion to 1 for AddNodeByClassName w/ better log
- Set some lesser used pins to AdvancedDisplay
- Clarify Node Handle display name

diff --git a/Engine/Plugins/Runtime/Metasound/Source/MetasoundEditor/Public/MetasoundEditorSubsystem.h b/Engine/Plugins/Runtime/Metasound/Source/MetasoundEditor/Public/MetasoundEditorSubsystem.h
--- a/Engine/Plugins/Runtime/Metasound/Source/MetasoundEditor/Public/MetasoundEditorSubsystem.h
+++ b/Engine/Plugins/Runtime/Metasound/Source/MetasoundEditor/Public/MetasoundEditorSubsystem.h
@@ -38,7 +38,7 @@ public:
 	UFUNCTION(BlueprintCallable, Category = "Audio|MetaSound|Builder|Editor", meta = (ExpandEnumAsExecs = "OutResult"))
 	void SetNodeLocation(
 		UPARAM(DisplayName = "Builder") UMetaSoundBuilderBase * InBuilder,
-		UPARAM(DisplayName = "Node") const FMetaSoundNodeHandle& InNode,
+		UPARAM(DisplayName = "Node Handle") const FMetaSoundNodeHandle& InNode,
 		UPARAM(DisplayName = "Location") const FVector2D& InLocation,
 		EMetaSoundBuilderResult& OutResult);
 	
diff --git a/Engine/Plugins/Runtime/Metasound/Source/MetasoundEngine/Private/MetasoundBuilderSubsystem.cpp b/Engine/Plugins/Runtime/Metasound/Source/MetasoundEngine/Private/MetasoundBuilderSubsystem.cpp
--- a/Engine/Plugins/Runtime/Metasound/Source/MetasoundEngine/Private/MetasoundBuilderSubsystem.cpp
+++ b/Engine/Plugins/Runtime/Metasound/Source/MetasoundEngine/Private/MetasoundBuilderSubsystem.cpp
@@ -191,6 +191,11 @@ FMetaSoundNodeHandle UMetaSoundBuilderBase::AddNode(const TScriptInterface<IMeta
 }
 
 FMetaSoundNodeHandle UMetaSoundBuilderBase::AddNodeByClassName(const FMetasoundFrontendClassName& ClassName, int32 MajorVersion, EMetaSoundBuilderResult& OutResult)
+{
+	return AddNodeByClassName(ClassName, OutResult, MajorVersion);
+}
+
+FMetaSoundNodeHandle UMetaSoundBuilderBase::AddNodeByClassName(const FMetasoundFrontendClassName& ClassName, EMetaSoundBuilderResult& OutResult, int32 MajorVersion)
 {
 	using namespace Metasound;
 	using namespace Metasound::Frontend;
diff --git a/Engine/Plugins/Runtime/Metasound/Source/MetasoundEngine/Public/MetasoundBuilderSubsystem.h b/Engine/Plugins/Runtime/Metasound/Source/MetasoundEngine/Public/MetasoundBuilderSubsystem.h
--- a/Engine/Plugins/Runtime/Metasound/Source/MetasoundEngine/Public/MetasoundBuilderSubsystem.h
+++ b/Engine/Plugins/Runtime/Metasound/Source/MetasoundEngine/Public/MetasoundBuilderSubsystem.h
@@ -138,7 +138,7 @@ public:
 
 	// Adds a graph output node with the given name, DataType, and sets output node's input to default value.
 	// Returns the new output node's input handle if it was successfully created, or an invalid handle if it failed.
-	UFUNCTION(BlueprintCallable, Category = "Audio|MetaSound|Builder", meta = (ExpandEnumAsExecs = "OutResult"))
+	UFUNCTION(BlueprintCallable, Category = "Audio|MetaSound|Builder", meta = (ExpandEnumAsExecs = "OutResult", AdvancedDisplay = "3"))
 	UPARAM(DisplayName = "Input Handle") FMetaSoundBuilderNodeInputHandle AddGraphOutputNode(FName Name, FName DataType, FMetasoundFrontendLiteral DefaultValue, EMetaSoundBuilderResult& OutResult, bool bIsConstructorOutput = false);
 
 	// Adds an interface registered with the given name to the graph, adding associated input and output nodes.
@@ -152,7 +152,10 @@ public:
 
 	// Adds node referencing the highest native class version of the given class name to the document.
 	// Returns a node handle to the created node if successful, or an invalid handle if it failed.
-	UFUNCTION(BlueprintCallable, Category = "Audio|MetaSound|Builder", DisplayName = "Add MetaSound Node By ClassName", meta = (ExpandEnumAsExecs = "OutResult"))
+	UFUNCTION(BlueprintCallable, Category = "Audio|MetaSound|Builder", DisplayName = "Add MetaSound Node By ClassName", meta = (ExpandEnumAsExecs = "OutResult", AdvancedDisplay = "2"))
+	UPARAM(DisplayName = "Node Handle") FMetaSoundNodeHandle AddNodeByClassName(const FMetasoundFrontendClassName& ClassName, EMetaSoundBuilderResult& OutResult, int32 MajorVersion = 1);
+	
+	UE_DEPRECATED(5.4, "This version of AddNodeByClassName is deprecated. Use the one with a default MajorVersion of 1.")
 	UPARAM(DisplayName = "Node Handle") FMetaSoundNodeHandle AddNodeByClassName(const FMetasoundFrontendClassName& ClassName, int32 MajorVersion, EMetaSoundBuilderResult& OutResult);
 
 	// Connects node output to a node input. Does *NOT* provide loop detection for performance reasons.  Loop detection is checked on class registration when built or played.
@@ -222,7 +225,7 @@ public:
 	UFUNCTION(BlueprintCallable, Category = "Audio|MetaSound|Builder", meta = (ExpandEnumAsExecs = "OutResult"))
 	UPARAM(DisplayName = "Input Handle") FMetaSoundBuilderNodeInputHandle FindNodeInputByName(const FMetaSoundNodeHandle& NodeHandle, FName InputName, EMetaSoundBuilderResult& OutResult);
 
-	// Returns node output by the given name.
+	// Returns all node inputs.
 	UFUNCTION(BlueprintCallable, Category = "Audio|MetaSound|Builder", meta = (ExpandEnumAsExecs = "OutResult"))
 	UPARAM(DisplayName = "Input  Handles") TArray<FMetaSoundBuilderNodeInputHandle> FindNodeInputs(const FMetaSoundNodeHandle& NodeHandle, EMetaSoundBuilderResult& OutResult);
 
diff --git a/Engine/Plugins/Runtime/Metasound/Source/MetasoundEngineTest/Private/EngineTestMetaSoundBuilder.h b/Engine/Plugins/Runtime/Metasound/Source/MetasoundEngineTest/Private/EngineTestMetaSoundBuilder.h
--- a/Engine/Plugins/Runtime/Metasound/Source/MetasoundEngineTest/Private/EngineTestMetaSoundBuilder.h
+++ b/Engine/Plugins/Runtime/Metasound/Source/MetasoundEngineTest/Private/EngineTestMetaSoundBuilder.h
@@ -98,7 +98,7 @@ bool FMetaSoundSourceBuilderCreateAndConnectTriGeneratorNodeLatentCommand::Updat
 	if (Builder)
 	{
 		// Tri Oscillator Node
-		FMetaSoundNodeHandle TriNode = Builder->AddNodeByClassName({ "UE", "Triangle", "Audio" }, 1, Result);
+		FMetaSoundNodeHandle TriNode = Builder->AddNodeByClassName({ "UE", "Triangle", "Audio" }, Result, 1);
 		Test.AddErrorIfFalse(Result == EMetaSoundBuilderResult::Succeeded && TriNode.IsSet(), TEXT("Failed to create node by class name 'UE:Triangle:Audio"));
 
 		FMetaSoundBuilderNodeOutputHandle TriNodeAudioOutput = Builder->FindNodeOutputByName(TriNode, "Audio", Result);
diff --git a/Engine/Plugins/Runtime/Metasound/Source/MetasoundEngineTest/Private/EngineTestMetasoundSource.cpp b/Engine/Plugins/Runtime/Metasound/Source/MetasoundEngineTest/Private/EngineTestMetasoundSource.cpp
--- a/Engine/Plugins/Runtime/Metasound/Source/MetasoundEngineTest/Private/EngineTestMetasoundSource.cpp
+++ b/Engine/Plugins/Runtime/Metasound/Source/MetasoundEngineTest/Private/EngineTestMetasoundSource.cpp
@@ -135,7 +135,7 @@ namespace EngineTestMetaSoundSourcePrivate
 		Test.AddErrorIfFalse(Result == EMetaSoundBuilderResult::Succeeded && FrequencyNodeOutput.IsSet(), TEXT("Failed to create new MetaSound graph input"));
 
 		// Sine Oscillator Node
-		const FMetaSoundNodeHandle OscNode = Builder.AddNodeByClassName({ "UE", "Sine", "Audio" }, 1, Result);
+		const FMetaSoundNodeHandle OscNode = Builder.AddNodeByClassName({ "UE", "Sine", "Audio" }, Result, 1);
 		Test.AddErrorIfFalse(Result == EMetaSoundBuilderResult::Succeeded && OscNode.IsSet(), TEXT("Failed to create new MetaSound node by class name"));
 
 		// Make connections:
@@ -690,7 +690,7 @@ bool FAudioMetasoundSourceBuilderAddRemoveNodes::RunTest(const FString& Paramete
 		TArray<FMetaSoundNodeHandle> GenNodes;
 		for (int32 i = 0; i < 8; ++i)
 		{
-			FMetaSoundNodeHandle OscNode = Builder.AddNodeByClassName({ "UE", OscType, "Audio" }, 1, Result);
+			FMetaSoundNodeHandle OscNode = Builder.AddNodeByClassName({ "UE", OscType, "Audio" }, Result, 1);
 			AddErrorIfFalse(Result == EMetaSoundBuilderResult::Succeeded, "Failed to add osc node to graph");
 
 			if (Result == EMetaSoundBuilderResult::Succeeded)
diff --git a/Engine/Plugins/Runtime/Metasound/Source/MetasoundFrontend/Private/MetasoundFrontendDocumentBuilder.cpp b/Engine/Plugins/Runtime/Metasound/Source/MetasoundFrontend/Private/MetasoundFrontendDocumentBuilder.cpp
--- a/Engine/Plugins/Runtime/Metasound/Source/MetasoundFrontend/Private/MetasoundFrontendDocumentBuilder.cpp
+++ b/Engine/Plugins/Runtime/Metasound/Source/MetasoundFrontend/Private/MetasoundFrontendDocumentBuilder.cpp
@@ -1081,7 +1081,7 @@ const FMetasoundFrontendNode* FMetaSoundFrontendDocumentBuilder::AddNodeByClassN
 		}
 	}
 
-	UE_LOG(LogMetaSound, Warning, TEXT("Failed to add new node by class name '%s': Class not found"), *InClassName.ToString());
+	UE_LOG(LogMetaSound, Warning, TEXT("Failed to add new node by class name '%s' and major version '%d': Class not found"), *InClassName.ToString(), InMajorVersion);
 	return nullptr;
 }
 
diff --git a/Engine/Plugins/Runtime/Metasound/Source/MetasoundFrontend/Public/MetasoundFrontendDocumentBuilder.h b/Engine/Plugins/Runtime/Metasound/Source/MetasoundFrontend/Public/MetasoundFrontendDocumentBuilder.h
--- a/Engine/Plugins/Runtime/Metasound/Source/MetasoundFrontend/Public/MetasoundFrontendDocumentBuilder.h
+++ b/Engine/Plugins/Runtime/Metasound/Source/MetasoundFrontend/Public/MetasoundFrontendDocumentBuilder.h
@@ -154,7 +154,7 @@ public:
 	bool AddInterface(FName InterfaceName);
 
 	const FMetasoundFrontendNode* AddGraphNode(const FMetasoundFrontendGraphClass& InClass, FGuid InNodeID = FGuid::NewGuid());
-	const FMetasoundFrontendNode* AddNodeByClassName(const FMetasoundFrontendClassName& InClassName, int32 InMajorVersion, FGuid InNodeID = FGuid::NewGuid());
+	const FMetasoundFrontendNode* AddNodeByClassName(const FMetasoundFrontendClassName& InClassName, int32 InMajorVersion = 1, FGuid InNodeID = FGuid::NewGuid());
 
 	// Returns whether or not the given edge can be added, which requires that its input
 	// is not already connected and the edge is valid (see function 'IsValidEdge').
