- Fix deviation from naming conventions
- Better name for insights tags on "sub-graphs" (patchs) within a metasound source.

diff --git a/Engine/Plugins/Runtime/Metasound/Source/MetasoundFrontend/Private/MetasoundFrontendGraph.cpp b/Engine/Plugins/Runtime/Metasound/Source/MetasoundFrontend/Private/MetasoundFrontendGraph.cpp
--- a/Engine/Plugins/Runtime/Metasound/Source/MetasoundFrontend/Private/MetasoundFrontendGraph.cpp
+++ b/Engine/Plugins/Runtime/Metasound/Source/MetasoundFrontend/Private/MetasoundFrontendGraph.cpp
@@ -739,7 +739,7 @@ namespace Metasound
 
 	TUniquePtr<FFrontendGraph> FFrontendGraphBuilder::CreateGraph(FBuildContext& InContext, const FMetasoundFrontendGraphClass& InGraphClass, const TSet<FName>& InTransmittableInputNames)
 	{
-		const FString GraphName = InGraphClass.Metadata.GetClassName().GetFullName().ToString();
+		const FString GraphName = InContext.DebugAssetName;
 
 		FBuildGraphContext BuildGraphContext
 		{
diff --git a/Engine/Plugins/Runtime/Metasound/Source/MetasoundGraphCore/Private/MetasoundOperatorBuilder.cpp b/Engine/Plugins/Runtime/Metasound/Source/MetasoundGraphCore/Private/MetasoundOperatorBuilder.cpp
--- a/Engine/Plugins/Runtime/Metasound/Source/MetasoundGraphCore/Private/MetasoundOperatorBuilder.cpp
+++ b/Engine/Plugins/Runtime/Metasound/Source/MetasoundGraphCore/Private/MetasoundOperatorBuilder.cpp
@@ -460,7 +460,7 @@ namespace Metasound
 				FOperatorFactorySharedRef Factory = Node->GetDefaultOperatorFactory();
 				if (ProfileOperators && Profiling::OperatorShouldBeProfiled(Node->GetMetadata()))
 				{
-					OperatorInfo.Operator = MakeUnique<ProfilingOperator>(Factory->CreateOperator(CreateParams, InOutContext.Results), Node->GetMetadata());
+					OperatorInfo.Operator = MakeUnique<FProfilingOperator>(Factory->CreateOperator(CreateParams, InOutContext.Results), Node);
 				}
 				else
 				{
diff --git a/Engine/Plugins/Runtime/Metasound/Source/MetasoundGraphCore/Private/MetasoundProfilingOperator.cpp b/Engine/Plugins/Runtime/Metasound/Source/MetasoundGraphCore/Private/MetasoundProfilingOperator.cpp
--- a/Engine/Plugins/Runtime/Metasound/Source/MetasoundGraphCore/Private/MetasoundProfilingOperator.cpp
+++ b/Engine/Plugins/Runtime/Metasound/Source/MetasoundGraphCore/Private/MetasoundProfilingOperator.cpp
@@ -83,25 +83,25 @@ namespace Metasound
 		
 	}
 
-	void ProfilingOperator::StaticReset(IOperator* InOperator, const IOperator::FResetParams& InParams)
+	void FProfilingOperator::StaticReset(IOperator* InOperator, const IOperator::FResetParams& InParams)
 	{
-		ProfilingOperator* DerivedOperator = static_cast<ProfilingOperator*>(InOperator);
+		FProfilingOperator* DerivedOperator = static_cast<FProfilingOperator*>(InOperator);
 		check(nullptr != DerivedOperator && DerivedOperator->ResetFunction);
 		METASOUND_TRACE_CPUPROFILER_EVENT_SCOPE_TEXT(*DerivedOperator->InsightsResetEventName);
 		DerivedOperator->ResetFunction(DerivedOperator->Operator.Get(), InParams);
 	}
 
-	void ProfilingOperator::StaticExecute(IOperator* InOperator)
+	void FProfilingOperator::StaticExecute(IOperator* InOperator)
 	{
-		ProfilingOperator* DerivedOperator = static_cast<ProfilingOperator*>(InOperator);
+		FProfilingOperator* DerivedOperator = static_cast<FProfilingOperator*>(InOperator);
 		check(nullptr != DerivedOperator && DerivedOperator->ExecuteFunction);
 		METASOUND_TRACE_CPUPROFILER_EVENT_SCOPE_TEXT(*DerivedOperator->InsightsExecuteEventName);
 		DerivedOperator->ExecuteFunction(DerivedOperator->Operator.Get());
 	}
 
-	void ProfilingOperator::StaticPostExecute(IOperator* InOperator)
+	void FProfilingOperator::StaticPostExecute(IOperator* InOperator)
 	{
-		ProfilingOperator* DerivedOperator = static_cast<ProfilingOperator*>(InOperator);
+		FProfilingOperator* DerivedOperator = static_cast<FProfilingOperator*>(InOperator);
 		check(nullptr != DerivedOperator && DerivedOperator->PostExecuteFunction);
 		METASOUND_TRACE_CPUPROFILER_EVENT_SCOPE_TEXT(*DerivedOperator->InsightsPostExecuteEventName);
 		DerivedOperator->PostExecuteFunction(DerivedOperator->Operator.Get());
diff --git a/Engine/Plugins/Runtime/Metasound/Source/MetasoundGraphCore/Private/MetasoundProfilingOperator.h b/Engine/Plugins/Runtime/Metasound/Source/MetasoundGraphCore/Private/MetasoundProfilingOperator.h
--- a/Engine/Plugins/Runtime/Metasound/Source/MetasoundGraphCore/Private/MetasoundProfilingOperator.h
+++ b/Engine/Plugins/Runtime/Metasound/Source/MetasoundGraphCore/Private/MetasoundProfilingOperator.h
@@ -16,23 +16,28 @@ namespace Metasound
 	}
 
 	// This is a wrapper around any IOperator that causes its functions to be timed for Insights
-	class ProfilingOperator : public IOperator
+	class FProfilingOperator : public IOperator
 	{
 	public:
-		ProfilingOperator(TUniquePtr<IOperator>&& WrappedOperator, const FNodeClassMetadata& NodeMetadata)
+		FProfilingOperator(TUniquePtr<IOperator>&& WrappedOperator, const INode* Node)
 			: Operator(MoveTemp(WrappedOperator))
 			, ResetFunction(Operator->GetResetFunction())
 			, ExecuteFunction(Operator->GetExecuteFunction())
 			, PostExecuteFunction(Operator->GetPostExecuteFunction())
 		{
 			check(Operator);
+			const FNodeClassMetadata& NodeMetadata = Node->GetMetadata();
 			FString BaseEventName = NodeMetadata.ClassName.GetName().ToString();
+			if (NodeMetadata.ClassName.GetName().IsNone())
+			{
+				BaseEventName = Node->GetInstanceName().ToString();
+			}
 			InsightsResetEventName = FString::Printf(TEXT("%s_RESET"), *BaseEventName);
 			InsightsExecuteEventName = FString::Printf(TEXT("%s_EXECUTE"), *BaseEventName);
 			InsightsPostExecuteEventName = FString::Printf(TEXT("%s_POSTEXECUTE"), *BaseEventName);
 		}
 
-		virtual ~ProfilingOperator() = default;
+		virtual ~FProfilingOperator() = default;
 
 		virtual void BindInputs(FInputVertexInterfaceData& InVertexData) override
 		{
