Prepare code to preprocess MetaSound cooked layers/pages
- Move reroute preprocessing to Builder API
- Fix for potential memory scribble when deleting output edges in builder
- Move builder to use deterministic guid generation where necessary
#rnx

diff --git a/Engine/Plugins/Runtime/Metasound/Source/MetasoundEngine/Private/MetasoundBuilderSubsystem.cpp b/Engine/Plugins/Runtime/Metasound/Source/MetasoundEngine/Private/MetasoundBuilderSubsystem.cpp
--- a/Engine/Plugins/Runtime/Metasound/Source/MetasoundEngine/Private/MetasoundBuilderSubsystem.cpp
+++ b/Engine/Plugins/Runtime/Metasound/Source/MetasoundEngine/Private/MetasoundBuilderSubsystem.cpp
@@ -14,6 +14,7 @@
 #include "MetasoundDynamicOperatorTransactor.h"
 #include "MetasoundFrontendDataTypeRegistry.h"
 #include "MetasoundFrontendDocument.h"
+#include "MetasoundFrontendDocumentIdGenerator.h"
 #include "MetasoundFrontendRegistries.h"
 #include "MetasoundFrontendSearchEngine.h"
 #include "MetasoundFrontendTransform.h"
@@ -86,11 +87,14 @@ FMetaSoundBuilderNodeOutputHandle UMetaSoundBuilderBase::AddGraphInputNode(FName
 		}
 		else
 		{
+			FDocumentIDGenerator& IDGenerator = FDocumentIDGenerator::Get();
+			const FMetasoundFrontendDocument& Doc = GetConstBuilder().GetDocument();
+
 			FMetasoundFrontendClassInput Description;
 			Description.Name = Name;
 			Description.TypeName = DataType;
-			Description.NodeID = FGuid::NewGuid();
-			Description.VertexID = FGuid::NewGuid();
+			Description.NodeID = IDGenerator.CreateNodeID(Doc);
+			Description.VertexID = IDGenerator.CreateVertexID(Doc);
 			Description.DefaultLiteral = static_cast<FMetasoundFrontendLiteral>(DefaultValue);
 			Description.AccessType = bIsConstructorInput ? EMetasoundFrontendVertexAccessType::Value : EMetasoundFrontendVertexAccessType::Reference;
 			Node = Builder.AddGraphInput(Description);
@@ -129,11 +133,14 @@ FMetaSoundBuilderNodeInputHandle UMetaSoundBuilderBase::AddGraphOutputNode(FName
 		}
 		else
 		{
+			FDocumentIDGenerator& IDGenerator = FDocumentIDGenerator::Get();
+			const FMetasoundFrontendDocument& Doc = GetConstBuilder().GetDocument();
+
 			FMetasoundFrontendClassOutput Description;
 			Description.Name = Name;
 			Description.TypeName = DataType;
-			Description.NodeID = FGuid::NewGuid();
-			Description.VertexID = FGuid::NewGuid();
+			Description.NodeID = IDGenerator.CreateNodeID(Doc);
+			Description.VertexID = IDGenerator.CreateVertexID(Doc);
 			Description.AccessType = bIsConstructorOutput ? EMetasoundFrontendVertexAccessType::Value : EMetasoundFrontendVertexAccessType::Reference;
 			Node = Builder.AddGraphOutput(Description);
 		}
@@ -638,21 +645,7 @@ FMetaSoundNodeHandle UMetaSoundBuilderBase::FindGraphOutputNode(FName OutputName
 
 UMetaSoundBuilderDocument* UMetaSoundBuilderBase::CreateTransientDocumentObject() const
 {
-	UMetaSoundBuilderDocument* DocObject = NewObject<UMetaSoundBuilderDocument>();
-	DocObject->SetBaseMetaSoundUClass(GetBuilderUClass());
-	return DocObject;
-}
-
-void UMetaSoundBuilderBase::InitFrontendBuilder()
-{
-	UMetaSoundBuilderDocument* DocObject = CreateTransientDocumentObject();
-	Builder = FMetaSoundFrontendDocumentBuilder(DocObject);
-	Builder.InitDocument();
-}
-
-void UMetaSoundBuilderBase::InitNodeLocations()
-{
-	Builder.InitNodeLocations();
+	return &UMetaSoundBuilderDocument::Create(GetBuilderUClass());
 }
 
 UObject* UMetaSoundBuilderBase::GetReferencedPresetAsset() const
@@ -684,6 +677,18 @@ UObject* UMetaSoundBuilderBase::GetReferencedPresetAsset() const
 	return nullptr;
 }
 
+void UMetaSoundBuilderBase::InitFrontendBuilder()
+{
+	UMetaSoundBuilderDocument& DocObject = UMetaSoundBuilderDocument::Create(GetBuilderUClass());
+	Builder = FMetaSoundFrontendDocumentBuilder(&DocObject);
+	Builder.InitDocument();
+}
+
+void UMetaSoundBuilderBase::InitNodeLocations()
+{
+	Builder.InitNodeLocations();
+}
+
 bool UMetaSoundBuilderBase::InterfaceIsDeclared(FName InterfaceName) const
 {
 	return Builder.IsInterfaceDeclared(InterfaceName);
@@ -1020,8 +1025,8 @@ void UMetaSoundSourceBuilder::InitFrontendBuilder()
 	TSharedRef<FDocumentModifyDelegates> DocumentDelegates = MakeShared<FDocumentModifyDelegates>();
 	InitDelegates(*DocumentDelegates);
 
-	UMetaSoundBuilderDocument* DocObject = CreateTransientDocumentObject();
-	Builder = FMetaSoundFrontendDocumentBuilder(DocObject, DocumentDelegates);
+	UMetaSoundBuilderDocument& DocObject = UMetaSoundBuilderDocument::Create(GetBuilderUClass());
+	Builder = FMetaSoundFrontendDocumentBuilder(&DocObject, DocumentDelegates);
 	Builder.InitDocument();
 }
 
@@ -1704,12 +1709,12 @@ void UMetaSoundBuilderSubsystem::Initialize(FSubsystemCollectionBase& Collection
 {
 	using namespace Metasound::Frontend;
 
-	IMetaSoundDocumentBuilderRegistry::Set([]() -> IMetaSoundDocumentBuilderRegistry&
+	IDocumentBuilderRegistry::Set([]() -> IDocumentBuilderRegistry&
 	{
 		check(GEngine);
 		UMetaSoundBuilderSubsystem* BuilderSubsystem = GEngine->GetEngineSubsystem<UMetaSoundBuilderSubsystem>();
 		check(BuilderSubsystem);
-		return static_cast<IMetaSoundDocumentBuilderRegistry&>(*BuilderSubsystem);
+		return static_cast<IDocumentBuilderRegistry&>(*BuilderSubsystem);
 	});
 }
 
diff --git a/Engine/Plugins/Runtime/Metasound/Source/MetasoundEngine/Private/MetasoundEngineAsset.h b/Engine/Plugins/Runtime/Metasound/Source/MetasoundEngine/Private/MetasoundEngineAsset.h
--- a/Engine/Plugins/Runtime/Metasound/Source/MetasoundEngine/Private/MetasoundEngineAsset.h
+++ b/Engine/Plugins/Runtime/Metasound/Source/MetasoundEngine/Private/MetasoundEngineAsset.h
@@ -107,7 +107,7 @@ namespace Metasound
 						{
 							constexpr bool bIsDeterministic = true;
 							FDocumentIDGenerator::FScopeDeterminism DeterminismScope = FDocumentIDGenerator::FScopeDeterminism(bIsDeterministic);
-							InMetaSound.CookGraph();
+							InMetaSound.CookMetaSound();
 						}
 					}
 				}
diff --git a/Engine/Plugins/Runtime/Metasound/Source/MetasoundEngine/Private/MetasoundSource.cpp b/Engine/Plugins/Runtime/Metasound/Source/MetasoundEngine/Private/MetasoundSource.cpp
--- a/Engine/Plugins/Runtime/Metasound/Source/MetasoundEngine/Private/MetasoundSource.cpp
+++ b/Engine/Plugins/Runtime/Metasound/Source/MetasoundEngine/Private/MetasoundSource.cpp
@@ -13,6 +13,7 @@
 #include "MetasoundAssetManager.h"
 #include "MetasoundAudioFormats.h"
 #include "MetasoundBuilderSubsystem.h"
+#include "MetasoundDocumentInterface.h"
 #include "MetasoundDynamicOperatorTransactor.h"
 #include "MetasoundEngineAsset.h"
 #include "MetasoundEngineEnvironment.h"
@@ -57,10 +58,9 @@ namespace Metasound
 #if !WITH_EDITOR 
 			if (Frontend::MetaSoundEnableCookDeterministicIDGeneration != 0)
 			{
-				// When without editor, don't AutoUpdate or PreprocessDocument at runtime. This should happen at cook or asset save.
-				// When with editor, those are needed because sounds are not necessarily saved before previewing
+				// When without editor, don't AutoUpdate or ResolveDocument at runtime. This only happens at cook or save.
+				// When with editor, those are needed because sounds are not necessarily saved before previewing.
 				RegOptions.bAutoUpdate = false;
-				RegOptions.bPreprocessDocument = false;
 			}
 #endif // !WITH_EDITOR
 			if (const UMetaSoundSettings* Settings = GetDefault<UMetaSoundSettings>())
@@ -1108,5 +1108,4 @@ TSharedPtr<Metasound::DynamicGraph::FDynamicOperatorTransactor> UMetaSoundSource
 {
 	return DynamicTransactor;
 }
-
 #undef LOCTEXT_NAMESPACE // MetaSound
diff --git a/Engine/Plugins/Runtime/Metasound/Source/MetasoundEngine/Public/MetasoundBuilderSubsystem.h b/Engine/Plugins/Runtime/Metasound/Source/MetasoundEngine/Public/MetasoundBuilderSubsystem.h
--- a/Engine/Plugins/Runtime/Metasound/Source/MetasoundEngine/Public/MetasoundBuilderSubsystem.h
+++ b/Engine/Plugins/Runtime/Metasound/Source/MetasoundEngine/Public/MetasoundBuilderSubsystem.h
@@ -29,6 +29,7 @@ class UMetaSoundSource;
 
 struct FMetasoundFrontendClassName;
 struct FMetasoundFrontendVersion;
+class FMetasoundAssetBase;
 
 enum class EMetaSoundOutputAudioFormat : uint8;
 
@@ -303,7 +304,7 @@ public:
 	// Returns if a given node output is connected.
 	UFUNCTION(BlueprintCallable, Category = "Audio|MetaSound|Builder")
 	UPARAM(DisplayName = "Connected") bool NodeOutputIsConnected(const FMetaSoundBuilderNodeOutputHandle& OutputHandle) const;
-		
+
 	// Returns whether this is a preset.
 	UFUNCTION(BlueprintCallable, Category = "Audio|MetaSound|Builder")
 	bool IsPreset() const;
@@ -454,10 +455,9 @@ protected:
 		return *MetaSound;
 	}
 
-	// Constructs a transient UMetaSoundBuilderDocument to be acted upon by this builder instance.
+	UE_DEPRECATED(5.4, "Use UMetaSoundBuilderDocument::Create instead")
 	UMetaSoundBuilderDocument* CreateTransientDocumentObject() const;
 
-
 	UPROPERTY()
 	FMetaSoundFrontendDocumentBuilder Builder;
 
@@ -535,7 +535,7 @@ private:
 
 /** The subsystem in charge of tracking MetaSound builders */
 UCLASS()
-class METASOUNDENGINE_API UMetaSoundBuilderSubsystem : public UEngineSubsystem, public Metasound::Frontend::IMetaSoundDocumentBuilderRegistry
+class METASOUNDENGINE_API UMetaSoundBuilderSubsystem : public UEngineSubsystem, public Metasound::Frontend::IDocumentBuilderRegistry
 {
 	GENERATED_BODY()
 
@@ -576,7 +576,7 @@ public:
 		EMetaSoundBuilderResult& OutResult,
 		EMetaSoundOutputAudioFormat OutputFormat = EMetaSoundOutputAudioFormat::Mono,
 		bool bIsOneShot = true);
-	
+
 	UFUNCTION(BlueprintCallable, Category = "Audio|MetaSound|Builder",  meta = (ExpandEnumAsExecs = "OutResult"))
 	UPARAM(DisplayName = "Patch Preset Builder") UMetaSoundPatchBuilder* CreatePatchPresetBuilder(FName BuilderName, const TScriptInterface<IMetaSoundDocumentInterface>& ReferencedPatchClass, EMetaSoundBuilderResult& OutResult);
 
diff --git a/Engine/Plugins/Runtime/Metasound/Source/MetasoundFrontend/Private/MetasoundAssetBase.cpp b/Engine/Plugins/Runtime/Metasound/Source/MetasoundFrontend/Private/MetasoundAssetBase.cpp
--- a/Engine/Plugins/Runtime/Metasound/Source/MetasoundFrontend/Private/MetasoundAssetBase.cpp
+++ b/Engine/Plugins/Runtime/Metasound/Source/MetasoundFrontend/Private/MetasoundAssetBase.cpp
@@ -16,6 +16,7 @@
 #include "IStructSerializerBackend.h"
 #include "Logging/LogMacros.h"
 #include "MetasoundAssetManager.h"
+#include "MetasoundDocumentInterface.h"
 #include "MetasoundFrontendController.h"
 #include "MetasoundFrontendDocument.h"
 #include "MetasoundFrontendDocumentBuilder.h"
@@ -232,25 +233,36 @@ void FMetasoundAssetBase::RegisterGraphWithFrontend(Metasound::Frontend::FMetaSo
 	CacheRegistryMetadata();
 #endif // WITH_EDITOR
 
-	FString AssetName;
-	const UObject* OwningAsset = GetOwningAsset();
-	if (ensure(OwningAsset))
-	{
-		AssetName = OwningAsset->GetName();
-	}
-	
+	TScriptInterface<IMetaSoundDocumentInterface> DocInterface = GetOwningAsset();
+	check(DocInterface.GetObject());
 	TSharedPtr<FMetasoundFrontendDocument> DocumentToRegister;
-	if (InRegistrationOptions.bPreprocessDocument)
+	const IMetaSoundDocumentInterface& Interface = *DocInterface.GetInterface();
+
+	auto CreateTransformedDocumentToRegister = [&Interface]()
 	{
-		DocumentToRegister = PreprocessDocument();
+		// Node template transform is performed on copy of local document to avoid overwriting editable data
+		UMetaSoundBuilderDocument& DocObject = UMetaSoundBuilderDocument::Create(Interface);
+		FMetaSoundFrontendDocumentBuilder DocBuilder(&DocObject);
+		DocBuilder.TransformTemplateNodes();
+		return MakeShared<FMetasoundFrontendDocument>(MoveTemp(DocObject.GetDocument()));
+	};
+
+#if WITH_EDITOR
+	DocumentToRegister = CreateTransformedDocumentToRegister();
+#else // !WITH_EDITOR
+	if (Frontend::MetaSoundEnableCookDeterministicIDGeneration != 0)
+	{
+		DocumentToRegister = MakeShared<FMetasoundFrontendDocument>(Interface.GetDocument());
 	}
 	else
 	{
-		DocumentToRegister = MakeShared<FMetasoundFrontendDocument>(GetDocumentChecked());
+		DocumentToRegister = CreateTransformedDocumentToRegister();
 	}
+#endif // !WITH_EDITOR
 
 	CacheRuntimeData(*DocumentToRegister);
 	FNodeClassInfo AssetClassInfo = GetAssetClassInfo();
+	const FString AssetName = DocInterface.GetObject()->GetName();
 	TUniquePtr<INodeRegistryEntry> RegistryEntry = MakeUnique<AssetBasePrivate::FNodeRegistryEntry>(AssetName, DocumentToRegister, AssetClassInfo.AssetPath);
 
 	UnregisterGraphWithFrontend();
@@ -264,30 +276,25 @@ void FMetasoundAssetBase::RegisterGraphWithFrontend(Metasound::Frontend::FMetaSo
 	}
 	else
 	{
-		FString ClassName;
-		if (OwningAsset)
-		{
-			if (UClass* Class = OwningAsset->GetClass())
-			{
-				ClassName = Class->GetName();
-			}
-		}
+		UClass* Class = DocInterface.GetObject()->GetClass();
+		check(Class);
+		const FString ClassName = Class->GetName();
 		UE_LOG(LogMetaSound, Error, TEXT("Registration failed for MetaSound node class '%s' of UObject class '%s'"), *AssetName, *ClassName);
 	}
 }
 
-void FMetasoundAssetBase::CookGraph()
+void FMetasoundAssetBase::CookMetaSound()
 {
 	using namespace Metasound;
 	using namespace Metasound::Frontend;
 
-	METASOUND_TRACE_CPUPROFILER_EVENT_SCOPE(MetaSoundAssetBase::CookGraph);
+	METASOUND_TRACE_CPUPROFILER_EVENT_SCOPE(MetaSoundAssetBase::CookMetaSound);
 	if (IsRegistered())
 	{
 		return;
 	}
 
-	RegisterAssetDependenciesForCook();
+	CookReferencedMetaSounds();
 	IMetaSoundAssetManager::GetChecked().AddOrUpdateAsset(*GetOwningAsset());
 
 	// Auto update must be done after all referenced asset classes are registered
@@ -305,20 +312,20 @@ void FMetasoundAssetBase::CookGraph()
 	CacheRegistryMetadata();
 #endif // WITH_EDITOR
 
-	FString AssetName;
-	const UObject* OwningAsset = GetOwningAsset();
-	if (ensure(OwningAsset))
-	{
-		AssetName = OwningAsset->GetName();
-	}
-	
-	TSharedPtr<FMetasoundFrontendDocument> DocumentToRegister = PreprocessDocument();
-	TSharedPtr<FMetasoundFrontendDocument> PreprocessedDocumentCopy = MakeShared<FMetasoundFrontendDocument>(*DocumentToRegister);
-	// Replace document with preprocessed document
-	SetDocument(*PreprocessedDocumentCopy, /*bMarkDirty=*/false);
+	TScriptInterface<IMetaSoundDocumentInterface> DocInterface = GetOwningAsset();
+	check(DocInterface.GetObject());
+	TSharedPtr<FMetasoundFrontendDocument> DocumentToRegister;
+	const IMetaSoundDocumentInterface& Interface = *DocInterface.GetInterface();
+
+	// Performs document transforms on local copy, which reduces document footprint & renders resolve unnecessary at runtime upon registration
+	FMetaSoundFrontendDocumentBuilder DocBuilder(DocInterface);
+	DocBuilder.TransformTemplateNodes();
+	const FMetaSoundFrontendDocumentBuilder& ConstBuilder = DocBuilder;
+	TSharedRef<FMetasoundFrontendDocument> ResolvedDocument = MakeShared<FMetasoundFrontendDocument>(ConstBuilder.GetDocument());
 
 	FNodeClassInfo AssetClassInfo = GetAssetClassInfo();
-	TUniquePtr<INodeRegistryEntry> RegistryEntry = MakeUnique<AssetBasePrivate::FNodeRegistryEntry>(AssetName, DocumentToRegister, AssetClassInfo.AssetPath);
+	const FString AssetName = DocInterface.GetObject()->GetName();
+	TUniquePtr<INodeRegistryEntry> RegistryEntry = MakeUnique<AssetBasePrivate::FNodeRegistryEntry>(AssetName, ResolvedDocument, AssetClassInfo.AssetPath);
 
 	UnregisterGraphWithFrontend();
 	RegistryKey = FMetasoundFrontendRegistryContainer::Get()->RegisterNode(MoveTemp(RegistryEntry));
@@ -331,15 +338,10 @@ void FMetasoundAssetBase::CookGraph()
 	}
 	else
 	{
-		FString ClassName;
-		if (OwningAsset)
-		{
-			if (UClass* Class = OwningAsset->GetClass())
-			{
-				ClassName = Class->GetName();
-			}
-		}
-		UE_LOG(LogMetaSound, Error, TEXT("While cooking graph, registration failed for MetaSound node class '%s' of UObject class '%s'"), *AssetName, *ClassName);
+		const UClass* Class = DocInterface.GetObject()->GetClass();
+		check(Class);
+		const FString ClassName = Class->GetName();
+		UE_LOG(LogMetaSound, Error, TEXT("Registration failed during cook for MetaSound node class '%s' of UObject class '%s'"), *AssetName, *ClassName);
 	}
 }
 
@@ -920,7 +922,7 @@ void FMetasoundAssetBase::RegisterAssetDependencies(const Metasound::Frontend::F
 	}
 }
 
-void FMetasoundAssetBase::RegisterAssetDependenciesForCook()
+void FMetasoundAssetBase::CookReferencedMetaSounds()
 {
 	using namespace Metasound::Frontend;
 
@@ -932,7 +934,7 @@ void FMetasoundAssetBase::RegisterAssetDependenciesForCook()
 		{
 			// TODO: Check for infinite recursion and error if so
 			AssetManager.AddOrUpdateAsset(*(Reference->GetOwningAsset()));
-			Reference->CookGraph();
+			Reference->CookMetaSound();
 		}
 	}
 }
@@ -974,12 +976,7 @@ void FMetasoundAssetBase::UpdateAssetRegistry()
 
 TSharedPtr<FMetasoundFrontendDocument> FMetasoundAssetBase::PreprocessDocument()
 {
-	using namespace Metasound::Frontend;
-
-	TSharedPtr<FMetasoundFrontendDocument> PreprocessedDocument = MakeShared<FMetasoundFrontendDocument>(GetDocumentChecked());
-	FDocumentTemplatePreprocessTransform TemplatePreprocessor;
-	TemplatePreprocessor.Transform(*PreprocessedDocument);
-	return PreprocessedDocument;
+	return nullptr;
 }
 
 const FMetasoundAssetBase::FRuntimeData& FMetasoundAssetBase::CacheRuntimeData(const FMetasoundFrontendDocument& InPreprocessedDoc)
diff --git a/Engine/Plugins/Runtime/Metasound/Source/MetasoundFrontend/Private/MetasoundDocumentInterface.cpp b/Engine/Plugins/Runtime/Metasound/Source/MetasoundFrontend/Private/MetasoundDocumentInterface.cpp
--- a/Engine/Plugins/Runtime/Metasound/Source/MetasoundFrontend/Private/MetasoundDocumentInterface.cpp
+++ b/Engine/Plugins/Runtime/Metasound/Source/MetasoundFrontend/Private/MetasoundDocumentInterface.cpp
@@ -7,10 +7,10 @@ namespace Metasound::Frontend
 {
 	namespace DocumentBuilderRegistryPrivate
 	{
-		TUniqueFunction<IMetaSoundDocumentBuilderRegistry&()> GetInstance;
+		TUniqueFunction<IDocumentBuilderRegistry&()> GetInstance;
 	} // namespace DocumentBuilderRegistryPrivate
 
-	IMetaSoundDocumentBuilderRegistry& IMetaSoundDocumentBuilderRegistry::GetChecked()
+	IDocumentBuilderRegistry& IDocumentBuilderRegistry::GetChecked()
 	{
 		using namespace DocumentBuilderRegistryPrivate;
 
@@ -18,7 +18,7 @@ namespace Metasound::Frontend
 		return GetInstance();
 	}
 
-	void IMetaSoundDocumentBuilderRegistry::Set(TUniqueFunction<IMetaSoundDocumentBuilderRegistry&()>&& InGetInstance)
+	void IDocumentBuilderRegistry::Set(TUniqueFunction<IDocumentBuilderRegistry&()>&& InGetInstance)
 	{
 		using namespace DocumentBuilderRegistryPrivate;
 
diff --git a/Engine/Plugins/Runtime/Metasound/Source/MetasoundFrontend/Private/MetasoundFrontendDocumentBuilder.cpp b/Engine/Plugins/Runtime/Metasound/Source/MetasoundFrontend/Private/MetasoundFrontendDocumentBuilder.cpp
--- a/Engine/Plugins/Runtime/Metasound/Source/MetasoundFrontend/Private/MetasoundFrontendDocumentBuilder.cpp
+++ b/Engine/Plugins/Runtime/Metasound/Source/MetasoundFrontend/Private/MetasoundFrontendDocumentBuilder.cpp
@@ -15,7 +15,9 @@
 #include "MetasoundFrontendController.h"
 #include "MetasoundFrontendDocumentCache.h"
 #include "MetasoundFrontendDocument.h"
+#include "MetasoundFrontendDocumentIdGenerator.h"
 #include "MetasoundFrontendDocumentModifyDelegates.h"
+#include "MetasoundFrontendNodeTemplateRegistry.h"
 #include "MetasoundFrontendRegistries.h"
 #include "MetasoundFrontendSearchEngine.h"
 #include "MetasoundFrontendTransform.h"
@@ -99,8 +101,9 @@ namespace Metasound::Frontend
 		class FModifyInterfacesImpl
 		{
 		public:
-			FModifyInterfacesImpl(FModifyInterfaceOptions&& InOptions)
+			FModifyInterfacesImpl(FMetasoundFrontendDocument& InDocument, FModifyInterfaceOptions&& InOptions)
 				: Options(MoveTemp(InOptions))
+				, Document(InDocument)
 			{
 				for (const FMetasoundFrontendInterface& FromInterface : Options.InterfacesToRemove)
 				{
@@ -110,19 +113,19 @@ namespace Metasound::Frontend
 
 				for (const FMetasoundFrontendInterface& ToInterface : Options.InterfacesToAdd)
 				{
-					Algo::Transform(ToInterface.Inputs, InputsToAdd, [&ToInterface](const FMetasoundFrontendClassInput& Input)
+					Algo::Transform(ToInterface.Inputs, InputsToAdd, [this, &ToInterface](const FMetasoundFrontendClassInput& Input)
 					{
 						FMetasoundFrontendClassInput NewInput = Input;
-						NewInput.NodeID = FGuid::NewGuid();
-						NewInput.VertexID = FGuid::NewGuid();
+						NewInput.NodeID = FDocumentIDGenerator::Get().CreateNodeID(Document);
+						NewInput.VertexID = FDocumentIDGenerator::Get().CreateVertexID(Document);
 						return FInputInterfacePair { MoveTemp(NewInput), &ToInterface };
 					});
 
-					Algo::Transform(ToInterface.Outputs, OutputsToAdd, [&ToInterface](const FMetasoundFrontendClassOutput& Output)
+					Algo::Transform(ToInterface.Outputs, OutputsToAdd, [this, &ToInterface](const FMetasoundFrontendClassOutput& Output)
 					{
 						FMetasoundFrontendClassOutput NewOutput = Output;
-						NewOutput.NodeID = FGuid::NewGuid();
-						NewOutput.VertexID = FGuid::NewGuid();
+						NewOutput.NodeID = FDocumentIDGenerator::Get().CreateNodeID(Document);
+						NewOutput.VertexID = FDocumentIDGenerator::Get().CreateVertexID(Document);
 						return FOutputInterfacePair { MoveTemp(NewOutput), &ToInterface };
 					});
 				}
@@ -364,33 +367,33 @@ namespace Metasound::Frontend
 #endif // WITH_EDITORONLY_DATA
 
 		public:
-			bool Execute(FMetaSoundFrontendDocumentBuilder& OutBuilder, FMetasoundFrontendDocument& OutDoc, FDocumentModifyDelegates& OutDelegates)
+			bool Execute(FMetaSoundFrontendDocumentBuilder& OutBuilder, FDocumentModifyDelegates& OutDelegates)
 			{
 				bool bDidEdit = false;
 
 				for (const FMetasoundFrontendInterface& Interface : Options.InterfacesToRemove)
 				{
-					if (OutDoc.Interfaces.Contains(Interface.Version))
+					if (Document.Interfaces.Contains(Interface.Version))
 					{
 						OutDelegates.InterfaceDelegates.OnRemovingInterface.Broadcast(Interface);
 						bDidEdit = true;
 #if WITH_EDITORONLY_DATA
-						OutDoc.Metadata.ModifyContext.AddInterfaceModified(Interface.Version.Name);
+						Document.Metadata.ModifyContext.AddInterfaceModified(Interface.Version.Name);
 #endif // WITH_EDITORONLY_DATA
-						OutDoc.Interfaces.Remove(Interface.Version);
+						Document.Interfaces.Remove(Interface.Version);
 					}
 				}
 
 				for (const FMetasoundFrontendInterface& Interface : Options.InterfacesToAdd)
 				{
 					bool bAlreadyInSet = false;
-					OutDoc.Interfaces.Add(Interface.Version, &bAlreadyInSet);
+					Document.Interfaces.Add(Interface.Version, &bAlreadyInSet);
 					if (!bAlreadyInSet)
 					{
 						OutDelegates.InterfaceDelegates.OnInterfaceAdded.Broadcast(Interface);
 						bDidEdit = true;
 #if WITH_EDITORONLY_DATA
-						OutDoc.Metadata.ModifyContext.AddInterfaceModified(Interface.Version.Name);
+						Document.Metadata.ModifyContext.AddInterfaceModified(Interface.Version.Name);
 #endif // WITH_EDITORONLY_DATA
 					}
 				}
@@ -408,7 +411,7 @@ namespace Metasound::Frontend
 #if WITH_EDITORONLY_DATA
 				if (bAddedVertices && Options.bSetDefaultNodeLocations)
 				{
-					TArray<FMetasoundFrontendNode>& Nodes = OutDoc.RootGraph.Graph.Nodes;
+					TArray<FMetasoundFrontendNode>& Nodes = Document.RootGraph.Graph.Nodes;
 					// Sort/Place Inputs
 					{
 						TSet<FName> NamesToSort;
@@ -442,6 +445,9 @@ namespace Metasound::Frontend
 
 			const FModifyInterfaceOptions Options;
 
+		private:
+			FMetasoundFrontendDocument& Document;
+
 			using FVertexPair = TTuple<FMetasoundFrontendClassVertex, FMetasoundFrontendClassVertex>;
 			TArray<FVertexPair> PairedInputs;
 			TArray<FVertexPair> PairedOutputs;
@@ -521,6 +527,37 @@ namespace Metasound::Frontend
 } // namespace Metasound::Frontend
 
 
+UMetaSoundBuilderDocument& UMetaSoundBuilderDocument::Create(const UClass& InBuilderClass)
+{
+	UMetaSoundBuilderDocument* DocObject = NewObject<UMetaSoundBuilderDocument>();
+	DocObject->MetaSoundUClass = InBuilderClass;
+	return *DocObject;
+}
+
+UMetaSoundBuilderDocument& UMetaSoundBuilderDocument::Create(const IMetaSoundDocumentInterface& InDocToCopy)
+{
+	UMetaSoundBuilderDocument* DocObject = NewObject<UMetaSoundBuilderDocument>();
+	DocObject->Document = InDocToCopy.GetDocument();
+	DocObject->MetaSoundUClass = InDocToCopy.GetBaseMetaSoundUClass();
+	return *DocObject;
+}
+
+const FMetasoundFrontendDocument& UMetaSoundBuilderDocument::GetDocument() const
+{
+	return Document;
+}
+
+const UClass& UMetaSoundBuilderDocument::GetBaseMetaSoundUClass() const
+{
+	checkf(MetaSoundUClass, TEXT("BaseMetaSoundUClass must be set upon creation of UMetaSoundBuilderDocument instance"));
+	return *MetaSoundUClass;
+}
+
+FMetasoundFrontendDocument& UMetaSoundBuilderDocument::GetDocument()
+{
+	return Document;
+}
+
 FMetaSoundFrontendDocumentBuilder::FMetaSoundFrontendDocumentBuilder()
 	: DocumentDelegates(MakeShared<Metasound::Frontend::FDocumentModifyDelegates>())
 {
@@ -562,7 +599,7 @@ const FMetasoundFrontendClass* FMetaSoundFrontendDocumentBuilder::AddDependency(
 		NewDependency.Metadata.SetType(EMetasoundFrontendClassType::External);
 	}
 
-	NewDependency.ID = FGuid::NewGuid();
+	NewDependency.ID = FDocumentIDGenerator::Get().CreateClassID(Document);
 	Dependency = &Document.Dependencies.Emplace_GetRef(MoveTemp(NewDependency));
 
 	const int32 NewIndex = Document.Dependencies.Num() - 1;
@@ -881,7 +918,7 @@ const FMetasoundFrontendNode* FMetaSoundFrontendDocumentBuilder::AddGraphInput(c
 			FMetasoundFrontendClassInput& NewInput = RootGraph.Interface.Inputs.Add_GetRef(InClassInput);
 			if (!NewInput.VertexID.IsValid())
 			{
-				NewInput.VertexID = FGuid::NewGuid();
+				NewInput.VertexID = FDocumentIDGenerator::Get().CreateVertexID(Document);
 			}
 
 			DocumentDelegates->InterfaceDelegates.OnInputAdded.Broadcast(NewIndex);
@@ -967,7 +1004,7 @@ const FMetasoundFrontendNode* FMetaSoundFrontendDocumentBuilder::AddGraphOutput(
 			FMetasoundFrontendClassOutput& NewOutput = RootGraph.Interface.Outputs.Add_GetRef(InClassOutput);
 			if (!NewOutput.VertexID.IsValid())
 			{
-				NewOutput.VertexID = FGuid::NewGuid();
+				NewOutput.VertexID = FDocumentIDGenerator::Get().CreateVertexID(Document);;
 			}
 
 			DocumentDelegates->InterfaceDelegates.OnOutputAdded.Broadcast(NewIndex);
@@ -1273,6 +1310,14 @@ const FMetasoundFrontendClass* FMetaSoundFrontendDocumentBuilder::FindDependency
 	return DocumentCache->FindDependency(RegistryKey);
 }
 
+TArray<const FMetasoundFrontendEdge*> FMetaSoundFrontendDocumentBuilder::FindEdges(const FGuid& InNodeID, const FGuid& InVertexID) const
+{
+	using namespace Metasound::Frontend;
+
+	const IDocumentGraphEdgeCache& EdgeCache = DocumentCache->GetEdgeCache();
+	return EdgeCache.FindEdges(InNodeID, InVertexID);
+}
+
 bool FMetaSoundFrontendDocumentBuilder::FindInterfaceInputNodes(FName InterfaceName, TArray<const FMetasoundFrontendNode*>& OutInputs) const
 {
 	using namespace Metasound::Frontend;
@@ -1466,18 +1511,16 @@ TArray<const FMetasoundFrontendVertex*> FMetaSoundFrontendDocumentBuilder::FindN
 	const FMetasoundFrontendDocument& Document = GetDocument();
 
 	TArray<const FMetasoundFrontendVertex*> Inputs;
-	if (const TArray<int32>* Indices = EdgeCache.FindEdgeIndicesFromNodeOutput(InOutputNodeID, InOutputVertexID))
+	const TArrayView<const int32> Indices = EdgeCache.FindEdgeIndicesFromNodeOutput(InOutputNodeID, InOutputVertexID);
+	Algo::Transform(Indices, Inputs, [&Document, &NodeCache, &ConnectedInputNodes](const int32& Index)
 	{
-		Algo::Transform(*Indices, Inputs, [&Document, &NodeCache, &ConnectedInputNodes](const int32& Index)
+		const FMetasoundFrontendEdge& Edge = Document.RootGraph.Graph.Edges[Index];
+		if (ConnectedInputNodes)
 		{
-			const FMetasoundFrontendEdge& Edge = Document.RootGraph.Graph.Edges[Index];
-			if (ConnectedInputNodes)
-			{
-				ConnectedInputNodes->Add(NodeCache.FindNode(Edge.ToNodeID));
-			}
-			return NodeCache.FindInputVertex(Edge.ToNodeID, Edge.ToVertexID);
-		});
-	}
+			ConnectedInputNodes->Add(NodeCache.FindNode(Edge.ToNodeID));
+		}
+		return NodeCache.FindInputVertex(Edge.ToNodeID, Edge.ToVertexID);
+	});
 	return Inputs;
 }
 
@@ -1785,30 +1828,59 @@ Metasound::Frontend::EInvalidEdgeReason FMetaSoundFrontendDocumentBuilder::IsVal
 
 	const IDocumentGraphNodeCache& NodeCache = DocumentCache->GetNodeCache();
 
-	const FMetasoundFrontendVertex* OutputVertex = NodeCache.FindOutputVertex(InEdge.FromNodeID, InEdge.FromVertexID);
-	const FMetasoundFrontendVertex* InputVertex = NodeCache.FindInputVertex(InEdge.ToNodeID, InEdge.ToVertexID);
+	bool bOutputRerouted = false;
+	const FMetasoundFrontendNode* OutputNode = nullptr;
+	const FMetasoundFrontendVertex* OutputVertex = NodeCache.FindReroutedOutputVertex(InEdge.FromNodeID, InEdge.FromVertexID, &OutputNode, &bOutputRerouted);
 	if (!OutputVertex)
 	{
-		return EInvalidEdgeReason::MissingOutput;
+		if (!bOutputRerouted)
+		{
+			// No concrete output found and not going through reroute
+			return EInvalidEdgeReason::MissingOutput;
+		}
 	}
 
-	if (!InputVertex)
+	bool bInputRerouted = false;
+	TArray<const FMetasoundFrontendNode*> InputNodes;
+	TArray<const FMetasoundFrontendVertex*> InputVertices = NodeCache.FindReroutedInputVertices(InEdge.ToNodeID, InEdge.ToVertexID, &InputNodes, &bInputRerouted);
+	if (InputVertices.IsEmpty())
 	{
-		return EInvalidEdgeReason::MissingInput;
+		if (!bInputRerouted)
+		{
+			// No concrete input found and not going through reroute
+			return EInvalidEdgeReason::MissingInput;
+		}
 	}
 
-	if (OutputVertex->TypeName != InputVertex->TypeName)
+	FName OutputType;
+	EMetasoundFrontendVertexAccessType OutputAccessType = EMetasoundFrontendVertexAccessType::Unset;
+	if (OutputVertex)
 	{
-		return EInvalidEdgeReason::MismatchedDataType;
+		OutputType = OutputVertex->TypeName;
+
+		const FMetasoundFrontendClassOutput* ClassOutput = FindNodeOutputClassOutput(OutputNode->GetID(), OutputVertex->VertexID);
+		check(ClassOutput);
+		OutputAccessType = ClassOutput->AccessType;
 	}
 
-	const FMetasoundFrontendClassOutput* ClassOutput = FindNodeOutputClassOutput(InEdge.FromNodeID, InEdge.FromVertexID);
-	check(ClassOutput);
+	FName InputType;
+	EMetasoundFrontendVertexAccessType InputAccessType = EMetasoundFrontendVertexAccessType::Unset;
+	if (!InputVertices.IsEmpty())
+	{
+		InputType = InputVertices.Last()->TypeName;
+
+		// Only have to check last typename as all should be the same in theory if more than one off a reroute node
+		const FMetasoundFrontendClassInput* ClassInput = FindNodeInputClassInput(InputNodes.Last()->GetID(), InputVertices.Last()->VertexID);
+		check(ClassInput);
+		InputAccessType = ClassInput->AccessType;
+	}
 
-	const FMetasoundFrontendClassInput* ClassInput = FindNodeInputClassInput(InEdge.ToNodeID, InEdge.ToVertexID);
-	check(ClassInput);
+	if (!InputType.IsNone() && !OutputType.IsNone() && OutputType != InputType)
+	{
+		return EInvalidEdgeReason::MismatchedDataType;
+	}
 
-	if (!FMetasoundFrontendClassVertex::CanConnectVertexAccessTypes(ClassOutput->AccessType, ClassInput->AccessType))
+	if (!FMetasoundFrontendClassVertex::CanConnectVertexAccessTypes(OutputAccessType, InputAccessType))
 	{
 		return EInvalidEdgeReason::MismatchedAccessType;
 	}
@@ -1820,9 +1892,132 @@ bool FMetaSoundFrontendDocumentBuilder::ModifyInterfaces(Metasound::Frontend::FM
 {
 	using namespace Metasound::Frontend;
 
-	DocumentBuilderPrivate::FModifyInterfacesImpl Context(MoveTemp(InOptions));
 	FMetasoundFrontendDocument& Doc = GetDocument();
-	return Context.Execute(*this, Doc, *DocumentDelegates);
+	DocumentBuilderPrivate::FModifyInterfacesImpl Context(Doc, MoveTemp(InOptions));
+	return Context.Execute(*this, *DocumentDelegates);
+}
+
+bool FMetaSoundFrontendDocumentBuilder::TransformTemplateNodes()
+{
+	using namespace Metasound::Frontend;
+
+	METASOUND_TRACE_CPUPROFILER_EVENT_SCOPE(FMetaSoundFrontendDocumentBuilder::TransformTemplateNodes);
+
+	bool bModified = false;
+
+	FMetasoundFrontendDocument& Document = GetDocument();
+	FMetasoundFrontendGraph& Graph = Document.RootGraph.Graph;
+	TArray<FMetasoundFrontendClass>& Dependencies = Document.Dependencies;
+	const TArray<FMetasoundFrontendEdge>& GraphEdges = Graph.Edges;
+
+	struct FTemplateParams
+	{
+		FGuid ClassID;
+		const INodeTemplate* Template = nullptr;
+	};
+
+	// 1. Find template dependencies to build
+	TArray<FTemplateParams> TemplateParams;
+	Algo::TransformIf(Dependencies, TemplateParams,
+		[](const FMetasoundFrontendClass& Class) { return Class.Metadata.GetType() == EMetasoundFrontendClassType::Template; },
+		[](const FMetasoundFrontendClass& Class)
+		{
+			const FNodeRegistryKey Key = NodeRegistryKey::CreateKey(Class.Metadata);
+			FGuid ClassID = Class.ID;
+			const INodeTemplate* Template = INodeTemplateRegistry::Get().FindTemplate(Key);
+			ensureMsgf(Template, TEXT("Template not found for template class reference '%s'"), *Class.Metadata.GetClassName().ToString());
+			return FTemplateParams { ClassID, Template };
+		}
+	);
+
+	for (FTemplateParams& Params : TemplateParams)
+	{
+		if (!Params.Template)
+		{
+			continue;
+		}
+
+		TUniquePtr<INodeTransform> NodeTransform = Params.Template->GenerateNodeTransform();
+		check(NodeTransform.IsValid());
+
+		TSet<const FMetasoundFrontendNode*> NodesToRemove;
+		TSet<TPair<FGuid, FGuid>> VerticesRemoved;
+
+		// 1a. Preprocess template nodes
+		for (const FMetasoundFrontendNode& Node : Graph.Nodes)
+		{
+			if (Params.ClassID == Node.ClassID)
+			{
+				NodeTransform->Transform(Node, *this);
+				NodesToRemove.Add(&Node);
+
+				auto GetNodeVertexGuidPair = [NodeID = Node.GetID()](const FMetasoundFrontendVertex& Vertex) { return TPair<FGuid, FGuid> { NodeID, Vertex.VertexID }; };
+				Algo::Transform(Node.Interface.Inputs, VerticesRemoved, GetNodeVertexGuidPair);
+				Algo::Transform(Node.Interface.Outputs, VerticesRemoved, GetNodeVertexGuidPair);
+
+				bModified = true;
+			}
+		}
+
+		// 1b. Remove template node from graph
+		constexpr bool bAllowShrinking = false;
+		for (int32 i = Graph.Nodes.Num() - 1; i >= 0; --i)
+		{
+			if (NodesToRemove.Contains(&Graph.Nodes[i]))
+			{
+				DocumentDelegates->NodeDelegates.OnRemoveSwappingNode.Broadcast(i, Graph.Nodes.Num() - 1);
+				Graph.Nodes.RemoveAtSwap(i, 1, bAllowShrinking);
+			}
+		}
+		Graph.Nodes.Shrink();
+
+		// 1c. Remove edges connecting template node from graph
+		for (int32 i = Graph.Edges.Num() - 1; i >= 0; --i)
+		{
+			const TPair<FGuid, FGuid> FromNodeVertexPair { Graph.Edges[i].FromNodeID, Graph.Edges[i].FromVertexID };
+			if (VerticesRemoved.Contains(FromNodeVertexPair))
+			{
+				DocumentDelegates->EdgeDelegates.OnRemoveSwappingEdge.Broadcast(i, GraphEdges.Num() - 1);
+				Graph.Edges.RemoveAtSwap(i, 1, bAllowShrinking);
+			}
+			else
+			{
+				const TPair<FGuid, FGuid> ToNodeVertexPair { Graph.Edges[i].ToNodeID, Graph.Edges[i].ToVertexID };
+				if (VerticesRemoved.Contains(ToNodeVertexPair))
+				{
+					DocumentDelegates->EdgeDelegates.OnRemoveSwappingEdge.Broadcast(i, GraphEdges.Num() - 1);
+					Graph.Edges.RemoveAtSwap(i, 1, bAllowShrinking);
+				}
+			}
+		}
+		Graph.Edges.Shrink();
+	}
+
+	// 2. Finally, Remove template classes from dependency list
+	{
+		TSet<FString> TemplateKeys;
+		Algo::Transform(TemplateParams, TemplateKeys, [](const FTemplateParams& Params)
+		{
+			if (Params.Template)
+			{
+				return NodeRegistryKey::CreateKey(Params.Template->GetFrontendClass().Metadata);
+			}
+
+			return FString();
+		});
+		constexpr bool bAllowShrinking = false;
+		for (int32 i = Dependencies.Num() - 1; i >= 0; --i)
+		{
+			const FMetasoundFrontendClass& Class = Dependencies[i];
+			if (TemplateKeys.Contains(NodeRegistryKey::CreateKey(Class.Metadata)))
+			{
+				DocumentDelegates->OnRemoveSwappingDependency.Broadcast(i, Dependencies.Num() - 1);
+				Dependencies.RemoveAtSwap(i, 1, bAllowShrinking);
+			}
+		}
+	}
+
+	return bModified;
 }
 
 void FMetaSoundFrontendDocumentBuilder::ReloadCacheInternal()
@@ -2038,11 +2233,15 @@ bool FMetaSoundFrontendDocumentBuilder::RemoveEdgesFromNodeOutput(const FGuid& I
 	using namespace Metasound::Frontend;
 
 	const IDocumentGraphEdgeCache& EdgeCache = DocumentCache->GetEdgeCache();
-	if (const TArray<int32>* Indices = EdgeCache.FindEdgeIndicesFromNodeOutput(InNodeID, InVertexID))
+	const TArrayView<const int32> Indices = EdgeCache.FindEdgeIndicesFromNodeOutput(InNodeID, InVertexID);
+	if (!Indices.IsEmpty())
 	{
 		FMetasoundFrontendDocument& Document = GetDocument();
 		FMetasoundFrontendGraph& Graph = Document.RootGraph.Graph;
-		TArray<int32> IndicesCopy = *Indices; // Copy off indices as the array may be modified when notifying the cache in the loop below
+
+		// Copy off indices and sort descending as the edge array will be modified when notifying the cache in the loop below
+		TArray<int32> IndicesCopy(Indices.GetData(), Indices.Num());
+		Algo::Sort(IndicesCopy, [](const int32& L, const int32& R) { return L > R; });
 		for (int32 Index : IndicesCopy)
 		{
 			const int32 LastIndex = Graph.Edges.Num() - 1;
diff --git a/Engine/Plugins/Runtime/Metasound/Source/MetasoundFrontend/Private/MetasoundFrontendDocumentCache.cpp b/Engine/Plugins/Runtime/Metasound/Source/MetasoundFrontend/Private/MetasoundFrontendDocumentCache.cpp
--- a/Engine/Plugins/Runtime/Metasound/Source/MetasoundFrontend/Private/MetasoundFrontendDocumentCache.cpp
+++ b/Engine/Plugins/Runtime/Metasound/Source/MetasoundFrontend/Private/MetasoundFrontendDocumentCache.cpp
@@ -7,6 +7,7 @@
 #include "MetasoundFrontendDocumentBuilder.h"
 #include "MetasoundFrontendRegistries.h"
 #include "MetasoundTrace.h"
+#include "NodeTemplates/MetasoundFrontendNodeTemplateReroute.h"
 
 
 namespace Metasound::Frontend
@@ -456,6 +457,195 @@ namespace Metasound::Frontend
 		return nullptr;
 	}
 
+	TArray<const FMetasoundFrontendVertex*> FDocumentGraphNodeCache::FindReroutedInputVertices(const FGuid& InNodeID, const FGuid& InVertexID, TArray<const FMetasoundFrontendNode*>* ConnectedNodes, bool* bOutIsRerouted) const
+	{
+		TArray<const FMetasoundFrontendVertex*> Vertices;
+
+		if (const FMetasoundFrontendNode* Node = FindNode(InNodeID))
+		{
+			const FMetasoundFrontendClass* Class = Parent->FindDependency(Node->ClassID);
+			check(Class);
+			if (Class->Metadata.GetClassName() == FRerouteNodeTemplate::ClassName)
+			{
+				if (bOutIsRerouted)
+				{
+					*bOutIsRerouted = true;
+				}
+
+				const IDocumentGraphEdgeCache& EdgeCache = Parent->GetEdgeCache();
+				const FMetasoundFrontendVertex& RerouteOutput = Node->Interface.Outputs.Last();
+				TArrayView<const int32> EdgeIndices = EdgeCache.FindEdgeIndicesFromNodeOutput(InNodeID, RerouteOutput.VertexID);
+				const FMetasoundFrontendDocument& Doc = Parent->GetDocument();
+				for (int32 Index : EdgeIndices)
+				{
+					const FMetasoundFrontendEdge& EdgeConnectedToOutput = Doc.RootGraph.Graph.Edges[Index];
+					TArray<const FMetasoundFrontendVertex*> ReroutedVertices;
+					if (ConnectedNodes)
+					{
+						TArray<const FMetasoundFrontendNode*> ReroutedConnectedNodes;
+						ReroutedVertices = FindReroutedInputVertices(EdgeConnectedToOutput.ToNodeID, EdgeConnectedToOutput.ToVertexID, &ReroutedConnectedNodes, bOutIsRerouted);
+						ConnectedNodes->Append(ReroutedConnectedNodes);
+					}
+					else
+					{
+						ReroutedVertices = FindReroutedInputVertices(EdgeConnectedToOutput.ToNodeID, EdgeConnectedToOutput.ToVertexID, nullptr, bOutIsRerouted);
+					}
+					Vertices.Append(MoveTemp(ReroutedVertices));
+				}
+			}
+			else
+			{
+				auto VertexMatchesPredicate = [&InVertexID](const FMetasoundFrontendVertex& Vertex)
+				{
+					return Vertex.VertexID == InVertexID;
+				};
+				const FMetasoundFrontendVertex* InputVertex = Node->Interface.Inputs.FindByPredicate(VertexMatchesPredicate);
+				check(InputVertex);
+				Vertices.Add(InputVertex);
+				if (ConnectedNodes)
+				{
+					ConnectedNodes->Add(Node);
+				}
+			}
+		}
+
+		return Vertices;
+	}
+
+	TArray<const FMetasoundFrontendVertex*> FDocumentGraphNodeCache::FindReroutedInputVertices(const FGuid& InNodeID, FName InVertexName, TArray<const FMetasoundFrontendNode*>* ConnectedNodes, bool* bOutIsRerouted) const
+	{
+		TArray<const FMetasoundFrontendVertex*> Vertices;
+
+		if (const FMetasoundFrontendNode* Node = FindNode(InNodeID))
+		{
+			const FMetasoundFrontendClass* Class = Parent->FindDependency(Node->ClassID);
+			check(Class);
+			if (Class->Metadata.GetClassName() == FRerouteNodeTemplate::ClassName)
+			{
+				if (bOutIsRerouted)
+				{
+					*bOutIsRerouted = true;
+				}
+
+				const IDocumentGraphEdgeCache& EdgeCache = Parent->GetEdgeCache();
+				const FMetasoundFrontendVertex& RerouteOutput = Node->Interface.Outputs.Last();
+				TArrayView<const int32> EdgeIndices = EdgeCache.FindEdgeIndicesFromNodeOutput(InNodeID, RerouteOutput.VertexID);
+				const FMetasoundFrontendDocument& Doc = Parent->GetDocument();
+				for (int32 Index : EdgeIndices)
+				{
+					const FMetasoundFrontendEdge& EdgeConnectedToOutput = Doc.RootGraph.Graph.Edges[Index];
+
+					TArray<const FMetasoundFrontendVertex*> ReroutedVertices;
+					if (ConnectedNodes)
+					{
+						TArray<const FMetasoundFrontendNode*> ReroutedConnectedNodes;
+						ReroutedVertices = FindReroutedInputVertices(EdgeConnectedToOutput.ToNodeID, EdgeConnectedToOutput.ToVertexID, &ReroutedConnectedNodes, bOutIsRerouted);
+						ConnectedNodes->Append(MoveTemp(ReroutedConnectedNodes));
+					}
+					else
+					{
+						ReroutedVertices = FindReroutedInputVertices(EdgeConnectedToOutput.ToNodeID, EdgeConnectedToOutput.ToVertexID, nullptr, bOutIsRerouted);
+					}
+					Vertices.Append(MoveTemp(ReroutedVertices));
+				}
+			}
+			else
+			{
+				auto VertexMatchesPredicate = [&InVertexName](const FMetasoundFrontendVertex& Vertex)
+				{
+					return Vertex.Name == InVertexName;
+				};
+				const FMetasoundFrontendVertex* InputVertex = Node->Interface.Inputs.FindByPredicate(VertexMatchesPredicate);
+				check(InputVertex);
+				Vertices.Add(InputVertex);
+				if (ConnectedNodes)
+				{
+					ConnectedNodes->Add(Node);
+				}
+			}
+		}
+
+		return Vertices;
+	}
+
+	const FMetasoundFrontendVertex* FDocumentGraphNodeCache::FindReroutedOutputVertex(const FGuid& InNodeID, const FGuid& InVertexID, const FMetasoundFrontendNode** ConnectedNode, bool* bOutIsRerouted) const
+	{
+		if (const FMetasoundFrontendNode* Node = FindNode(InNodeID))
+		{
+			const FMetasoundFrontendClass* Class = Parent->FindDependency(Node->ClassID);
+			check(Class);
+			if (Class->Metadata.GetClassName() == FRerouteNodeTemplate::ClassName)
+			{
+				if (bOutIsRerouted)
+				{
+					*bOutIsRerouted = true;
+				}
+
+				const IDocumentGraphEdgeCache& EdgeCache = Parent->GetEdgeCache();
+				const FMetasoundFrontendVertex& RerouteInput = Node->Interface.Inputs.Last();
+				const FMetasoundFrontendDocument& Doc = Parent->GetDocument();
+				if (const int32* ConnectedEdgeIndex = EdgeCache.FindEdgeIndexToNodeInput(InNodeID, RerouteInput.VertexID))
+				{
+					const FMetasoundFrontendEdge& EdgeConnectedToInput = Doc.RootGraph.Graph.Edges[*ConnectedEdgeIndex];
+					return FindReroutedOutputVertex(EdgeConnectedToInput.FromNodeID, EdgeConnectedToInput.FromVertexID, ConnectedNode, bOutIsRerouted);
+				}
+
+				return nullptr;
+			}
+
+			if (ConnectedNode)
+			{
+				*ConnectedNode = Node;
+			}
+			auto VertexMatchesPredicate = [&InVertexID](const FMetasoundFrontendVertex& Vertex)
+			{
+				return Vertex.VertexID == InVertexID;
+			};
+			return Node->Interface.Outputs.FindByPredicate(VertexMatchesPredicate);
+		}
+
+		return nullptr;
+	}
+
+	const FMetasoundFrontendVertex* FDocumentGraphNodeCache::FindReroutedOutputVertex(const FGuid& InNodeID, FName InVertexName, const FMetasoundFrontendNode** ConnectedNode,  bool* bOutIsRerouted) const
+	{
+		if (const FMetasoundFrontendNode* Node = FindNode(InNodeID))
+		{
+			const FMetasoundFrontendClass* Class = Parent->FindDependency(Node->ClassID);
+			check(Class);
+			if (Class->Metadata.GetClassName() == FRerouteNodeTemplate::ClassName)
+			{
+				if (bOutIsRerouted)
+				{
+					*bOutIsRerouted = true;
+				}
+
+				const IDocumentGraphEdgeCache& EdgeCache = Parent->GetEdgeCache();
+				const FMetasoundFrontendVertex& RerouteInput = Node->Interface.Inputs.Last();
+				const FMetasoundFrontendDocument& Doc = Parent->GetDocument();
+				if (const int32* ConnectedEdgeIndex = EdgeCache.FindEdgeIndexToNodeInput(InNodeID, RerouteInput.VertexID))
+				{
+					const FMetasoundFrontendEdge& EdgeConnectedToInput = Doc.RootGraph.Graph.Edges[*ConnectedEdgeIndex];
+					return FindReroutedOutputVertex(EdgeConnectedToInput.FromNodeID, EdgeConnectedToInput.FromVertexID, ConnectedNode, bOutIsRerouted);
+				}
+
+				return nullptr;
+			}
+
+			auto VertexMatchesPredicate = [&InVertexName](const FMetasoundFrontendVertex& Vertex)
+			{
+				return Vertex.Name == InVertexName;
+			};
+			if (ConnectedNode)
+			{
+				*ConnectedNode = Node;
+			}
+			return Node->Interface.Outputs.FindByPredicate(VertexMatchesPredicate);
+		}
+
+		return nullptr;
+	}
+
 	void FDocumentGraphNodeCache::OnNodeAdded(int32 InNewIndex)
 	{
 		const FMetasoundFrontendDocument& Document = Parent->GetDocument();
@@ -542,9 +732,14 @@ namespace Metasound::Frontend
 		return InputToEdgeIndex.Find({ InNodeID, InVertexID });
 	}
 
-	const TArray<int32>* FDocumentGraphEdgeCache::FindEdgeIndicesFromNodeOutput(const FGuid& InNodeID, const FGuid& InVertexID) const
+	const TArrayView<const int32> FDocumentGraphEdgeCache::FindEdgeIndicesFromNodeOutput(const FGuid& InNodeID, const FGuid& InVertexID) const
 	{
-		return OutputToEdgeIndices.Find({ InNodeID, InVertexID });
+		if (const TArray<int32>* Array = OutputToEdgeIndices.Find({ InNodeID, InVertexID }))
+		{
+			return TArrayView<const int32>(*Array);
+		}
+
+		return { };
 	}
 
 	bool FDocumentGraphEdgeCache::IsNodeInputConnected(const FGuid& InNodeID, const FGuid& InVertexID) const
diff --git a/Engine/Plugins/Runtime/Metasound/Source/MetasoundFrontend/Private/MetasoundFrontendDocumentCache.h b/Engine/Plugins/Runtime/Metasound/Source/MetasoundFrontend/Private/MetasoundFrontendDocumentCache.h
--- a/Engine/Plugins/Runtime/Metasound/Source/MetasoundFrontend/Private/MetasoundFrontendDocumentCache.h
+++ b/Engine/Plugins/Runtime/Metasound/Source/MetasoundFrontend/Private/MetasoundFrontendDocumentCache.h
@@ -31,7 +31,7 @@ namespace Metasound::Frontend
 
 		virtual TArray<const FMetasoundFrontendEdge*> FindEdges(const FGuid& InNodeID, const FGuid& InVertexID) const override;
 		virtual const int32* FindEdgeIndexToNodeInput(const FGuid& InNodeID, const FGuid& InVertexID) const override;
-		virtual const TArray<int32>* FindEdgeIndicesFromNodeOutput(const FGuid& InNodeID, const FGuid& InVertexID) const override;
+		virtual const TArrayView<const int32> FindEdgeIndicesFromNodeOutput(const FGuid& InNodeID, const FGuid& InVertexID) const override;
 
 		void Init(FEdgeModifyDelegates& OutDelegates);
 
@@ -75,6 +75,11 @@ namespace Metasound::Frontend
 		virtual const FMetasoundFrontendVertex* FindOutputVertex(const FGuid& InNodeID, const FGuid& InVertexID) const override;
 		virtual const FMetasoundFrontendVertex* FindOutputVertex(const FGuid& InNodeID, FName InVertexName) const override;
 
+		virtual TArray<const FMetasoundFrontendVertex*> FindReroutedInputVertices(const FGuid& InNodeID, const FGuid& InVertexID, TArray<const FMetasoundFrontendNode*>* ConnectedNodes = nullptr, bool* bOutIsRerouted = nullptr) const override;
+		virtual TArray<const FMetasoundFrontendVertex*> FindReroutedInputVertices(const FGuid& InNodeID, FName InVertexName, TArray<const FMetasoundFrontendNode*>* ConnectedNodes = nullptr, bool* bOutIsRerouted = nullptr) const override;
+		virtual const FMetasoundFrontendVertex* FindReroutedOutputVertex(const FGuid& InNodeID, const FGuid& InVertexID, const FMetasoundFrontendNode** ConnectedNodes = nullptr, bool* bOutIsRerouted = nullptr) const override;
+		virtual const FMetasoundFrontendVertex* FindReroutedOutputVertex(const FGuid& InNodeID, FName InVertexName, const FMetasoundFrontendNode** ConnectedNodes = nullptr, bool* bOutIsRerouted = nullptr) const override;
+
 		void Init(FNodeModifyDelegates& OutDelegates);
 
 	private:
diff --git a/Engine/Plugins/Runtime/Metasound/Source/MetasoundFrontend/Private/MetasoundFrontendDocumentTemplatePreprocessor.cpp b/Engine/Plugins/Runtime/Metasound/Source/MetasoundFrontend/Private/MetasoundFrontendDocumentTemplatePreprocessor.cpp
--- a/Engine/Plugins/Runtime/Metasound/Source/MetasoundFrontend/Private/MetasoundFrontendDocumentTemplatePreprocessor.cpp
+++ b/Engine/Plugins/Runtime/Metasound/Source/MetasoundFrontend/Private/MetasoundFrontendDocumentTemplatePreprocessor.cpp
@@ -13,15 +13,9 @@ namespace Metasound
 {
 	namespace Frontend
 	{
-		bool FDocumentTemplatePreprocessTransform::Transform(FDocumentHandle InDocument) const
-		{
-			unimplemented();
-			return false;
-		}
-
 		bool FDocumentTemplatePreprocessTransform::Transform(FMetasoundFrontendDocument& InOutDoc) const
 		{
-			METASOUND_TRACE_CPUPROFILER_EVENT_SCOPE(Metasound::Frontend::AssetBase::Preprocessor);
+			PRAGMA_DISABLE_DEPRECATION_WARNINGS
 
 			bool bModified = false;
 
@@ -56,7 +50,7 @@ namespace Metasound
 					continue;
 				}
 
-				TUniquePtr<INodeTransform> NodeTransform = Params.Template->GenerateNodeTransform(InOutDoc);
+				TUniquePtr<INodeTransform> NodeTransform = Params.Template->GenerateNodeTransform();
 				check(NodeTransform.IsValid());
 
 				TSet<const FMetasoundFrontendNode*> NodesToRemove;
@@ -128,6 +122,8 @@ namespace Metasound
 				}, bAllowShrinking);
 			}
 
+			PRAGMA_ENABLE_DEPRECATION_WARNINGS
+
 			return bModified;
 		}
 	} // namespace Frontend
diff --git a/Engine/Plugins/Runtime/Metasound/Source/MetasoundFrontend/Private/MetasoundFrontendTransform.cpp b/Engine/Plugins/Runtime/Metasound/Source/MetasoundFrontend/Private/MetasoundFrontendTransform.cpp
--- a/Engine/Plugins/Runtime/Metasound/Source/MetasoundFrontend/Private/MetasoundFrontendTransform.cpp
+++ b/Engine/Plugins/Runtime/Metasound/Source/MetasoundFrontend/Private/MetasoundFrontendTransform.cpp
@@ -43,6 +43,16 @@ namespace Metasound
 			return Transform(FDocumentController::CreateDocumentHandle(DocAccessPtr));
 		}
 
+		bool INodeTransform::Transform(const FMetasoundFrontendNode& InNode, FMetaSoundFrontendDocumentBuilder& OutBuilder) const
+		{
+			return false;
+		}
+
+		bool INodeTransform::Transform(FMetasoundFrontendNode& InOutNode) const
+		{
+			return false;
+		}
+
 		FModifyRootGraphInterfaces::FModifyRootGraphInterfaces(const TArray<FMetasoundFrontendInterface>& InInterfacesToRemove, const TArray<FMetasoundFrontendInterface>& InInterfacesToAdd)
 			: InterfacesToRemove(InInterfacesToRemove)
 			, InterfacesToAdd(InInterfacesToAdd)
diff --git a/Engine/Plugins/Runtime/Metasound/Source/MetasoundFrontend/Private/NodeTemplates/MetasoundFrontendDocumentTemplatePreprocessor.h b/Engine/Plugins/Runtime/Metasound/Source/MetasoundFrontend/Private/NodeTemplates/MetasoundFrontendDocumentTemplatePreprocessor.h
--- a/Engine/Plugins/Runtime/Metasound/Source/MetasoundFrontend/Private/NodeTemplates/MetasoundFrontendDocumentTemplatePreprocessor.h
+++ b/Engine/Plugins/Runtime/Metasound/Source/MetasoundFrontend/Private/NodeTemplates/MetasoundFrontendDocumentTemplatePreprocessor.h
@@ -9,14 +9,10 @@ namespace Metasound
 {
 	namespace Frontend
 	{
-		// Preprocesses all template node classes and respective node instances within
-		// a given document and removes template classes from document dependencies list.
 		class FDocumentTemplatePreprocessTransform : public IDocumentTransform
 		{
 		public:
-			UE_DEPRECATED(5.1, "Transforms acting on handles/controllers are deprecated.")
-			virtual bool Transform(FDocumentHandle InDocument) const override;
-
+			UE_DEPRECATED(5.4, "TemplatePreprocessTransform is now FMetasoundFrontendDocumentBuilder::TransformTemplateNodes")
 			virtual bool Transform(FMetasoundFrontendDocument& InOutDocument) const override;
 		};
 
diff --git a/Engine/Plugins/Runtime/Metasound/Source/MetasoundFrontend/Private/NodeTemplates/MetasoundFrontendNodeTemplateReroute.cpp b/Engine/Plugins/Runtime/Metasound/Source/MetasoundFrontend/Private/NodeTemplates/MetasoundFrontendNodeTemplateReroute.cpp
--- a/Engine/Plugins/Runtime/Metasound/Source/MetasoundFrontend/Private/NodeTemplates/MetasoundFrontendNodeTemplateReroute.cpp
+++ b/Engine/Plugins/Runtime/Metasound/Source/MetasoundFrontend/Private/NodeTemplates/MetasoundFrontendNodeTemplateReroute.cpp
@@ -4,204 +4,197 @@
 
 #include "Algo/AnyOf.h"
 #include "MetasoundFrontendDataTypeRegistry.h"
+#include "MetasoundFrontendDocumentBuilder.h"
 #include "MetasoundFrontendRegistries.h"
 #include "NodeTemplates/MetasoundFrontendDocumentTemplatePreprocessor.h"
 
 
-namespace Metasound
+namespace Metasound::Frontend
 {
-	namespace Frontend
+	namespace ReroutePrivate
 	{
-		namespace ReroutePrivate
+		class FRerouteNodeTemplatePreprocessTransform : public INodeTransform
 		{
-			class FRerouteNodeTemplatePreprocessTransform : public FNodeTemplatePreprocessTransformBase
+		public:
+			FRerouteNodeTemplatePreprocessTransform() = default;
+			virtual ~FRerouteNodeTemplatePreprocessTransform() = default;
+
+			virtual bool Transform(const FMetasoundFrontendNode& InNode, FMetaSoundFrontendDocumentBuilder& OutBuilder) const override;
+		};
+
+		bool FRerouteNodeTemplatePreprocessTransform::Transform(const FMetasoundFrontendNode& InNode, FMetaSoundFrontendDocumentBuilder& OutBuilder) const
+		{
+			using namespace ReroutePrivate;
+				
+			// Find the input and output edges for this node
+			const FMetasoundFrontendEdge* InputEdge = nullptr;
 			{
-			public:
-				FRerouteNodeTemplatePreprocessTransform(FMetasoundFrontendDocument& InDocument)
-					: FNodeTemplatePreprocessTransformBase(InDocument)
+				if (!ensure(InNode.Interface.Inputs.Num() == 1))
 				{
-					TArray<FMetasoundFrontendEdge>& Edges = Graph.Edges;
-					for (int32 i = 0; i < Edges.Num(); ++i)
-					{
-						FMetasoundFrontendEdge& Edge = Edges[i];
-						InputEdgeMap.Add(Edge.GetToVertexHandle(), &Edge);
-						OutputEdgeMap.FindOrAdd(Edge.GetFromVertexHandle()).Add(&Edge);
-					}
+					return false;
 				}
 
-				virtual ~FRerouteNodeTemplatePreprocessTransform() = default;
-
-				virtual bool Transform(FMetasoundFrontendNode& InOutNode) const override;
-
-			private:
-				mutable TMap<FMetasoundFrontendVertexHandle, FMetasoundFrontendEdge*> InputEdgeMap;
-				mutable TMap<FMetasoundFrontendVertexHandle, TArray<FMetasoundFrontendEdge*>> OutputEdgeMap;
-			};
+				const FMetasoundFrontendVertex& InputVertex = InNode.Interface.Inputs.Last();
+				FMetasoundFrontendVertexHandle InputNodeVertexHandle { InNode.GetID(), InputVertex.VertexID };
+				TArray<const FMetasoundFrontendEdge*> InputEdges = OutBuilder.FindEdges(InNode.GetID(), InputVertex.VertexID);
 
-			bool FRerouteNodeTemplatePreprocessTransform::Transform(FMetasoundFrontendNode& InOutNode) const
-			{
-				using namespace ReroutePrivate;
-				
-				// Find the input and output edges for this node
-				const FMetasoundFrontendEdge* InputEdge = nullptr;
+				// This can happen if the reroute node isn't provided an input, so its perfectly
+				// acceptable to just ignore this node as it ultimately provides no sourced input.
+				if (InputEdges.IsEmpty())
 				{
-					if (!ensure(InOutNode.Interface.Inputs.Num() == 1))
-					{
-						return false;
-					}
-
-					const FMetasoundFrontendVertex& InputVertex = InOutNode.Interface.Inputs.Last();
-					FMetasoundFrontendVertexHandle InputNodeVertexHandle { InOutNode.GetID(), InputVertex.VertexID };
-					InputEdge = InputEdgeMap.FindRef(InputNodeVertexHandle);
-
-					// This can happen if the reroute node isn't provided an input, so its perfectly
-					// acceptable to just ignore this node as it ultimately provides no sourced input.
-					if (!InputEdge)
-					{
-						return false;
-					}
+					return false;
 				}
+				InputEdge = InputEdges.Last();
+			}
 
-				TArray<FMetasoundFrontendEdge*>* OutputEdges = nullptr;
-				const FMetasoundFrontendVertex& OutputVertex = InOutNode.Interface.Outputs.Last();
-				const FMetasoundFrontendVertexHandle OutputVertexHandle { InOutNode.GetID(), OutputVertex.VertexID };
+			TArray<const FMetasoundFrontendEdge*> OutputEdges;
+			const FMetasoundFrontendVertex& OutputVertex = InNode.Interface.Outputs.Last();
+			const FMetasoundFrontendVertexHandle OutputVertexHandle { InNode.GetID(), OutputVertex.VertexID };
+			{
+				if (!ensure(InNode.Interface.Outputs.Num() == 1))
 				{
-					if (!ensure(InOutNode.Interface.Outputs.Num() == 1))
-					{
-						return false;
-					}
-
-					OutputEdges = OutputEdgeMap.Find(OutputVertexHandle);
-
-					// This can happen if the reroute node isn't provided any outputs to connect to, so its
-					// perfectly acceptable to just ignore this node as it ultimately provides no sourced input.
-					if (!OutputEdges)
-					{
-						return false;
-					}
+					return false;
 				}
 
-				// Update the output edges with the input edge 
-				FMetasoundFrontendVertexHandle NewOutputEdgeNodeVertexHandle { InputEdge->FromNodeID, InputEdge->FromVertexID };
+				OutputEdges = OutBuilder.FindEdges(InNode.GetID(), OutputVertex.VertexID);
 
-				for (FMetasoundFrontendEdge* OutputEdge : *OutputEdges)
+				// This can happen if the reroute node isn't provided any outputs to connect to, so its
+				// perfectly acceptable to just ignore this node as it ultimately provides no sourced input.
+				if (OutputEdges.IsEmpty())
 				{
-					OutputEdge->FromNodeID = NewOutputEdgeNodeVertexHandle.NodeID;
-					OutputEdge->FromVertexID = NewOutputEdgeNodeVertexHandle.VertexID;
+					return false;
 				}
+			}
+
+			// Update the output edges with the input edge 
+			FMetasoundFrontendVertexHandle NewOutputEdgeNodeVertexHandle { InputEdge->FromNodeID, InputEdge->FromVertexID };
 
-				OutputEdgeMap.FindOrAdd(NewOutputEdgeNodeVertexHandle).Append(*OutputEdges);
-				OutputEdgeMap.Remove(OutputVertexHandle);
+			for (const FMetasoundFrontendEdge* OutputEdge : OutputEdges)
+			{
+				check(OutputEdge);
+
+				// Must copied before removal as pointer will become invalid
+				FMetasoundFrontendEdge Edge = *OutputEdge;
+				OutBuilder.RemoveEdge(Edge);
 
-				return true;
+				Edge.FromNodeID = InputEdge->FromNodeID;
+				Edge.FromVertexID = InputEdge->FromVertexID;
+				OutBuilder.AddEdge(MoveTemp(Edge));
 			}
+
+			return true;
 		}
+	}
 
-		const FMetasoundFrontendClassName FRerouteNodeTemplate::ClassName { "UE", "Reroute", "" };
+	const FMetasoundFrontendClassName FRerouteNodeTemplate::ClassName { "UE", "Reroute", "" };
 
-		const FMetasoundFrontendVersion FRerouteNodeTemplate::Version { ClassName.GetFullName(), { 1, 0 } };
+	const FMetasoundFrontendVersion FRerouteNodeTemplate::Version { ClassName.GetFullName(), { 1, 0 } };
 
-		TUniquePtr<INodeTransform> FRerouteNodeTemplate::GenerateNodeTransform(FMetasoundFrontendDocument& InPreprocessedDocument) const
-		{
-			using namespace ReroutePrivate;
-			return TUniquePtr<INodeTransform>(new FRerouteNodeTemplatePreprocessTransform(InPreprocessedDocument));
-		}
+	TUniquePtr<INodeTransform> FRerouteNodeTemplate::GenerateNodeTransform(FMetasoundFrontendDocument& InPreprocessedDocument) const
+	{
+		return GenerateNodeTransform();
+	}
 
-		const FMetasoundFrontendClass& FRerouteNodeTemplate::GetFrontendClass() const
+	TUniquePtr<INodeTransform> FRerouteNodeTemplate::GenerateNodeTransform() const
+	{
+		using namespace ReroutePrivate;
+		return TUniquePtr<INodeTransform>(new FRerouteNodeTemplatePreprocessTransform());
+	}
+
+	const FMetasoundFrontendClass& FRerouteNodeTemplate::GetFrontendClass() const
+	{
+		auto CreateFrontendClass = []()
 		{
-			auto CreateFrontendClass = []()
-			{
-				FMetasoundFrontendClass Class;
-				Class.Metadata.SetClassName(ClassName);
+			FMetasoundFrontendClass Class;
+			Class.Metadata.SetClassName(ClassName);
 
 #if WITH_EDITOR
-				Class.Metadata.SetSerializeText(false);
-				Class.Metadata.SetAuthor(Metasound::PluginAuthor);
-				Class.Metadata.SetDescription(Metasound::PluginNodeMissingPrompt);
-
-				FMetasoundFrontendClassStyleDisplay& StyleDisplay = Class.Style.Display;
-				StyleDisplay.ImageName = "MetasoundEditor.Graph.Node.Class.Reroute";
-				StyleDisplay.bShowInputNames = false;
-				StyleDisplay.bShowOutputNames = false;
-				StyleDisplay.bShowLiterals = false;
-				StyleDisplay.bShowName = false;
+			Class.Metadata.SetSerializeText(false);
+			Class.Metadata.SetAuthor(Metasound::PluginAuthor);
+			Class.Metadata.SetDescription(Metasound::PluginNodeMissingPrompt);
+
+			FMetasoundFrontendClassStyleDisplay& StyleDisplay = Class.Style.Display;
+			StyleDisplay.ImageName = "MetasoundEditor.Graph.Node.Class.Reroute";
+			StyleDisplay.bShowInputNames = false;
+			StyleDisplay.bShowOutputNames = false;
+			StyleDisplay.bShowLiterals = false;
+			StyleDisplay.bShowName = false;
 #endif // WITH_EDITOR
 
-				Class.Metadata.SetType(EMetasoundFrontendClassType::Template);
-				Class.Metadata.SetVersion(Version.Number);
+			Class.Metadata.SetType(EMetasoundFrontendClassType::Template);
+			Class.Metadata.SetVersion(Version.Number);
 
 
-				return Class;
-			};
+			return Class;
+		};
 
-			static const FMetasoundFrontendClass FrontendClass = CreateFrontendClass();
-			return FrontendClass;
-		}
+		static const FMetasoundFrontendClass FrontendClass = CreateFrontendClass();
+		return FrontendClass;
+	}
 
-		FMetasoundFrontendNodeInterface FRerouteNodeTemplate::CreateNodeInterfaceFromDataType(FName InDataType)
-		{
-			auto CreateNewVertex = [&] { return FMetasoundFrontendVertex { "Value", InDataType, FGuid::NewGuid() }; };
+	FMetasoundFrontendNodeInterface FRerouteNodeTemplate::CreateNodeInterfaceFromDataType(FName InDataType)
+	{
+		auto CreateNewVertex = [&] { return FMetasoundFrontendVertex { "Value", InDataType, FGuid::NewGuid() }; };
 
-			FMetasoundFrontendNodeInterface NewInterface;
-			NewInterface.Inputs.Add(CreateNewVertex());
-			NewInterface.Outputs.Add(CreateNewVertex());
+		FMetasoundFrontendNodeInterface NewInterface;
+		NewInterface.Inputs.Add(CreateNewVertex());
+		NewInterface.Outputs.Add(CreateNewVertex());
 
-			return NewInterface;
-		}
+		return NewInterface;
+	}
 
-		const FNodeRegistryKey& FRerouteNodeTemplate::GetRegistryKey()
-		{
-			static const FNodeRegistryKey RegistryKey = NodeRegistryKey::CreateKey(
-				EMetasoundFrontendClassType::Template,
-				ClassName.ToString(), 
-				Version.Number.Major, 
-				Version.Number.Minor);
+	const FNodeRegistryKey& FRerouteNodeTemplate::GetRegistryKey()
+	{
+		static const FNodeRegistryKey RegistryKey = NodeRegistryKey::CreateKey(
+			EMetasoundFrontendClassType::Template,
+			ClassName.ToString(), 
+			Version.Number.Major, 
+			Version.Number.Minor);
 
-			return RegistryKey;
-		}
+		return RegistryKey;
+	}
 
-		const FMetasoundFrontendVersion& FRerouteNodeTemplate::GetVersion() const
-		{
-			return Version;
-		}
+	const FMetasoundFrontendVersion& FRerouteNodeTemplate::GetVersion() const
+	{
+		return Version;
+	}
 
 #if WITH_EDITOR
-		bool FRerouteNodeTemplate::HasRequiredConnections(FConstNodeHandle InNodeHandle) const
-		{
-			TArray<FConstOutputHandle> Outputs = InNodeHandle->GetConstOutputs();
-			TArray<FConstInputHandle> Inputs = InNodeHandle->GetConstInputs();
+	bool FRerouteNodeTemplate::HasRequiredConnections(FConstNodeHandle InNodeHandle) const
+	{
+		TArray<FConstOutputHandle> Outputs = InNodeHandle->GetConstOutputs();
+		TArray<FConstInputHandle> Inputs = InNodeHandle->GetConstInputs();
 
-			const bool bConnectedToNonRerouteOutputs = Algo::AnyOf(Outputs, [](const FConstOutputHandle& OutputHandle) { return Frontend::FindReroutedOutput(OutputHandle)->IsValid(); });
-			const bool bConnectedToNonRerouteInputs = Algo::AnyOf(Inputs, [](const FConstInputHandle& InputHandle)
-			{
-				TArray<FConstInputHandle> Inputs;
-				Frontend::FindReroutedInputs(InputHandle, Inputs);
-				return !Inputs.IsEmpty();
-			});
+		const bool bConnectedToNonRerouteOutputs = Algo::AnyOf(Outputs, [](const FConstOutputHandle& OutputHandle) { return Frontend::FindReroutedOutput(OutputHandle)->IsValid(); });
+		const bool bConnectedToNonRerouteInputs = Algo::AnyOf(Inputs, [](const FConstInputHandle& InputHandle)
+		{
+			TArray<FConstInputHandle> Inputs;
+			Frontend::FindReroutedInputs(InputHandle, Inputs);
+			return !Inputs.IsEmpty();
+		});
 
-			return bConnectedToNonRerouteOutputs || bConnectedToNonRerouteOutputs == bConnectedToNonRerouteInputs;
-		}
+		return bConnectedToNonRerouteOutputs || bConnectedToNonRerouteOutputs == bConnectedToNonRerouteInputs;
+	}
 #endif // WITH_EDITOR
 
-		bool FRerouteNodeTemplate::IsValidNodeInterface(const FMetasoundFrontendNodeInterface& InNodeInterface) const
+	bool FRerouteNodeTemplate::IsValidNodeInterface(const FMetasoundFrontendNodeInterface& InNodeInterface) const
+	{
+		if (InNodeInterface.Inputs.Num() != 1)
 		{
-			if (InNodeInterface.Inputs.Num() != 1)
-			{
-				return false;
-			}
+			return false;
+		}
 			
-			if (InNodeInterface.Outputs.Num() != 1)
-			{
-				return false;
-			}
-
-			const FName DataType = InNodeInterface.Inputs.Last().TypeName;
-			if (DataType != InNodeInterface.Outputs.Last().TypeName)
-			{
-				return false;
-			}
+		if (InNodeInterface.Outputs.Num() != 1)
+		{
+			return false;
+		}
 
-			return IDataTypeRegistry::Get().IsRegistered(DataType);
+		const FName DataType = InNodeInterface.Inputs.Last().TypeName;
+		if (DataType != InNodeInterface.Outputs.Last().TypeName)
+		{
+			return false;
 		}
-	} // namespace Frontend
-} // namespace Metasound
\ No newline at end of file
+
+		return IDataTypeRegistry::Get().IsRegistered(DataType);
+	}
+} // namespace Metasound::Frontend
diff --git a/Engine/Plugins/Runtime/Metasound/Source/MetasoundFrontend/Public/MetasoundAssetBase.h b/Engine/Plugins/Runtime/Metasound/Source/MetasoundFrontend/Public/MetasoundAssetBase.h
--- a/Engine/Plugins/Runtime/Metasound/Source/MetasoundFrontend/Public/MetasoundAssetBase.h
+++ b/Engine/Plugins/Runtime/Metasound/Source/MetasoundFrontend/Public/MetasoundAssetBase.h
@@ -15,8 +15,11 @@
 #include "UObject/SoftObjectPath.h"
 #include "UObject/WeakObjectPtrTemplates.h"
 
+
 // Forward Declarations
 class UEdGraph;
+struct FMetaSoundFrontendDocumentBuilder;
+
 
 namespace Metasound
 {
@@ -75,8 +78,8 @@ public:
 	// Unregisters the root graph of the given asset with the MetaSound Frontend.
 	void UnregisterGraphWithFrontend();
 
-	// Prepares graph for cook, including autoupdating and preprocessing graph and registering node, similar to RegisterGraphWithFrontend. 
-	void CookGraph();
+	// Cooks this MetaSound and recursively checks and cooks referenced MetaSounds if necessary. Cook includes autoupdating and resolving the document, which is then registered with the MetaSound Frontend.
+	void CookMetaSound();
 
 	// Sets/overwrites the root class metadata
 	UE_DEPRECATED(5.3, "Directly setting graph class Metadata is no longer be supported. Use the FMetaSoundFrontendDocumentBuilder to modify class data.")
@@ -223,9 +226,14 @@ protected:
 
 
 	bool AutoUpdate(bool bInLogWarningsOnDroppedConnection);
-	void RegisterAssetDependencies(const Metasound::Frontend::FMetaSoundAssetRegistrationOptions& InRegistrationOptions);
-	void RegisterAssetDependenciesForCook();
+	void CookReferencedMetaSounds();
+
+	// Ensures all referenced graph classes are registered (or re-registers depending on options).
+	void RegisterAssetDependencies(const Metasound::Frontend::FMetaSoundAssetRegistrationOptions & InRegistrationOptions);
+
+	UE_DEPRECATED(5.4, "Template node transformation moved to private implementation. A MetaSound asset will likely never have the function Process(...). Without a Process function, you cannot have preprocessing.")
 	TSharedPtr<FMetasoundFrontendDocument> PreprocessDocument();
+
 private:
 #if WITH_EDITORONLY_DATA
 	void UpdateAssetRegistry();
diff --git a/Engine/Plugins/Runtime/Metasound/Source/MetasoundFrontend/Public/MetasoundDocumentInterface.h b/Engine/Plugins/Runtime/Metasound/Source/MetasoundFrontend/Public/MetasoundDocumentInterface.h
--- a/Engine/Plugins/Runtime/Metasound/Source/MetasoundFrontend/Public/MetasoundDocumentInterface.h
+++ b/Engine/Plugins/Runtime/Metasound/Source/MetasoundFrontend/Public/MetasoundDocumentInterface.h
@@ -42,7 +42,7 @@ private:
 
 namespace Metasound::Frontend
 {
-	class METASOUNDFRONTEND_API IMetaSoundDocumentBuilderRegistry
+	class METASOUNDFRONTEND_API IDocumentBuilderRegistry
 	{
 	public:
 		// Returns delegates used to mutate any internal builder cached state or notify listeners of external system that
@@ -51,9 +51,12 @@ namespace Metasound::Frontend
 		// mutate all documents using MetasoundDocumentBuilders).
 		virtual const FDocumentModifyDelegates* FindModifyDelegates(const FMetasoundFrontendClassName& InClassName) const = 0;
 
-		static IMetaSoundDocumentBuilderRegistry& GetChecked();
+		static IDocumentBuilderRegistry& GetChecked();
 
 	protected:
-		static void Set(TUniqueFunction<IMetaSoundDocumentBuilderRegistry&()>&& InGetInstance);
+		static void Set(TUniqueFunction<IDocumentBuilderRegistry&()>&& InGetInstance);
 	};
+
+	// Deprecated: use shorter name
+	using IMetaSoundDocumentBuilderRegistry = IDocumentBuilderRegistry;
 } // namespace Metasound::Frontend
diff --git a/Engine/Plugins/Runtime/Metasound/Source/MetasoundFrontend/Public/MetasoundFrontendDocumentBuilder.h b/Engine/Plugins/Runtime/Metasound/Source/MetasoundFrontend/Public/MetasoundFrontendDocumentBuilder.h
--- a/Engine/Plugins/Runtime/Metasound/Source/MetasoundFrontend/Public/MetasoundFrontendDocumentBuilder.h
+++ b/Engine/Plugins/Runtime/Metasound/Source/MetasoundFrontend/Public/MetasoundFrontendDocumentBuilder.h
@@ -92,30 +92,20 @@ class METASOUNDFRONTEND_API UMetaSoundBuilderDocument : public UObject, public I
 	GENERATED_BODY()
 
 public:
+	// Create and return a valid builder document with the provided class
+	static UMetaSoundBuilderDocument& Create(const UClass& InBuilderClass);
+
+	// Create and return a valid builder document which copies the provided interface's document & class
+	static UMetaSoundBuilderDocument& Create(const IMetaSoundDocumentInterface& InDocToCopy);
+
 	// Returns the document
-	virtual const FMetasoundFrontendDocument& GetDocument() const override
-	{
-		return Document;
-	}
+	virtual const FMetasoundFrontendDocument& GetDocument() const override;
 
 	// Base MetaSoundClass that document is published to.
-	virtual const UClass& GetBaseMetaSoundUClass() const final override
-	{
-		checkf(MetaSoundUClass, TEXT("BaseMetaSoundUClass must be set upon creation of UMetaSoundBuilderDocument instance"));
-		return *MetaSoundUClass;
-	}
+	virtual const UClass& GetBaseMetaSoundUClass() const final override;
 
 private:
-	// Sets the base MetaSound UClass (to be called immediately following UObject construction)
-	void SetBaseMetaSoundUClass(const UClass& InMetaSoundClass)
-	{
-		MetaSoundUClass = &InMetaSoundClass;
-	}
-
-	virtual FMetasoundFrontendDocument& GetDocument() override
-	{
-		return Document;
-	}
+	virtual FMetasoundFrontendDocument& GetDocument() override;
 
 	UPROPERTY(Transient)
 	FMetasoundFrontendDocument Document;
@@ -123,8 +113,7 @@ private:
 	UPROPERTY(Transient)
 	TObjectPtr<const UClass> MetaSoundUClass = nullptr;
 
-	friend struct FMetaSoundFrontendDocumentBuilder;
-	friend class UMetaSoundBuilderBase;
+	friend class FMetasoundAssetBase;
 };
 
 // Builder used to support dynamically generating MetaSound documents at runtime. Builder contains caches that speed up
@@ -172,6 +161,7 @@ public:
 
 	const FMetasoundFrontendClass* FindDependency(const FGuid& InClassID) const;
 	const FMetasoundFrontendClass* FindDependency(const FMetasoundFrontendClassMetadata& InMetadata) const;
+	TArray<const FMetasoundFrontendEdge*> FindEdges(const FGuid& InNodeID, const FGuid& InVertexID) const;
 
 	bool FindInterfaceInputNodes(FName InterfaceName, TArray<const FMetasoundFrontendNode*>& OutInputs) const;
 	bool FindInterfaceOutputNodes(FName InterfaceName, TArray<const FMetasoundFrontendNode*>& OutOutputs) const;
@@ -253,6 +243,10 @@ public:
 	bool SwapGraphOutput(const FMetasoundFrontendClassVertex& InExistingOutputVertex, const FMetasoundFrontendClassVertex& NewOutputVertex);
 	bool UpdateDependencyClassNames(const TMap<FMetasoundFrontendClassName, FMetasoundFrontendClassName>& OldToNewReferencedClassNames);
 
+	// Transforms template nodes within the given builder's document, which can include swapping associated edges and/or
+	// replacing nodes with other, registry-defined concrete node class instances. Returns true if any template nodes were processed.
+	bool TransformTemplateNodes();
+
 private:
 	using FFinalizeNodeFunctionRef = TFunctionRef<void(FMetasoundFrontendNode&, const Metasound::Frontend::FNodeRegistryKey&)>;
 
diff --git a/Engine/Plugins/Runtime/Metasound/Source/MetasoundFrontend/Public/MetasoundFrontendDocumentCacheInterface.h b/Engine/Plugins/Runtime/Metasound/Source/MetasoundFrontend/Public/MetasoundFrontendDocumentCacheInterface.h
--- a/Engine/Plugins/Runtime/Metasound/Source/MetasoundFrontend/Public/MetasoundFrontendDocumentCacheInterface.h
+++ b/Engine/Plugins/Runtime/Metasound/Source/MetasoundFrontend/Public/MetasoundFrontendDocumentCacheInterface.h
@@ -40,10 +40,29 @@ namespace Metasound::Frontend
 		virtual TArray<const FMetasoundFrontendVertex*> FindNodeInputs(const FGuid& InNodeID, FName TypeName = FName()) const = 0;
 		virtual TArray<const FMetasoundFrontendVertex*> FindNodeOutputs(const FGuid& InNodeID, FName TypeName = FName()) const = 0;
 
+		// Returns corresponding input vertex if it exists
 		virtual const FMetasoundFrontendVertex* FindInputVertex(const FGuid& InNodeID, const FGuid& InVertexID) const = 0;
+
+		// Returns corresponding input vertex if it exists
 		virtual const FMetasoundFrontendVertex* FindInputVertex(const FGuid& InNodeID, FName InVertexName) const = 0;
+
+		// Returns corresponding output vertex if it exists
 		virtual const FMetasoundFrontendVertex* FindOutputVertex(const FGuid& InNodeID, const FGuid& InVertexID) const = 0;
+
+		// Returns corresponding output vertex if it exists
 		virtual const FMetasoundFrontendVertex* FindOutputVertex(const FGuid& InNodeID, FName InVertexName) const = 0;
+
+		// Recursively finds rerouted input vertices if they exist. Returns corresponding nodes to vertices if array provided. Sets boolean determining if connected to reroute if provided (optional).
+		virtual TArray<const FMetasoundFrontendVertex*> FindReroutedInputVertices(const FGuid& InNodeID, const FGuid& InVertexID, TArray<const FMetasoundFrontendNode*>* ConnectedNodes = nullptr, bool* bOutIsRerouted = nullptr) const = 0;
+
+		// Recursively finds rerouted input vertices if they exist. Returns corresponding nodes to vertices if array provided. Sets boolean determining if connected to reroute if provided (optional).
+		virtual TArray<const FMetasoundFrontendVertex*> FindReroutedInputVertices(const FGuid& InNodeID, FName InVertexName, TArray<const FMetasoundFrontendNode*>* ConnectedNodes = nullptr, bool* bOutIsRerouted = nullptr) const = 0;
+
+		// Recursively finds rerouted output vertex if it exist. Returns corresponding nodes to vertices if pointer provided. Sets boolean determining if connected to reroute if provided (optional).
+		virtual const FMetasoundFrontendVertex* FindReroutedOutputVertex(const FGuid& InNodeID, const FGuid& InVertexID, const FMetasoundFrontendNode** ConnectedNodes = nullptr, bool* bOutIsRerouted = nullptr) const = 0;
+
+		// Recursively finds rerouted output vertex if it exist. Returns corresponding node to vertices if pointer provided. Sets boolean determining if connected to reroute if provided (optional).
+		virtual const FMetasoundFrontendVertex* FindReroutedOutputVertex(const FGuid& InNodeID, FName InVertexName, const FMetasoundFrontendNode** ConnectedNodes = nullptr, bool* bOutIsRerouted = nullptr) const = 0;
 	};
 
 	/** Interface for querying cached document edges. */
@@ -59,7 +78,8 @@ namespace Metasound::Frontend
 
 		virtual TArray<const FMetasoundFrontendEdge*> FindEdges(const FGuid& InNodeID, const FGuid& InVertexID) const = 0;
 		virtual const int32* FindEdgeIndexToNodeInput(const FGuid& InNodeID, const FGuid& InVertexID) const = 0;
-		virtual const TArray<int32>* FindEdgeIndicesFromNodeOutput(const FGuid& InNodeID, const FGuid& InVertexID) const = 0;
+
+		virtual const TArrayView<const int32> FindEdgeIndicesFromNodeOutput(const FGuid& InNodeID, const FGuid& InVertexID) const = 0;
 	};
 
 	/** Interface for querying cached document interface. */
diff --git a/Engine/Plugins/Runtime/Metasound/Source/MetasoundFrontend/Public/MetasoundFrontendDocumentIdGenerator.h b/Engine/Plugins/Runtime/Metasound/Source/MetasoundFrontend/Public/MetasoundFrontendDocumentIdGenerator.h
--- a/Engine/Plugins/Runtime/Metasound/Source/MetasoundFrontend/Public/MetasoundFrontendDocumentIdGenerator.h
+++ b/Engine/Plugins/Runtime/Metasound/Source/MetasoundFrontend/Public/MetasoundFrontendDocumentIdGenerator.h
@@ -43,7 +43,7 @@ namespace Metasound
 		 * }
 		 *
 		 */
-		class FDocumentIDGenerator
+		class METASOUNDFRONTEND_API FDocumentIDGenerator
 		{
 		public:
 			class METASOUNDFRONTEND_API FScopeDeterminism final
diff --git a/Engine/Plugins/Runtime/Metasound/Source/MetasoundFrontend/Public/MetasoundFrontendNodeTemplateRegistry.h b/Engine/Plugins/Runtime/Metasound/Source/MetasoundFrontend/Public/MetasoundFrontendNodeTemplateRegistry.h
--- a/Engine/Plugins/Runtime/Metasound/Source/MetasoundFrontend/Public/MetasoundFrontendNodeTemplateRegistry.h
+++ b/Engine/Plugins/Runtime/Metasound/Source/MetasoundFrontend/Public/MetasoundFrontendNodeTemplateRegistry.h
@@ -24,8 +24,11 @@ namespace Metasound
 		public:
 			virtual ~INodeTemplate() = default;
 
-			// Generates node transform that can cache document state for use when individually preprocessing nodes.
-			virtual TUniquePtr<INodeTransform> GenerateNodeTransform(FMetasoundFrontendDocument& InDocument) const = 0;
+			UE_DEPRECATED(5.4, "Use version that does not provide a preprocessed document")
+			virtual TUniquePtr<INodeTransform> GenerateNodeTransform(FMetasoundFrontendDocument& InDocument) const { return nullptr; }
+
+			// Generates node transform that is used to preprocess nodes.
+			virtual TUniquePtr<INodeTransform> GenerateNodeTransform() const { return nullptr; }
 
 			// Returns the class definition for the given node class template.
 			virtual const FMetasoundFrontendClass& GetFrontendClass() const = 0;
diff --git a/Engine/Plugins/Runtime/Metasound/Source/MetasoundFrontend/Public/MetasoundFrontendTransform.h b/Engine/Plugins/Runtime/Metasound/Source/MetasoundFrontend/Public/MetasoundFrontendTransform.h
--- a/Engine/Plugins/Runtime/Metasound/Source/MetasoundFrontend/Public/MetasoundFrontendTransform.h
+++ b/Engine/Plugins/Runtime/Metasound/Source/MetasoundFrontend/Public/MetasoundFrontendTransform.h
@@ -9,240 +9,244 @@
 #include "MetasoundVertex.h"
 #include "Templates/Function.h"
 
-namespace Metasound
+
+// Forward Declarations
+struct FMetaSoundFrontendDocumentBuilder;
+
+
+namespace Metasound::Frontend
 {
-	namespace Frontend
+	namespace DocumentTransform
 	{
-
-		namespace DocumentTransform
-		{
 #if WITH_EDITOR
-			using FGetNodeDisplayNameProjection = TFunction<FText(const FNodeHandle&)>;
-			using FGetNodeDisplayNameProjectionRef = TFunctionRef<FText(const FNodeHandle&)>;
+		using FGetNodeDisplayNameProjection = TFunction<FText(const FNodeHandle&)>;
+		using FGetNodeDisplayNameProjectionRef = TFunctionRef<FText(const FNodeHandle&)>;
 
-			METASOUNDFRONTEND_API void RegisterNodeDisplayNameProjection(FGetNodeDisplayNameProjection&& InNameProjection);
-			METASOUNDFRONTEND_API FGetNodeDisplayNameProjectionRef GetNodeDisplayNameProjection();
+		METASOUNDFRONTEND_API void RegisterNodeDisplayNameProjection(FGetNodeDisplayNameProjection&& InNameProjection);
+		METASOUNDFRONTEND_API FGetNodeDisplayNameProjectionRef GetNodeDisplayNameProjection();
 #endif // WITH_EDITOR
-		}
+	}
 
-		/** Interface for transforms applied to documents. */
-		class METASOUNDFRONTEND_API IDocumentTransform
-		{
-		public:
-			virtual ~IDocumentTransform() = default;
-
-			/** Return true if InDocument was modified, false otherwise. */
-			virtual bool Transform(FDocumentHandle InDocument) const = 0;
-
-			/** Return true if InDocument was modified, false otherwise.
-			  * This function is soft deprecated.  It is not pure virtual
-			  * to grandfather in old transform implementation. Old transforms
-			  * should be deprecated and rewritten to use the Controller-less
-			  * API in the interest of better performance and simplicity.
-			  */
-			virtual bool Transform(FMetasoundFrontendDocument& InOutDocument) const;
-		};
+	/** Interface for transforms applied to documents. */
+	class METASOUNDFRONTEND_API IDocumentTransform
+	{
+	public:
+		virtual ~IDocumentTransform() = default;
+
+		/** Return true if InDocument was modified, false otherwise. */
+		virtual bool Transform(FDocumentHandle InDocument) const = 0;
+
+		/** Return true if InDocument was modified, false otherwise.
+			* This function is soft deprecated.  It is not pure virtual
+			* to grandfather in old transform implementation. Old transforms
+			* should be deprecated and rewritten to use the Controller-less
+			* API in the interest of better performance and simplicity.
+			*/
+		virtual bool Transform(FMetasoundFrontendDocument& InOutDocument) const;
+	};
+
+	/** Interface for transforms applied to a graph. */
+	class METASOUNDFRONTEND_API IGraphTransform
+	{
+	public:
+		virtual ~IGraphTransform() = default;
 
-		/** Interface for transforms applied to a graph. */
-		class METASOUNDFRONTEND_API IGraphTransform
-		{
-		public:
-			virtual ~IGraphTransform() = default;
+		UE_DEPRECATED(5.3, "Deprecated: unused function which encouraged breaking const behavior of transform state.")
+		virtual FMetasoundFrontendDocument& GetOwningDocument() const { static FMetasoundFrontendDocument BaseDoc; return BaseDoc; };
 
-			UE_DEPRECATED(5.3, "Deprecated: unused function which encouraged breaking const behavior of transform state.")
-			virtual FMetasoundFrontendDocument& GetOwningDocument() const { static FMetasoundFrontendDocument BaseDoc; return BaseDoc; };
+		/** Return true if the graph was modified, false otherwise. */
+		virtual bool Transform(FMetasoundFrontendGraph& InOutGraph) const = 0;
+	};
 
-			/** Return true if the graph was modified, false otherwise. */
-			virtual bool Transform(FMetasoundFrontendGraph& InOutGraph) const = 0;
-		};
+	/** Interface for transforming a node. */
+	class METASOUNDFRONTEND_API INodeTransform
+	{
+	public:
+		virtual ~INodeTransform() = default;
 
-		/** Interface for transforming a node. */
-		class METASOUNDFRONTEND_API INodeTransform
-		{
-		public:
-			virtual ~INodeTransform() = default;
+		UE_DEPRECATED(5.3, "Deprecated: unused function which encouraged breaking const behavior of transform state.")
+		virtual FMetasoundFrontendDocument& GetOwningDocument() const { static FMetasoundFrontendDocument BaseDoc; return BaseDoc; };
 
-			UE_DEPRECATED(5.3, "Deprecated: unused function which encouraged breaking const behavior of transform state.")
-			virtual FMetasoundFrontendDocument& GetOwningDocument() const { static FMetasoundFrontendDocument BaseDoc; return BaseDoc; };
+		UE_DEPRECATED(5.3, "Deprecated: unused function which encouraged breaking const behavior of transform state.")
+		virtual FMetasoundFrontendGraph& GetOwningGraph() const { static FMetasoundFrontendGraph BaseGraph; return BaseGraph; }
 
-			UE_DEPRECATED(5.3, "Deprecated: unused function which encouraged breaking const behavior of transform state.")
-			virtual FMetasoundFrontendGraph& GetOwningGraph() const { static FMetasoundFrontendGraph BaseGraph; return BaseGraph; }
+		/** Return true if the node was modified, false otherwise. */
+		virtual bool Transform(const FMetasoundFrontendNode& InNode, FMetaSoundFrontendDocumentBuilder& OutBuilder) const;
 
-			/** Return true if the node was modified, false otherwise. */
-			virtual bool Transform(FMetasoundFrontendNode& InOutNode) const = 0;
-		};
+		UE_DEPRECATED(5.4, "Use transform which provides builder and const node representation to apply node transform to mutate underlying builder document")
+		virtual bool Transform(FMetasoundFrontendNode& InOutNode) const;
+	};
 
-		/** Adds or swaps document members (inputs, outputs) and removing any document members where necessary and adding those missing. */
-		class METASOUNDFRONTEND_API FModifyRootGraphInterfaces : public IDocumentTransform
-		{
-		public:
-			FModifyRootGraphInterfaces(const TArray<FMetasoundFrontendInterface>& InInterfacesToRemove, const TArray<FMetasoundFrontendInterface>& InInterfacesToAdd);
-			FModifyRootGraphInterfaces(const TArray<FMetasoundFrontendVersion>& InInterfaceVersionsToRemove, const TArray<FMetasoundFrontendVersion>& InInterfaceVersionsToAdd);
+	/** Adds or swaps document members (inputs, outputs) and removing any document members where necessary and adding those missing. */
+	class METASOUNDFRONTEND_API FModifyRootGraphInterfaces : public IDocumentTransform
+	{
+	public:
+		FModifyRootGraphInterfaces(const TArray<FMetasoundFrontendInterface>& InInterfacesToRemove, const TArray<FMetasoundFrontendInterface>& InInterfacesToAdd);
+		FModifyRootGraphInterfaces(const TArray<FMetasoundFrontendVersion>& InInterfaceVersionsToRemove, const TArray<FMetasoundFrontendVersion>& InInterfaceVersionsToAdd);
 
 #if WITH_EDITOR
-			// Whether or not to propagate node locations to new members. Setting to false
-			// results in members not having a default physical location in the editor graph.
-			void SetDefaultNodeLocations(bool bInSetDefaultNodeLocations);
+		// Whether or not to propagate node locations to new members. Setting to false
+		// results in members not having a default physical location in the editor graph.
+		void SetDefaultNodeLocations(bool bInSetDefaultNodeLocations);
 #endif // WITH_EDITOR
 
-			// Override function used to match removed members with added members, allowing
-			// transform to preserve connections made between removed interface members & new interface members
-			// that may be related but not be named the same.
-			void SetNamePairingFunction(const TFunction<bool(FName, FName)>& InNamePairingFunction);
+		// Override function used to match removed members with added members, allowing
+		// transform to preserve connections made between removed interface members & new interface members
+		// that may be related but not be named the same.
+		void SetNamePairingFunction(const TFunction<bool(FName, FName)>& InNamePairingFunction);
 
-			virtual bool Transform(FDocumentHandle InDocument) const override;
-			virtual bool Transform(FMetasoundFrontendDocument& InOutDocument) const override;
+		virtual bool Transform(FDocumentHandle InDocument) const override;
+		virtual bool Transform(FMetasoundFrontendDocument& InOutDocument) const override;
 
-		private:
-			bool AddMissingVertices(FGraphHandle GraphHandle) const;
-			void Init(const TFunction<bool(FName, FName)>* InNamePairingFunction = nullptr);
-			bool SwapPairedVertices(FGraphHandle GraphHandle) const;
-			bool RemoveUnsupportedVertices(FGraphHandle GraphHandle) const;
-			bool UpdateInterfacesInternal(FDocumentHandle DocumentHandle) const;
+	private:
+		bool AddMissingVertices(FGraphHandle GraphHandle) const;
+		void Init(const TFunction<bool(FName, FName)>* InNamePairingFunction = nullptr);
+		bool SwapPairedVertices(FGraphHandle GraphHandle) const;
+		bool RemoveUnsupportedVertices(FGraphHandle GraphHandle) const;
+		bool UpdateInterfacesInternal(FDocumentHandle DocumentHandle) const;
 
 #if WITH_EDITORONLY_DATA
-			void UpdateAddedVertexNodePositions(FGraphHandle GraphHandle) const;
+		void UpdateAddedVertexNodePositions(FGraphHandle GraphHandle) const;
 
-			bool bSetDefaultNodeLocations = true;
+		bool bSetDefaultNodeLocations = true;
 #endif // WITH_EDITORONLY_DATA
 
-			TArray<FMetasoundFrontendInterface> InterfacesToRemove;
-			TArray<FMetasoundFrontendInterface> InterfacesToAdd;
+		TArray<FMetasoundFrontendInterface> InterfacesToRemove;
+		TArray<FMetasoundFrontendInterface> InterfacesToAdd;
 
-			using FVertexPair = TTuple<FMetasoundFrontendClassVertex, FMetasoundFrontendClassVertex>;
-			TArray<FVertexPair> PairedInputs;
-			TArray<FVertexPair> PairedOutputs;
+		using FVertexPair = TTuple<FMetasoundFrontendClassVertex, FMetasoundFrontendClassVertex>;
+		TArray<FVertexPair> PairedInputs;
+		TArray<FVertexPair> PairedOutputs;
 
-			struct FInputData
-			{
-				FMetasoundFrontendClassInput Input;
-				const FMetasoundFrontendInterface* InputInterface = nullptr;
-			};
+		struct FInputData
+		{
+			FMetasoundFrontendClassInput Input;
+			const FMetasoundFrontendInterface* InputInterface = nullptr;
+		};
 
-			struct FOutputData
-			{
-				FMetasoundFrontendClassOutput Output;
-				const FMetasoundFrontendInterface* OutputInterface = nullptr;
-			};
+		struct FOutputData
+		{
+			FMetasoundFrontendClassOutput Output;
+			const FMetasoundFrontendInterface* OutputInterface = nullptr;
+		};
 
-			TArray<FInputData> InputsToAdd;
-			TArray<FMetasoundFrontendClassInput> InputsToRemove;
-			TArray<FOutputData> OutputsToAdd;
-			TArray<FMetasoundFrontendClassOutput> OutputsToRemove;
+		TArray<FInputData> InputsToAdd;
+		TArray<FMetasoundFrontendClassInput> InputsToRemove;
+		TArray<FOutputData> OutputsToAdd;
+		TArray<FMetasoundFrontendClassOutput> OutputsToRemove;
 
-		};
+	};
 
-		/** Updates document's given interface to the most recent version. */
-		class METASOUNDFRONTEND_API FUpdateRootGraphInterface : public IDocumentTransform
+	/** Updates document's given interface to the most recent version. */
+	class METASOUNDFRONTEND_API FUpdateRootGraphInterface : public IDocumentTransform
+	{
+	public:
+		FUpdateRootGraphInterface(const FMetasoundFrontendVersion& InInterfaceVersion, const FString& InOwningAssetName=FString(TEXT("Unknown")))
+			: InterfaceVersion(InInterfaceVersion)
+			, OwningAssetName(InOwningAssetName)
 		{
-		public:
-			FUpdateRootGraphInterface(const FMetasoundFrontendVersion& InInterfaceVersion, const FString& InOwningAssetName=FString(TEXT("Unknown")))
-				: InterfaceVersion(InInterfaceVersion)
-				, OwningAssetName(InOwningAssetName)
-			{
-			}
+		}
 
-			bool Transform(FDocumentHandle InDocument) const override;
+		bool Transform(FDocumentHandle InDocument) const override;
 
-		private:
-			void GetUpdatePathForDocument(const FMetasoundFrontendVersion& InCurrentVersion, const FMetasoundFrontendVersion& InTargetVersion, TArray<const IInterfaceRegistryEntry*>& OutUpgradePath) const;
-			bool UpdateDocumentInterface(const TArray<const IInterfaceRegistryEntry*>& InUpgradePath, FDocumentHandle InDocument) const;
+	private:
+		void GetUpdatePathForDocument(const FMetasoundFrontendVersion& InCurrentVersion, const FMetasoundFrontendVersion& InTargetVersion, TArray<const IInterfaceRegistryEntry*>& OutUpgradePath) const;
+		bool UpdateDocumentInterface(const TArray<const IInterfaceRegistryEntry*>& InUpgradePath, FDocumentHandle InDocument) const;
 
-			FMetasoundFrontendVersion InterfaceVersion;
-			FString OwningAssetName;
-		};
+		FMetasoundFrontendVersion InterfaceVersion;
+		FString OwningAssetName;
+	};
 
-		/** Completely rebuilds the graph connecting a preset's inputs to the reference
-		 * document's root graph. It maintains previously set input values entered upon 
-		 * the presets wrapping graph. */
-		class METASOUNDFRONTEND_API FRebuildPresetRootGraph : public IDocumentTransform
+	/** Completely rebuilds the graph connecting a preset's inputs to the reference
+		* document's root graph. It maintains previously set input values entered upon 
+		* the presets wrapping graph. */
+	class METASOUNDFRONTEND_API FRebuildPresetRootGraph : public IDocumentTransform
+	{
+	public:
+		/** Create transform.
+			* @param InReferenceDocument - The document containing the wrapped MetaSound graph.
+			*/
+		FRebuildPresetRootGraph(FConstDocumentHandle InReferencedDocument)
+			: ReferencedDocument(InReferencedDocument)
 		{
-		public:
-			/** Create transform.
-			 * @param InReferenceDocument - The document containing the wrapped MetaSound graph.
-			 */
-			FRebuildPresetRootGraph(FConstDocumentHandle InReferencedDocument)
-				: ReferencedDocument(InReferencedDocument)
-			{
-			}
+		}
 
-			FRebuildPresetRootGraph(const FMetasoundFrontendDocument& InReferencedDocument);
+		FRebuildPresetRootGraph(const FMetasoundFrontendDocument& InReferencedDocument);
 
-			virtual bool Transform(FDocumentHandle InDocument) const override;
-			virtual bool Transform(FMetasoundFrontendDocument& InOutDocument) const override;
+		virtual bool Transform(FDocumentHandle InDocument) const override;
+		virtual bool Transform(FMetasoundFrontendDocument& InOutDocument) const override;
 
-		private:
+	private:
 
-			// Get the class inputs needed for this preset. Input literals set on 
-			// the preset graph will be used if they are set and are marked as inheriting
-			// the default from the referenced graph.
-			TArray<FMetasoundFrontendClassInput> GenerateRequiredClassInputs(FDocumentHandle& InDocumentHandle, const FConstGraphHandle& InParentGraph, TSet<FName>& OutInputsInheritingDefault) const;
+		// Get the class inputs needed for this preset. Input literals set on 
+		// the preset graph will be used if they are set and are marked as inheriting
+		// the default from the referenced graph.
+		TArray<FMetasoundFrontendClassInput> GenerateRequiredClassInputs(FDocumentHandle& InDocumentHandle, const FConstGraphHandle& InParentGraph, TSet<FName>& OutInputsInheritingDefault) const;
 
-			// Get the class Outputs needed for this preset.
-			TArray<FMetasoundFrontendClassOutput> GenerateRequiredClassOutputs(FDocumentHandle& InDocumentHandle, const FConstGraphHandle& InParentGraph) const;
+		// Get the class Outputs needed for this preset.
+		TArray<FMetasoundFrontendClassOutput> GenerateRequiredClassOutputs(FDocumentHandle& InDocumentHandle, const FConstGraphHandle& InParentGraph) const;
 
-			// Add inputs to parent graph and connect to wrapped graph node.
-			void AddAndConnectInputs(const TArray<FMetasoundFrontendClassInput>& InClassInputs, FGraphHandle& InParentGraphHandle, FNodeHandle& InReferencedNode) const;
+		// Add inputs to parent graph and connect to wrapped graph node.
+		void AddAndConnectInputs(const TArray<FMetasoundFrontendClassInput>& InClassInputs, FGraphHandle& InParentGraphHandle, FNodeHandle& InReferencedNode) const;
 
-			// Add outputs to parent graph and connect to wrapped graph node.
-			void AddAndConnectOutputs(const TArray<FMetasoundFrontendClassOutput>& InClassOutputs, FGraphHandle& InParentGraphHandle, FNodeHandle& InReferencedNode) const;
+		// Add outputs to parent graph and connect to wrapped graph node.
+		void AddAndConnectOutputs(const TArray<FMetasoundFrontendClassOutput>& InClassOutputs, FGraphHandle& InParentGraphHandle, FNodeHandle& InReferencedNode) const;
 
-			FConstDocumentHandle ReferencedDocument = IDocumentController::GetInvalidHandle();
-		};
+		FConstDocumentHandle ReferencedDocument = IDocumentController::GetInvalidHandle();
+	};
 
-		/** Automatically updates all nodes and respective dependencies in graph where
-		  * newer versions exist in the loaded MetaSound Class Node Registry.
-		  */
-		class METASOUNDFRONTEND_API FAutoUpdateRootGraph : public IDocumentTransform
+	/** Automatically updates all nodes and respective dependencies in graph where
+		* newer versions exist in the loaded MetaSound Class Node Registry.
+		*/
+	class METASOUNDFRONTEND_API FAutoUpdateRootGraph : public IDocumentTransform
+	{
+	public:
+		/** Construct an AutoUpdate transform
+			*
+			* @param InDebugAssetPath - Asset path used for debug logs on warnings and errors.
+			* @param bInLogWarningOnDroppedConnections - If true, warnings will be logged if a node update results in a dropped connection.
+			*/
+		FAutoUpdateRootGraph(FString&& InDebugAssetPath, bool bInLogWarningOnDroppedConnection)
+			: DebugAssetPath(MoveTemp(InDebugAssetPath))
+			, bLogWarningOnDroppedConnection(bInLogWarningOnDroppedConnection)
 		{
-		public:
-			/** Construct an AutoUpdate transform
-			 *
-			 * @param InDebugAssetPath - Asset path used for debug logs on warnings and errors.
-			 * @param bInLogWarningOnDroppedConnections - If true, warnings will be logged if a node update results in a dropped connection.
-			 */
-			FAutoUpdateRootGraph(FString&& InDebugAssetPath, bool bInLogWarningOnDroppedConnection)
-				: DebugAssetPath(MoveTemp(InDebugAssetPath))
-				, bLogWarningOnDroppedConnection(bInLogWarningOnDroppedConnection)
-			{
-			}
-
-			bool Transform(FDocumentHandle InDocument) const override;
-
-		private:
-			const FString DebugAssetPath;
-			bool bLogWarningOnDroppedConnection;
-		};
+		}
+
+		bool Transform(FDocumentHandle InDocument) const override;
+
+	private:
+		const FString DebugAssetPath;
+		bool bLogWarningOnDroppedConnection;
+	};
 
-		/** Sets the document's graph class, optionally updating the namespace and variant. */
-		class METASOUNDFRONTEND_API FRenameRootGraphClass : public IDocumentTransform
+	/** Sets the document's graph class, optionally updating the namespace and variant. */
+	class METASOUNDFRONTEND_API FRenameRootGraphClass : public IDocumentTransform
+	{
+		const FMetasoundFrontendClassName NewClassName;
+
+	public:
+		/* Generates and assigns the document's root graph class a unique name using the provided guid as the name field and
+		 * (optionally) namespace and variant.
+		 * Returns true if transform succeeded, false if not.
+		*/
+		static bool Generate(FDocumentHandle InDocument, const FGuid& InGuid, const FName Namespace = { }, const FName Variant = { })
 		{
-			const FMetasoundFrontendClassName NewClassName;
-
-		public:
-			/* Generates and assigns the docment's root graph class a unique name using the provided guid as the name field and
-			 * (optionally) namespace and variant.
-			 * Returns true if transform succeeded, false if not.
-			 */
-			static bool Generate(FDocumentHandle InDocument, const FGuid& InGuid, const FName Namespace = { }, const FName Variant = { })
-			{
-				const FMetasoundFrontendClassName GeneratedClassName = { Namespace, *InGuid.ToString(), Variant };
-				return FRenameRootGraphClass(GeneratedClassName).Transform(InDocument);
-			}
-
-			static bool Generate(FMetasoundFrontendDocument& InDocument, const FGuid& InGuid, const FName Namespace = { }, const FName Variant = { })
-			{
-				const FMetasoundFrontendClassName GeneratedClassName = { Namespace, *InGuid.ToString(), Variant };
-				return FRenameRootGraphClass(GeneratedClassName).Transform(InDocument);
-			}
-
-			FRenameRootGraphClass(const FMetasoundFrontendClassName InClassName)
-				: NewClassName(InClassName)
-			{
-			}
-
-			bool Transform(FDocumentHandle InDocument) const override;
-			bool Transform(FMetasoundFrontendDocument& InOutDocument) const override;
-		};
-	}
-}
+			const FMetasoundFrontendClassName GeneratedClassName = { Namespace, *InGuid.ToString(), Variant };
+			return FRenameRootGraphClass(GeneratedClassName).Transform(InDocument);
+		}
+
+		static bool Generate(FMetasoundFrontendDocument& InDocument, const FGuid& InGuid, const FName Namespace = { }, const FName Variant = { })
+		{
+			const FMetasoundFrontendClassName GeneratedClassName = { Namespace, *InGuid.ToString(), Variant };
+			return FRenameRootGraphClass(GeneratedClassName).Transform(InDocument);
+		}
+
+		FRenameRootGraphClass(const FMetasoundFrontendClassName InClassName)
+			: NewClassName(InClassName)
+		{
+		}
+
+		bool Transform(FDocumentHandle InDocument) const override;
+		bool Transform(FMetasoundFrontendDocument& InOutDocument) const override;
+	};
+} // Metasound::Frontend
diff --git a/Engine/Plugins/Runtime/Metasound/Source/MetasoundFrontend/Public/NodeTemplates/MetasoundFrontendNodeTemplateReroute.h b/Engine/Plugins/Runtime/Metasound/Source/MetasoundFrontend/Public/NodeTemplates/MetasoundFrontendNodeTemplateReroute.h
--- a/Engine/Plugins/Runtime/Metasound/Source/MetasoundFrontend/Public/NodeTemplates/MetasoundFrontendNodeTemplateReroute.h
+++ b/Engine/Plugins/Runtime/Metasound/Source/MetasoundFrontend/Public/NodeTemplates/MetasoundFrontendNodeTemplateReroute.h
@@ -25,7 +25,10 @@ namespace Metasound
 
 			virtual ~FRerouteNodeTemplate() = default;
 
+			UE_DEPRECATED(5.4, "Use version that does not require mutating a provided PreprocessedDocument")
 			virtual TUniquePtr<INodeTransform> GenerateNodeTransform(FMetasoundFrontendDocument& InPreprocessedDocument) const override;
+
+			virtual TUniquePtr<INodeTransform> GenerateNodeTransform() const override;
 			virtual const FMetasoundFrontendClass& GetFrontendClass() const override;
 			virtual const FMetasoundFrontendVersion& GetVersion() const override;
 			virtual bool IsValidNodeInterface(const FMetasoundFrontendNodeInterface& InNodeInterface) const override;
