- Refactor node template transforms to support addition/removal of self (iterate & pass node ID instead of reference to avoid invalidating calling iterator)
- Misc clean-up to reroute transform to make clearer and more concise
- Add ability to override tooltip for template node warnings & place existing reroute specific warning in reroute transform
#rnx

diff --git a/Engine/Plugins/Runtime/Metasound/Source/MetasoundEditor/Private/MetasoundEditorGraphNode.cpp b/Engine/Plugins/Runtime/Metasound/Source/MetasoundEditor/Private/MetasoundEditorGraphNode.cpp
--- a/Engine/Plugins/Runtime/Metasound/Source/MetasoundEditor/Private/MetasoundEditorGraphNode.cpp
+++ b/Engine/Plugins/Runtime/Metasound/Source/MetasoundEditor/Private/MetasoundEditorGraphNode.cpp
@@ -863,9 +863,10 @@ void UMetasoundEditorGraphExternalNode::Validate(Metasound::Editor::FGraphNodeVa
 			else
 			{
 #if WITH_EDITOR
-				if (!Template->HasRequiredConnections(NodeHandle))
+				FString Message;
+				if (!Template->HasRequiredConnections(NodeHandle, &Message))
 				{
-					OutResult.SetMessage(EMessageSeverity::Warning, TEXT("Reroute node(s) missing non-reroute input connection(s)."));
+					OutResult.SetMessage(EMessageSeverity::Warning, Message);
 				}
 #endif // WITH_EDITOR
 			}
diff --git a/Engine/Plugins/Runtime/Metasound/Source/MetasoundFrontend/Private/MetasoundFrontendDocumentBuilder.cpp b/Engine/Plugins/Runtime/Metasound/Source/MetasoundFrontend/Private/MetasoundFrontendDocumentBuilder.cpp
--- a/Engine/Plugins/Runtime/Metasound/Source/MetasoundFrontend/Private/MetasoundFrontendDocumentBuilder.cpp
+++ b/Engine/Plugins/Runtime/Metasound/Source/MetasoundFrontend/Private/MetasoundFrontendDocumentBuilder.cpp
@@ -1941,59 +1941,23 @@ bool FMetaSoundFrontendDocumentBuilder::TransformTemplateNodes()
 		check(NodeTransform.IsValid());
 
 		TSet<const FMetasoundFrontendNode*> NodesToRemove;
-		TSet<TPair<FGuid, FGuid>> VerticesRemoved;
 
-		// 1a. Preprocess template nodes
-		for (const FMetasoundFrontendNode& Node : Graph.Nodes)
-		{
-			if (Params.ClassID == Node.ClassID)
-			{
-				NodeTransform->Transform(Node, *this);
-				NodesToRemove.Add(&Node);
-
-				auto GetNodeVertexGuidPair = [NodeID = Node.GetID()](const FMetasoundFrontendVertex& Vertex) { return TPair<FGuid, FGuid> { NodeID, Vertex.VertexID }; };
-				Algo::Transform(Node.Interface.Inputs, VerticesRemoved, GetNodeVertexGuidPair);
-				Algo::Transform(Node.Interface.Outputs, VerticesRemoved, GetNodeVertexGuidPair);
-
-				bModified = true;
-			}
-		}
+		// 2. Execute generated template node transform on copy of node array,
+		// which allows for addition/removal of nodes to/from original array container
+		// without template transform having to worry about mutation while iterating
+		TArray<FGuid> TemplateNodeIDs;
+		Algo::TransformIf(Graph.Nodes, TemplateNodeIDs,
+			[ClassID = Params.ClassID](const FMetasoundFrontendNode& Node) { return ClassID == Node.ClassID; },
+			[](const FMetasoundFrontendNode& Node) { return Node.GetID(); });
 
-		// 1b. Remove template node from graph
-		constexpr bool bAllowShrinking = false;
-		for (int32 i = Graph.Nodes.Num() - 1; i >= 0; --i)
+		for (const FGuid& NodeID : TemplateNodeIDs)
 		{
-			if (NodesToRemove.Contains(&Graph.Nodes[i]))
-			{
-				DocumentDelegates->NodeDelegates.OnRemoveSwappingNode.Broadcast(i, Graph.Nodes.Num() - 1);
-				Graph.Nodes.RemoveAtSwap(i, 1, bAllowShrinking);
-			}
-		}
-		Graph.Nodes.Shrink();
-
-		// 1c. Remove edges connecting template node from graph
-		for (int32 i = Graph.Edges.Num() - 1; i >= 0; --i)
-		{
-			const TPair<FGuid, FGuid> FromNodeVertexPair { Graph.Edges[i].FromNodeID, Graph.Edges[i].FromVertexID };
-			if (VerticesRemoved.Contains(FromNodeVertexPair))
-			{
-				DocumentDelegates->EdgeDelegates.OnRemoveSwappingEdge.Broadcast(i, GraphEdges.Num() - 1);
-				Graph.Edges.RemoveAtSwap(i, 1, bAllowShrinking);
-			}
-			else
-			{
-				const TPair<FGuid, FGuid> ToNodeVertexPair { Graph.Edges[i].ToNodeID, Graph.Edges[i].ToVertexID };
-				if (VerticesRemoved.Contains(ToNodeVertexPair))
-				{
-					DocumentDelegates->EdgeDelegates.OnRemoveSwappingEdge.Broadcast(i, GraphEdges.Num() - 1);
-					Graph.Edges.RemoveAtSwap(i, 1, bAllowShrinking);
-				}
-			}
+			bModified = true;
+			NodeTransform->Transform(NodeID, *this);
 		}
-		Graph.Edges.Shrink();
 	}
 
-	// 2. Finally, Remove template classes from dependency list
+	// 4. Remove template classes from dependency list
 	{
 		TSet<FString> TemplateKeys;
 		Algo::Transform(TemplateParams, TemplateKeys, [](const FTemplateParams& Params)
@@ -2015,6 +1979,7 @@ bool FMetaSoundFrontendDocumentBuilder::TransformTemplateNodes()
 				Dependencies.RemoveAtSwap(i, 1, bAllowShrinking);
 			}
 		}
+		Dependencies.Shrink();
 	}
 
 	return bModified;
diff --git a/Engine/Plugins/Runtime/Metasound/Source/MetasoundFrontend/Private/MetasoundFrontendTransform.cpp b/Engine/Plugins/Runtime/Metasound/Source/MetasoundFrontend/Private/MetasoundFrontendTransform.cpp
--- a/Engine/Plugins/Runtime/Metasound/Source/MetasoundFrontend/Private/MetasoundFrontendTransform.cpp
+++ b/Engine/Plugins/Runtime/Metasound/Source/MetasoundFrontend/Private/MetasoundFrontendTransform.cpp
@@ -43,7 +43,7 @@ namespace Metasound
 			return Transform(FDocumentController::CreateDocumentHandle(DocAccessPtr));
 		}
 
-		bool INodeTransform::Transform(const FMetasoundFrontendNode& InNode, FMetaSoundFrontendDocumentBuilder& OutBuilder) const
+		bool INodeTransform::Transform(const FGuid& InNodeID, FMetaSoundFrontendDocumentBuilder& OutBuilder) const
 		{
 			return false;
 		}
diff --git a/Engine/Plugins/Runtime/Metasound/Source/MetasoundFrontend/Private/NodeTemplates/MetasoundFrontendNodeTemplateReroute.cpp b/Engine/Plugins/Runtime/Metasound/Source/MetasoundFrontend/Private/NodeTemplates/MetasoundFrontendNodeTemplateReroute.cpp
--- a/Engine/Plugins/Runtime/Metasound/Source/MetasoundFrontend/Private/NodeTemplates/MetasoundFrontendNodeTemplateReroute.cpp
+++ b/Engine/Plugins/Runtime/Metasound/Source/MetasoundFrontend/Private/NodeTemplates/MetasoundFrontendNodeTemplateReroute.cpp
@@ -13,78 +13,81 @@ namespace Metasound::Frontend
 {
 	namespace ReroutePrivate
 	{
-		class FRerouteNodeTemplatePreprocessTransform : public INodeTransform
+		class FRerouteNodeTemplateTransform : public INodeTransform
 		{
 		public:
-			FRerouteNodeTemplatePreprocessTransform() = default;
-			virtual ~FRerouteNodeTemplatePreprocessTransform() = default;
+			FRerouteNodeTemplateTransform() = default;
+			virtual ~FRerouteNodeTemplateTransform() = default;
 
-			virtual bool Transform(const FMetasoundFrontendNode& InNode, FMetaSoundFrontendDocumentBuilder& OutBuilder) const override;
+			virtual bool Transform(const FGuid& InNodeID, FMetaSoundFrontendDocumentBuilder& OutBuilder) const override;
 		};
 
-		bool FRerouteNodeTemplatePreprocessTransform::Transform(const FMetasoundFrontendNode& InNode, FMetaSoundFrontendDocumentBuilder& OutBuilder) const
+		bool FRerouteNodeTemplateTransform::Transform(const FGuid& InNodeID, FMetaSoundFrontendDocumentBuilder& OutBuilder) const
 		{
-			using namespace ReroutePrivate;
-				
-			// Find the input and output edges for this node
-			const FMetasoundFrontendEdge* InputEdge = nullptr;
+			FMetasoundFrontendEdge InputEdge;
+			TArray<FMetasoundFrontendEdge> OutputEdges;
+
+			if (const FMetasoundFrontendNode* Node = OutBuilder.FindNode(InNodeID))
 			{
-				if (!ensure(InNode.Interface.Inputs.Num() == 1))
+				if (!ensureMsgf(Node->Interface.Inputs.Num() == 1, TEXT("Template nodes must only have one input")))
 				{
 					return false;
 				}
 
-				const FMetasoundFrontendVertex& InputVertex = InNode.Interface.Inputs.Last();
-				FMetasoundFrontendVertexHandle InputNodeVertexHandle { InNode.GetID(), InputVertex.VertexID };
-				TArray<const FMetasoundFrontendEdge*> InputEdges = OutBuilder.FindEdges(InNode.GetID(), InputVertex.VertexID);
-
-				// This can happen if the reroute node isn't provided an input, so its perfectly
-				// acceptable to just ignore this node as it ultimately provides no sourced input.
-				if (InputEdges.IsEmpty())
+				if (!ensureMsgf(Node->Interface.Outputs.Num() == 1, TEXT("Template nodes must only have one output")))
 				{
 					return false;
 				}
-				InputEdge = InputEdges.Last();
-			}
 
-			TArray<const FMetasoundFrontendEdge*> OutputEdges;
-			const FMetasoundFrontendVertex& OutputVertex = InNode.Interface.Outputs.Last();
-			const FMetasoundFrontendVertexHandle OutputVertexHandle { InNode.GetID(), OutputVertex.VertexID };
-			{
-				if (!ensure(InNode.Interface.Outputs.Num() == 1))
+				// Copy input edge to mutate from fields and avoid pointer going out of scope when template node is removed below
 				{
-					return false;
+					const FMetasoundFrontendVertex& InputVertex = Node->Interface.Inputs.Last();
+					TArray<const FMetasoundFrontendEdge*> InputEdges = OutBuilder.FindEdges(Node->GetID(), InputVertex.VertexID);
+					if (!InputEdges.IsEmpty())
+					{
+						InputEdge = *InputEdges.Last();
+					}
 				}
 
-				OutputEdges = OutBuilder.FindEdges(InNode.GetID(), OutputVertex.VertexID);
-
-				// This can happen if the reroute node isn't provided any outputs to connect to, so its
-				// perfectly acceptable to just ignore this node as it ultimately provides no sourced input.
-				if (OutputEdges.IsEmpty())
+				// Copy output edges to mutate from fields and avoid pointer going out of scope when swapping below
 				{
-					return false;
+					const FMetasoundFrontendVertex& OutputVertex = Node->Interface.Outputs.Last();
+					TArray<const FMetasoundFrontendEdge*> CurrentOutputEdges = OutBuilder.FindEdges(Node->GetID(), OutputVertex.VertexID);
+					Algo::Transform(CurrentOutputEdges, OutputEdges, [](const FMetasoundFrontendEdge* CurrentEdge)
+					{
+						check(CurrentEdge);
+						return *CurrentEdge;
+					});
 				}
-			}
 
-			// Update the output edges with the input edge 
-			FMetasoundFrontendVertexHandle NewOutputEdgeNodeVertexHandle { InputEdge->FromNodeID, InputEdge->FromVertexID };
-
-			for (const FMetasoundFrontendEdge* OutputEdge : OutputEdges)
+				// Remove the template node
+				OutBuilder.RemoveNode(Node->GetID());
+			}
+			else
 			{
-				check(OutputEdge);
+				const FMetaSoundFrontendDocumentBuilder& ConstBuilder = OutBuilder;
+				ensureMsgf(false, TEXT("Failed to find node with ID '%s' when template node transform was given a valid ID for builder '%s'."),
+					*InNodeID.ToString(),
+					*ConstBuilder.GetDocument().RootGraph.Metadata.GetClassName().GetFullName().ToString());
+			}
 
-				// Must copied before removal as pointer will become invalid
-				FMetasoundFrontendEdge Edge = *OutputEdge;
-				OutBuilder.RemoveEdge(Edge);
+			// Add new connections from reroute source node to reroute destination node. Either could be another reroute,
+			// which is valid because said node will subsequently get processed.
+			if (InputEdge.GetFromVertexHandle().IsSet())
+			{
+				for (FMetasoundFrontendEdge& OutputEdge : OutputEdges)
+				{
+					OutputEdge.FromNodeID = InputEdge.FromNodeID;
+					OutputEdge.FromVertexID = InputEdge.FromVertexID;
+					OutBuilder.AddEdge(MoveTemp(OutputEdge));
+				}
 
-				Edge.FromNodeID = InputEdge->FromNodeID;
-				Edge.FromVertexID = InputEdge->FromVertexID;
-				OutBuilder.AddEdge(MoveTemp(Edge));
+				return !OutputEdges.IsEmpty();
 			}
 
-			return true;
+			return false;
 		}
-	}
+	} // namespace ReroutePrivate
 
 	const FMetasoundFrontendClassName FRerouteNodeTemplate::ClassName { "UE", "Reroute", "" };
 
@@ -98,7 +101,7 @@ namespace Metasound::Frontend
 	TUniquePtr<INodeTransform> FRerouteNodeTemplate::GenerateNodeTransform() const
 	{
 		using namespace ReroutePrivate;
-		return TUniquePtr<INodeTransform>(new FRerouteNodeTemplatePreprocessTransform());
+		return TUniquePtr<INodeTransform>(new FRerouteNodeTemplateTransform());
 	}
 
 	const FMetasoundFrontendClass& FRerouteNodeTemplate::GetFrontendClass() const
@@ -160,7 +163,7 @@ namespace Metasound::Frontend
 	}
 
 #if WITH_EDITOR
-	bool FRerouteNodeTemplate::HasRequiredConnections(FConstNodeHandle InNodeHandle) const
+	bool FRerouteNodeTemplate::HasRequiredConnections(FConstNodeHandle InNodeHandle, FString* OutMessage) const
 	{
 		TArray<FConstOutputHandle> Outputs = InNodeHandle->GetConstOutputs();
 		TArray<FConstInputHandle> Inputs = InNodeHandle->GetConstInputs();
@@ -173,7 +176,13 @@ namespace Metasound::Frontend
 			return !Inputs.IsEmpty();
 		});
 
-		return bConnectedToNonRerouteOutputs || bConnectedToNonRerouteOutputs == bConnectedToNonRerouteInputs;
+		const bool bHasRequiredConnections = bConnectedToNonRerouteOutputs || bConnectedToNonRerouteOutputs == bConnectedToNonRerouteInputs;
+		if (!bHasRequiredConnections && OutMessage)
+		{
+			*OutMessage = TEXT("Reroute node(s) missing non-reroute input connection(s).");
+		}
+
+		return bHasRequiredConnections;
 	}
 #endif // WITH_EDITOR
 
diff --git a/Engine/Plugins/Runtime/Metasound/Source/MetasoundFrontend/Public/MetasoundFrontendNodeTemplateRegistry.h b/Engine/Plugins/Runtime/Metasound/Source/MetasoundFrontend/Public/MetasoundFrontendNodeTemplateRegistry.h
--- a/Engine/Plugins/Runtime/Metasound/Source/MetasoundFrontend/Public/MetasoundFrontendNodeTemplateRegistry.h
+++ b/Engine/Plugins/Runtime/Metasound/Source/MetasoundFrontend/Public/MetasoundFrontendNodeTemplateRegistry.h
@@ -43,7 +43,7 @@ namespace Metasound
 #if WITH_EDITOR
 			// Returns whether or not the given node template has the necessary
 			// required connections to be preprocessed (editor only).
-			virtual bool HasRequiredConnections(FConstNodeHandle InNodeHandle) const = 0;
+			virtual bool HasRequiredConnections(FConstNodeHandle InNodeHandle, FString* OutMessage = nullptr) const = 0;
 #endif // WITH_EDITOR
 		};
 
diff --git a/Engine/Plugins/Runtime/Metasound/Source/MetasoundFrontend/Public/MetasoundFrontendTransform.h b/Engine/Plugins/Runtime/Metasound/Source/MetasoundFrontend/Public/MetasoundFrontendTransform.h
--- a/Engine/Plugins/Runtime/Metasound/Source/MetasoundFrontend/Public/MetasoundFrontendTransform.h
+++ b/Engine/Plugins/Runtime/Metasound/Source/MetasoundFrontend/Public/MetasoundFrontendTransform.h
@@ -71,7 +71,7 @@ namespace Metasound::Frontend
 		virtual FMetasoundFrontendGraph& GetOwningGraph() const { static FMetasoundFrontendGraph BaseGraph; return BaseGraph; }
 
 		/** Return true if the node was modified, false otherwise. */
-		virtual bool Transform(const FMetasoundFrontendNode& InNode, FMetaSoundFrontendDocumentBuilder& OutBuilder) const;
+		virtual bool Transform(const FGuid& InNodeID, FMetaSoundFrontendDocumentBuilder& OutBuilder) const;
 
 		UE_DEPRECATED(5.4, "Use transform which provides builder and const node representation to apply node transform to mutate underlying builder document")
 		virtual bool Transform(FMetasoundFrontendNode& InOutNode) const;
diff --git a/Engine/Plugins/Runtime/Metasound/Source/MetasoundFrontend/Public/NodeTemplates/MetasoundFrontendNodeTemplateReroute.h b/Engine/Plugins/Runtime/Metasound/Source/MetasoundFrontend/Public/NodeTemplates/MetasoundFrontendNodeTemplateReroute.h
--- a/Engine/Plugins/Runtime/Metasound/Source/MetasoundFrontend/Public/NodeTemplates/MetasoundFrontendNodeTemplateReroute.h
+++ b/Engine/Plugins/Runtime/Metasound/Source/MetasoundFrontend/Public/NodeTemplates/MetasoundFrontendNodeTemplateReroute.h
@@ -34,7 +34,7 @@ namespace Metasound
 			virtual bool IsValidNodeInterface(const FMetasoundFrontendNodeInterface& InNodeInterface) const override;
 
 #if WITH_EDITOR
-			virtual bool HasRequiredConnections(FConstNodeHandle InNodeHandle) const override;
+			virtual bool HasRequiredConnections(FConstNodeHandle InNodeHandle, FString* OutMessage = nullptr) const override;
 #endif // WITH_EDITOR
 		};
 	} // namespace Frontend
