[MetasoundEditor] (Experimental code!) Show oscilloscope when doing mouse over in audio connections

#rnx

diff --git a/Engine/Plugins/Runtime/Metasound/Source/MetasoundEditor/Private/MetasoundEditorGraphBuilder.cpp b/Engine/Plugins/Runtime/Metasound/Source/MetasoundEditor/Private/MetasoundEditorGraphBuilder.cpp
--- a/Engine/Plugins/Runtime/Metasound/Source/MetasoundEditor/Private/MetasoundEditorGraphBuilder.cpp
+++ b/Engine/Plugins/Runtime/Metasound/Source/MetasoundEditor/Private/MetasoundEditorGraphBuilder.cpp
@@ -209,10 +209,14 @@ namespace Metasound
 				DataType = OutputHandle->GetDataType();
 			}
 
+			const UMetasoundEditorSettings* MetaSoundSettings = GetDefault<UMetasoundEditorSettings>();
+			const bool bIsDataTypeFAudioBufferSupported = DataType == GetMetasoundDataTypeName<FAudioBuffer>() && MetaSoundSettings && MetaSoundSettings->bShowOscilloscopeOnAudioPinMouseOver;
+
 			const bool bIsSupportedType = DataType == GetMetasoundDataTypeName<float>()
 				|| DataType == GetMetasoundDataTypeName<int32>()
 				|| DataType == GetMetasoundDataTypeName<FString>()
-				|| DataType == GetMetasoundDataTypeName<bool>();
+				|| DataType == GetMetasoundDataTypeName<bool>()
+				|| bIsDataTypeFAudioBufferSupported;
 
 			if (!bIsSupportedType)
 			{
diff --git a/Engine/Plugins/Runtime/Metasound/Source/MetasoundEditor/Private/MetasoundEditorGraphConnectionManager.cpp b/Engine/Plugins/Runtime/Metasound/Source/MetasoundEditor/Private/MetasoundEditorGraphConnectionManager.cpp
--- a/Engine/Plugins/Runtime/Metasound/Source/MetasoundEditor/Private/MetasoundEditorGraphConnectionManager.cpp
+++ b/Engine/Plugins/Runtime/Metasound/Source/MetasoundEditor/Private/MetasoundEditorGraphConnectionManager.cpp
@@ -1,10 +1,12 @@
 // Copyright Epic Games, Inc. All Rights Reserved.
 #include "MetasoundEditorGraphConnectionManager.h"
 
+#include "Analysis/MetasoundFrontendVertexAnalyzerAudioBuffer.h"
 #include "Analysis/MetasoundFrontendVertexAnalyzerEnvelopeFollower.h"
 #include "Analysis/MetasoundFrontendVertexAnalyzerForwardValue.h"
 #include "Analysis/MetasoundFrontendVertexAnalyzerTriggerDensity.h"
 #include "MetasoundEditorGraphBuilder.h"
+#include "MetasoundEditorSettings.h"
 #include "MetasoundVertex.h"
 
 
@@ -23,6 +25,7 @@ namespace Metasound
 			GraphAnalyzerView->AddAnalyzerForAllSupportedOutputs(FVertexAnalyzerForwardInt::GetAnalyzerName());
 			GraphAnalyzerView->AddAnalyzerForAllSupportedOutputs(FVertexAnalyzerForwardString::GetAnalyzerName());
 
+			GraphAnalyzerView->AddAnalyzerForAllSupportedOutputs(FVertexAnalyzerAudioBuffer::GetAnalyzerName());
 			GraphAnalyzerView->AddAnalyzerForAllSupportedOutputs(FVertexAnalyzerEnvelopeFollower::GetAnalyzerName());
 			GraphAnalyzerView->AddAnalyzerForAllSupportedOutputs(FVertexAnalyzerTriggerDensity::GetAnalyzerName());
 		}
@@ -73,6 +76,16 @@ namespace Metasound
 			Window.SetNumZeroed(InWindowSize);
 		}
 
+		void FGraphConnectionManager::TrackAudioPin(const FGuid& InNodeID, FVertexName InOutputName, FName InAnalyzerName, Audio::FPatchInput& InPatchInput)
+		{
+			TrackedAudioPins.Add({ InNodeID, InOutputName, InAnalyzerName }, MakeShared<Audio::FPatchInput>(InPatchInput));
+		}
+
+		void FGraphConnectionManager::UntrackAudioPin(const FGuid& InNodeID, FVertexName InOutputName, FName InAnalyzerName)
+		{
+			TrackedAudioPins.Remove({ InNodeID, InOutputName, InAnalyzerName });
+		}
+
 		void FGraphConnectionManager::Update(float InDeltaTime)
 		{
 			using namespace Frontend;
@@ -118,6 +131,23 @@ namespace Metasound
 				{
 					Pair.Value.Update(InDeltaTime);
 				}
+
+				// Audio Pins update
+				if (const UMetasoundEditorSettings* MetaSoundSettings = GetDefault<UMetasoundEditorSettings>();
+				    MetaSoundSettings && MetaSoundSettings->bShowOscilloscopeOnAudioPinMouseOver)
+				{
+					auto UpdateAudio = [this](const FAnalyzerAddress& InAnalyzerAddress, const FAudioBuffer& InAudioBuffer)
+					{
+						TSharedRef<Audio::FPatchInput>* PatchInput = TrackedAudioPins.Find({ InAnalyzerAddress.NodeID, InAnalyzerAddress.OutputName, InAnalyzerAddress.AnalyzerName });
+
+						if (PatchInput)
+						{
+							(*PatchInput)->PushAudio(InAudioBuffer.GetData(), InAudioBuffer.Num());
+						}
+					};
+				
+					UpdateConnections<FVertexAnalyzerAudioBuffer, FAudioBuffer>(InDeltaTime, UpdateAudio);
+				}
 			}
 		}
 
diff --git a/Engine/Plugins/Runtime/Metasound/Source/MetasoundEditor/Private/MetasoundEditorGraphConnectionManager.h b/Engine/Plugins/Runtime/Metasound/Source/MetasoundEditor/Private/MetasoundEditorGraphConnectionManager.h
--- a/Engine/Plugins/Runtime/Metasound/Source/MetasoundEditor/Private/MetasoundEditorGraphConnectionManager.h
+++ b/Engine/Plugins/Runtime/Metasound/Source/MetasoundEditor/Private/MetasoundEditorGraphConnectionManager.h
@@ -6,6 +6,7 @@
 #include "Analysis/MetasoundFrontendVertexAnalyzerForwardValue.h"
 #include "Components/AudioComponent.h"
 #include "DSP/Dsp.h"
+#include "DSP/MultithreadedPatching.h"
 #include "DSP/VolumeFader.h"
 #include "HAL/Platform.h"
 #include "MetasoundAssetBase.h"
@@ -99,6 +100,9 @@ namespace Metasound
 			// Retains existing values provided from update up to given size.
 			void TrackValue(const FGuid& InNodeID, FVertexName InOutputName, FName InAnalyzerName, int32 InWindowSize);
 
+			void TrackAudioPin(const FGuid& InNodeID, FVertexName InOutputName, FName InAnalyzerName, Audio::FPatchInput& InPatchInput);
+			void UntrackAudioPin(const FGuid& InNodeID, FVertexName InOutputName, FName InAnalyzerName);
+
 			void Update(float InDeltaTime);
 
 			template <typename TAnalyzerName>
@@ -212,7 +216,10 @@ namespace Metasound
 				}
 			};
 
+			struct FAudioBufferKey : FWindowValueKey{};
+
 			TMap<FWindowValueKey, FFloatMovingWindow> WindowedValues;
+			TMap<FAudioBufferKey, TSharedRef<Audio::FPatchInput>> TrackedAudioPins;
 
 			TWeakObjectPtr<const UAudioComponent> AudioComponent;
 			TMap<FString, Audio::FVolumeFader> ConnectionFaders;
diff --git a/Engine/Plugins/Runtime/Metasound/Source/MetasoundEditor/Private/MetasoundEditorSettings.h b/Engine/Plugins/Runtime/Metasound/Source/MetasoundEditor/Private/MetasoundEditorSettings.h
--- a/Engine/Plugins/Runtime/Metasound/Source/MetasoundEditor/Private/MetasoundEditorSettings.h
+++ b/Engine/Plugins/Runtime/Metasound/Source/MetasoundEditor/Private/MetasoundEditorSettings.h
@@ -154,6 +154,10 @@ public:
 	UPROPERTY(EditAnywhere, config, Category = GraphAnimation, meta = (ShowOnlyInnerProperties))
 	FMetasoundAnalyzerAnimationSettings AnalyzerAnimationSettings;
 
+	/** Whether to show a small oscilloscope when hovering audio pins. */
+	UPROPERTY(EditAnywhere, config, DisplayName = "Show Oscilloscope On Audio Pin Mouse Over", Category = Visualization)
+	bool bShowOscilloscopeOnAudioPinMouseOver = false;
+
 	/** Determines which details view to show in Metasounds Editor */
 	UPROPERTY(Transient)
 	EMetasoundActiveDetailView DetailView = EMetasoundActiveDetailView::Metasound;
diff --git a/Engine/Plugins/Runtime/Metasound/Source/MetasoundEditor/Private/MetasoundPinAudioInspector.cpp b/Engine/Plugins/Runtime/Metasound/Source/MetasoundEditor/Private/MetasoundPinAudioInspector.cpp
new file mode 100644
--- /dev/null
+++ b/Engine/Plugins/Runtime/Metasound/Source/MetasoundEditor/Private/MetasoundPinAudioInspector.cpp
@@ -0,0 +1,133 @@
+// Copyright Epic Games, Inc. All Rights Reserved.
+#include "MetasoundPinAudioInspector.h"
+
+#include "Analysis/MetasoundFrontendVertexAnalyzerAudioBuffer.h"
+#include "AudioBusSubsystem.h"
+#include "AudioMixerDevice.h"
+#include "AudioOscilloscope.h"
+#include "MetasoundEditor.h"
+#include "MetasoundEditorGraphBuilder.h"
+#include "MetasoundEditorGraphNode.h"
+
+#define LOCTEXT_NAMESPACE "MetasoundEditor"
+
+namespace Metasound
+{
+	namespace Editor
+	{
+		namespace FMetasoundPinAudioInspectorPrivate
+		{
+			UEdGraphPin* ResolvePinObjectAsOutput(UEdGraphPin* InPin)
+			{
+				if (InPin && !InPin->LinkedTo.IsEmpty())
+				{
+					// Swap to show connected output if input (Only ever one)
+					if (InPin->Direction == EGPD_Input)
+					{
+						InPin = InPin->LinkedTo.Last();
+						check(InPin->Direction == EGPD_Output);
+					}
+				}
+
+				return InPin;
+			}
+		}
+
+		FMetasoundPinAudioInspector::FMetasoundPinAudioInspector(FEdGraphPinReference InPinRef)
+			: GraphPinObj(FMetasoundPinAudioInspectorPrivate::ResolvePinObjectAsOutput(InPinRef.Get()))
+		{
+			// Initialize Oscilloscope
+			UWorld* EditorWorld = GEditor->GetEditorWorldContext().World();
+			check(EditorWorld);
+
+			constexpr int32 NumChannels = 1; // Audio wires are currently mono signals
+
+			Oscilloscope = MakeShared<AudioWidgets::FAudioOscilloscope>(EditorWorld,
+																		NumChannels,
+																		/*InTimeWindowMs*/     10.0f,
+																		/*InMaxTimeWindowMs*/  10.0f,
+																		/*InAnalysisPeriodMs*/ 10.0f,
+																		/*InPanelLayoutType*/  EAudioPanelLayoutType::Basic);
+
+			PinAudioInspectorWidget = SNew(SMetasoundPinAudioInspector)
+									  .VisualizationWidget(Oscilloscope->GetPanelWidget());
+
+			Oscilloscope->StartProcessing();
+
+			// Set PatchInput from Oscilloscope AudioBus
+			Audio::FMixerDevice* MixerDevice = static_cast<Audio::FMixerDevice*>(EditorWorld->GetAudioDeviceRaw());
+			check(MixerDevice);
+
+			UAudioBusSubsystem* AudioBusSubsystem = MixerDevice->GetSubsystem<UAudioBusSubsystem>();
+			check(AudioBusSubsystem);
+
+			const uint32 AudioBusId = Oscilloscope->GetAudioBus()->GetUniqueID();
+			PatchInput = AudioBusSubsystem->AddPatchInputForAudioBus(Audio::FAudioBusKey(AudioBusId), MixerDevice->GetNumOutputFrames(), NumChannels);
+
+			// Track Audio Pin
+			if (GraphPinObj && !GraphPinObj->LinkedTo.IsEmpty() && GraphPinObj->PinType.PinCategory == FGraphBuilder::PinCategoryAudio)
+			{
+				if (FGraphConnectionManager* ConnectionManager = GetConnectionManager())
+				{
+					const Frontend::FConstOutputHandle OutputHandle = FGraphBuilder::FindReroutedOutputHandleFromPin(GraphPinObj);
+					const FGuid NodeID = OutputHandle->GetOwningNodeID();
+					const FName OutputName = OutputHandle->GetName();
+					const FName AnalyzerName = Frontend::FVertexAnalyzerAudioBuffer::GetAnalyzerName();
+
+					ConnectionManager->TrackAudioPin(NodeID, OutputName, AnalyzerName, PatchInput);
+				}
+			}
+		}
+
+		FMetasoundPinAudioInspector::~FMetasoundPinAudioInspector()
+		{
+			// Untrack audio pin
+			if (FGraphConnectionManager* ConnectionManager = GetConnectionManager())
+			{
+				const Frontend::FConstOutputHandle OutputHandle = FGraphBuilder::FindReroutedOutputHandleFromPin(GraphPinObj);
+				const FGuid NodeID = OutputHandle->GetOwningNodeID();
+				const FName OutputName = OutputHandle->GetName();
+				const FName AnalyzerName = Frontend::FVertexAnalyzerAudioBuffer::GetAnalyzerName();
+
+				ConnectionManager->UntrackAudioPin(NodeID, OutputName, AnalyzerName);
+			}
+
+			// Stop AudioBus
+			if (UWorld* EditorWorld = GEditor->GetEditorWorldContext().World())
+			{
+				if (Audio::FMixerDevice* MixerDevice = static_cast<Audio::FMixerDevice*>(EditorWorld->GetAudioDeviceRaw()))
+				{
+					if (UAudioBusSubsystem* AudioBusSubsystem = MixerDevice->GetSubsystem<UAudioBusSubsystem>())
+					{
+						const uint32 AudioBusId = Oscilloscope->GetAudioBus()->GetUniqueID();
+						AudioBusSubsystem->StopAudioBus(Audio::FAudioBusKey(AudioBusId));
+					}
+				}
+			}
+		}
+
+		TSharedPtr<SMetasoundPinAudioInspector> FMetasoundPinAudioInspector::GetWidget()
+		{
+			return PinAudioInspectorWidget;
+		}
+
+		const UMetasoundEditorGraphNode& FMetasoundPinAudioInspector::GetReroutedNode() const
+		{
+			const UMetasoundEditorGraphNode* Node = FGraphBuilder::FindReroutedOutputPin(GraphPinObj) ? Cast<UMetasoundEditorGraphNode>(GraphPinObj->GetOwningNode()) 
+																									  : nullptr;
+			check(Node);
+
+			return *Node;
+		}
+
+		FGraphConnectionManager* FMetasoundPinAudioInspector::GetConnectionManager()
+		{
+			const UMetasoundEditorGraphNode& Node = GetReroutedNode();
+			TSharedPtr<FEditor> Editor = FGraphBuilder::GetEditorForNode(Node);
+
+			return Editor.IsValid() ? &Editor->GetConnectionManager() : nullptr;
+		}
+	} // namespace Editor
+} // namespace Metasound
+
+#undef LOCTEXT_NAMESPACE
diff --git a/Engine/Plugins/Runtime/Metasound/Source/MetasoundEditor/Private/MetasoundPinAudioInspector.h b/Engine/Plugins/Runtime/Metasound/Source/MetasoundEditor/Private/MetasoundPinAudioInspector.h
new file mode 100644
--- /dev/null
+++ b/Engine/Plugins/Runtime/Metasound/Source/MetasoundEditor/Private/MetasoundPinAudioInspector.h
@@ -0,0 +1,41 @@
+// Copyright Epic Games, Inc. All Rights Reserved.
+#pragma once
+
+#include "DSP/MultithreadedPatching.h"
+#include "EdGraph/EdGraphPin.h"
+#include "MetasoundEditorGraphConnectionManager.h"
+#include "SMetasoundPinAudioInspector.h"
+#include "Templates/SharedPointer.h"
+
+namespace AudioWidgets
+{
+	class FAudioOscilloscope;
+}
+
+class UMetasoundEditorGraphNode;
+
+namespace Metasound
+{
+	namespace Editor
+	{
+		class METASOUNDEDITOR_API FMetasoundPinAudioInspector
+		{
+		public:
+			explicit FMetasoundPinAudioInspector(FEdGraphPinReference InPinRef);
+			virtual ~FMetasoundPinAudioInspector();
+
+			TSharedPtr<SMetasoundPinAudioInspector> GetWidget();
+			
+		private:
+			const UMetasoundEditorGraphNode& GetReroutedNode() const;
+			FGraphConnectionManager* GetConnectionManager();
+
+			UEdGraphPin* GraphPinObj = nullptr;
+
+			TSharedPtr<AudioWidgets::FAudioOscilloscope> Oscilloscope = nullptr;
+			Audio::FPatchInput PatchInput;
+			
+			TSharedPtr<SMetasoundPinAudioInspector> PinAudioInspectorWidget = nullptr;
+		};
+	} // namespace Editor
+} // namespace Metasound
diff --git a/Engine/Plugins/Runtime/Metasound/Source/MetasoundEditor/Private/SMetasoundGraphPin.h b/Engine/Plugins/Runtime/Metasound/Source/MetasoundEditor/Private/SMetasoundGraphPin.h
--- a/Engine/Plugins/Runtime/Metasound/Source/MetasoundEditor/Private/SMetasoundGraphPin.h
+++ b/Engine/Plugins/Runtime/Metasound/Source/MetasoundEditor/Private/SMetasoundGraphPin.h
@@ -13,6 +13,7 @@
 #include "MetasoundEditorGraphMemberDefaults.h"
 #include "MetasoundEditorGraphNode.h"
 #include "MetasoundFrontendController.h"
+#include "MetasoundPinAudioInspector.h"
 #include "NodeTemplates/MetasoundFrontendNodeTemplateReroute.h"
 #include "SGraphPin.h"
 #include "SGraphNodeKnot.h"
@@ -38,6 +39,8 @@ namespace Metasound
 		{
 			TSharedPtr<SMetasoundPinValueInspector> PinInspector;
 
+			TSharedPtr<FMetasoundPinAudioInspector> PinAudioInspector;
+
 		protected:
 			static TWeakPtr<FPinValueInspectorTooltip> OpenPinInspector(UEdGraphPin& InPin, TSharedPtr<SMetasoundPinValueInspector>& OutPinInspector)
 			{
@@ -102,6 +105,62 @@ namespace Metasound
 				}
 			}
 
+			static TWeakPtr<FPinValueInspectorTooltip> OpenPinAudioInspector(UEdGraphPin& InPin, TSharedPtr<FMetasoundPinAudioInspector>& OutPinAudioInspector)
+			{
+				TSharedPtr<FMetasoundPinAudioInspector> NewPinAudioInspector = MakeShared<FMetasoundPinAudioInspector>(&InPin);
+				TWeakPtr<FPinValueInspectorTooltip> NewTooltip = FPinValueInspectorTooltip::SummonTooltip(&InPin, NewPinAudioInspector->GetWidget());
+
+				if (NewTooltip.IsValid())
+				{
+					OutPinAudioInspector = NewPinAudioInspector;
+					return NewTooltip;
+				}
+
+				return nullptr;
+			}
+
+			static void UpdatePinAudioInspector(UEdGraphPin& InPin,
+				const bool bIsHoveringPin,
+				TSharedPtr<FMetasoundPinAudioInspector>& OutPinAudioInspector,
+				TWeakPtr<FPinValueInspectorTooltip>& OutInspectorTooltip,
+				TFunctionRef<void(FVector2D&)> InGetTooltipLocation)
+			{
+				const bool bCanInspectPin = FGraphBuilder::CanInspectPin(&InPin);
+
+				if (bIsHoveringPin && bCanInspectPin)
+				{
+					if (!OutPinAudioInspector.IsValid())
+					{
+						OutInspectorTooltip = OpenPinAudioInspector(InPin, OutPinAudioInspector);
+						TSharedPtr<FPinValueInspectorTooltip> NewTooltip = OutInspectorTooltip.Pin();
+						if (NewTooltip.IsValid())
+						{
+							FVector2D TooltipLocation;
+							InGetTooltipLocation(TooltipLocation);
+							NewTooltip->MoveTooltip(TooltipLocation);
+						}
+					}
+				}
+				else if (OutPinAudioInspector.IsValid())
+				{
+					TSharedPtr<FPinValueInspectorTooltip> InspectorTooltip = OutInspectorTooltip.Pin();
+					if (InspectorTooltip.IsValid())
+					{
+						if (InspectorTooltip->TooltipCanClose())
+						{
+							constexpr bool bForceDismiss = true;
+							InspectorTooltip->TryDismissTooltip(bForceDismiss);
+							OutInspectorTooltip.Reset();
+							OutPinAudioInspector.Reset();
+						}
+					}
+					else
+					{
+						OutPinAudioInspector.Reset();
+					}
+				}
+			}
+
 			void CacheAccessType()
 			{
 				AccessType = EMetasoundFrontendVertexAccessType::Unset;
@@ -371,6 +430,32 @@ namespace Metasound
 						{
 							ParentPinType::GetInteractiveTooltipLocation(OutTooltipLocation);
 						});
+
+					// Audio Pin Inspector update
+					if (const UMetasoundEditorSettings* MetasoundSettings = GetDefault<UMetasoundEditorSettings>();
+					    MetasoundSettings && MetasoundSettings->bShowOscilloscopeOnAudioPinMouseOver)
+					{
+						FName DataType;
+						if (GraphPin->Direction == EGPD_Input)
+						{
+							const Metasound::Frontend::FConstInputHandle InputHandle = FGraphBuilder::GetConstInputHandleFromPin(GraphPin);
+							DataType = InputHandle->GetDataType();
+						}
+						else
+						{
+							const Metasound::Frontend::FConstOutputHandle OutputHandle = FGraphBuilder::GetConstOutputHandleFromPin(GraphPin);
+							DataType = OutputHandle->GetDataType();
+						}
+
+						if (DataType == GetMetasoundDataTypeName<FAudioBuffer>())
+						{
+							UpdatePinAudioInspector(*GraphPin, bIsHoveringPin, PinAudioInspector, ParentPinType::ValueInspectorTooltip,
+								[this](FVector2D& OutTooltipLocation)
+								{
+									ParentPinType::GetInteractiveTooltipLocation(OutTooltipLocation);
+								});
+						}
+					}
 				}
 			}
 		};
diff --git a/Engine/Plugins/Runtime/Metasound/Source/MetasoundEditor/Private/SMetasoundPinAudioInspector.cpp b/Engine/Plugins/Runtime/Metasound/Source/MetasoundEditor/Private/SMetasoundPinAudioInspector.cpp
new file mode 100644
--- /dev/null
+++ b/Engine/Plugins/Runtime/Metasound/Source/MetasoundEditor/Private/SMetasoundPinAudioInspector.cpp
@@ -0,0 +1,30 @@
+// Copyright Epic Games, Inc. All Rights Reserved.
+#include "SMetasoundPinAudioInspector.h"
+
+#include "Widgets/Layout/SBox.h"
+
+#define LOCTEXT_NAMESPACE "MetasoundEditor"
+
+namespace Metasound
+{
+	namespace Editor
+	{
+		void SMetasoundPinAudioInspector::Construct(const FArguments& InArgs)
+		{
+			SPinValueInspector::FArguments InspectorArgs;
+			SPinValueInspector::Construct(InspectorArgs);
+
+			ChildSlot
+			[
+				SNew(SBox)
+				.MaxDesiredWidth(250)
+				.MaxDesiredHeight(200)
+				[
+					InArgs._VisualizationWidget
+				]
+			];
+		}
+	} // namespace Editor
+} // namespace Metasound
+
+#undef LOCTEXT_NAMESPACE
diff --git a/Engine/Plugins/Runtime/Metasound/Source/MetasoundEditor/Private/SMetasoundPinAudioInspector.h b/Engine/Plugins/Runtime/Metasound/Source/MetasoundEditor/Private/SMetasoundPinAudioInspector.h
new file mode 100644
--- /dev/null
+++ b/Engine/Plugins/Runtime/Metasound/Source/MetasoundEditor/Private/SMetasoundPinAudioInspector.h
@@ -0,0 +1,27 @@
+// Copyright Epic Games, Inc. All Rights Reserved.
+#pragma once
+
+#include "SPinValueInspector.h"
+#include "Templates/SharedPointer.h"
+#include "Widgets/SNullWidget.h"
+#include "Widgets/SWidget.h"
+
+namespace Metasound
+{
+	namespace Editor
+	{
+		class METASOUNDEDITOR_API SMetasoundPinAudioInspector : public SPinValueInspector
+		{
+
+		public:
+			SLATE_BEGIN_ARGS(SMetasoundPinAudioInspector)
+			: _VisualizationWidget(SNullWidget::NullWidget)
+			{}
+
+			SLATE_ARGUMENT(TSharedRef<SWidget>, VisualizationWidget)
+			SLATE_END_ARGS()
+
+			void Construct(const FArguments& InArgs);
+		};
+	} // namespace Editor
+} // namespace Metasound
diff --git a/Engine/Plugins/Runtime/Metasound/Source/MetasoundEngine/Private/MetasoundEngineModule.cpp b/Engine/Plugins/Runtime/Metasound/Source/MetasoundEngine/Private/MetasoundEngineModule.cpp
--- a/Engine/Plugins/Runtime/Metasound/Source/MetasoundEngine/Private/MetasoundEngineModule.cpp
+++ b/Engine/Plugins/Runtime/Metasound/Source/MetasoundEngine/Private/MetasoundEngineModule.cpp
@@ -13,6 +13,7 @@
 #include "MetasoundWave.h"
 #include "MetasoundWaveTable.h"
 #include "Analysis/MetasoundFrontendAnalyzerRegistry.h"
+#include "Analysis/MetasoundFrontendVertexAnalyzerAudioBuffer.h"
 #include "Analysis/MetasoundFrontendVertexAnalyzerEnvelopeFollower.h"
 #include "Analysis/MetasoundFrontendVertexAnalyzerForwardValue.h"
 #include "Analysis/MetasoundFrontendVertexAnalyzerTriggerDensity.h"
@@ -65,12 +66,13 @@ class FMetasoundEngineModule : public IMetasoundEngineModule
 		FMetasoundFrontendRegistryContainer::Get()->RegisterPendingNodes();
 
 		// Register Analyzers
+		METASOUND_REGISTER_VERTEX_ANALYZER_FACTORY(Frontend::FVertexAnalyzerAudioBuffer)
 		METASOUND_REGISTER_VERTEX_ANALYZER_FACTORY(Frontend::FVertexAnalyzerEnvelopeFollower)
-		METASOUND_REGISTER_VERTEX_ANALYZER_FACTORY(Frontend::FVertexAnalyzerTriggerDensity)
 		METASOUND_REGISTER_VERTEX_ANALYZER_FACTORY(Frontend::FVertexAnalyzerForwardBool)
 		METASOUND_REGISTER_VERTEX_ANALYZER_FACTORY(Frontend::FVertexAnalyzerForwardFloat)
 		METASOUND_REGISTER_VERTEX_ANALYZER_FACTORY(Frontend::FVertexAnalyzerForwardInt)
 		METASOUND_REGISTER_VERTEX_ANALYZER_FACTORY(Frontend::FVertexAnalyzerForwardString)
+		METASOUND_REGISTER_VERTEX_ANALYZER_FACTORY(Frontend::FVertexAnalyzerTriggerDensity)
 		METASOUND_REGISTER_VERTEX_ANALYZER_FACTORY(Frontend::FVertexAnalyzerTriggerToTime)
 
 		// Register passthrough output analyzers
diff --git a/Engine/Plugins/Runtime/Metasound/Source/MetasoundFrontend/Private/Analysis/MetasoundFrontendVertexAnalyzerAudioBuffer.cpp b/Engine/Plugins/Runtime/Metasound/Source/MetasoundFrontend/Private/Analysis/MetasoundFrontendVertexAnalyzerAudioBuffer.cpp
new file mode 100644
--- /dev/null
+++ b/Engine/Plugins/Runtime/Metasound/Source/MetasoundFrontend/Private/Analysis/MetasoundFrontendVertexAnalyzerAudioBuffer.cpp
@@ -0,0 +1,42 @@
+// Copyright Epic Games, Inc. All Rights Reserved.
+
+#include "Analysis/MetasoundFrontendVertexAnalyzerAudioBuffer.h"
+
+#include "MetasoundAudioBuffer.h"
+#include "MetasoundDataReference.h"
+
+
+namespace Metasound
+{
+	namespace Frontend
+	{
+		const FAnalyzerOutput& FVertexAnalyzerAudioBuffer::FOutputs::GetValue()
+		{
+			static FAnalyzerOutput Value = { "AudioBuffer", GetMetasoundDataTypeName<FAudioBuffer>() };
+			return Value;
+		}
+
+		const FName& FVertexAnalyzerAudioBuffer::GetAnalyzerName()
+		{
+			static const FName AnalyzerName = "UE.Audio.AudioBuffer"; return AnalyzerName;
+		}
+
+		const FName& FVertexAnalyzerAudioBuffer::GetDataType()
+		{
+			return GetMetasoundDataTypeName<FAudioBuffer>();
+		}
+
+		FVertexAnalyzerAudioBuffer::FVertexAnalyzerAudioBuffer(const FCreateAnalyzerParams& InParams)
+			: FVertexAnalyzerBase(InParams.AnalyzerAddress, InParams.VertexDataReference)
+			, AudioBuffer(TDataWriteReference<FAudioBuffer>::CreateNew())
+		{
+			FVertexAnalyzerBase::BindOutputData<FAudioBuffer>(FOutputs::GetValue().Name, InParams.OperatorSettings, FAudioBufferReadRef(AudioBuffer));
+		}
+
+		void FVertexAnalyzerAudioBuffer::Execute()
+		{
+			*AudioBuffer = GetVertexData<FAudioBuffer>();
+			MarkOutputDirty();
+		}
+	} // namespace Frontend
+} // namespace Metasound
diff --git a/Engine/Plugins/Runtime/Metasound/Source/MetasoundFrontend/Public/Analysis/MetasoundFrontendVertexAnalyzerAudioBuffer.h b/Engine/Plugins/Runtime/Metasound/Source/MetasoundFrontend/Public/Analysis/MetasoundFrontendVertexAnalyzerAudioBuffer.h
new file mode 100644
--- /dev/null
+++ b/Engine/Plugins/Runtime/Metasound/Source/MetasoundFrontend/Public/Analysis/MetasoundFrontendVertexAnalyzerAudioBuffer.h
@@ -0,0 +1,44 @@
+// Copyright Epic Games, Inc. All Rights Reserved.
+#pragma once
+
+#include "Analysis/MetasoundFrontendAnalyzerFactory.h"
+#include "Analysis/MetasoundFrontendVertexAnalyzer.h"
+#include "Containers/Array.h"
+#include "MetasoundDataReferenceCollection.h"
+
+
+namespace Metasound
+{
+	namespace Frontend
+	{
+		class METASOUNDFRONTEND_API FVertexAnalyzerAudioBuffer : public FVertexAnalyzerBase
+		{
+		public:
+			static const FName& GetAnalyzerName();
+			static const FName& GetDataType();
+
+			struct METASOUNDFRONTEND_API FOutputs
+			{
+				static const FAnalyzerOutput& GetValue();
+			};
+
+			class METASOUNDFRONTEND_API FFactory : public TVertexAnalyzerFactory<FVertexAnalyzerAudioBuffer>
+			{
+			public:
+				virtual const TArray<FAnalyzerOutput>& GetAnalyzerOutputs() const override
+				{
+					static const TArray<FAnalyzerOutput> Outputs { FOutputs::GetValue() };
+					return Outputs;
+				}
+			};
+
+			FVertexAnalyzerAudioBuffer(const FCreateAnalyzerParams& InParams);
+			virtual ~FVertexAnalyzerAudioBuffer() = default;
+
+			virtual void Execute() override;
+
+		private:
+			FAudioBufferWriteRef AudioBuffer;
+		};
+	} // namespace Frontend
+} // namespace Metasound
