- Fix for crash when transforming reroute template nodes
- Add more extensible API to transforms for opting in or out of applying access type dynamically for inputs/outputs to builder API
- Added debug cvar for viewing node data in MetaSound editor
#rnx

diff --git a/Engine/Plugins/Runtime/Metasound/Source/MetasoundEditor/Private/MetasoundEditorGraphNode.cpp b/Engine/Plugins/Runtime/Metasound/Source/MetasoundEditor/Private/MetasoundEditorGraphNode.cpp
--- a/Engine/Plugins/Runtime/Metasound/Source/MetasoundEditor/Private/MetasoundEditorGraphNode.cpp
+++ b/Engine/Plugins/Runtime/Metasound/Source/MetasoundEditor/Private/MetasoundEditorGraphNode.cpp
@@ -7,6 +7,7 @@
 #include "Engine/Font.h"
 #include "Framework/Commands/GenericCommands.h"
 #include "GraphEditorActions.h"
+#include "HAL/IConsoleManager.h"
 #include "Logging/TokenizedMessage.h"
 #include "Metasound.h"
 #include "MetasoundAssetManager.h"
@@ -35,19 +36,24 @@
 
 #define LOCTEXT_NAMESPACE "MetaSoundEditor"
 
-namespace Metasound
+namespace Metasound::Editor
 {
-	namespace Editor
+	namespace GraphNodePrivate
 	{
-		namespace GraphNodePrivate
-		{
-			static const FString MissingConcreteOutputConnectionFormat = TEXT(
-				"Reroute connection for pin '{0}' does not provide a concrete output. "
-				"Resulting literal value is undefined and may result in unintended results."
-			);
-		}
-	} 
-}
+		static const FString MissingConcreteOutputConnectionFormat = TEXT(
+			"Reroute connection for pin '{0}' does not provide a concrete output. "
+			"Resulting literal value is undefined and may result in unintended results."
+		);
+
+		int32 ShowNodeDebugData = 0;
+		FAutoConsoleVariableRef CVarShowNodeDebugData(
+			TEXT("au.MetaSound.Editor.Debug.ShowNodeDebugData"),
+			ShowNodeDebugData,
+			TEXT("If enabled, shows debug data such as node IDs, vertex IDs, vertex names, and class names when hovering over node titles and pins in the MetaSound asset editor.\n")
+			TEXT("0: Disabled (default), !0: Enabled"),
+			ECVF_Default);
+	} // GraphNodePrivate
+} // Metasound::Editor
 
 UMetasoundEditorGraphNode::UMetasoundEditorGraphNode(const FObjectInitializer& ObjectInitializer)
 	: Super(ObjectInitializer)
@@ -347,11 +353,32 @@ void UMetasoundEditorGraphNode::GetPinHoverText(const UEdGraphPin& Pin, FString&
 		{
 			OutHoverText = InputHandle->GetTooltip().ToString();
 		}
+
+		if (GraphNodePrivate::ShowNodeDebugData)
+		{
+			OutHoverText = FString::Format(TEXT("{0}\nVertex Name: {1}\nDataType: {2}\nID: {3}"),
+			{
+				OutHoverText,
+				InputHandle->GetName().ToString(),
+				InputHandle->GetDataType().ToString(),
+				InputHandle->GetID().ToString(),
+			});
+		}
 	}
 	else // Pin.Direction == EGPD_Output
 	{
 		FConstOutputHandle OutputHandle = FGraphBuilder::FindReroutedConstOutputHandleFromPin(&Pin);
 		OutHoverText = OutputHandle->GetTooltip().ToString();
+		if (GraphNodePrivate::ShowNodeDebugData)
+		{
+			OutHoverText = FString::Format(TEXT("{0}\nVertex Name: {1}\nDataType: {2}\nID: {3}"),
+			{
+				OutHoverText,
+				OutputHandle->GetName().ToString(),
+				OutputHandle->GetDataType().ToString(),
+				OutputHandle->GetID().ToString(),
+			});
+		}
 	}
 }
 
@@ -556,7 +583,20 @@ void UMetasoundEditorGraphNode::GetNodeContextMenuActions(UToolMenu* Menu, UGrap
 
 FText UMetasoundEditorGraphNode::GetTooltipText() const
 {
-	return GetConstNodeHandle()->GetDescription();
+	using namespace Metasound::Editor;
+	using namespace Metasound::Frontend;
+
+	FConstNodeHandle Node = GetConstNodeHandle();
+	FText Description = Node->GetDescription();
+	if (GraphNodePrivate::ShowNodeDebugData)
+	{
+		Description = FText::Format(LOCTEXT("Metasound_DebugNodeTooltipText", "{0}\nClass Name: {1}\nNode ID: {2}"),
+			Description,
+			FText::FromString(Node->GetClassMetadata().GetClassName().ToString()),
+			FText::FromString(Node->GetID().ToString())
+		);
+	}
+	return Description;
 }
 
 FText UMetasoundEditorGraphNode::GetDisplayName() const
diff --git a/Engine/Plugins/Runtime/Metasound/Source/MetasoundFrontend/Private/MetasoundFrontendDocumentBuilder.cpp b/Engine/Plugins/Runtime/Metasound/Source/MetasoundFrontend/Private/MetasoundFrontendDocumentBuilder.cpp
--- a/Engine/Plugins/Runtime/Metasound/Source/MetasoundFrontend/Private/MetasoundFrontendDocumentBuilder.cpp
+++ b/Engine/Plugins/Runtime/Metasound/Source/MetasoundFrontend/Private/MetasoundFrontendDocumentBuilder.cpp
@@ -1477,7 +1477,7 @@ const FMetasoundFrontendVertex* FMetaSoundFrontendDocumentBuilder::FindNodeInput
 	return NodeCache.FindInputVertex(InNodeID, InVertexName);
 }
 
-const FMetasoundFrontendClassInput* FMetaSoundFrontendDocumentBuilder::FindNodeInputClassInput(const FGuid& InNodeID, const FGuid& InVertexID) const
+const FMetasoundFrontendClassInput* FMetaSoundFrontendDocumentBuilder::FindNodeInputClassInput(const FGuid& InNodeID, const FName& InVertexName) const
 {
 	using namespace Metasound::Frontend;
 
@@ -1488,8 +1488,12 @@ const FMetasoundFrontendClassInput* FMetaSoundFrontendDocumentBuilder::FindNodeI
 	{
 		if (const FMetasoundFrontendClass* Class = DocumentCache->FindDependency(Node->ClassID))
 		{
-			auto InputMatchesVertex = [&InVertexID](const FMetasoundFrontendClassVertex& ClassVertex) { return ClassVertex.VertexID == InVertexID; };
-			return Class->Interface.Inputs.FindByPredicate(InputMatchesVertex);
+			if (const FMetasoundFrontendVertex* Vertex = NodeCache.FindInputVertex(InNodeID, InVertexName))
+			{
+				const FName VertexName = Vertex->Name;
+				auto InputMatchesVertexName = [&VertexName](const FMetasoundFrontendClassVertex& ClassVertex) { return ClassVertex.Name == VertexName; };
+				return Class->Interface.Inputs.FindByPredicate(InputMatchesVertexName);
+			}
 		}
 	}
 
@@ -1538,7 +1542,7 @@ const FMetasoundFrontendVertex* FMetaSoundFrontendDocumentBuilder::FindNodeOutpu
 	return NodeCache.FindOutputVertex(InNodeID, InVertexName);
 }
 
-const FMetasoundFrontendClassOutput* FMetaSoundFrontendDocumentBuilder::FindNodeOutputClassOutput(const FGuid& InNodeID, const FGuid& InVertexID) const
+const FMetasoundFrontendClassOutput* FMetaSoundFrontendDocumentBuilder::FindNodeOutputClassOutput(const FGuid& InNodeID, const FName& InVertexName) const
 {
 	using namespace Metasound::Frontend;
 
@@ -1549,8 +1553,12 @@ const FMetasoundFrontendClassOutput* FMetaSoundFrontendDocumentBuilder::FindNode
 	{
 		if (const FMetasoundFrontendClass* Class = DocumentCache->FindDependency(Node->ClassID))
 		{
-			auto OutputMatchesVertex = [&InVertexID](const FMetasoundFrontendClassVertex& ClassVertex) { return ClassVertex.VertexID == InVertexID; };
-			return Class->Interface.Outputs.FindByPredicate(OutputMatchesVertex);
+			if (const FMetasoundFrontendVertex* Vertex = NodeCache.FindOutputVertex(InNodeID, InVertexName))
+			{
+				const FName VertexName = Vertex->Name;
+				auto OutputMatchesVertexName = [&VertexName](const FMetasoundFrontendClassVertex& ClassVertex) { return ClassVertex.Name == VertexName; };
+				return Class->Interface.Outputs.FindByPredicate(OutputMatchesVertexName);
+			}
 		}
 	}
 
@@ -1593,6 +1601,21 @@ const FTopLevelAssetPath FMetaSoundFrontendDocumentBuilder::GetBuilderClassPath(
 	return Interface->GetBaseMetaSoundUClass().GetClassPathName();
 }
 
+const FString FMetaSoundFrontendDocumentBuilder::GetDebugName() const
+{
+	using namespace Metasound::Frontend;
+
+	const FMetasoundFrontendClassMetadata& Metadata = GetDocument().RootGraph.Metadata;
+	const FMetasoundFrontendVersionNumber& Version = Metadata.GetVersion();
+	const FNodeRegistryKey RegKey = NodeRegistryKey::CreateKey(EMetasoundFrontendClassType::External, Metadata.GetClassName().ToString(), Version.Major, Version.Minor);
+	if (const FMetasoundAssetBase* Asset = IMetaSoundAssetManager::GetChecked().TryLoadAssetFromKey(RegKey))
+	{
+		return Asset->GetOwningAssetName();
+	}
+
+	return Metadata.GetClassName().ToString();
+}
+
 FMetasoundFrontendDocument& FMetaSoundFrontendDocumentBuilder::GetDocument()
 {
 	IMetaSoundDocumentInterface* Interface = DocumentInterface.GetInterface();
@@ -1617,6 +1640,86 @@ const IMetaSoundDocumentInterface& FMetaSoundFrontendDocumentBuilder::GetDocumen
 	return *Interface;
 }
 
+EMetasoundFrontendVertexAccessType FMetaSoundFrontendDocumentBuilder::GetNodeInputAccessType(const FGuid& InNodeID, const FGuid& InVertexID) const
+{
+	using namespace Metasound::Frontend;
+
+	const IDocumentGraphNodeCache& NodeCache = DocumentCache->GetNodeCache();
+	if (const int32* NodeIndex = NodeCache.FindNodeIndex(InNodeID))
+	{
+		const FMetasoundFrontendGraph& Graph = GetDocument().RootGraph.Graph;
+		const FMetasoundFrontendNode& Node = Graph.Nodes[*NodeIndex];
+		auto IsVertexID = [&InVertexID](const FMetasoundFrontendVertex& Vertex) { return Vertex.VertexID == InVertexID; };
+		if (const FMetasoundFrontendClass* Class = DocumentCache->FindDependency(Node.ClassID))
+		{
+			const EMetasoundFrontendClassType ClassType = Class->Metadata.GetType();
+			switch (ClassType)
+			{
+				case EMetasoundFrontendClassType::Template:
+				{
+					const FNodeRegistryKey Key = NodeRegistryKey::CreateKey(Class->Metadata);
+					const INodeTemplate* Template = INodeTemplateRegistry::Get().FindTemplate(Key);
+					if (ensureMsgf(Template, TEXT("Failed to find MetaSound node template registered with key '%s'"), *Key))
+					{
+						if (Template->IsInputAccessTypeDynamic())
+						{
+							return Template->GetNodeInputAccessType(*this, InNodeID, InVertexID);
+						}
+					}
+				}
+				break;
+
+				case EMetasoundFrontendClassType::Output:
+				{
+					const FMetasoundFrontendClassInput& ClassInput = Class->Interface.Inputs.Last();
+					return ClassInput.AccessType;
+				}
+
+				default:
+				break;
+			}
+			static_assert(static_cast<uint32>(EMetasoundFrontendClassType::Invalid) == 10, "Potential missing case coverage for EMetasoundFrontendClassType");
+
+			if (const FMetasoundFrontendVertex* Vertex = Node.Interface.Inputs.FindByPredicate(IsVertexID))
+			{
+				auto IsClassInput = [VertexName = Vertex->Name](const FMetasoundFrontendClassInput& Input) { return Input.Name == VertexName; };
+				if (const FMetasoundFrontendClassInput* ClassInput = Class->Interface.Inputs.FindByPredicate(IsClassInput))
+				{
+					return ClassInput->AccessType;
+				}
+			}
+		}
+	}
+
+	return EMetasoundFrontendVertexAccessType::Unset;
+}
+
+const FMetasoundFrontendLiteral* FMetaSoundFrontendDocumentBuilder::GetNodeInputClassDefault(const FGuid& InNodeID, const FGuid& InVertexID) const
+{
+	using namespace Metasound::Frontend;
+
+	const IDocumentGraphNodeCache& NodeCache = DocumentCache->GetNodeCache();
+	if (const int32* NodeIndex = NodeCache.FindNodeIndex(InNodeID))
+	{
+		const FMetasoundFrontendGraph& Graph = GetDocument().RootGraph.Graph;
+		const FMetasoundFrontendNode& Node = Graph.Nodes[*NodeIndex];
+		auto IsVertexID = [&InVertexID](const FMetasoundFrontendVertex& Vertex) { return Vertex.VertexID == InVertexID; };
+		if (const FMetasoundFrontendVertex* Vertex = Node.Interface.Inputs.FindByPredicate(IsVertexID))
+		{
+			if (const FMetasoundFrontendClass* Class = DocumentCache->FindDependency(Node.ClassID))
+			{
+				auto IsClassInput = [VertexName = Vertex->Name](const FMetasoundFrontendClassInput& Input) { return Input.Name == VertexName; };
+				if (const FMetasoundFrontendClassInput* ClassInput = Class->Interface.Inputs.FindByPredicate(IsClassInput))
+				{
+					return &ClassInput->DefaultLiteral;
+				}
+			}
+		}
+	}
+
+	return nullptr;
+}
+
 const FMetasoundFrontendLiteral* FMetaSoundFrontendDocumentBuilder::GetNodeInputDefault(const FGuid& InNodeID, const FGuid& InVertexID) const
 {
 	using namespace Metasound::Frontend;
@@ -1646,7 +1749,7 @@ const FMetasoundFrontendLiteral* FMetaSoundFrontendDocumentBuilder::GetNodeInput
 	return nullptr;
 }
 
-const FMetasoundFrontendLiteral* FMetaSoundFrontendDocumentBuilder::GetNodeInputClassDefault(const FGuid& InNodeID, const FGuid& InVertexID) const
+EMetasoundFrontendVertexAccessType FMetaSoundFrontendDocumentBuilder::GetNodeOutputAccessType(const FGuid& InNodeID, const FGuid& InVertexID) const
 {
 	using namespace Metasound::Frontend;
 
@@ -1655,27 +1758,49 @@ const FMetasoundFrontendLiteral* FMetaSoundFrontendDocumentBuilder::GetNodeInput
 	{
 		const FMetasoundFrontendGraph& Graph = GetDocument().RootGraph.Graph;
 		const FMetasoundFrontendNode& Node = Graph.Nodes[*NodeIndex];
-
-		auto IsVertex = [&InVertexID](const FMetasoundFrontendVertex& Vertex) { return Vertex.VertexID == InVertexID; };
-		const int32 VertexIndex = Node.Interface.Inputs.IndexOfByPredicate(IsVertex);
-		if (VertexIndex != INDEX_NONE)
+		if (const FMetasoundFrontendClass* Class = DocumentCache->FindDependency(Node.ClassID))
 		{
-			const FMetasoundFrontendVertex& NodeInput = Node.Interface.Inputs[VertexIndex];
-			if (const FMetasoundFrontendClass* Class = DocumentCache->FindDependency(Node.ClassID))
+			const EMetasoundFrontendClassType ClassType = Class->Metadata.GetType();
+			switch (ClassType)
 			{
-// 				auto IsClassInput = [&NodeInput](const FMetasoundFrontendClassInput& Input)
-// 				{
-// 					return FMetasoundFrontendVertex::IsFunctionalEquivalent(Input, NodeInput);
-// 				};
-				if (const FMetasoundFrontendClassInput* ClassInput = Class->Interface.Inputs.FindByPredicate(IsVertex))
+				case EMetasoundFrontendClassType::Template:
 				{
-					return &ClassInput->DefaultLiteral;
+					const FNodeRegistryKey Key = NodeRegistryKey::CreateKey(Class->Metadata);
+					const INodeTemplate* Template = INodeTemplateRegistry::Get().FindTemplate(Key);
+					if (ensureMsgf(Template, TEXT("Failed to find MetaSound node template registered with key '%s'"), *Key))
+					{
+						if (Template->IsOutputAccessTypeDynamic())
+						{
+							return Template->GetNodeOutputAccessType(*this, InNodeID, InVertexID);
+						}
+					}
+				}
+				break;
+
+				case EMetasoundFrontendClassType::Input:
+				{
+					const FMetasoundFrontendClassOutput& ClassOutput = Class->Interface.Outputs.Last();
+					return ClassOutput.AccessType;
+				}
+
+				default:
+				break;
+			}
+			static_assert(static_cast<uint32>(EMetasoundFrontendClassType::Invalid) == 10, "Potential missing case coverage for EMetasoundFrontendClassType");
+
+			auto IsVertexID = [&InVertexID](const FMetasoundFrontendVertex& Vertex) { return Vertex.VertexID == InVertexID; };
+			if (const FMetasoundFrontendVertex* Vertex = Node.Interface.Outputs.FindByPredicate(IsVertexID))
+			{
+				auto IsClassInput = [VertexName = Vertex->Name](const FMetasoundFrontendClassInput& Output) { return Output.Name == VertexName; };
+				if (const FMetasoundFrontendClassOutput* ClassOutput = Class->Interface.Outputs.FindByPredicate(IsClassInput))
+				{
+					return ClassOutput->AccessType;
 				}
 			}
 		}
 	}
 
-	return nullptr;
+	return EMetasoundFrontendVertexAccessType::Unset;
 }
 
 void FMetaSoundFrontendDocumentBuilder::InitGraphClassMetadata(FMetasoundFrontendClassMetadata& InOutMetadata, bool bResetVersion, const FMetasoundFrontendClassName* NewClassName)
@@ -1828,58 +1953,27 @@ Metasound::Frontend::EInvalidEdgeReason FMetaSoundFrontendDocumentBuilder::IsVal
 
 	const IDocumentGraphNodeCache& NodeCache = DocumentCache->GetNodeCache();
 
-	bool bOutputRerouted = false;
-	const FMetasoundFrontendNode* OutputNode = nullptr;
-	const FMetasoundFrontendVertex* OutputVertex = NodeCache.FindReroutedOutputVertex(InEdge.FromNodeID, InEdge.FromVertexID, &OutputNode, &bOutputRerouted);
+	const FMetasoundFrontendVertex* OutputVertex = NodeCache.FindOutputVertex(InEdge.FromNodeID, InEdge.FromVertexID);
 	if (!OutputVertex)
 	{
-		if (!bOutputRerouted)
-		{
-			// No concrete output found and not going through reroute
-			return EInvalidEdgeReason::MissingOutput;
-		}
-	}
-
-	bool bInputRerouted = false;
-	TArray<const FMetasoundFrontendNode*> InputNodes;
-	TArray<const FMetasoundFrontendVertex*> InputVertices = NodeCache.FindReroutedInputVertices(InEdge.ToNodeID, InEdge.ToVertexID, &InputNodes, &bInputRerouted);
-	if (InputVertices.IsEmpty())
-	{
-		if (!bInputRerouted)
-		{
-			// No concrete input found and not going through reroute
-			return EInvalidEdgeReason::MissingInput;
-		}
+		return EInvalidEdgeReason::MissingOutput;
 	}
 
-	FName OutputType;
-	EMetasoundFrontendVertexAccessType OutputAccessType = EMetasoundFrontendVertexAccessType::Unset;
-	if (OutputVertex)
+	const FMetasoundFrontendVertex* InputVertex = NodeCache.FindInputVertex(InEdge.ToNodeID, InEdge.ToVertexID);
+	if (!InputVertex)
 	{
-		OutputType = OutputVertex->TypeName;
-
-		const FMetasoundFrontendClassOutput* ClassOutput = FindNodeOutputClassOutput(OutputNode->GetID(), OutputVertex->VertexID);
-		check(ClassOutput);
-		OutputAccessType = ClassOutput->AccessType;
-	}
-
-	FName InputType;
-	EMetasoundFrontendVertexAccessType InputAccessType = EMetasoundFrontendVertexAccessType::Unset;
-	if (!InputVertices.IsEmpty())
-	{
-		InputType = InputVertices.Last()->TypeName;
-
-		// Only have to check last typename as all should be the same in theory if more than one off a reroute node
-		const FMetasoundFrontendClassInput* ClassInput = FindNodeInputClassInput(InputNodes.Last()->GetID(), InputVertices.Last()->VertexID);
-		check(ClassInput);
-		InputAccessType = ClassInput->AccessType;
+		return EInvalidEdgeReason::MissingInput;
 	}
 
-	if (!InputType.IsNone() && !OutputType.IsNone() && OutputType != InputType)
+	if (OutputVertex->TypeName != InputVertex->TypeName)
 	{
 		return EInvalidEdgeReason::MismatchedDataType;
 	}
 
+	// TODO: Add cycle detection here
+
+	const EMetasoundFrontendVertexAccessType OutputAccessType = GetNodeOutputAccessType(InEdge.FromNodeID, InEdge.FromVertexID);
+	const EMetasoundFrontendVertexAccessType InputAccessType = GetNodeInputAccessType(InEdge.ToNodeID, InEdge.ToVertexID);
 	if (!FMetasoundFrontendClassVertex::CanConnectVertexAccessTypes(OutputAccessType, InputAccessType))
 	{
 		return EInvalidEdgeReason::MismatchedAccessType;
diff --git a/Engine/Plugins/Runtime/Metasound/Source/MetasoundFrontend/Private/MetasoundFrontendNodeTemplateRegistry.cpp b/Engine/Plugins/Runtime/Metasound/Source/MetasoundFrontend/Private/MetasoundFrontendNodeTemplateRegistry.cpp
--- a/Engine/Plugins/Runtime/Metasound/Source/MetasoundFrontend/Private/MetasoundFrontendNodeTemplateRegistry.cpp
+++ b/Engine/Plugins/Runtime/Metasound/Source/MetasoundFrontend/Private/MetasoundFrontendNodeTemplateRegistry.cpp
@@ -10,101 +10,123 @@
 #include "MetasoundTrace.h"
 
 
-namespace Metasound
+namespace Metasound::Frontend
 {
-	namespace Frontend
+	bool FNodeTemplateBase::IsValidNodeInterface(const FMetasoundFrontendNodeInterface& InNodeInterface) const
 	{
-		class FNodeTemplateRegistry : public INodeTemplateRegistry
-		{
-		public:
-			FNodeTemplateRegistry() = default;
-			virtual ~FNodeTemplateRegistry() = default;
+		return true;
+	}
 
-			virtual const INodeTemplate* FindTemplate(const FNodeRegistryKey& InKey) const override;
+	bool FNodeTemplateBase::IsInputAccessTypeDynamic() const
+	{
+		return false;
+	}
 
-			void Register(TUniquePtr<INodeTemplate>&& InEntry);
-			void Unregister(const FNodeRegistryKey& InKey);
+	bool FNodeTemplateBase::IsOutputAccessTypeDynamic() const
+	{
+		return false;
+	}
 
-		private:
-			TMap<FNodeRegistryKey, TUniquePtr<INodeTemplate>> Templates;
-		};
+	EMetasoundFrontendVertexAccessType FNodeTemplateBase::GetNodeInputAccessType(const FMetaSoundFrontendDocumentBuilder& InBuilder, const FGuid& InNodeID, const FGuid& InVertexID) const
+	{
+		return EMetasoundFrontendVertexAccessType::Unset;
+	}
+
+	EMetasoundFrontendVertexAccessType FNodeTemplateBase::GetNodeOutputAccessType(const FMetaSoundFrontendDocumentBuilder& InBuilder, const FGuid& InNodeID, const FGuid& InVertexID) const
+	{
+		return EMetasoundFrontendVertexAccessType::Unset;
+	}
+
+	class FNodeTemplateRegistry : public INodeTemplateRegistry
+	{
+	public:
+		FNodeTemplateRegistry() = default;
+		virtual ~FNodeTemplateRegistry() = default;
 
-		void FNodeTemplateRegistry::Register(TUniquePtr<INodeTemplate>&& InTemplate)
+		virtual const INodeTemplate* FindTemplate(const FNodeRegistryKey& InKey) const override;
+
+		void Register(TUniquePtr<INodeTemplate>&& InEntry);
+		void Unregister(const FNodeRegistryKey& InKey);
+
+	private:
+		TMap<FNodeRegistryKey, TUniquePtr<INodeTemplate>> Templates;
+	};
+
+	void FNodeTemplateRegistry::Register(TUniquePtr<INodeTemplate>&& InTemplate)
+	{
+		if (ensure(InTemplate.IsValid()))
 		{
-			if (ensure(InTemplate.IsValid()))
+			const FMetasoundFrontendVersion& Version = InTemplate->GetVersion();
+			const FNodeRegistryKey Key = NodeRegistryKey::CreateKey(InTemplate->GetFrontendClass().Metadata);
+			if (ensure(NodeRegistryKey::IsValid(Key)))
 			{
-				const FMetasoundFrontendVersion& Version = InTemplate->GetVersion();
-				const FNodeRegistryKey Key = NodeRegistryKey::CreateKey(InTemplate->GetFrontendClass().Metadata);
-				if (ensure(NodeRegistryKey::IsValid(Key)))
-				{
-					Templates.Add(Key, MoveTemp(InTemplate));
-				}
+				Templates.Add(Key, MoveTemp(InTemplate));
 			}
 		}
+	}
 
-		void FNodeTemplateRegistry::Unregister(const FNodeRegistryKey& InKey)
+	void FNodeTemplateRegistry::Unregister(const FNodeRegistryKey& InKey)
+	{
+		ensure(Templates.Remove(InKey) > 0);
+	}
+
+	const INodeTemplate* FNodeTemplateRegistry::FindTemplate(const FNodeRegistryKey& InKey) const
+	{
+		if (const TUniquePtr<INodeTemplate>* TemplatePtr = Templates.Find(InKey))
 		{
-			ensure(Templates.Remove(InKey) > 0);
+			return TemplatePtr->Get();
 		}
 
-		const INodeTemplate* FNodeTemplateRegistry::FindTemplate(const FNodeRegistryKey& InKey) const
+		return nullptr;
+	}
+
+	INodeTemplateRegistry& INodeTemplateRegistry::Get()
+	{
+		static FNodeTemplateRegistry Registry;
+		return Registry;
+	}
+
+	void RegisterNodeTemplate(TUniquePtr<INodeTemplate>&& InTemplate)
+	{
+		class FTemplateRegistryEntry : public INodeRegistryTemplateEntry
 		{
-			if (const TUniquePtr<INodeTemplate>* TemplatePtr = Templates.Find(InKey))
+			const FNodeClassInfo ClassInfo;
+			const FMetasoundFrontendClass FrontendClass;
+
+		public:
+			FTemplateRegistryEntry(const INodeTemplate& InNodeTemplate)
+				: ClassInfo(InNodeTemplate.GetFrontendClass().Metadata)
+				, FrontendClass(InNodeTemplate.GetFrontendClass())
 			{
-				return TemplatePtr->Get();
 			}
 
-			return nullptr;
-		}
-
-		INodeTemplateRegistry& INodeTemplateRegistry::Get()
-		{
-			static FNodeTemplateRegistry Registry;
-			return Registry;
-		}
+			virtual ~FTemplateRegistryEntry() = default;
 
-		void RegisterNodeTemplate(TUniquePtr<INodeTemplate>&& InTemplate)
-		{
-			class FTemplateRegistryEntry : public INodeRegistryTemplateEntry
+			virtual const FNodeClassInfo& GetClassInfo() const override
 			{
-				const FNodeClassInfo ClassInfo;
-				const FMetasoundFrontendClass FrontendClass;
-
-			public:
-				FTemplateRegistryEntry(const INodeTemplate& InNodeTemplate)
-					: ClassInfo(InNodeTemplate.GetFrontendClass().Metadata)
-					, FrontendClass(InNodeTemplate.GetFrontendClass())
-				{
-				}
-
-				virtual ~FTemplateRegistryEntry() = default;
-
-				virtual const FNodeClassInfo& GetClassInfo() const override
-				{
-					return ClassInfo;
-				}
-
-				/** Return a FMetasoundFrontendClass which describes the node. */
-				virtual const FMetasoundFrontendClass& GetFrontendClass() const override
-				{
-					return FrontendClass;
-				}
-			};
-
-			TUniquePtr<INodeRegistryTemplateEntry> RegEntry = TUniquePtr<INodeRegistryTemplateEntry>(new FTemplateRegistryEntry(*InTemplate.Get()));
-			FRegistryContainerImpl::Get().RegisterNodeTemplate(MoveTemp(RegEntry));
-
-			static_cast<FNodeTemplateRegistry&>(INodeTemplateRegistry::Get()).Register(MoveTemp(InTemplate));
-		}
+				return ClassInfo;
+			}
 
-		void UnregisterNodeTemplate(const FMetasoundFrontendVersion& InVersion)
-		{
-			const FNodeRegistryKey Key = NodeRegistryKey::CreateKey(EMetasoundFrontendClassType::Template, InVersion.Name.ToString(), InVersion.Number.Major, InVersion.Number.Minor);
-			if (ensure(NodeRegistryKey::IsValid(Key)))
+			/** Return a FMetasoundFrontendClass which describes the node. */
+			virtual const FMetasoundFrontendClass& GetFrontendClass() const override
 			{
-				FRegistryContainerImpl::Get().UnregisterNodeTemplate(Key);
-				static_cast<FNodeTemplateRegistry&>(INodeTemplateRegistry::Get()).Unregister(Key);
+				return FrontendClass;
 			}
+		};
+
+		TUniquePtr<INodeRegistryTemplateEntry> RegEntry = TUniquePtr<INodeRegistryTemplateEntry>(new FTemplateRegistryEntry(*InTemplate.Get()));
+		FRegistryContainerImpl::Get().RegisterNodeTemplate(MoveTemp(RegEntry));
+
+		static_cast<FNodeTemplateRegistry&>(INodeTemplateRegistry::Get()).Register(MoveTemp(InTemplate));
+	}
+
+	void UnregisterNodeTemplate(const FMetasoundFrontendVersion& InVersion)
+	{
+		const FNodeRegistryKey Key = NodeRegistryKey::CreateKey(EMetasoundFrontendClassType::Template, InVersion.Name.ToString(), InVersion.Number.Major, InVersion.Number.Minor);
+		if (ensure(NodeRegistryKey::IsValid(Key)))
+		{
+			FRegistryContainerImpl::Get().UnregisterNodeTemplate(Key);
+			static_cast<FNodeTemplateRegistry&>(INodeTemplateRegistry::Get()).Unregister(Key);
 		}
-	} // namespace Frontend
-} // namespace Metasound
+	}
+} // namespace Metasound::Frontend
diff --git a/Engine/Plugins/Runtime/Metasound/Source/MetasoundFrontend/Private/NodeTemplates/MetasoundFrontendNodeTemplateReroute.cpp b/Engine/Plugins/Runtime/Metasound/Source/MetasoundFrontend/Private/NodeTemplates/MetasoundFrontendNodeTemplateReroute.cpp
--- a/Engine/Plugins/Runtime/Metasound/Source/MetasoundFrontend/Private/NodeTemplates/MetasoundFrontendNodeTemplateReroute.cpp
+++ b/Engine/Plugins/Runtime/Metasound/Source/MetasoundFrontend/Private/NodeTemplates/MetasoundFrontendNodeTemplateReroute.cpp
@@ -3,8 +3,10 @@
 #include "NodeTemplates/MetasoundFrontendNodeTemplateReroute.h"
 
 #include "Algo/AnyOf.h"
+#include "MetasoundAssetManager.h"
 #include "MetasoundFrontendDataTypeRegistry.h"
 #include "MetasoundFrontendDocumentBuilder.h"
+#include "MetasoundFrontendNodeTemplateRegistry.h"
 #include "MetasoundFrontendRegistries.h"
 #include "NodeTemplates/MetasoundFrontendDocumentTemplatePreprocessor.h"
 
@@ -27,14 +29,17 @@ namespace Metasound::Frontend
 			FMetasoundFrontendEdge InputEdge;
 			TArray<FMetasoundFrontendEdge> OutputEdges;
 
-			if (const FMetasoundFrontendNode* Node = OutBuilder.FindNode(InNodeID))
+			const FMetasoundFrontendNode* Node = OutBuilder.FindNode(InNodeID);
+			if (ensureMsgf(Node, TEXT("Failed to find node with ID '%s' when reroute template node transform was given a valid ID for builder '%s'."),
+				*InNodeID.ToString(),
+				*OutBuilder.GetDebugName()))
 			{
-				if (!ensureMsgf(Node->Interface.Inputs.Num() == 1, TEXT("Template nodes must only have one input")))
+				if (!ensureMsgf(Node->Interface.Inputs.Num() == 1, TEXT("Reroute nodes must only have one input")))
 				{
 					return false;
 				}
 
-				if (!ensureMsgf(Node->Interface.Outputs.Num() == 1, TEXT("Template nodes must only have one output")))
+				if (!ensureMsgf(Node->Interface.Outputs.Num() == 1, TEXT("Reroute nodes must only have one output")))
 				{
 					return false;
 				}
@@ -62,27 +67,21 @@ namespace Metasound::Frontend
 
 				// Remove the template node
 				OutBuilder.RemoveNode(Node->GetID());
-			}
-			else
-			{
-				const FMetaSoundFrontendDocumentBuilder& ConstBuilder = OutBuilder;
-				ensureMsgf(false, TEXT("Failed to find node with ID '%s' when template node transform was given a valid ID for builder '%s'."),
-					*InNodeID.ToString(),
-					*ConstBuilder.GetDocument().RootGraph.Metadata.GetClassName().GetFullName().ToString());
-			}
 
-			// Add new connections from reroute source node to reroute destination node. Either could be another reroute,
-			// which is valid because said node will subsequently get processed.
-			if (InputEdge.GetFromVertexHandle().IsSet())
-			{
-				for (FMetasoundFrontendEdge& OutputEdge : OutputEdges)
+				// Add new connections from reroute source node to reroute destination node. Either could be another reroute,
+				// which is valid because said node will subsequently get processed.
+				if (InputEdge.GetFromVertexHandle().IsSet())
 				{
-					OutputEdge.FromNodeID = InputEdge.FromNodeID;
-					OutputEdge.FromVertexID = InputEdge.FromVertexID;
-					OutBuilder.AddEdge(MoveTemp(OutputEdge));
-				}
+					bool bModified = !OutputEdges.IsEmpty();
+					for (FMetasoundFrontendEdge& OutputEdge : OutputEdges)
+					{
+						OutputEdge.FromNodeID = InputEdge.FromNodeID;
+						OutputEdge.FromVertexID = InputEdge.FromVertexID;
+						OutBuilder.AddEdge(MoveTemp(OutputEdge));
+					}
 
-				return !OutputEdges.IsEmpty();
+					return bModified;
+				}
 			}
 
 			return false;
@@ -146,6 +145,61 @@ namespace Metasound::Frontend
 		return NewInterface;
 	}
 
+	EMetasoundFrontendVertexAccessType FRerouteNodeTemplate::GetNodeInputAccessType(const FMetaSoundFrontendDocumentBuilder& InBuilder, const FGuid& InNodeID, const FGuid& InVertexID) const
+	{
+		// Recursive search up DAG for first connected non-reroute node's input access type
+		if (const FMetasoundFrontendNode* Node = InBuilder.FindNode(InNodeID))
+		{
+			// Should only ever be one
+			const FMetasoundFrontendVertex& RerouteOutput = Node->Interface.Outputs.Last();
+
+			TArray<const FMetasoundFrontendNode*> ConnectedNodes;
+			TArray<const FMetasoundFrontendVertex*> ConnectedInputs = InBuilder.FindNodeInputsConnectedToNodeOutput(InNodeID, InVertexID, &ConnectedNodes);
+			for (int32 Index = 0; Index < ConnectedNodes.Num(); ++Index)
+			{
+				const FMetasoundFrontendNode* ConnectedNode = ConnectedNodes[Index];
+				if (const FMetasoundFrontendClass* ConnectedNodeClass = InBuilder.FindDependency(ConnectedNode->ClassID))
+				{
+					const FMetasoundFrontendVertex* ConnectedInput = ConnectedInputs[Index];
+					if (ConnectedNodeClass->Metadata.GetClassName() == ClassName)
+					{
+						return this->GetNodeInputAccessType(InBuilder, ConnectedNode->GetID(), ConnectedInput->VertexID);
+					}
+
+					return InBuilder.GetNodeInputAccessType(ConnectedNode->GetID(), ConnectedInput->VertexID);
+				}
+			}
+		}
+
+		return EMetasoundFrontendVertexAccessType::Unset;
+	}
+
+	EMetasoundFrontendVertexAccessType FRerouteNodeTemplate::GetNodeOutputAccessType(const FMetaSoundFrontendDocumentBuilder& InBuilder, const FGuid& InNodeID, const FGuid& InVertexID) const
+	{
+		// Depth-first recursive search for first connected non-reroute node's output access type
+		if (const FMetasoundFrontendNode* Node = InBuilder.FindNode(InNodeID))
+		{
+			// Should only ever be one
+			const FMetasoundFrontendVertex& RerouteInput = Node->Interface.Inputs.Last();
+
+			const FMetasoundFrontendNode* ConnectedNode = nullptr;
+			if (const FMetasoundFrontendVertex* ConnectedOutput = InBuilder.FindNodeOutputConnectedToNodeInput(InNodeID, RerouteInput.VertexID, &ConnectedNode))
+			{
+				if (const FMetasoundFrontendClass* ConnectedNodeClass = InBuilder.FindDependency(ConnectedNode->ClassID))
+				{
+					if (ConnectedNodeClass->Metadata.GetClassName() == ClassName)
+					{
+						return this->GetNodeOutputAccessType(InBuilder, ConnectedNode->GetID(), ConnectedOutput->VertexID);
+					}
+
+					return InBuilder.GetNodeOutputAccessType(ConnectedNode->GetID(), ConnectedOutput->VertexID);
+				}
+			}
+		}
+
+		return EMetasoundFrontendVertexAccessType::Unset;
+	}
+
 	const FNodeRegistryKey& FRerouteNodeTemplate::GetRegistryKey()
 	{
 		static const FNodeRegistryKey RegistryKey = NodeRegistryKey::CreateKey(
@@ -186,6 +240,16 @@ namespace Metasound::Frontend
 	}
 #endif // WITH_EDITOR
 
+	bool FRerouteNodeTemplate::IsInputAccessTypeDynamic() const
+	{
+		return true;
+	}
+
+	bool FRerouteNodeTemplate::IsOutputAccessTypeDynamic() const
+	{
+		return true;
+	}
+
 	bool FRerouteNodeTemplate::IsValidNodeInterface(const FMetasoundFrontendNodeInterface& InNodeInterface) const
 	{
 		if (InNodeInterface.Inputs.Num() != 1)
diff --git a/Engine/Plugins/Runtime/Metasound/Source/MetasoundFrontend/Public/MetasoundFrontendDocument.h b/Engine/Plugins/Runtime/Metasound/Source/MetasoundFrontend/Public/MetasoundFrontendDocument.h
--- a/Engine/Plugins/Runtime/Metasound/Source/MetasoundFrontend/Public/MetasoundFrontendDocument.h
+++ b/Engine/Plugins/Runtime/Metasound/Source/MetasoundFrontend/Public/MetasoundFrontendDocument.h
@@ -121,7 +121,7 @@ UENUM()
 enum class EMetasoundFrontendClassType : uint8
 {
 	// The MetaSound class is defined externally, in compiled code or in another document.
-	External,
+	External = 0,
 
 	// The MetaSound class is a graph within the containing document.
 	Graph,
@@ -148,7 +148,7 @@ enum class EMetasoundFrontendClassType : uint8
 	VariableMutator,
 
 	// The MetaSound class is defined only by the Frontend, and associatively
-	// performs a functional replacement operation in a pre-build step.
+	// performs a functional operation within the given document in a registration/cook step.
 	Template,
 
 	Invalid UMETA(Hidden)
@@ -596,6 +596,26 @@ struct FMetasoundFrontendEdge
 	{
 		return FMetasoundFrontendVertexHandle { ToNodeID, ToVertexID };
 	}
+
+	friend bool operator==(const FMetasoundFrontendEdge& InLHS, const FMetasoundFrontendEdge& InRHS)
+	{
+		return InLHS.FromNodeID == InRHS.FromNodeID
+			&& InLHS.FromVertexID == InRHS.FromVertexID
+			&& InLHS.ToNodeID == InRHS.ToNodeID
+			&& InLHS.ToVertexID == InRHS.ToVertexID;
+	}
+
+	friend bool operator!=(const FMetasoundFrontendEdge& InLHS, const FMetasoundFrontendEdge& InRHS)
+	{
+		return !(InLHS == InRHS);
+	}
+
+	friend FORCEINLINE uint32 GetTypeHash(const FMetasoundFrontendEdge& InEdge)
+	{
+		const int32 FromHash = HashCombineFast(InEdge.FromNodeID.A, InEdge.FromVertexID.B);
+		const int32 ToHash = HashCombineFast(InEdge.ToNodeID.C, InEdge.ToVertexID.D);
+		return HashCombineFast(FromHash, ToHash);
+	}
 };
 
 USTRUCT()
diff --git a/Engine/Plugins/Runtime/Metasound/Source/MetasoundFrontend/Public/MetasoundFrontendDocumentBuilder.h b/Engine/Plugins/Runtime/Metasound/Source/MetasoundFrontend/Public/MetasoundFrontendDocumentBuilder.h
--- a/Engine/Plugins/Runtime/Metasound/Source/MetasoundFrontend/Public/MetasoundFrontendDocumentBuilder.h
+++ b/Engine/Plugins/Runtime/Metasound/Source/MetasoundFrontend/Public/MetasoundFrontendDocumentBuilder.h
@@ -183,11 +183,14 @@ public:
 	TArray<const FMetasoundFrontendVertex*> FindNodeOutputs(const FGuid& InNodeID, FName TypeName = FName()) const;
 	const FMetasoundFrontendVertex* FindNodeOutputConnectedToNodeInput(const FGuid& InInputNodeID, const FGuid& InInputVertexID, const FMetasoundFrontendNode** ConnectedOutputNode = nullptr) const;
 
+	const FString GetDebugName() const;
 	const FMetasoundFrontendDocument& GetDocument() const;
 	const Metasound::Frontend::FDocumentModifyDelegates& GetDocumentDelegates() const;
 	const IMetaSoundDocumentInterface& GetDocumentInterface() const;
+	EMetasoundFrontendVertexAccessType GetNodeInputAccessType(const FGuid& InNodeID, const FGuid& InVertexID) const;
 	const FMetasoundFrontendLiteral* GetNodeInputClassDefault(const FGuid& InNodeID, const FGuid& InVertexID) const;
 	const FMetasoundFrontendLiteral* GetNodeInputDefault(const FGuid& InNodeID, const FGuid& InVertexID) const;
+	EMetasoundFrontendVertexAccessType GetNodeOutputAccessType(const FGuid& InNodeID, const FGuid& InVertexID) const;
 
 	// Initializes GraphClass Metadata, optionally resetting the version back to 1.0 and/or creating a unique class name if a name is not provided.
 	static void InitGraphClassMetadata(FMetasoundFrontendClassMetadata& InOutMetadata, bool bResetVersion = false, const FMetasoundFrontendClassName* NewClassName = nullptr);
@@ -256,8 +259,8 @@ private:
 	FMetasoundFrontendNode* AddNodeInternal(const FMetasoundFrontendClassMetadata& InClassMetadata, Metasound::Frontend::FFinalizeNodeFunctionRef FinalizeNode, FGuid InNodeID = FGuid::NewGuid());
 
 	const TSet<FMetasoundFrontendVersion>* FindNodeClassInterfaces(const FGuid& InNodeID) const;
-	const FMetasoundFrontendClassInput* FindNodeInputClassInput(const FGuid& InNodeID, const FGuid& InVertexID) const;
-	const FMetasoundFrontendClassOutput* FindNodeOutputClassOutput(const FGuid& InNodeID, const FGuid& InVertexID) const;
+	const FMetasoundFrontendClassInput* FindNodeInputClassInput(const FGuid& InNodeID, const FName& InVertexName) const;
+	const FMetasoundFrontendClassOutput* FindNodeOutputClassOutput(const FGuid& InNodeID, const FName& InVertexName) const;
 
 	void ReloadCacheInternal();
 
diff --git a/Engine/Plugins/Runtime/Metasound/Source/MetasoundFrontend/Public/MetasoundFrontendNodeTemplateRegistry.h b/Engine/Plugins/Runtime/Metasound/Source/MetasoundFrontend/Public/MetasoundFrontendNodeTemplateRegistry.h
--- a/Engine/Plugins/Runtime/Metasound/Source/MetasoundFrontend/Public/MetasoundFrontendNodeTemplateRegistry.h
+++ b/Engine/Plugins/Runtime/Metasound/Source/MetasoundFrontend/Public/MetasoundFrontendNodeTemplateRegistry.h
@@ -10,58 +10,80 @@
 #include "UObject/NoExportTypes.h"
 
 
-namespace Metasound
+// Forward Declarations
+struct FMetaSoundFrontendDocumentBuilder;
+
+
+namespace Metasound::Frontend
 {
-	namespace Frontend
+	/**
+		* Base interface for a node template, which acts in place of frontend node class and respective instance(s).
+		* Instances are preprocessed, allowing for custom graph manipulation prior to generating a respective runtime
+		* graph operator representation.
+		*/
+	class METASOUNDFRONTEND_API INodeTemplate
 	{
-		/**
-		  * Base interface for a node template, which acts in place of frontend node class and respective instance(s).
-		  * Instances are preprocessed, allowing for custom graph manipulation prior to generating a respective runtime
-		  * graph operator representation.
-		  */
-		class METASOUNDFRONTEND_API INodeTemplate
-		{
-		public:
-			virtual ~INodeTemplate() = default;
+	public:
+		virtual ~INodeTemplate() = default;
 
-			UE_DEPRECATED(5.4, "Use version that does not provide a preprocessed document")
-			virtual TUniquePtr<INodeTransform> GenerateNodeTransform(FMetasoundFrontendDocument& InDocument) const { return nullptr; }
+		UE_DEPRECATED(5.4, "Use version that does not provide a preprocessed document")
+		virtual TUniquePtr<INodeTransform> GenerateNodeTransform(FMetasoundFrontendDocument& InDocument) const { return nullptr; }
 
-			// Generates node transform that is used to preprocess nodes.
-			virtual TUniquePtr<INodeTransform> GenerateNodeTransform() const { return nullptr; }
+		// Generates node transform that is used to preprocess nodes.
+		virtual TUniquePtr<INodeTransform> GenerateNodeTransform() const { return nullptr; }
 
-			// Returns the class definition for the given node class template.
-			virtual const FMetasoundFrontendClass& GetFrontendClass() const = 0;
+		// Returns the class definition for the given node class template.
+		virtual const FMetasoundFrontendClass& GetFrontendClass() const = 0;
 
-			// Returns the version of the given node class template.
-			virtual const FMetasoundFrontendVersion& GetVersion() const = 0;
+		// Returns access type of the given input within the provided builder's document
+		virtual EMetasoundFrontendVertexAccessType GetNodeInputAccessType(const FMetaSoundFrontendDocumentBuilder& InBuilder, const FGuid& InNodeID, const FGuid& InVertexID) const = 0;
 
-			// Given the provided node interface, returns whether or not it conforms to an expected format
-			// that can be successfully manipulated by a generated node template preprocessor.
-			virtual bool IsValidNodeInterface(const FMetasoundFrontendNodeInterface& InNodeInterface) const = 0;
+		// Returns access type of the given output within the provided builder's document
+		virtual EMetasoundFrontendVertexAccessType GetNodeOutputAccessType(const FMetaSoundFrontendDocumentBuilder& InBuilder, const FGuid& InNodeID, const FGuid& InVertexID) const = 0;
+
+		// Returns the version of the given node class template.
+		virtual const FMetasoundFrontendVersion& GetVersion() const = 0;
 
 #if WITH_EDITOR
-			// Returns whether or not the given node template has the necessary
-			// required connections to be preprocessed (editor only).
-			virtual bool HasRequiredConnections(FConstNodeHandle InNodeHandle, FString* OutMessage = nullptr) const = 0;
+		// Returns whether or not the given node template has the necessary
+		// required connections to be preprocessed (editor only).
+		virtual bool HasRequiredConnections(FConstNodeHandle InNodeHandle, FString* OutMessage = nullptr) const = 0;
 #endif // WITH_EDITOR
-		};
-
-		class METASOUNDFRONTEND_API INodeTemplateRegistry
-		{
-		public:
-			// Returns singleton template registry.
-			static INodeTemplateRegistry& Get();
-
-			virtual ~INodeTemplateRegistry() = default;
-
-			// Find a template with the given key. Returns null if entry not found with given key.
-			virtual const INodeTemplate* FindTemplate(const FNodeRegistryKey& InKey) const = 0;
-		};
-
-		// Register & Unregister are not publicly accessible implementation as the API
-		// is in beta and, currently, only to be used by internal implementation (ex. reroute nodes).
-		void RegisterNodeTemplate(TUniquePtr<INodeTemplate>&& InTemplate);
-		void UnregisterNodeTemplate(const FMetasoundFrontendVersion& InNodeTemplateVersion);
-	} // namespace Frontend
-} // namespace Metasound
+
+		// Returns whether template can dynamically assign a node's input access type (as opposed to it being assigned on the class input definition)
+		virtual bool IsInputAccessTypeDynamic() const = 0;
+
+		// Returns whether template can dynamically assign a node's output's access type (as opposed to it being assigned on the class output definition)
+		virtual bool IsOutputAccessTypeDynamic() const = 0;
+
+		// Given the provided node interface, returns whether or not it conforms to an expected format
+		// that can be successfully manipulated by a generated node template transform.
+		virtual bool IsValidNodeInterface(const FMetasoundFrontendNodeInterface& InNodeInterface) const = 0;
+	};
+
+	class METASOUNDFRONTEND_API FNodeTemplateBase : public INodeTemplate
+	{
+		virtual bool IsValidNodeInterface(const FMetasoundFrontendNodeInterface& InNodeInterface) const override;
+		virtual bool IsInputAccessTypeDynamic() const override;
+		virtual bool IsOutputAccessTypeDynamic() const override;
+		virtual EMetasoundFrontendVertexAccessType GetNodeInputAccessType(const FMetaSoundFrontendDocumentBuilder& InBuilder, const FGuid& InNodeID, const FGuid& InVertexID) const override;
+		virtual EMetasoundFrontendVertexAccessType GetNodeOutputAccessType(const FMetaSoundFrontendDocumentBuilder& InBuilder, const FGuid& InNodeID, const FGuid& InVertexID) const override;
+	};
+
+	class METASOUNDFRONTEND_API INodeTemplateRegistry
+	{
+	public:
+		// Returns singleton template registry.
+		static INodeTemplateRegistry& Get();
+
+		virtual ~INodeTemplateRegistry() = default;
+
+		// Find a template with the given key. Returns null if entry not found with given key.
+		virtual const INodeTemplate* FindTemplate(const FNodeRegistryKey& InKey) const = 0;
+	};
+
+	// Register & Unregister are not publicly accessible implementation as the API
+	// is in beta and, currently, only to be used by internal implementation (ex. reroute nodes).
+	void RegisterNodeTemplate(TUniquePtr<INodeTemplate>&& InTemplate);
+	void UnregisterNodeTemplate(const FMetasoundFrontendVersion& InNodeTemplateVersion);
+} // namespace Metasound::Frontend
diff --git a/Engine/Plugins/Runtime/Metasound/Source/MetasoundFrontend/Public/NodeTemplates/MetasoundFrontendNodeTemplateReroute.h b/Engine/Plugins/Runtime/Metasound/Source/MetasoundFrontend/Public/NodeTemplates/MetasoundFrontendNodeTemplateReroute.h
--- a/Engine/Plugins/Runtime/Metasound/Source/MetasoundFrontend/Public/NodeTemplates/MetasoundFrontendNodeTemplateReroute.h
+++ b/Engine/Plugins/Runtime/Metasound/Source/MetasoundFrontend/Public/NodeTemplates/MetasoundFrontendNodeTemplateReroute.h
@@ -3,39 +3,41 @@
 
 #include "MetasoundFrontendController.h"
 #include "MetasoundFrontendDocument.h"
+#include "MetasoundFrontendDocumentBuilder.h"
 #include "MetasoundFrontendNodeTemplateRegistry.h"
 #include "MetasoundFrontendTransform.h"
 #include "Templates/SharedPointer.h"
 #include "UObject/NoExportTypes.h"
 
 
-namespace Metasound
+namespace Metasound::Frontend
 {
-	namespace Frontend
+	class METASOUNDFRONTEND_API FRerouteNodeTemplate : public INodeTemplate
 	{
-		class METASOUNDFRONTEND_API FRerouteNodeTemplate : public INodeTemplate
-		{
-		public:
-			static const FMetasoundFrontendClassName ClassName;
-			static const FMetasoundFrontendVersion Version;
+	public:
+		static const FMetasoundFrontendClassName ClassName;
+		static const FMetasoundFrontendVersion Version;
 
-			static const FNodeRegistryKey& GetRegistryKey();
+		static const FNodeRegistryKey& GetRegistryKey();
 
-			static FMetasoundFrontendNodeInterface CreateNodeInterfaceFromDataType(FName InDataType);
+		static FMetasoundFrontendNodeInterface CreateNodeInterfaceFromDataType(FName InDataType);
 
-			virtual ~FRerouteNodeTemplate() = default;
+		virtual ~FRerouteNodeTemplate() = default;
 
-			UE_DEPRECATED(5.4, "Use version that does not require mutating a provided PreprocessedDocument")
-			virtual TUniquePtr<INodeTransform> GenerateNodeTransform(FMetasoundFrontendDocument& InPreprocessedDocument) const override;
+		UE_DEPRECATED(5.4, "Use version that does not require mutating a provided PreprocessedDocument")
+		virtual TUniquePtr<INodeTransform> GenerateNodeTransform(FMetasoundFrontendDocument& InPreprocessedDocument) const override;
 
-			virtual TUniquePtr<INodeTransform> GenerateNodeTransform() const override;
-			virtual const FMetasoundFrontendClass& GetFrontendClass() const override;
-			virtual const FMetasoundFrontendVersion& GetVersion() const override;
-			virtual bool IsValidNodeInterface(const FMetasoundFrontendNodeInterface& InNodeInterface) const override;
+		virtual TUniquePtr<INodeTransform> GenerateNodeTransform() const override;
+		virtual const FMetasoundFrontendClass& GetFrontendClass() const override;
+		virtual EMetasoundFrontendVertexAccessType GetNodeInputAccessType(const FMetaSoundFrontendDocumentBuilder& InBuilder, const FGuid& InNodeID, const FGuid& InVertexID) const override;
+		virtual EMetasoundFrontendVertexAccessType GetNodeOutputAccessType(const FMetaSoundFrontendDocumentBuilder& InBuilder, const FGuid& InNodeID, const FGuid& InVertexID) const override;
+		virtual const FMetasoundFrontendVersion& GetVersion() const override;
+		virtual bool IsInputAccessTypeDynamic() const override;
+		virtual bool IsOutputAccessTypeDynamic() const override;
+		virtual bool IsValidNodeInterface(const FMetasoundFrontendNodeInterface& InNodeInterface) const override;
 
 #if WITH_EDITOR
-			virtual bool HasRequiredConnections(FConstNodeHandle InNodeHandle, FString* OutMessage = nullptr) const override;
+		virtual bool HasRequiredConnections(FConstNodeHandle InNodeHandle, FString* OutMessage = nullptr) const override;
 #endif // WITH_EDITOR
-		};
-	} // namespace Frontend
-} // namespace Metasound
+	};
+} // namespace Metasound::Frontend
