[Metasounds] Adding unreal insights markup:
* Added asset name metadata to "RegisterWithFrontEnd"
* Added per-asset counter for the number of concurrent operators alive at a time.

diff --git a/Engine/Plugins/Runtime/Metasound/Source/MetasoundFrontend/MetasoundFrontend.Build.cs b/Engine/Plugins/Runtime/Metasound/Source/MetasoundFrontend/MetasoundFrontend.Build.cs
--- a/Engine/Plugins/Runtime/Metasound/Source/MetasoundFrontend/MetasoundFrontend.Build.cs
+++ b/Engine/Plugins/Runtime/Metasound/Source/MetasoundFrontend/MetasoundFrontend.Build.cs
@@ -20,6 +20,7 @@ namespace UnrealBuildTool.Rules
 					"CoreUObject",
 					"Serialization",
 					"SignalProcessing",
+					"TraceLog",
 					"MetasoundGraphCore"
 
 				}
diff --git a/Engine/Plugins/Runtime/Metasound/Source/MetasoundFrontend/Private/MetasoundAssetBase.cpp b/Engine/Plugins/Runtime/Metasound/Source/MetasoundFrontend/Private/MetasoundAssetBase.cpp
--- a/Engine/Plugins/Runtime/Metasound/Source/MetasoundFrontend/Private/MetasoundAssetBase.cpp
+++ b/Engine/Plugins/Runtime/Metasound/Source/MetasoundFrontend/Private/MetasoundAssetBase.cpp
@@ -36,6 +36,8 @@
 #include "NodeTemplates/MetasoundFrontendDocumentTemplatePreprocessor.h"
 #include "StructSerializer.h"
 #include "Templates/SharedPointer.h"
+#include "Trace/Trace.h"
+#include "Trace/Trace.inl"
 #include "UObject/MetaData.h"
 
 #define LOCTEXT_NAMESPACE "MetaSound"
@@ -178,12 +180,19 @@ namespace Metasound
 
 const FString FMetasoundAssetBase::FileExtension(TEXT(".metasound"));
 
+UE_TRACE_EVENT_BEGIN(Cpu, RegisterGraphWithFrontEnd, NoSync)
+	UE_TRACE_EVENT_FIELD(UE::Trace::WideString, AssetName)
+UE_TRACE_EVENT_END()
+
 void FMetasoundAssetBase::RegisterGraphWithFrontend(Metasound::Frontend::FMetaSoundAssetRegistrationOptions InRegistrationOptions)
 {
 	using namespace Metasound;
 	using namespace Metasound::Frontend;
 
 	METASOUND_TRACE_CPUPROFILER_EVENT_SCOPE(MetaSoundAssetBase::RegisterGraphWithFrontend);
+	UE_TRACE_LOG_SCOPED_T(Cpu, RegisterGraphWithFrontEnd, CpuChannel)
+		<< RegisterGraphWithFrontEnd.AssetName(*GetOwningAssetName());
+	
 	if (!InRegistrationOptions.bForceReregister)
 	{
 		if (IsRegistered())
diff --git a/Engine/Plugins/Runtime/Metasound/Source/MetasoundGenerator/Private/MetasoundGenerator.cpp b/Engine/Plugins/Runtime/Metasound/Source/MetasoundGenerator/Private/MetasoundGenerator.cpp
--- a/Engine/Plugins/Runtime/Metasound/Source/MetasoundGenerator/Private/MetasoundGenerator.cpp
+++ b/Engine/Plugins/Runtime/Metasound/Source/MetasoundGenerator/Private/MetasoundGenerator.cpp
@@ -288,6 +288,12 @@ namespace Metasound
 		FDataTransmissionCenter::Get().UnregisterDataChannelIfUnconnected(ParameterPackSendAddress);
 	}
 
+	static const FName InstanceCounterCategory(TEXT("Metasound/Active_Generators"));
+	const FName& FMetasoundGenerator::GetCategoryName() const
+	{
+		return InstanceCounterCategory;
+	}
+
 	FDelegateHandle FMetasoundGenerator::AddGraphSetCallback(FOnSetGraph::FDelegate&& Delegate)
 	{
 		FScopeLock SetPendingGraphLock(&PendingGraphMutex);
@@ -301,6 +307,7 @@ namespace Metasound
 
 	void FMetasoundGenerator::InitBase(const FMetasoundGeneratorInitParams& InInitParams)
 	{
+		FConcurrentMetasoundInstanceCounter::Init(InInitParams.MetaSoundName);
 		MetasoundName = InInitParams.MetaSoundName;
 		NumChannels = InInitParams.AudioOutputNames.Num();
 		NumSamplesPerExecute = NumChannels * NumFramesPerExecute;
diff --git a/Engine/Plugins/Runtime/Metasound/Source/MetasoundGenerator/Private/MetasoundInstanceCounter.cpp b/Engine/Plugins/Runtime/Metasound/Source/MetasoundGenerator/Private/MetasoundInstanceCounter.cpp
new file mode 100644
--- /dev/null
+++ b/Engine/Plugins/Runtime/Metasound/Source/MetasoundGenerator/Private/MetasoundInstanceCounter.cpp
@@ -0,0 +1,128 @@
+// Copyright Epic Games, Inc. All Rights Reserved.
+
+#include "MetasoundInstanceCounter.h"
+
+namespace Metasound
+{
+	FConcurrentMetasoundInstanceCounter::FConcurrentMetasoundInstanceCounter(const FName& InName)
+	: InstanceName(InName)
+	{
+		Increment();
+	}
+
+	FConcurrentMetasoundInstanceCounter::FConcurrentMetasoundInstanceCounter(const FString& InName)
+	: InstanceName(InName)
+	{
+		Increment();
+	}
+
+	// dtor
+	FConcurrentMetasoundInstanceCounter::~FConcurrentMetasoundInstanceCounter()
+	{
+		Decrement();
+	}
+
+	void FConcurrentMetasoundInstanceCounter::Init(const FName& InName)
+	{
+		InstanceName = InName;
+		Increment();
+	}
+
+	void FConcurrentMetasoundInstanceCounter::Init(const FString& InName)
+	{
+		InstanceName = FName(InName);
+		Increment();
+	}
+
+	// static interface
+	int64 FConcurrentMetasoundInstanceCounter::GetCountForName(const FName& InName)
+	{
+		FScopeLock Lock(&MapCritSec);
+		if (FStats* Stats = StatsMap.Find(InName))
+		{
+			return Stats->GetCount();
+		}
+
+		return 0;
+	}
+
+	int64 FConcurrentMetasoundInstanceCounter::GetPeakCountForName(const FName& InName)
+	{
+		FScopeLock Lock(&MapCritSec);
+		if (FStats* Stats = StatsMap.Find(InName))
+		{
+			return Stats->GetPeakCount();
+		}
+
+		return 0;
+	}
+
+	void FConcurrentMetasoundInstanceCounter::Increment()
+	{
+		FScopeLock Lock(&MapCritSec);
+		GetOrAddStats().Increment();
+	}
+
+	void FConcurrentMetasoundInstanceCounter::Decrement()
+	{
+		FScopeLock Lock(&MapCritSec);
+		GetOrAddStats().Decrement();
+	}
+
+	FConcurrentMetasoundInstanceCounter::FStats& FConcurrentMetasoundInstanceCounter::GetOrAddStats()
+	{
+		FScopeLock Lock(&MapCritSec);
+
+		// avoid re-constructing the trace counter string unless it's new
+		if (StatsMap.Contains(InstanceName))
+		{
+			return StatsMap[InstanceName];
+		}
+
+#if COUNTERSTRACE_ENABLED
+		return StatsMap.Emplace(InstanceName, FString::Printf(TEXT("%s - %s"), *GetCategoryName().ToString(), *InstanceName.ToString()));
+#else
+		return StatsMap.Add(InstanceName);
+#endif
+	}
+
+	// ctor
+#if COUNTERSTRACE_ENABLED
+	FConcurrentMetasoundInstanceCounter::FStats::FStats(const FString& InName)
+	: TraceCounter(MakeUnique<FCountersTrace::FCounterInt>(TraceCounterNameType_Dynamic, *InName, TraceCounterDisplayHint_None))
+	{
+	}
+#endif
+
+	// FConcurrentMetasoundInstanceCounter::FStats::
+
+	void FConcurrentMetasoundInstanceCounter::FStats::Increment()
+	{
+		ensure(TraceCounter);
+		TraceCounter->Increment();
+		PeakCount = FMath::Max(PeakCount, GetCount());
+	}
+
+
+	void FConcurrentMetasoundInstanceCounter::FStats::Decrement()
+	{
+		ensure(TraceCounter);
+		TraceCounter->Decrement();
+	}
+
+	int64 FConcurrentMetasoundInstanceCounter::FStats::GetCount()
+	{
+		ensure(TraceCounter);
+#if COUNTERSTRACE_ENABLED
+		return TraceCounter->Get();
+#else
+		return TraceCounter->GetValue();
+#endif
+	}
+
+	int64 FConcurrentMetasoundInstanceCounter::FStats::GetPeakCount()
+	{
+		return PeakCount;
+	}
+
+} // namespace Metasound
diff --git a/Engine/Plugins/Runtime/Metasound/Source/MetasoundGenerator/Public/MetasoundGenerator.h b/Engine/Plugins/Runtime/Metasound/Source/MetasoundGenerator/Public/MetasoundGenerator.h
--- a/Engine/Plugins/Runtime/Metasound/Source/MetasoundGenerator/Public/MetasoundGenerator.h
+++ b/Engine/Plugins/Runtime/Metasound/Source/MetasoundGenerator/Public/MetasoundGenerator.h
@@ -3,6 +3,7 @@
 
 #include "MetasoundExecutableOperator.h"
 #include "MetasoundGraphOperator.h"
+#include "MetasoundInstanceCounter.h"
 #include "MetasoundOperatorBuilder.h"
 #include "MetasoundOperatorInterface.h"
 #include "MetasoundParameterPack.h"
@@ -93,7 +94,9 @@ namespace Metasound
 
 	DECLARE_TS_MULTICAST_DELEGATE(FOnSetGraph);
 
-	class METASOUNDGENERATOR_API FMetasoundGenerator : public ISoundGenerator
+	class METASOUNDGENERATOR_API FMetasoundGenerator
+	: public ISoundGenerator
+	, FConcurrentMetasoundInstanceCounter
 	{
 	public:
 		using FOperatorUniquePtr = TUniquePtr<Metasound::IOperator>;
@@ -112,6 +115,8 @@ namespace Metasound
 
 		virtual ~FMetasoundGenerator();
 
+		virtual const FName& GetCategoryName() const override;
+
 		/** Set the value of a graph's input data using the assignment operator.
 		 *
 		 * @param InName - The name of the graph's input data reference.
diff --git a/Engine/Plugins/Runtime/Metasound/Source/MetasoundGenerator/Public/MetasoundInstanceCounter.h b/Engine/Plugins/Runtime/Metasound/Source/MetasoundGenerator/Public/MetasoundInstanceCounter.h
new file mode 100644
--- /dev/null
+++ b/Engine/Plugins/Runtime/Metasound/Source/MetasoundGenerator/Public/MetasoundInstanceCounter.h
@@ -0,0 +1,79 @@
+// Copyright Epic Games, Inc. All Rights Reserved.
+#pragma once
+
+#include "CoreTypes.h"
+#include "HAL/PlatformAtomics.h"
+#include "ProfilingDebugging/CountersTrace.h"
+#include "ProfilingDebugging/CsvProfiler.h"
+#include "Stats/Stats.h"
+#include "Templates/UniquePtr.h"
+#include "UObject/NameTypes.h"
+
+namespace Metasound
+{
+
+#if COUNTERSTRACE_ENABLED
+	using CounterType = FCountersTrace::FCounterInt;
+#else
+	using CounterType = FThreadSafeCounter64;
+#endif
+
+	class FConcurrentMetasoundInstanceCounter
+	{
+	public:
+		// ctor
+		FConcurrentMetasoundInstanceCounter() = default;
+		FConcurrentMetasoundInstanceCounter(const FName& InName);
+		FConcurrentMetasoundInstanceCounter(const FString& InName);
+
+		// dtor
+		virtual ~FConcurrentMetasoundInstanceCounter();
+
+		// for non-RAII clients
+		void Init(const FName& InName);
+		void Init(const FString& InName);
+
+		// pure virtual interface
+		virtual const FName& GetCategoryName() const = 0;
+
+		// static interface
+		static int64 GetCountForName(const FName& InName);
+		static int64 GetPeakCountForName(const FName& InName);
+
+
+	private:
+		struct FStats
+		{
+		public:
+			// ctor
+#if COUNTERSTRACE_ENABLED
+			FStats(const FString& InName);
+#else
+			FStats() = default;
+#endif
+
+			void Increment();
+			void Decrement();
+
+			int64 GetCount();
+			int64 GetPeakCount();
+
+		private:
+			TUniquePtr<CounterType> TraceCounter;
+			int64 PeakCount;
+
+		}; // struct FStats
+
+	// non-static data
+		FName InstanceName;
+
+		void Increment();
+		void Decrement();
+		FStats& GetOrAddStats();
+
+
+	// static data
+		inline static TMap<FName, FStats> StatsMap;
+		inline static FCriticalSection MapCritSec;
+	}; // class FConcurrentMetasoundInstanceCounter
+} // namespace Metasound
