Follow up tweaks to instance counting (code review comments after PFaS went through).

diff --git a/Engine/Plugins/Runtime/Metasound/Source/MetasoundGenerator/Private/MetasoundGenerator.cpp b/Engine/Plugins/Runtime/Metasound/Source/MetasoundGenerator/Private/MetasoundGenerator.cpp
--- a/Engine/Plugins/Runtime/Metasound/Source/MetasoundGenerator/Private/MetasoundGenerator.cpp
+++ b/Engine/Plugins/Runtime/Metasound/Source/MetasoundGenerator/Private/MetasoundGenerator.cpp
@@ -288,12 +288,6 @@ namespace Metasound
 		FDataTransmissionCenter::Get().UnregisterDataChannelIfUnconnected(ParameterPackSendAddress);
 	}
 
-	static const FName InstanceCounterCategory(TEXT("Metasound/Active_Generators"));
-	const FName& FMetasoundGenerator::GetCategoryName() const
-	{
-		return InstanceCounterCategory;
-	}
-
 	FDelegateHandle FMetasoundGenerator::AddGraphSetCallback(FOnSetGraph::FDelegate&& Delegate)
 	{
 		FScopeLock SetPendingGraphLock(&PendingGraphMutex);
@@ -307,7 +301,7 @@ namespace Metasound
 
 	void FMetasoundGenerator::InitBase(const FMetasoundGeneratorInitParams& InInitParams)
 	{
-		FConcurrentMetasoundInstanceCounter::Init(InInitParams.MetaSoundName);
+		InstanceCounter.Init(FName(InInitParams.MetaSoundName));
 		MetasoundName = InInitParams.MetaSoundName;
 		NumChannels = InInitParams.AudioOutputNames.Num();
 		NumSamplesPerExecute = NumChannels * NumFramesPerExecute;
diff --git a/Engine/Plugins/Runtime/Metasound/Source/MetasoundGenerator/Private/MetasoundInstanceCounter.cpp b/Engine/Plugins/Runtime/Metasound/Source/MetasoundGenerator/Private/MetasoundInstanceCounter.cpp
--- a/Engine/Plugins/Runtime/Metasound/Source/MetasoundGenerator/Private/MetasoundInstanceCounter.cpp
+++ b/Engine/Plugins/Runtime/Metasound/Source/MetasoundGenerator/Private/MetasoundInstanceCounter.cpp
@@ -10,12 +10,6 @@ namespace Metasound
 		Increment();
 	}
 
-	FConcurrentMetasoundInstanceCounter::FConcurrentMetasoundInstanceCounter(const FString& InName)
-	: InstanceName(InName)
-	{
-		Increment();
-	}
-
 	// dtor
 	FConcurrentMetasoundInstanceCounter::~FConcurrentMetasoundInstanceCounter()
 	{
@@ -28,12 +22,6 @@ namespace Metasound
 		Increment();
 	}
 
-	void FConcurrentMetasoundInstanceCounter::Init(const FString& InName)
-	{
-		InstanceName = FName(InName);
-		Increment();
-	}
-
 	// static interface
 	int64 FConcurrentMetasoundInstanceCounter::GetCountForName(const FName& InName)
 	{
@@ -69,6 +57,7 @@ namespace Metasound
 		GetOrAddStats().Decrement();
 	}
 
+	static const FString InstanceCounterCategory(TEXT("Metasound/Active_Generators"));
 	FConcurrentMetasoundInstanceCounter::FStats& FConcurrentMetasoundInstanceCounter::GetOrAddStats()
 	{
 		FScopeLock Lock(&MapCritSec);
@@ -80,7 +69,7 @@ namespace Metasound
 		}
 
 #if COUNTERSTRACE_ENABLED
-		return StatsMap.Emplace(InstanceName, FString::Printf(TEXT("%s - %s"), *GetCategoryName().ToString(), *InstanceName.ToString()));
+		return StatsMap.Emplace(InstanceName, FString::Printf(TEXT("%s - %s"), *InstanceCounterCategory, *InstanceName.ToString()));
 #else
 		return StatsMap.Add(InstanceName);
 #endif
diff --git a/Engine/Plugins/Runtime/Metasound/Source/MetasoundGenerator/Public/MetasoundGenerator.h b/Engine/Plugins/Runtime/Metasound/Source/MetasoundGenerator/Public/MetasoundGenerator.h
--- a/Engine/Plugins/Runtime/Metasound/Source/MetasoundGenerator/Public/MetasoundGenerator.h
+++ b/Engine/Plugins/Runtime/Metasound/Source/MetasoundGenerator/Public/MetasoundGenerator.h
@@ -94,9 +94,7 @@ namespace Metasound
 
 	DECLARE_TS_MULTICAST_DELEGATE(FOnSetGraph);
 
-	class METASOUNDGENERATOR_API FMetasoundGenerator
-	: public ISoundGenerator
-	, FConcurrentMetasoundInstanceCounter
+	class METASOUNDGENERATOR_API FMetasoundGenerator : public ISoundGenerator
 	{
 	public:
 		using FOperatorUniquePtr = TUniquePtr<Metasound::IOperator>;
@@ -115,8 +113,6 @@ namespace Metasound
 
 		virtual ~FMetasoundGenerator();
 
-		virtual const FName& GetCategoryName() const override;
-
 		/** Set the value of a graph's input data using the assignment operator.
 		 *
 		 * @param InName - The name of the graph's input data reference.
@@ -242,6 +238,7 @@ namespace Metasound
 
 		virtual TUniquePtr<IOperator> ReleaseGraphOperator();
 		FInputVertexInterfaceData ReleaseInputVertexData();
+		FConcurrentMetasoundInstanceCounter InstanceCounter;
 
 		/** Release the graph operator and remove any references to data owned by
 		 * the graph operator.
diff --git a/Engine/Plugins/Runtime/Metasound/Source/MetasoundGenerator/Public/MetasoundInstanceCounter.h b/Engine/Plugins/Runtime/Metasound/Source/MetasoundGenerator/Public/MetasoundInstanceCounter.h
--- a/Engine/Plugins/Runtime/Metasound/Source/MetasoundGenerator/Public/MetasoundInstanceCounter.h
+++ b/Engine/Plugins/Runtime/Metasound/Source/MetasoundGenerator/Public/MetasoundInstanceCounter.h
@@ -24,17 +24,12 @@ namespace Metasound
 		// ctor
 		FConcurrentMetasoundInstanceCounter() = default;
 		FConcurrentMetasoundInstanceCounter(const FName& InName);
-		FConcurrentMetasoundInstanceCounter(const FString& InName);
 
 		// dtor
 		virtual ~FConcurrentMetasoundInstanceCounter();
 
 		// for non-RAII clients
 		void Init(const FName& InName);
-		void Init(const FString& InName);
-
-		// pure virtual interface
-		virtual const FName& GetCategoryName() const = 0;
 
 		// static interface
 		static int64 GetCountForName(const FName& InName);
