[Backout] - CL 27692414 and CL 27732161 - Smoke Crash
Original CL Desc
-----------------------------------------------------------------
[Metasounds] Adding unreal insights markup:
* Added asset name metadata to "RegisterWithFrontEnd"
* Added per-asset counter for the number of concurrent operators alive at a time.

diff --git a/Engine/Plugins/Runtime/Metasound/Source/MetasoundFrontend/MetasoundFrontend.Build.cs b/Engine/Plugins/Runtime/Metasound/Source/MetasoundFrontend/MetasoundFrontend.Build.cs
--- a/Engine/Plugins/Runtime/Metasound/Source/MetasoundFrontend/MetasoundFrontend.Build.cs
+++ b/Engine/Plugins/Runtime/Metasound/Source/MetasoundFrontend/MetasoundFrontend.Build.cs
@@ -20,7 +20,6 @@ namespace UnrealBuildTool.Rules
 					"CoreUObject",
 					"Serialization",
 					"SignalProcessing",
-					"TraceLog",
 					"MetasoundGraphCore"
 
 				}
diff --git a/Engine/Plugins/Runtime/Metasound/Source/MetasoundFrontend/Private/MetasoundAssetBase.cpp b/Engine/Plugins/Runtime/Metasound/Source/MetasoundFrontend/Private/MetasoundAssetBase.cpp
--- a/Engine/Plugins/Runtime/Metasound/Source/MetasoundFrontend/Private/MetasoundAssetBase.cpp
+++ b/Engine/Plugins/Runtime/Metasound/Source/MetasoundFrontend/Private/MetasoundAssetBase.cpp
@@ -36,8 +36,6 @@
 #include "NodeTemplates/MetasoundFrontendDocumentTemplatePreprocessor.h"
 #include "StructSerializer.h"
 #include "Templates/SharedPointer.h"
-#include "Trace/Trace.h"
-#include "Trace/Trace.inl"
 #include "UObject/MetaData.h"
 
 #define LOCTEXT_NAMESPACE "MetaSound"
@@ -180,19 +178,12 @@ namespace Metasound
 
 const FString FMetasoundAssetBase::FileExtension(TEXT(".metasound"));
 
-UE_TRACE_EVENT_BEGIN(Cpu, RegisterGraphWithFrontEnd, NoSync)
-	UE_TRACE_EVENT_FIELD(UE::Trace::WideString, AssetName)
-UE_TRACE_EVENT_END()
-
 void FMetasoundAssetBase::RegisterGraphWithFrontend(Metasound::Frontend::FMetaSoundAssetRegistrationOptions InRegistrationOptions)
 {
 	using namespace Metasound;
 	using namespace Metasound::Frontend;
 
 	METASOUND_TRACE_CPUPROFILER_EVENT_SCOPE(MetaSoundAssetBase::RegisterGraphWithFrontend);
-	UE_TRACE_LOG_SCOPED_T(Cpu, RegisterGraphWithFrontEnd, CpuChannel)
-		<< RegisterGraphWithFrontEnd.AssetName(*GetOwningAssetName());
-	
 	if (!InRegistrationOptions.bForceReregister)
 	{
 		if (IsRegistered())
diff --git a/Engine/Plugins/Runtime/Metasound/Source/MetasoundGenerator/Private/MetasoundGenerator.cpp b/Engine/Plugins/Runtime/Metasound/Source/MetasoundGenerator/Private/MetasoundGenerator.cpp
--- a/Engine/Plugins/Runtime/Metasound/Source/MetasoundGenerator/Private/MetasoundGenerator.cpp
+++ b/Engine/Plugins/Runtime/Metasound/Source/MetasoundGenerator/Private/MetasoundGenerator.cpp
@@ -301,7 +301,6 @@ namespace Metasound
 
 	void FMetasoundGenerator::InitBase(const FMetasoundGeneratorInitParams& InInitParams)
 	{
-		InstanceCounter.Init(FName(InInitParams.MetaSoundName));
 		MetasoundName = InInitParams.MetaSoundName;
 		NumChannels = InInitParams.AudioOutputNames.Num();
 		NumSamplesPerExecute = NumChannels * NumFramesPerExecute;
diff --git a/Engine/Plugins/Runtime/Metasound/Source/MetasoundGenerator/Private/MetasoundInstanceCounter.cpp b/Engine/Plugins/Runtime/Metasound/Source/MetasoundGenerator/Private/MetasoundInstanceCounter.cpp
deleted file mode 100644
diff --git a/Engine/Plugins/Runtime/Metasound/Source/MetasoundGenerator/Public/MetasoundGenerator.h b/Engine/Plugins/Runtime/Metasound/Source/MetasoundGenerator/Public/MetasoundGenerator.h
--- a/Engine/Plugins/Runtime/Metasound/Source/MetasoundGenerator/Public/MetasoundGenerator.h
+++ b/Engine/Plugins/Runtime/Metasound/Source/MetasoundGenerator/Public/MetasoundGenerator.h
@@ -3,7 +3,6 @@
 
 #include "MetasoundExecutableOperator.h"
 #include "MetasoundGraphOperator.h"
-#include "MetasoundInstanceCounter.h"
 #include "MetasoundOperatorBuilder.h"
 #include "MetasoundOperatorInterface.h"
 #include "MetasoundParameterPack.h"
@@ -238,7 +237,6 @@ namespace Metasound
 
 		virtual TUniquePtr<IOperator> ReleaseGraphOperator();
 		FInputVertexInterfaceData ReleaseInputVertexData();
-		FConcurrentMetasoundInstanceCounter InstanceCounter;
 
 		/** Release the graph operator and remove any references to data owned by
 		 * the graph operator.
diff --git a/Engine/Plugins/Runtime/Metasound/Source/MetasoundGenerator/Public/MetasoundInstanceCounter.h b/Engine/Plugins/Runtime/Metasound/Source/MetasoundGenerator/Public/MetasoundInstanceCounter.h
deleted file mode 100644
