[MetaSounds] Automated test & fix-up for getting class, node, access type, and ctor state of node inputs/outputs
#rnx

diff --git a/Engine/Plugins/Runtime/Metasound/Source/MetasoundEngine/Private/MetasoundBuilderSubsystem.cpp b/Engine/Plugins/Runtime/Metasound/Source/MetasoundEngine/Private/MetasoundBuilderSubsystem.cpp
--- a/Engine/Plugins/Runtime/Metasound/Source/MetasoundEngine/Private/MetasoundBuilderSubsystem.cpp
+++ b/Engine/Plugins/Runtime/Metasound/Source/MetasoundEngine/Private/MetasoundBuilderSubsystem.cpp
@@ -866,6 +866,12 @@ FMetasoundFrontendLiteral UMetaSoundBuilderBase::GetNodeInputClassDefault(const
 	return { };
 }
 
+bool UMetaSoundBuilderBase::GetNodeInputIsConstructorPin(const FMetaSoundBuilderNodeInputHandle& InputHandle) const
+{
+	const EMetasoundFrontendVertexAccessType AccessType = Builder.GetNodeInputAccessType(InputHandle.NodeID, InputHandle.VertexID);
+	return AccessType == EMetasoundFrontendVertexAccessType::Value;
+}
+
 void UMetaSoundBuilderBase::GetNodeOutputData(const FMetaSoundBuilderNodeOutputHandle& OutputHandle, FName& Name, FName& DataType, EMetaSoundBuilderResult& OutResult)
 {
 	if (const FMetasoundFrontendVertex* Vertex = Builder.FindNodeOutput(OutputHandle.NodeID, OutputHandle.VertexID))
@@ -882,6 +888,12 @@ void UMetaSoundBuilderBase::GetNodeOutputData(const FMetaSoundBuilderNodeOutputH
 	}
 }
 
+bool UMetaSoundBuilderBase::GetNodeOutputIsConstructorPin(const FMetaSoundBuilderNodeOutputHandle& OutputHandle) const
+{
+	const EMetasoundFrontendVertexAccessType AccessType = Builder.GetNodeOutputAccessType(OutputHandle.NodeID, OutputHandle.VertexID);
+	return AccessType == EMetasoundFrontendVertexAccessType::Value;
+}
+
 void UMetaSoundBuilderBase::UpdateDependencyClassNames(const TMap<FMetasoundFrontendClassName, FMetasoundFrontendClassName>& OldToNewReferencedClassNames)
 {
 	Builder.UpdateDependencyClassNames(OldToNewReferencedClassNames);
diff --git a/Engine/Plugins/Runtime/Metasound/Source/MetasoundEngine/Public/MetasoundBuilderSubsystem.h b/Engine/Plugins/Runtime/Metasound/Source/MetasoundEngine/Public/MetasoundBuilderSubsystem.h
--- a/Engine/Plugins/Runtime/Metasound/Source/MetasoundEngine/Public/MetasoundBuilderSubsystem.h
+++ b/Engine/Plugins/Runtime/Metasound/Source/MetasoundEngine/Public/MetasoundBuilderSubsystem.h
@@ -281,10 +281,18 @@ public:
 	UFUNCTION(BlueprintCallable, Category = "Audio|MetaSound|Builder", meta = (ExpandEnumAsExecs = "OutResult"))
 	FMetasoundFrontendLiteral GetNodeInputClassDefault(const FMetaSoundBuilderNodeInputHandle& InputHandle, EMetaSoundBuilderResult& OutResult);
 
+	// Returns whether the given node input is a constructor pin
+	UFUNCTION(BlueprintCallable, Category = "Audio|MetaSound|Builder")
+	bool GetNodeInputIsConstructorPin(const FMetaSoundBuilderNodeInputHandle& InputHandle) const;
+
 	// Returns node output's data if valid (including things like name and datatype).
 	UFUNCTION(BlueprintCallable, Category = "Audio|MetaSound|Builder", meta = (ExpandEnumAsExecs = "OutResult"))
 	void GetNodeOutputData(const FMetaSoundBuilderNodeOutputHandle& OutputHandle, FName& Name, FName& DataType, EMetaSoundBuilderResult& OutResult);
 	
+	// Returns whether the given node output is a constructor pin
+	UFUNCTION(BlueprintCallable, Category = "Audio|MetaSound|Builder")
+	bool GetNodeOutputIsConstructorPin(const FMetaSoundBuilderNodeOutputHandle& OutputHandle) const;
+
 	// Return the asset referenced by this preset builder. Returns nullptr if the builder is not a preset.
 	UFUNCTION(BlueprintCallable, Category = "Audio|MetaSound|Builder")
 	UObject* GetReferencedPresetAsset() const;
diff --git a/Engine/Plugins/Runtime/Metasound/Source/MetasoundEngineTest/Private/EngineTestMetaSoundBuilder.h b/Engine/Plugins/Runtime/Metasound/Source/MetasoundEngineTest/Private/EngineTestMetaSoundBuilder.h
--- a/Engine/Plugins/Runtime/Metasound/Source/MetasoundEngineTest/Private/EngineTestMetaSoundBuilder.h
+++ b/Engine/Plugins/Runtime/Metasound/Source/MetasoundEngineTest/Private/EngineTestMetaSoundBuilder.h
@@ -124,7 +124,7 @@ bool FBuilderRemoveFromRootLatentCommand::Update()
 }
 
 // Creates a collection of MetaSound patches from builders and exercises bindings all input and output interfaces by connecting and disconnecting using various builder API calls.
-IMPLEMENT_SIMPLE_AUTOMATION_TEST(FAudioMetaSoundBuilderInterfaceBindingConnectAndDisconnect, "Audio.Metasound.InterfaceBindingConnectAndDisconnect", EAutomationTestFlags::EditorContext | EAutomationTestFlags::EngineFilter)
+IMPLEMENT_SIMPLE_AUTOMATION_TEST(FAudioMetaSoundBuilderInterfaceBindingConnectAndDisconnect, "Audio.Metasound.Builder.InterfaceBindingConnectAndDisconnect", EAutomationTestFlags::EditorContext | EAutomationTestFlags::EngineFilter)
 bool FAudioMetaSoundBuilderInterfaceBindingConnectAndDisconnect::RunTest(const FString& Parameters)
 {
 	using namespace EngineTestMetaSoundPatchBuilderPrivate;
@@ -191,4 +191,84 @@ bool FAudioMetaSoundBuilderInterfaceBindingConnectAndDisconnect::RunTest(const F
 
 	return true;
 }
+
+// This test creates a MetaSound patch, then adds and connects sin oscillator, attempts to retrieve a default input set, then clears it, and finally retrieves the class default
+IMPLEMENT_SIMPLE_AUTOMATION_TEST(FAudioMetasoundNodeClassQueryFunctions, "Audio.Metasound.Builder.NodeClassQueryFunctions", EAutomationTestFlags::EditorContext | EAutomationTestFlags::EngineFilter)
+bool FAudioMetasoundNodeClassQueryFunctions::RunTest(const FString& Parameters)
+{
+	using namespace EngineTestMetaSoundPatchBuilderPrivate;
+	using namespace Metasound;
+	using namespace Metasound::Frontend;
+
+	constexpr float NodeDefaultFreq = 100.0f;
+
+	EMetaSoundBuilderResult Result = EMetaSoundBuilderResult::Failed;
+	UMetaSoundPatchBuilder& Builder = CreatePatchBuilderChecked(*this, "DefaultLiteralAssignment", { });
+
+	// Sine Oscillator Node
+	const FMetaSoundNodeHandle OscNode = Builder.AddNodeByClassName({ "UE", "Sine", "Audio" }, Result, 1);
+	AddErrorIfFalse(Result == EMetaSoundBuilderResult::Succeeded && OscNode.IsSet(), TEXT("Failed to create new MetaSound node by class name"));
+
+	// Make connections:
+	const FMetaSoundBuilderNodeInputHandle OscNodeFrequencyInput = Builder.FindNodeInputByName(OscNode, "Frequency", Result);
+	AddErrorIfFalse(Result == EMetaSoundBuilderResult::Succeeded && OscNodeFrequencyInput.IsSet(), TEXT("Failed to find Sine Oscillator node input 'Frequency'"));
+
+	FMetasoundFrontendLiteral Literal;
+	Literal.Set(100.0f);
+
+	Builder.SetNodeInputDefault(OscNodeFrequencyInput, Literal, Result);
+	AddErrorIfFalse(Result == EMetaSoundBuilderResult::Succeeded, TEXT("Failed to 'SetNodeInputDefault'"));
+
+	Literal = Builder.GetNodeInputDefault(OscNodeFrequencyInput, Result);
+	AddErrorIfFalse(Result == EMetaSoundBuilderResult::Succeeded, TEXT("Failed to 'GetNodeInputDefault'"));
+
+	float Freq = 0.0f;
+	bool bLiteralIsFloat = Literal.TryGet(Freq);
+	AddErrorIfFalse(bLiteralIsFloat, TEXT("Failed to retrieve node literal as 'float'"));
+	AddErrorIfFalse(FMath::IsNearlyEqual(NodeDefaultFreq, Freq), TEXT("'Freq' node default not set to provided default"));
+
+	Builder.RemoveNodeInputDefault(OscNodeFrequencyInput, Result);
+	AddErrorIfFalse(Result == EMetaSoundBuilderResult::Succeeded, TEXT("Failed to 'RemoveNodeInputDefault'"));
+
+	Literal = Builder.GetNodeInputClassDefault(OscNodeFrequencyInput, Result);
+	AddErrorIfFalse(Result == EMetaSoundBuilderResult::Succeeded, TEXT("Failed to 'GetNodeInputClassDefault' on c++ Sin Osc node"));
+
+	bLiteralIsFloat = Literal.TryGet(Freq);
+	AddErrorIfFalse(bLiteralIsFloat, TEXT("Failed to retrieve class input literal as 'float'"));
+	AddErrorIfFalse(FMath::IsNearlyEqual(440.0f, Freq), TEXT("'Freq' node default not set to expected sin freq default of 440.0f Hz"));
+
+	// Output Default Test
+	{
+		Literal.Set(123.0f);
+		const FMetaSoundBuilderNodeInputHandle InputTest = Builder.AddGraphOutputNode("Output", GetMetasoundDataTypeName<float>(), Literal, Result);
+		AddErrorIfFalse(Result == EMetaSoundBuilderResult::Succeeded && InputTest.IsSet(), TEXT("Failed to create new MetaSound graph input"));
+
+		bLiteralIsFloat = Literal.TryGet(Freq);
+		AddErrorIfFalse(bLiteralIsFloat, TEXT("Failed to retrieve class input literal as 'float'"));
+		AddErrorIfFalse(FMath::IsNearlyEqual(123.0f, Freq), TEXT("'Freq' node default not set to expected sin freq default of 440.0f Hz"));
+
+		const bool bIsConstructorPin = Builder.GetNodeInputIsConstructorPin(InputTest);
+		AddErrorIfFalse(!bIsConstructorPin, TEXT("Pin is not constructor pin but GetNodeInputIsConstructorPin is returning true"));
+	}
+
+	// Output Ctor tests
+	{
+		const FMetaSoundBuilderNodeInputHandle CtorTest = Builder.AddGraphOutputNode("CtorOutputTest", GetMetasoundDataTypeName<float>(), Literal, Result, true);
+		AddErrorIfFalse(Result == EMetaSoundBuilderResult::Succeeded && CtorTest.IsSet(), TEXT("Failed to create new MetaSound graph input"));
+
+		const bool bIsConstructorPin = Builder.GetNodeInputIsConstructorPin(CtorTest);
+		AddErrorIfFalse(bIsConstructorPin, TEXT("Output handle is constructor pin but GetNodeInputIsConstructorPin is returning false"));
+	}
+
+	// Input Ctor tests
+	{
+		const FMetaSoundBuilderNodeOutputHandle CtorTest = Builder.AddGraphInputNode("CtorOutput", GetMetasoundDataTypeName<float>(), Literal, Result, true);
+		AddErrorIfFalse(Result == EMetaSoundBuilderResult::Succeeded && CtorTest.IsSet(), TEXT("Failed to create new MetaSound graph input"));
+
+		const bool bIsConstructorPin = Builder.GetNodeOutputIsConstructorPin(CtorTest);
+		AddErrorIfFalse(bIsConstructorPin, TEXT("Input handle is constructor pin but GetNodeInputIsConstructorPin is returning false"));
+	}
+
+	return true;
+}
 #endif // WITH_DEV_AUTOMATION_TESTS
diff --git a/Engine/Plugins/Runtime/Metasound/Source/MetasoundEngineTest/Private/EngineTestMetasoundSource.cpp b/Engine/Plugins/Runtime/Metasound/Source/MetasoundEngineTest/Private/EngineTestMetasoundSource.cpp
--- a/Engine/Plugins/Runtime/Metasound/Source/MetasoundEngineTest/Private/EngineTestMetasoundSource.cpp
+++ b/Engine/Plugins/Runtime/Metasound/Source/MetasoundEngineTest/Private/EngineTestMetasoundSource.cpp
@@ -400,7 +400,7 @@ bool FAudioMetasoundSourceTest::RunTest(const FString& Parameters)
  }
 
 // This test creates a MetaSound source from a SourceBuilder, adds a simple sine tone generator with a connected graph input frequency, and attempts to audition it.
-IMPLEMENT_SIMPLE_AUTOMATION_TEST(FAudioMetasoundSourceBuilderTest, "Audio.Metasound.AuditionMetasoundSource", EAutomationTestFlags::EditorContext | EAutomationTestFlags::EngineFilter)
+IMPLEMENT_SIMPLE_AUTOMATION_TEST(FAudioMetasoundSourceBuilderTest, "Audio.Metasound.Builder.AuditionMetasoundSource", EAutomationTestFlags::EditorContext | EAutomationTestFlags::EngineFilter)
 bool FAudioMetasoundSourceBuilderTest::RunTest(const FString& Parameters)
 {
 	using namespace EngineTestMetaSoundPatchBuilderPrivate;
@@ -427,7 +427,7 @@ bool FAudioMetasoundSourceBuilderTest::RunTest(const FString& Parameters)
 }
 
 // This test creates a MetaSound source from a SourceBuilder, adds a simple sine tone generator with a connected graph input frequency, and attempts to change the frequency and audition it rapidly 100 times.
-IMPLEMENT_SIMPLE_AUTOMATION_TEST(FAudioMetasoundSourceBuilderTestSpamAudition, "Audio.Metasound.SpamAuditionMetasoundSource", EAutomationTestFlags::EditorContext | EAutomationTestFlags::EngineFilter)
+IMPLEMENT_SIMPLE_AUTOMATION_TEST(FAudioMetasoundSourceBuilderTestSpamAudition, "Audio.Metasound.Builder.SpamAuditionMetasoundSource", EAutomationTestFlags::EditorContext | EAutomationTestFlags::EngineFilter)
 bool FAudioMetasoundSourceBuilderTestSpamAudition::RunTest(const FString& Parameters)
 {
 	using namespace EngineTestMetaSoundPatchBuilderPrivate;
@@ -472,7 +472,7 @@ bool FAudioMetasoundSourceBuilderTestSpamAudition::RunTest(const FString& Parame
 // 7. Restarts audition of the original AudioComponent (which should continue after all the prior changes)
 // 8. Swaps to use a tri tone generator
 // 9. Stops original component, completing the test
-IMPLEMENT_SIMPLE_AUTOMATION_TEST(FAudioMetasoundSourceBuilderLiveUpdateNode, "Audio.Metasound.LiveUpdateMultipleMetaSoundSources", EAutomationTestFlags::EditorContext | EAutomationTestFlags::EngineFilter)
+IMPLEMENT_SIMPLE_AUTOMATION_TEST(FAudioMetasoundSourceBuilderLiveUpdateNode, "Audio.Metasound.Builder.LiveUpdateMultipleMetaSoundSources", EAutomationTestFlags::EditorContext | EAutomationTestFlags::EngineFilter)
 bool FAudioMetasoundSourceBuilderLiveUpdateNode::RunTest(const FString& Parameters)
 {
 	using namespace EngineTestMetaSoundSourcePrivate;
@@ -547,7 +547,7 @@ bool FAudioMetasoundSourceBuilderLiveUpdateNode::RunTest(const FString& Paramete
 // This test creates a MetaSound source from a SourceBuilder, adds a simple sine tone generator with a connected graph input frequency, attempts to audition it,
 // disconnects frequency input, sets the sinosc frequency literal value to a new value, and finally removes the literal value default to have it return to the
 // class default.
-IMPLEMENT_SIMPLE_AUTOMATION_TEST(FAudioMetasoundSourceBuilderLiveUpdateLiteral, "Audio.Metasound.LiveUpdateLiteralMetaSoundSource", EAutomationTestFlags::EditorContext | EAutomationTestFlags::EngineFilter)
+IMPLEMENT_SIMPLE_AUTOMATION_TEST(FAudioMetasoundSourceBuilderLiveUpdateLiteral, "Audio.Metasound.Builder.LiveUpdateLiteralMetaSoundSource", EAutomationTestFlags::EditorContext | EAutomationTestFlags::EngineFilter)
 bool FAudioMetasoundSourceBuilderLiveUpdateLiteral::RunTest(const FString& Parameters)
 {
 	using namespace EngineTestMetaSoundPatchBuilderPrivate;
@@ -594,7 +594,7 @@ bool FAudioMetasoundSourceBuilderLiveUpdateLiteral::RunTest(const FString& Param
 }
 
 // This test creates a MetaSound source from a SourceBuilder, then adds and finally removes an interface using the builder API, and verifies it as well as its members were added to the document.
-IMPLEMENT_SIMPLE_AUTOMATION_TEST(FAudioMetasoundSourceBuilderMutateInterface, "Audio.Metasound.MutateInterface", EAutomationTestFlags::EditorContext | EAutomationTestFlags::EngineFilter)
+IMPLEMENT_SIMPLE_AUTOMATION_TEST(FAudioMetasoundSourceBuilderMutateInterface, "Audio.Metasound.Builder.MutateInterface", EAutomationTestFlags::EditorContext | EAutomationTestFlags::EngineFilter)
 bool FAudioMetasoundSourceBuilderMutateInterface::RunTest(const FString& Parameters)
 {
 	using namespace Audio;
@@ -653,7 +653,7 @@ bool FAudioMetasoundSourceBuilderMutateInterface::RunTest(const FString& Paramet
 }
 
 // This test creates a MetaSound source from a SourceBuilder, then adds and connects saw oscillator nodes setting their input freq to a new chromatic tone which cascades up, and finally removes all of the added nodes.
-IMPLEMENT_SIMPLE_AUTOMATION_TEST(FAudioMetasoundSourceBuilderAddRemoveNodes, "Audio.Metasound.AddRemoveNodes", EAutomationTestFlags::EditorContext | EAutomationTestFlags::EngineFilter)
+IMPLEMENT_SIMPLE_AUTOMATION_TEST(FAudioMetasoundSourceBuilderAddRemoveNodes, "Audio.Metasound.Builder.AddRemoveNodes", EAutomationTestFlags::EditorContext | EAutomationTestFlags::EngineFilter)
 bool FAudioMetasoundSourceBuilderAddRemoveNodes::RunTest(const FString& Parameters)
 {
 	using namespace Audio;
diff --git a/Engine/Plugins/Runtime/Metasound/Source/MetasoundFrontend/Private/MetasoundFrontendDocumentBuilder.cpp b/Engine/Plugins/Runtime/Metasound/Source/MetasoundFrontend/Private/MetasoundFrontendDocumentBuilder.cpp
--- a/Engine/Plugins/Runtime/Metasound/Source/MetasoundFrontend/Private/MetasoundFrontendDocumentBuilder.cpp
+++ b/Engine/Plugins/Runtime/Metasound/Source/MetasoundFrontend/Private/MetasoundFrontendDocumentBuilder.cpp
@@ -1477,29 +1477,6 @@ const FMetasoundFrontendVertex* FMetaSoundFrontendDocumentBuilder::FindNodeInput
 	return NodeCache.FindInputVertex(InNodeID, InVertexName);
 }
 
-const FMetasoundFrontendClassInput* FMetaSoundFrontendDocumentBuilder::FindNodeInputClassInput(const FGuid& InNodeID, const FName& InVertexName) const
-{
-	using namespace Metasound::Frontend;
-
-	const FMetasoundFrontendDocument& Document = GetDocument();
-	const IDocumentGraphNodeCache& NodeCache = DocumentCache->GetNodeCache();
-
-	if (const FMetasoundFrontendNode* Node = NodeCache.FindNode(InNodeID))
-	{
-		if (const FMetasoundFrontendClass* Class = DocumentCache->FindDependency(Node->ClassID))
-		{
-			if (const FMetasoundFrontendVertex* Vertex = NodeCache.FindInputVertex(InNodeID, InVertexName))
-			{
-				const FName VertexName = Vertex->Name;
-				auto InputMatchesVertexName = [&VertexName](const FMetasoundFrontendClassVertex& ClassVertex) { return ClassVertex.Name == VertexName; };
-				return Class->Interface.Inputs.FindByPredicate(InputMatchesVertexName);
-			}
-		}
-	}
-
-	return nullptr;
-}
-
 TArray<const FMetasoundFrontendVertex*> FMetaSoundFrontendDocumentBuilder::FindNodeInputs(const FGuid& InNodeID, FName TypeName) const
 {
 	return DocumentCache->GetNodeCache().FindNodeInputs(InNodeID, TypeName);
@@ -1542,29 +1519,6 @@ const FMetasoundFrontendVertex* FMetaSoundFrontendDocumentBuilder::FindNodeOutpu
 	return NodeCache.FindOutputVertex(InNodeID, InVertexName);
 }
 
-const FMetasoundFrontendClassOutput* FMetaSoundFrontendDocumentBuilder::FindNodeOutputClassOutput(const FGuid& InNodeID, const FName& InVertexName) const
-{
-	using namespace Metasound::Frontend;
-
-	const FMetasoundFrontendDocument& Document = GetDocument();
-	const IDocumentGraphNodeCache& NodeCache = DocumentCache->GetNodeCache();
-
-	if (const FMetasoundFrontendNode* Node = NodeCache.FindNode(InNodeID))
-	{
-		if (const FMetasoundFrontendClass* Class = DocumentCache->FindDependency(Node->ClassID))
-		{
-			if (const FMetasoundFrontendVertex* Vertex = NodeCache.FindOutputVertex(InNodeID, InVertexName))
-			{
-				const FName VertexName = Vertex->Name;
-				auto OutputMatchesVertexName = [&VertexName](const FMetasoundFrontendClassVertex& ClassVertex) { return ClassVertex.Name == VertexName; };
-				return Class->Interface.Outputs.FindByPredicate(OutputMatchesVertexName);
-			}
-		}
-	}
-
-	return nullptr;
-}
-
 TArray<const FMetasoundFrontendVertex*> FMetaSoundFrontendDocumentBuilder::FindNodeOutputs(const FGuid& InNodeID, FName TypeName) const
 {
 	return DocumentCache->GetNodeCache().FindNodeOutputs(InNodeID, TypeName);
@@ -1708,10 +1662,27 @@ const FMetasoundFrontendLiteral* FMetaSoundFrontendDocumentBuilder::GetNodeInput
 		{
 			if (const FMetasoundFrontendClass* Class = DocumentCache->FindDependency(Node.ClassID))
 			{
-				auto IsClassInput = [VertexName = Vertex->Name](const FMetasoundFrontendClassInput& Input) { return Input.Name == VertexName; };
-				if (const FMetasoundFrontendClassInput* ClassInput = Class->Interface.Inputs.FindByPredicate(IsClassInput))
+				const EMetasoundFrontendClassType ClassType = Class->Metadata.GetType();
+				switch (ClassType)
 				{
-					return &ClassInput->DefaultLiteral;
+					case EMetasoundFrontendClassType::Output:
+					{
+						const FMetasoundFrontendClassInput& ClassInput = Class->Interface.Inputs.Last();
+						return &ClassInput.DefaultLiteral;
+					}
+					break;
+
+					default:
+					{
+						auto IsClassInput = [VertexName = Vertex->Name](const FMetasoundFrontendClassInput& Input) { return Input.Name == VertexName; };
+						if (const FMetasoundFrontendClassInput* ClassInput = Class->Interface.Inputs.FindByPredicate(IsClassInput))
+						{
+							return &ClassInput->DefaultLiteral;
+						}
+						static_assert(static_cast<uint32>(EMetasoundFrontendClassType::Invalid) == 10, "Potential missing case coverage for EMetasoundFrontendClassType "
+							"(default may not be sufficient for newly added class types)");
+					}
+					break;
 				}
 			}
 		}
diff --git a/Engine/Plugins/Runtime/Metasound/Source/MetasoundFrontend/Public/MetasoundFrontendDocumentBuilder.h b/Engine/Plugins/Runtime/Metasound/Source/MetasoundFrontend/Public/MetasoundFrontendDocumentBuilder.h
--- a/Engine/Plugins/Runtime/Metasound/Source/MetasoundFrontend/Public/MetasoundFrontendDocumentBuilder.h
+++ b/Engine/Plugins/Runtime/Metasound/Source/MetasoundFrontend/Public/MetasoundFrontendDocumentBuilder.h
@@ -259,8 +259,6 @@ private:
 	FMetasoundFrontendNode* AddNodeInternal(const FMetasoundFrontendClassMetadata& InClassMetadata, Metasound::Frontend::FFinalizeNodeFunctionRef FinalizeNode, FGuid InNodeID = FGuid::NewGuid());
 
 	const TSet<FMetasoundFrontendVersion>* FindNodeClassInterfaces(const FGuid& InNodeID) const;
-	const FMetasoundFrontendClassInput* FindNodeInputClassInput(const FGuid& InNodeID, const FName& InVertexName) const;
-	const FMetasoundFrontendClassOutput* FindNodeOutputClassOutput(const FGuid& InNodeID, const FName& InVertexName) const;
 
 	void ReloadCacheInternal();
 
