Making space bar preview a metasound

* Refactored metasound source and patch asset actions to use new asset definitions
* Fixed issue with preset icons not showing up on restart
* Utilized sound base asset def

diff --git a/Engine/Plugins/Runtime/Metasound/Source/MetasoundEditor/MetasoundEditor.Build.cs b/Engine/Plugins/Runtime/Metasound/Source/MetasoundEditor/MetasoundEditor.Build.cs
--- a/Engine/Plugins/Runtime/Metasound/Source/MetasoundEditor/MetasoundEditor.Build.cs
+++ b/Engine/Plugins/Runtime/Metasound/Source/MetasoundEditor/MetasoundEditor.Build.cs
@@ -45,8 +45,10 @@ namespace UnrealBuildTool.Rules
 				{
 					"ApplicationCore",
 					"AudioAnalyzer",
+					"AudioEditor",
 					"AudioMixer",
 					"AudioWidgets",
+					"AssetDefinition",
 					"ClassViewer",
 					"ContentBrowser",
 					"Core",
diff --git a/Engine/Plugins/Runtime/Metasound/Source/MetasoundEditor/Private/MetasoundAssetDefinitions.cpp b/Engine/Plugins/Runtime/Metasound/Source/MetasoundEditor/Private/MetasoundAssetDefinitions.cpp
new file mode 100644
--- /dev/null
+++ b/Engine/Plugins/Runtime/Metasound/Source/MetasoundEditor/Private/MetasoundAssetDefinitions.cpp
@@ -0,0 +1,419 @@
+// Copyright Epic Games, Inc. All Rights Reserved.
+
+#include "MetasoundAssetDefinitions.h"
+
+#include "AssetRegistry/AssetRegistryModule.h"
+
+#include "Components/AudioComponent.h"
+#include "ContentBrowserMenuContexts.h"
+#include "ContentBrowserModule.h"
+#include "Editor/EditorPerProjectUserSettings.h"
+#include "IAssetTools.h"
+#include "IContentBrowserSingleton.h"
+#include "Metasound.h"
+#include "MetasoundAssetBase.h"
+#include "MetasoundAssetManager.h"
+#include "MetasoundBuilderSubsystem.h"
+#include "MetasoundSource.h"
+#include "MetasoundEditor.h"
+#include "MetasoundEditorGraphBuilder.h"
+#include "MetasoundEditorModule.h"
+#include "MetasoundEditorSettings.h"
+#include "MetasoundEditorSubsystem.h"
+#include "MetasoundFactory.h"
+#include "MetasoundUObjectRegistry.h"
+#include "Misc/AssetRegistryInterface.h"
+#include "Modules/ModuleManager.h"
+#include "ObjectEditorUtils.h"
+#include "Styling/AppStyle.h"
+#include "Styling/CoreStyle.h"
+#include "Styling/ISlateStyle.h"
+#include "Styling/SlateStyleRegistry.h"
+#include "Styling/SlateTypes.h"
+#include "ToolMenus.h"
+#include "ToolMenuSection.h"
+#include "Widgets/Input/SButton.h"
+#include "Widgets/SOverlay.h"
+
+
+#define LOCTEXT_NAMESPACE "MetaSoundEditor"
+
+
+static bool GetIsPreset(const FSoftObjectPath& InSourcePath)
+{
+	using namespace Metasound;
+	using namespace Metasound::Editor;
+	using namespace Metasound::Frontend;
+
+	bool bIsPreset = false;
+	// If object in memory, try to resolve but not load
+	if (UObject* Object = InSourcePath.ResolveObject())
+	{
+		const FMetasoundAssetBase* MetaSoundAsset = IMetasoundUObjectRegistry::Get().GetObjectAsAssetBase(Object);
+		if (MetaSoundAsset)
+		{
+			bIsPreset = MetaSoundAsset->GetDocumentChecked().RootGraph.PresetOptions.bIsPreset;
+		}
+	}
+	// Otherwise, try to pull from asset registry, but avoid load as this call
+	// would then be slow and is called from the ContentBrowser many times
+	else
+	{
+		FAssetData IconAssetData;
+		const FAssetRegistryModule& AssetRegistryModule = FModuleManager::GetModuleChecked<FAssetRegistryModule>("AssetRegistry");
+		UE::AssetRegistry::EExists Exists = AssetRegistryModule.TryGetAssetByObjectPath(InSourcePath, IconAssetData);
+		if (Exists == UE::AssetRegistry::EExists::Exists)
+		{
+			IconAssetData.GetTagValue(AssetTags::IsPreset, bIsPreset);
+		}
+	}
+
+	return bIsPreset;
+}
+
+static const FSlateBrush* GetClassBrush(const FAssetData& InAssetData, FName InClassName, bool bIsThumbnail = false)
+{
+	using namespace Metasound::Editor;
+
+	const bool bIsPreset = GetIsPreset(InAssetData.ToSoftObjectPath());
+	FString BrushName = FString::Printf(TEXT("MetasoundEditor.%s"), *InClassName.ToString());
+	if (bIsPreset)
+	{
+		BrushName += TEXT(".Preset");
+	}
+	BrushName += bIsThumbnail ? TEXT(".Thumbnail") : TEXT(".Icon");
+
+	return &Style::GetSlateBrushSafe(FName(*BrushName));
+}
+
+
+FLinearColor UAssetDefinition_MetaSoundPatch::GetAssetColor() const
+{
+	if (const ISlateStyle* MetasoundStyle = FSlateStyleRegistry::FindSlateStyle("MetaSoundStyle"))
+	{
+		return MetasoundStyle->GetColor("MetaSoundPatch.Color").ToFColorSRGB();
+	}
+
+	return FColor::White;
+}
+
+TSoftClassPtr<UObject> UAssetDefinition_MetaSoundPatch::GetAssetClass() const
+{
+	return UMetaSoundPatch::StaticClass();
+}
+
+TConstArrayView<FAssetCategoryPath> UAssetDefinition_MetaSoundPatch::GetAssetCategories() const
+{
+	static const auto Pinned_Categories = { EAssetCategoryPaths::Audio };
+	static const auto Categories = { EAssetCategoryPaths::Audio / LOCTEXT("AssetSoundMetaSoundsSubMenu", "MetaSounds") };
+
+	if (GetDefault<UMetasoundEditorSettings>()->bPinMetaSoundPatchInAssetMenu)
+	{
+		return Pinned_Categories;
+	}
+
+	return Categories;
+}
+
+EAssetCommandResult UAssetDefinition_MetaSoundPatch::OpenAssets(const FAssetOpenArgs& OpenArgs) const
+{
+	for (UMetaSoundPatch* Metasound : OpenArgs.LoadObjects<UMetaSoundPatch>())
+	{
+		TSharedRef<Metasound::Editor::FEditor> NewEditor = MakeShared<Metasound::Editor::FEditor>();
+		NewEditor->InitMetasoundEditor(OpenArgs.GetToolkitMode(), OpenArgs.ToolkitHost, Metasound);
+	}
+	return EAssetCommandResult::Handled;
+}
+
+const FSlateBrush* UAssetDefinition_MetaSoundPatch::GetThumbnailBrush(const FAssetData& InAssetData, const FName InClassName) const
+{
+	constexpr bool bIsThumbnail = true;
+	return GetClassBrush(InAssetData, InClassName, bIsThumbnail);
+}
+
+const FSlateBrush* UAssetDefinition_MetaSoundPatch::GetIconBrush(const FAssetData& InAssetData, const FName InClassName) const
+{
+	return GetClassBrush(InAssetData, InClassName);
+}
+
+FLinearColor UAssetDefinition_MetaSoundSource::GetAssetColor() const
+{
+ 	if (const ISlateStyle* MetasoundStyle = FSlateStyleRegistry::FindSlateStyle("MetaSoundStyle"))
+ 	{
+ 		return MetasoundStyle->GetColor("MetaSoundSource.Color").ToFColorSRGB();
+ 	}
+ 
+ 	return FColor::White;
+}
+
+TSoftClassPtr<UObject> UAssetDefinition_MetaSoundSource::GetAssetClass() const
+{
+	return UMetaSoundSource::StaticClass();
+}
+
+TConstArrayView<FAssetCategoryPath> UAssetDefinition_MetaSoundSource::GetAssetCategories() const
+{
+	static const auto Pinned_Categories = { EAssetCategoryPaths::Audio };
+	static const auto Categories = { EAssetCategoryPaths::Audio / LOCTEXT("AssetSoundMetaSoundSourceSubMenu", "MetaSounds") };
+
+	if (GetDefault<UMetasoundEditorSettings>()->bPinMetaSoundSourceInAssetMenu)
+	{
+		return Pinned_Categories;
+	}
+
+	return Categories;
+}
+
+EAssetCommandResult UAssetDefinition_MetaSoundSource::OpenAssets(const FAssetOpenArgs& OpenArgs) const
+{
+	for (UMetaSoundSource* Metasound : OpenArgs.LoadObjects<UMetaSoundSource>())
+	{
+		TSharedRef<Metasound::Editor::FEditor> NewEditor = MakeShared<Metasound::Editor::FEditor>();
+		NewEditor->InitMetasoundEditor(OpenArgs.GetToolkitMode(), OpenArgs.ToolkitHost, Metasound);
+	}
+
+	return EAssetCommandResult::Handled;
+}
+
+const FSlateBrush* UAssetDefinition_MetaSoundSource::GetThumbnailBrush(const FAssetData& InAssetData, const FName InClassName) const
+{
+	constexpr bool bIsThumbnail = true;
+	return GetClassBrush(InAssetData, InClassName, bIsThumbnail);
+}
+
+const FSlateBrush* UAssetDefinition_MetaSoundSource::GetIconBrush(const FAssetData& InAssetData, const FName InClassName) const
+{
+	return GetClassBrush(InAssetData, InClassName);
+}
+
+void UAssetDefinition_MetaSoundSource::ExecutePlaySound(const FToolMenuContext& InContext)
+{
+	if (UMetaSoundSource* MetaSoundSource = UContentBrowserAssetContextMenuContext::LoadSingleSelectedAsset<UMetaSoundSource>(InContext))
+	{
+		// If editor is open, call into it to play to start all visualization requirements therein
+		// specific to auditioning MetaSounds (ex. priming audio bus used for volume metering, playtime
+		// widget, etc.)
+		TSharedPtr<Metasound::Editor::FEditor> Editor = Metasound::Editor::FGraphBuilder::GetEditorForMetasound(*MetaSoundSource);
+		if (Editor.IsValid())
+		{
+			Editor->Play();
+			return;
+		}
+
+		Metasound::Editor::FGraphBuilder::FGraphBuilder::RegisterGraphWithFrontend(*MetaSoundSource);
+		UAssetDefinition_SoundBase::ExecutePlaySound(InContext);
+	}
+}
+
+void UAssetDefinition_MetaSoundSource::ExecuteStopSound(const FToolMenuContext& InContext)
+{
+	if (UMetaSoundSource* MetaSoundSource = UContentBrowserAssetContextMenuContext::LoadSingleSelectedAsset<UMetaSoundSource>(InContext))
+	{
+		TSharedPtr<Metasound::Editor::FEditor> Editor = Metasound::Editor::FGraphBuilder::GetEditorForMetasound(*MetaSoundSource);
+		if (Editor.IsValid())
+		{
+			Editor->Stop();
+			return;
+		}
+
+		UAssetDefinition_SoundBase::ExecuteStopSound(InContext);
+	}
+}
+
+bool UAssetDefinition_MetaSoundSource::CanExecutePlayCommand(const FToolMenuContext& InContext)
+{
+	return UAssetDefinition_SoundBase::CanExecutePlayCommand(InContext);
+}
+
+ECheckBoxState UAssetDefinition_MetaSoundSource::IsActionCheckedMute(const FToolMenuContext& InContext)
+{
+	return UAssetDefinition_SoundBase::IsActionCheckedMute(InContext);
+}
+
+ECheckBoxState UAssetDefinition_MetaSoundSource::IsActionCheckedSolo(const FToolMenuContext& InContext)
+{
+	return UAssetDefinition_SoundBase::IsActionCheckedSolo(InContext);
+}
+
+void UAssetDefinition_MetaSoundSource::ExecuteMuteSound(const FToolMenuContext& InContext)
+{
+	UAssetDefinition_SoundBase::ExecuteMuteSound(InContext);
+}
+
+void UAssetDefinition_MetaSoundSource::ExecuteSoloSound(const FToolMenuContext& InContext)
+{
+	UAssetDefinition_SoundBase::ExecuteSoloSound(InContext);
+}
+
+bool UAssetDefinition_MetaSoundSource::CanExecuteMuteCommand(const FToolMenuContext& InContext)
+{
+	return UAssetDefinition_SoundBase::CanExecuteMuteCommand(InContext);
+}
+
+bool UAssetDefinition_MetaSoundSource::CanExecuteSoloCommand(const FToolMenuContext& InContext)
+{
+	return UAssetDefinition_SoundBase::CanExecuteSoloCommand(InContext);
+}
+
+TSharedPtr<SWidget> UAssetDefinition_MetaSoundSource::GetThumbnailOverlay(const FAssetData& InAssetData) const
+{
+	return UAssetDefinition_SoundBase::GetSoundBaseThumbnailOverlay(InAssetData);
+}
+
+EAssetCommandResult UAssetDefinition_MetaSoundSource::ActivateAssets(const FAssetActivateArgs& ActivateArgs) const
+{
+	if (ActivateArgs.ActivationMethod == EAssetActivationMethod::Previewed)
+	{
+		if (UMetaSoundSource* MetaSoundSource = ActivateArgs.LoadFirstValid<UMetaSoundSource>())
+		{
+			TSharedPtr<Metasound::Editor::FEditor> Editor = Metasound::Editor::FGraphBuilder::GetEditorForMetasound(*MetaSoundSource);
+			UAudioComponent* PreviewComp = GEditor->GetPreviewAudioComponent();
+	
+			// If the editor is open, we need to stop or start the editor so it can light up while previewing it in the CB
+			if (Editor.IsValid())
+			{
+				if (PreviewComp && PreviewComp->IsPlaying())
+				{
+					if (!MetaSoundSource || PreviewComp->Sound == MetaSoundSource)
+					{
+						Editor->Stop();
+					}
+				}
+				else
+				{
+					Editor->Play();
+				}
+
+				return EAssetCommandResult::Handled;
+			}
+			else
+			{
+				return UAssetDefinition_SoundBase::ActivateSoundBase(ActivateArgs);
+			}
+		}
+	}
+	return EAssetCommandResult::Unhandled;
+}
+
+namespace MenuExtension_MetaSoundSourceTemplate
+{
+	template <typename TClass>
+	void ExecuteCreateMetaSoundPreset(const FToolMenuContext& MenuContext)
+	{
+		if (const UContentBrowserAssetContextMenuContext* Context = UContentBrowserAssetContextMenuContext::FindContextWithAssets(MenuContext))
+		{
+			for (TClass* ReferencedMetaSound : Context->LoadSelectedObjects<TClass>())
+			{
+				FString PackagePath;
+				FString AssetName;
+
+				IAssetTools::Get().CreateUniqueAssetName(ReferencedMetaSound->GetOutermost()->GetName(), TEXT("_Preset"), PackagePath, AssetName);
+
+				EMetaSoundBuilderResult BuilderResult;
+				UMetaSoundBuilderBase& Builder = UMetaSoundBuilderSubsystem::GetChecked().CreatePresetBuilder(FName(AssetName), ReferencedMetaSound, BuilderResult);
+
+				if (BuilderResult != EMetaSoundBuilderResult::Succeeded)
+				{
+					UE_LOG(LogMetaSound, Error, TEXT("Error creating a preset builder for MetaSound '%s'"), *AssetName);
+					return;
+				}
+
+				UMetaSoundEditorSubsystem& MetaSoundEditorSubsystem = UMetaSoundEditorSubsystem::GetChecked();
+				MetaSoundEditorSubsystem.BuildToAsset(&Builder, MetaSoundEditorSubsystem.GetDefaultAuthor(), AssetName, FPackageName::GetLongPackagePath(PackagePath), BuilderResult);
+				if (BuilderResult != EMetaSoundBuilderResult::Succeeded)
+				{
+					UE_LOG(LogMetaSound, Error, TEXT("Error building to asset when creating preset '%s'"), *AssetName);
+				}
+			}
+		}
+	}
+
+ 	static FDelayedAutoRegisterHelper DelayedAutoRegister(EDelayedRegisterRunPhase::EndOfEngineInit, [] {
+ 		UToolMenus::RegisterStartupCallback(FSimpleMulticastDelegate::FDelegate::CreateLambda([]()
+ 		{
+ 			FToolMenuOwnerScoped OwnerScoped(UE_MODULE_NAME);
+ 
+ 			{
+				UToolMenu* Menu = UE::ContentBrowser::ExtendToolMenu_AssetContextMenu(UMetaSoundSource::StaticClass());
+
+				FToolMenuSection& Section = Menu->FindOrAddSection("GetAssetActions");
+				{
+
+					Section.AddDynamicEntry(NAME_None, FNewToolMenuSectionDelegate::CreateLambda([](FToolMenuSection& InSection)
+					{
+						{
+							const TAttribute<FText> Label = LOCTEXT("Sound_PlaySound", "Play");
+							const TAttribute<FText> ToolTip = LOCTEXT("Sound_PlaySoundTooltip", "Plays the selected sound.");
+							const FSlateIcon Icon = FSlateIcon(FAppStyle::GetAppStyleSetName(), "MediaAsset.AssetActions.Play.Small");
+
+							FToolUIAction UIAction;
+							UIAction.ExecuteAction = FToolMenuExecuteAction::CreateStatic(&UAssetDefinition_MetaSoundSource::ExecutePlaySound);
+							UIAction.CanExecuteAction = FToolMenuCanExecuteAction::CreateStatic(&UAssetDefinition_MetaSoundSource::CanExecutePlayCommand);
+							InSection.AddMenuEntry("Sound_PlaySound", Label, ToolTip, Icon, UIAction);
+						}
+						{
+							const TAttribute<FText> Label = LOCTEXT("Sound_StopSound", "Stop");
+							const TAttribute<FText> ToolTip = LOCTEXT("Sound_StopSoundTooltip", "Stops the selected sounds.");
+							const FSlateIcon Icon = FSlateIcon(FAppStyle::GetAppStyleSetName(), "MediaAsset.AssetActions.Stop.Small");
+
+							FToolUIAction UIAction;
+							UIAction.ExecuteAction = FToolMenuExecuteAction::CreateStatic(&UAssetDefinition_MetaSoundSource::ExecuteStopSound);
+							InSection.AddMenuEntry("Sound_StopSound", Label, ToolTip, Icon, UIAction);
+						}
+						{
+							const TAttribute<FText> Label = LOCTEXT("Sound_MuteSound", "Mute");
+							const TAttribute<FText> ToolTip = LOCTEXT("Sound_MuteSoundTooltip", "Mutes the selected sounds.");
+							const FSlateIcon Icon = FSlateIcon(FAppStyle::GetAppStyleSetName(), "MediaAsset.AssetActions.Mute.Small");
+
+							FToolUIAction UIAction;
+							UIAction.ExecuteAction = FToolMenuExecuteAction::CreateStatic(&UAssetDefinition_MetaSoundSource::ExecuteMuteSound);
+							UIAction.CanExecuteAction = FToolMenuCanExecuteAction::CreateStatic(&UAssetDefinition_MetaSoundSource::CanExecuteMuteCommand);
+							UIAction.GetActionCheckState = FToolMenuGetActionCheckState::CreateStatic(&UAssetDefinition_MetaSoundSource::IsActionCheckedMute);
+							InSection.AddMenuEntry("Sound_SoundMute", Label, ToolTip, Icon, UIAction, EUserInterfaceActionType::ToggleButton);
+						}
+						{
+							const TAttribute<FText> Label = LOCTEXT("Sound_SoloSound", "Solo");
+							const TAttribute<FText> ToolTip = LOCTEXT("Sound_SoloSoundTooltip", "Solos the selected sounds.");
+							const FSlateIcon Icon = FSlateIcon(FAppStyle::GetAppStyleSetName(), "MediaAsset.AssetActions.Solo.Small");
+
+							FToolUIAction UIAction;
+							UIAction.ExecuteAction = FToolMenuExecuteAction::CreateStatic(&UAssetDefinition_MetaSoundSource::ExecuteSoloSound);
+							UIAction.CanExecuteAction = FToolMenuCanExecuteAction::CreateStatic(&UAssetDefinition_MetaSoundSource::CanExecuteSoloCommand);
+							UIAction.GetActionCheckState = FToolMenuGetActionCheckState::CreateStatic(&UAssetDefinition_MetaSoundSource::IsActionCheckedSolo);
+							InSection.AddMenuEntry("Sound_StopSolo", Label, ToolTip, Icon, UIAction, EUserInterfaceActionType::ToggleButton);
+						}
+						{
+							const TAttribute<FText> Label = LOCTEXT("MetaSoundSource_CreatePreset", "Create MetaSound Source Preset");
+							const TAttribute<FText> ToolTip = LOCTEXT("MetaSoundSource_CreatePresetToolTip", "Creates a MetaSoundSource Preset using the selected MetaSound's root graph as a reference.");
+							const FSlateIcon Icon = Metasound::Editor::Style::CreateSlateIcon("ClassIcon.MetasoundSource");
+
+							FToolUIAction UIAction;
+							UIAction.ExecuteAction = FToolMenuExecuteAction::CreateStatic(&ExecuteCreateMetaSoundPreset<UMetaSoundSource>);
+							InSection.AddMenuEntry("MetaSoundSource_CreatePreset", Label, ToolTip, Icon, UIAction);
+						}
+					}));
+				}
+ 			}
+	
+			{
+				UToolMenu* Menu = UE::ContentBrowser::ExtendToolMenu_AssetContextMenu(UMetaSoundPatch::StaticClass());
+
+				FToolMenuSection& Section = Menu->FindOrAddSection("GetAssetActions");
+				Section.AddDynamicEntry(NAME_None, FNewToolMenuSectionDelegate::CreateLambda([](FToolMenuSection& InSection)
+				{
+					const TAttribute<FText> Label = LOCTEXT("MetaSoundPatch_CreatePreset", "Create MetaSound Patch Preset");
+					const TAttribute<FText> ToolTip = LOCTEXT("MetaSoundPatch_CreatePresetToolTip", "Creates a MetaSoundSource Patch Preset using the selected MetaSound Patch's root graph as a reference.");
+
+					const FSlateIcon Icon = Metasound::Editor::Style::CreateSlateIcon("ClassIcon.MetasoundPatch");
+
+					FToolUIAction UIAction;
+					UIAction.ExecuteAction = FToolMenuExecuteAction::CreateStatic(&ExecuteCreateMetaSoundPreset<UMetaSoundPatch>);
+					InSection.AddMenuEntry("MetaSoundPatch_CreatePreset", Label, ToolTip, Icon, UIAction);
+				}));
+			}
+ 		}));
+	});
+}
+
+
+#undef LOCTEXT_NAMESPACE //MetaSoundEditor
diff --git a/Engine/Plugins/Runtime/Metasound/Source/MetasoundEditor/Private/MetasoundAssetDefinitions.h b/Engine/Plugins/Runtime/Metasound/Source/MetasoundEditor/Private/MetasoundAssetDefinitions.h
new file mode 100644
--- /dev/null
+++ b/Engine/Plugins/Runtime/Metasound/Source/MetasoundEditor/Private/MetasoundAssetDefinitions.h
@@ -0,0 +1,58 @@
+// Copyright Epic Games, Inc. All Rights Reserved.
+#pragma once
+
+#include "AssetTypeActions/AssetDefinition_SoundBase.h"
+#include "MetasoundAssetDefinitions.generated.h"
+
+
+UCLASS()
+class UAssetDefinition_MetaSoundPatch : public UAssetDefinitionDefault
+{
+	GENERATED_BODY()
+
+public:
+	// UAssetDefinition Begin
+	virtual FText GetAssetDisplayName() const override { return NSLOCTEXT("AssetTypeActions", "AssetTypeActions_MetaSoundPatch", "MetaSound Patch"); }
+	virtual FLinearColor GetAssetColor() const override;
+	virtual TSoftClassPtr<UObject> GetAssetClass() const override;
+	virtual TConstArrayView<FAssetCategoryPath> GetAssetCategories() const override;
+	virtual EAssetCommandResult OpenAssets(const FAssetOpenArgs& OpenArgs) const override;
+
+	virtual const FSlateBrush* GetThumbnailBrush(const FAssetData& InAssetData, const FName InClassName) const override;
+	virtual const FSlateBrush* GetIconBrush(const FAssetData& InAssetData, const FName InClassName) const override;
+	// UAssetDefinition End
+
+};
+
+UCLASS()
+class UAssetDefinition_MetaSoundSource : public UAssetDefinitionDefault
+{
+	GENERATED_BODY()
+
+public:
+	// UAssetDefinition Begin
+	virtual FText GetAssetDisplayName() const override { return NSLOCTEXT("AssetTypeActions", "AssetTypeActions_MetaSoundSource", "MetaSound Source"); }
+	virtual FLinearColor GetAssetColor() const override;
+	virtual TSoftClassPtr<UObject> GetAssetClass() const override;
+	virtual TConstArrayView<FAssetCategoryPath> GetAssetCategories() const override;
+	virtual EAssetCommandResult OpenAssets(const FAssetOpenArgs& OpenArgs) const override;
+
+	virtual TSharedPtr<SWidget> GetThumbnailOverlay(const FAssetData& InAssetData) const override;
+	virtual EAssetCommandResult ActivateAssets(const FAssetActivateArgs& ActivateArgs) const override;
+
+ 	virtual const FSlateBrush* GetThumbnailBrush(const FAssetData& InAssetData, const FName InClassName) const override;
+ 	virtual const FSlateBrush* GetIconBrush(const FAssetData& InAssetData, const FName InClassName) const override;
+	// UAssetDefinition End
+
+	// Menu Extension statics
+	static void ExecutePlaySound(const FToolMenuContext& InContext);
+	static void ExecuteStopSound(const FToolMenuContext& InContext);
+	static bool CanExecutePlayCommand(const FToolMenuContext& InContext);
+	static ECheckBoxState IsActionCheckedMute(const FToolMenuContext& InContext);
+	static ECheckBoxState IsActionCheckedSolo(const FToolMenuContext& InContext);
+	static void ExecuteMuteSound(const FToolMenuContext& InContext);
+	static void ExecuteSoloSound(const FToolMenuContext& InContext);
+	static bool CanExecuteMuteCommand(const FToolMenuContext& InContext);
+	static bool CanExecuteSoloCommand(const FToolMenuContext& InContext);
+
+};
\ No newline at end of file
diff --git a/Engine/Plugins/Runtime/Metasound/Source/MetasoundEditor/Private/MetasoundAssetTypeActions.cpp b/Engine/Plugins/Runtime/Metasound/Source/MetasoundEditor/Private/MetasoundAssetTypeActions.cpp
deleted file mode 100644
diff --git a/Engine/Plugins/Runtime/Metasound/Source/MetasoundEditor/Private/MetasoundAssetTypeActions.h b/Engine/Plugins/Runtime/Metasound/Source/MetasoundEditor/Private/MetasoundAssetTypeActions.h
deleted file mode 100644
diff --git a/Engine/Plugins/Runtime/Metasound/Source/MetasoundEditor/Private/MetasoundEditorModule.cpp b/Engine/Plugins/Runtime/Metasound/Source/MetasoundEditor/Private/MetasoundEditorModule.cpp
--- a/Engine/Plugins/Runtime/Metasound/Source/MetasoundEditor/Private/MetasoundEditorModule.cpp
+++ b/Engine/Plugins/Runtime/Metasound/Source/MetasoundEditor/Private/MetasoundEditorModule.cpp
@@ -15,7 +15,6 @@
 #include "ISettingsModule.h"
 #include "Metasound.h"
 #include "MetasoundAssetSubsystem.h"
-#include "MetasoundAssetTypeActions.h"
 #include "MetasoundAudioBuffer.h"
 #include "MetasoundBuilderSubsystem.h"
 #include "MetasoundDetailCustomization.h"
@@ -731,9 +730,6 @@ namespace Metasound
 				// Register Metasound asset type actions
 				IAssetTools& AssetTools = FModuleManager::LoadModuleChecked<FAssetToolsModule>(AssetToolName).Get();
 
-				AddAssetAction<FAssetTypeActions_MetaSoundPatch>(AssetTools, AssetActions);
-				AddAssetAction<FAssetTypeActions_MetaSoundSource>(AssetTools, AssetActions);
-
 				FPropertyEditorModule& PropertyModule = FModuleManager::LoadModuleChecked<FPropertyEditorModule>("PropertyEditor");
 				
 				PropertyModule.RegisterCustomClassLayout(
@@ -798,9 +794,6 @@ namespace Metasound
 					GetMutableDefault<UMetasoundEditorSettings>()
 				);
 
-				FAssetTypeActions_MetaSoundPatch::RegisterMenuActions();
-				FAssetTypeActions_MetaSoundSource::RegisterMenuActions();
-
 				FAssetRegistryModule& AssetRegistryModule = FModuleManager::LoadModuleChecked<FAssetRegistryModule>("AssetRegistry");
 				if (AssetRegistryModule.Get().IsLoadingAssets())
 				{
diff --git a/Engine/Plugins/Runtime/Metasound/Source/MetasoundFrontend/Private/MetasoundAssetBase.cpp b/Engine/Plugins/Runtime/Metasound/Source/MetasoundFrontend/Private/MetasoundAssetBase.cpp
--- a/Engine/Plugins/Runtime/Metasound/Source/MetasoundFrontend/Private/MetasoundAssetBase.cpp
+++ b/Engine/Plugins/Runtime/Metasound/Source/MetasoundFrontend/Private/MetasoundAssetBase.cpp
@@ -962,8 +962,8 @@ void FMetasoundAssetBase::UpdateAssetRegistry()
 	FNodeClassName ClassName = DocumentClassMetadata.GetClassName().ToNodeClassName();
 	FMetasoundFrontendClass GraphClass;
 
+	AssetClassInfo.bIsPreset = DocumentClassGraph.PresetOptions.bIsPreset;
 	AssetClassInfo.Version = DocumentClassMetadata.GetVersion();
-
 	AssetClassInfo.InputTypes.Reset();
 	Algo::Transform(GraphClass.Interface.Inputs, AssetClassInfo.InputTypes, [] (const FMetasoundFrontendClassInput& Input) { return Input.TypeName; });
 
diff --git a/Engine/Source/Editor/AudioEditor/Private/AssetTypeActions/AssetDefinition_SoundBase.cpp b/Engine/Source/Editor/AudioEditor/Private/AssetTypeActions/AssetDefinition_SoundBase.cpp
--- a/Engine/Source/Editor/AudioEditor/Private/AssetTypeActions/AssetDefinition_SoundBase.cpp
+++ b/Engine/Source/Editor/AudioEditor/Private/AssetTypeActions/AssetDefinition_SoundBase.cpp
@@ -60,11 +60,11 @@ namespace UE::AudioEditor
 	}
 }
 
-EAssetCommandResult UAssetDefinition_SoundBase::ActivateAssets(const FAssetActivateArgs& ActivateArgs) const
+EAssetCommandResult UAssetDefinition_SoundBase::ActivateSoundBase(const FAssetActivateArgs& ActivateArgs)
 {
-	if ( ActivateArgs.ActivationMethod == EAssetActivationMethod::Previewed )
+	if (ActivateArgs.ActivationMethod == EAssetActivationMethod::Previewed)
 	{
-		if ( USoundBase* TargetSound = ActivateArgs.LoadFirstValid<USoundBase>() )
+		if (USoundBase* TargetSound = ActivateArgs.LoadFirstValid<USoundBase>())
 		{
 			UAudioComponent* PreviewComp = GEditor->GetPreviewAudioComponent();
 			if (PreviewComp && PreviewComp->IsPlaying())
@@ -84,17 +84,16 @@ EAssetCommandResult UAssetDefinition_SoundBase::ActivateAssets(const FAssetActiv
 				// Not already playing, play the target sound cue if it exists
 				UE::AudioEditor::PlaySound(TargetSound);
 			}
-			
+
 			return EAssetCommandResult::Handled;
 		}
 	}
-
-	return Super::ActivateAssets(ActivateArgs);
+	return EAssetCommandResult::Unhandled;
 }
 
-TSharedPtr<SWidget> UAssetDefinition_SoundBase::GetThumbnailOverlay(const FAssetData& InAssetData) const
+TSharedPtr<SWidget> UAssetDefinition_SoundBase::GetSoundBaseThumbnailOverlay(const FAssetData& InAssetData)
 {
-	auto OnGetDisplayBrushLambda = [this, InAssetData]() -> const FSlateBrush*
+	auto OnGetDisplayBrushLambda = [InAssetData]() -> const FSlateBrush*
 	{
 		if (UE::AudioEditor::IsSoundPlaying(InAssetData))
 		{
@@ -104,7 +103,7 @@ TSharedPtr<SWidget> UAssetDefinition_SoundBase::GetThumbnailOverlay(const FAsset
 		return FAppStyle::GetBrush("MediaAsset.AssetActions.Play.Large");
 	};
 
-	auto OnClickedLambda = [this, InAssetData]() -> FReply
+	auto OnClickedLambda = [InAssetData]() -> FReply
 	{
 		if (UE::AudioEditor::IsSoundPlaying(InAssetData))
 		{
@@ -118,7 +117,7 @@ TSharedPtr<SWidget> UAssetDefinition_SoundBase::GetThumbnailOverlay(const FAsset
 		return FReply::Handled();
 	};
 
-	auto OnToolTipTextLambda = [this, InAssetData]() -> FText
+	auto OnToolTipTextLambda = [InAssetData]() -> FText
 	{
 		if (UE::AudioEditor::IsSoundPlaying(InAssetData))
 		{
@@ -134,7 +133,7 @@ TSharedPtr<SWidget> UAssetDefinition_SoundBase::GetThumbnailOverlay(const FAsset
 		.VAlign(VAlign_Center)
 		.Padding(FMargin(2));
 
-	auto OnGetVisibilityLambda = [this, Box, InAssetData]() -> EVisibility
+	auto OnGetVisibilityLambda = [Box, InAssetData]() -> EVisibility
 	{
 		if (Box.IsValid() && (Box->IsHovered() || UE::AudioEditor::IsSoundPlaying(InAssetData)))
 		{
@@ -164,12 +163,25 @@ TSharedPtr<SWidget> UAssetDefinition_SoundBase::GetThumbnailOverlay(const FAsset
 	return Box;
 }
 
+
+EAssetCommandResult UAssetDefinition_SoundBase::ActivateAssets(const FAssetActivateArgs& ActivateArgs) const
+{
+	if (ActivateSoundBase(ActivateArgs) == EAssetCommandResult::Handled)
+	{
+		return EAssetCommandResult::Handled;
+	}
+	return Super::ActivateAssets(ActivateArgs);
+}
+
+TSharedPtr<SWidget> UAssetDefinition_SoundBase::GetThumbnailOverlay(const FAssetData& InAssetData) const
+{
+	return GetSoundBaseThumbnailOverlay(InAssetData);
+}
+
 // Menu Extensions
 //--------------------------------------------------------------------
 
-namespace MenuExtension_SoundBase
-{
-	void ExecutePlaySound(const FToolMenuContext& InContext)
+	void UAssetDefinition_SoundBase::ExecutePlaySound(const FToolMenuContext& InContext)
 	{
 		if (USoundBase* Sound = UContentBrowserAssetContextMenuContext::LoadSingleSelectedAsset<USoundBase>(InContext))
 		{
@@ -178,12 +190,12 @@ namespace MenuExtension_SoundBase
 		}
 	}
 
-	void ExecuteStopSound(const FToolMenuContext& InContext)
+	void UAssetDefinition_SoundBase::ExecuteStopSound(const FToolMenuContext& InContext)
 	{
 		UE::AudioEditor::StopSound();
 	}
 
-	bool CanExecutePlayCommand(const FToolMenuContext& InContext)
+	bool UAssetDefinition_SoundBase::CanExecutePlayCommand(const FToolMenuContext& InContext)
 	{
 		if (const UContentBrowserAssetContextMenuContext* CBContext = UContentBrowserAssetContextMenuContext::FindContextWithAssets(InContext))
 		{
@@ -193,7 +205,7 @@ namespace MenuExtension_SoundBase
 		return false;
 	}
 	
-	ECheckBoxState IsActionCheckedMute(const FToolMenuContext& InContext)
+	ECheckBoxState UAssetDefinition_SoundBase::IsActionCheckedMute(const FToolMenuContext& InContext)
 	{
 #if ENABLE_AUDIO_DEBUG
 		if (FAudioDeviceManager* ADM = GEditor->GetAudioDeviceManager())
@@ -224,7 +236,7 @@ namespace MenuExtension_SoundBase
 		return ECheckBoxState::Unchecked;
 	}
 	
-	ECheckBoxState IsActionCheckedSolo(const FToolMenuContext& InContext)
+	ECheckBoxState UAssetDefinition_SoundBase::IsActionCheckedSolo(const FToolMenuContext& InContext)
 	{
 #if ENABLE_AUDIO_DEBUG
 		// If *any* of the selection are solod, show the tick box as ticked.
@@ -255,7 +267,7 @@ namespace MenuExtension_SoundBase
 		return ECheckBoxState::Unchecked;
 	}
 
-	void ExecuteMuteSound(const FToolMenuContext& InContext)
+	void UAssetDefinition_SoundBase::ExecuteMuteSound(const FToolMenuContext& InContext)
 	{
 #if ENABLE_AUDIO_DEBUG
 		if (FAudioDeviceManager* ADM = GEditor->GetAudioDeviceManager())
@@ -282,7 +294,7 @@ namespace MenuExtension_SoundBase
 #endif
 	}
 
-	void ExecuteSoloSound(const FToolMenuContext& InContext)
+	void UAssetDefinition_SoundBase::ExecuteSoloSound(const FToolMenuContext& InContext)
 	{
 #if ENABLE_AUDIO_DEBUG
 		if (FAudioDeviceManager* ADM = GEditor->GetAudioDeviceManager())
@@ -310,7 +322,7 @@ namespace MenuExtension_SoundBase
 #endif
 	}	
 
-	bool CanExecuteMuteCommand(const FToolMenuContext& InContext)
+	bool UAssetDefinition_SoundBase::CanExecuteMuteCommand(const FToolMenuContext& InContext)
 	{
 #if ENABLE_AUDIO_DEBUG
 		if (FAudioDeviceManager* ADM = GEditor->GetAudioDeviceManager())
@@ -344,7 +356,7 @@ namespace MenuExtension_SoundBase
 		return false;
 	}
 
-	bool CanExecuteSoloCommand(const FToolMenuContext& InContext)
+	bool UAssetDefinition_SoundBase::CanExecuteSoloCommand(const FToolMenuContext& InContext)
 	{
 #if ENABLE_AUDIO_DEBUG
 		if (FAudioDeviceManager* ADM = GEditor->GetAudioDeviceManager())
@@ -377,61 +389,80 @@ namespace MenuExtension_SoundBase
 #endif
 		return false;
 	}
-	
-	static FDelayedAutoRegisterHelper DelayedAutoRegister(EDelayedRegisterRunPhase::EndOfEngineInit, []{ 
-		UToolMenus::RegisterStartupCallback(FSimpleMulticastDelegate::FDelegate::CreateLambda([]()
-		{
-			FToolMenuOwnerScoped OwnerScoped(UE_MODULE_NAME);
-			UToolMenu* Menu = UE::ContentBrowser::ExtendToolMenu_AssetContextMenu(USoundBase::StaticClass());
-		
-			FToolMenuSection& Section = Menu->FindOrAddSection("GetAssetActions");
-			Section.AddDynamicEntry(NAME_None, FNewToolMenuSectionDelegate::CreateLambda([](FToolMenuSection& InSection)
+
+namespace MenuExtension_SoundBase
+{
+	static void RegisterAssetActions(UClass* InClass)
+	{
+		UToolMenu* Menu = UE::ContentBrowser::ExtendToolMenu_AssetContextMenu(InClass);
+
+		FToolMenuSection& Section = Menu->FindOrAddSection("GetAssetActions");
+		Section.AddDynamicEntry(NAME_None, FNewToolMenuSectionDelegate::CreateLambda([](FToolMenuSection& InSection)
 			{
 				if (const UContentBrowserAssetContextMenuContext* Context = UContentBrowserAssetContextMenuContext::FindContextWithAssets(InSection))
 				{
+					if (Context->SelectedAssets.Num() > 0)
 					{
-						const TAttribute<FText> Label = LOCTEXT("Sound_PlaySound", "Play");
-						const TAttribute<FText> ToolTip = LOCTEXT("Sound_PlaySoundTooltip", "Plays the selected sound.");
-						const FSlateIcon Icon = FSlateIcon(FAppStyle::GetAppStyleSetName(), "MediaAsset.AssetActions.Play.Small");
-						
-						FToolUIAction UIAction;
-						UIAction.ExecuteAction = FToolMenuExecuteAction::CreateStatic(&ExecutePlaySound);
-						UIAction.CanExecuteAction = FToolMenuCanExecuteAction::CreateStatic(&CanExecutePlayCommand);
-						InSection.AddMenuEntry("Sound_PlaySound", Label, ToolTip, Icon, UIAction);
-					}
-					{
-						const TAttribute<FText> Label = LOCTEXT("Sound_StopSound", "Stop");
-						const TAttribute<FText> ToolTip = LOCTEXT("Sound_StopSoundTooltip", "Stops the selected sounds.");
-						const FSlateIcon Icon = FSlateIcon(FAppStyle::GetAppStyleSetName(), "MediaAsset.AssetActions.Stop.Small");
-						
-						FToolUIAction UIAction;
-						UIAction.ExecuteAction = FToolMenuExecuteAction::CreateStatic(&ExecuteStopSound);
-						InSection.AddMenuEntry("Sound_StopSound", Label, ToolTip, Icon, UIAction);
-					}
-					{
-						const TAttribute<FText> Label = LOCTEXT("Sound_MuteSound", "Mute");
-						const TAttribute<FText> ToolTip = LOCTEXT("Sound_MuteSoundTooltip", "Mutes the selected sounds.");
-						const FSlateIcon Icon = FSlateIcon(FAppStyle::GetAppStyleSetName(), "MediaAsset.AssetActions.Mute.Small");
-		
-						FToolUIAction UIAction;
-						UIAction.ExecuteAction = FToolMenuExecuteAction::CreateStatic(&ExecuteMuteSound);
-						UIAction.CanExecuteAction = FToolMenuCanExecuteAction::CreateStatic(&CanExecuteMuteCommand);
-						UIAction.GetActionCheckState = FToolMenuGetActionCheckState::CreateStatic(&IsActionCheckedMute);
-						InSection.AddMenuEntry("Sound_SoundMute", Label, ToolTip, Icon, UIAction, EUserInterfaceActionType::ToggleButton);
-					}
-					{
-						const TAttribute<FText> Label = LOCTEXT("Sound_SoloSound", "Solo");
-						const TAttribute<FText> ToolTip = LOCTEXT("Sound_SoloSoundTooltip", "Solos the selected sounds.");
-						const FSlateIcon Icon = FSlateIcon(FAppStyle::GetAppStyleSetName(), "MediaAsset.AssetActions.Solo.Small");
-
-						FToolUIAction UIAction;
-						UIAction.ExecuteAction = FToolMenuExecuteAction::CreateStatic(&ExecuteSoloSound);
-						UIAction.CanExecuteAction = FToolMenuCanExecuteAction::CreateStatic(&CanExecuteSoloCommand);
-						UIAction.GetActionCheckState = FToolMenuGetActionCheckState::CreateStatic(&IsActionCheckedSolo);
-						InSection.AddMenuEntry("Sound_StopSolo", Label, ToolTip, Icon, UIAction, EUserInterfaceActionType::ToggleButton);
+						const FAssetData& AssetData = Context->SelectedAssets[0];
+						if (AssetData.AssetClassPath == USoundWave::StaticClass()->GetClassPathName() || AssetData.AssetClassPath == USoundCue::StaticClass()->GetClassPathName())
+						{
+							{
+								const TAttribute<FText> Label = LOCTEXT("Sound_PlaySound", "Play");
+								const TAttribute<FText> ToolTip = LOCTEXT("Sound_PlaySoundTooltip", "Plays the selected sound.");
+								const FSlateIcon Icon = FSlateIcon(FAppStyle::GetAppStyleSetName(), "MediaAsset.AssetActions.Play.Small");
+
+								FToolUIAction UIAction;
+								UIAction.ExecuteAction = FToolMenuExecuteAction::CreateStatic(&UAssetDefinition_SoundBase::ExecutePlaySound);
+								UIAction.CanExecuteAction = FToolMenuCanExecuteAction::CreateStatic(&UAssetDefinition_SoundBase::CanExecutePlayCommand);
+								InSection.AddMenuEntry("Sound_PlaySound", Label, ToolTip, Icon, UIAction);
+							}
+							{
+								const TAttribute<FText> Label = LOCTEXT("Sound_StopSound", "Stop");
+								const TAttribute<FText> ToolTip = LOCTEXT("Sound_StopSoundTooltip", "Stops the selected sounds.");
+								const FSlateIcon Icon = FSlateIcon(FAppStyle::GetAppStyleSetName(), "MediaAsset.AssetActions.Stop.Small");
+
+								FToolUIAction UIAction;
+								UIAction.ExecuteAction = FToolMenuExecuteAction::CreateStatic(&UAssetDefinition_SoundBase::ExecuteStopSound);
+								InSection.AddMenuEntry("Sound_StopSound", Label, ToolTip, Icon, UIAction);
+							}
+							{
+								const TAttribute<FText> Label = LOCTEXT("Sound_MuteSound", "Mute");
+								const TAttribute<FText> ToolTip = LOCTEXT("Sound_MuteSoundTooltip", "Mutes the selected sounds.");
+								const FSlateIcon Icon = FSlateIcon(FAppStyle::GetAppStyleSetName(), "MediaAsset.AssetActions.Mute.Small");
+
+								FToolUIAction UIAction;
+								UIAction.ExecuteAction = FToolMenuExecuteAction::CreateStatic(&UAssetDefinition_SoundBase::ExecuteMuteSound);
+								UIAction.CanExecuteAction = FToolMenuCanExecuteAction::CreateStatic(&UAssetDefinition_SoundBase::CanExecuteMuteCommand);
+								UIAction.GetActionCheckState = FToolMenuGetActionCheckState::CreateStatic(&UAssetDefinition_SoundBase::IsActionCheckedMute);
+								InSection.AddMenuEntry("Sound_SoundMute", Label, ToolTip, Icon, UIAction, EUserInterfaceActionType::ToggleButton);
+							}
+							{
+								const TAttribute<FText> Label = LOCTEXT("Sound_SoloSound", "Solo");
+								const TAttribute<FText> ToolTip = LOCTEXT("Sound_SoloSoundTooltip", "Solos the selected sounds.");
+								const FSlateIcon Icon = FSlateIcon(FAppStyle::GetAppStyleSetName(), "MediaAsset.AssetActions.Solo.Small");
+
+								FToolUIAction UIAction;
+								UIAction.ExecuteAction = FToolMenuExecuteAction::CreateStatic(&UAssetDefinition_SoundBase::ExecuteSoloSound);
+								UIAction.CanExecuteAction = FToolMenuCanExecuteAction::CreateStatic(&UAssetDefinition_SoundBase::CanExecuteSoloCommand);
+								UIAction.GetActionCheckState = FToolMenuGetActionCheckState::CreateStatic(&UAssetDefinition_SoundBase::IsActionCheckedSolo);
+								InSection.AddMenuEntry("Sound_StopSolo", Label, ToolTip, Icon, UIAction, EUserInterfaceActionType::ToggleButton);
+							}
+						}
 					}
 				}
 			}));
+	}
+
+	
+	static FDelayedAutoRegisterHelper DelayedAutoRegister(EDelayedRegisterRunPhase::EndOfEngineInit, []{ 
+		UToolMenus::RegisterStartupCallback(FSimpleMulticastDelegate::FDelegate::CreateLambda([]()
+		{
+			FToolMenuOwnerScoped OwnerScoped(UE_MODULE_NAME);
+
+			// Note: we can't do this with USoundBase because UMetaSoundSource and others may need to do special actions on these actions.
+			// Since these actions are registered via delayed static callbacks, there's not a clean why to do this with inheritance.
+			RegisterAssetActions(USoundCue::StaticClass());
+			RegisterAssetActions(USoundWave::StaticClass());
 		}));
 	});
 }
diff --git a/Engine/Source/Editor/AudioEditor/Public/AssetTypeActions/AssetDefinition_SoundBase.h b/Engine/Source/Editor/AudioEditor/Public/AssetTypeActions/AssetDefinition_SoundBase.h
--- a/Engine/Source/Editor/AudioEditor/Public/AssetTypeActions/AssetDefinition_SoundBase.h
+++ b/Engine/Source/Editor/AudioEditor/Public/AssetTypeActions/AssetDefinition_SoundBase.h
@@ -5,6 +5,7 @@
 #include "CoreMinimal.h"
 #include "Sound/SoundBase.h"
 #include "AssetDefinitionDefault.h"
+#include "ContentBrowserMenuContexts.h"
 
 #include "AssetDefinition_SoundBase.generated.h"
 
@@ -36,4 +37,20 @@ public:
 	virtual TSharedPtr<SWidget> GetThumbnailOverlay(const FAssetData& InAssetData) const override;
 	virtual EAssetCommandResult ActivateAssets(const FAssetActivateArgs& ActivateArgs) const override;
 	// UAssetDefinition End
+
+	// Menu Extension statics
+	static void ExecutePlaySound(const FToolMenuContext& InContext);
+	static void ExecuteStopSound(const FToolMenuContext& InContext);
+	static bool CanExecutePlayCommand(const FToolMenuContext& InContext);
+	static ECheckBoxState IsActionCheckedMute(const FToolMenuContext& InContext);
+	static ECheckBoxState IsActionCheckedSolo(const FToolMenuContext& InContext);
+	static void ExecuteMuteSound(const FToolMenuContext& InContext);
+	static void ExecuteSoloSound(const FToolMenuContext& InContext);
+	static bool CanExecuteMuteCommand(const FToolMenuContext& InContext);
+	static bool CanExecuteSoloCommand(const FToolMenuContext& InContext);
+
+	// Asset definition static utilities
+	static TSharedPtr<SWidget> GetSoundBaseThumbnailOverlay(const FAssetData& InAssetData);
+	static EAssetCommandResult ActivateSoundBase(const FAssetActivateArgs& ActivateArgs);
+
 };
