Fixing compile error

diff --git a/Engine/Plugins/Runtime/Metasound/Source/MetasoundEditor/Private/MetasoundAssetDefinitions.cpp b/Engine/Plugins/Runtime/Metasound/Source/MetasoundEditor/Private/MetasoundAssetDefinitions.cpp
--- a/Engine/Plugins/Runtime/Metasound/Source/MetasoundEditor/Private/MetasoundAssetDefinitions.cpp
+++ b/Engine/Plugins/Runtime/Metasound/Source/MetasoundEditor/Private/MetasoundAssetDefinitions.cpp
@@ -73,8 +73,6 @@ static bool GetIsPreset(const FSoftObjectPath& InSourcePath)
 
 static const FSlateBrush* GetClassBrush(const FAssetData& InAssetData, FName InClassName, bool bIsThumbnail = false)
 {
-	using namespace Metasound::Editor;
-
 	const bool bIsPreset = GetIsPreset(InAssetData.ToSoftObjectPath());
 	FString BrushName = FString::Printf(TEXT("MetasoundEditor.%s"), *InClassName.ToString());
 	if (bIsPreset)
@@ -83,7 +81,7 @@ static const FSlateBrush* GetClassBrush(const FAssetData& InAssetData, FName InC
 	}
 	BrushName += bIsThumbnail ? TEXT(".Thumbnail") : TEXT(".Icon");
 
-	return &Style::GetSlateBrushSafe(FName(*BrushName));
+	return &Metasound::Editor::Style::GetSlateBrushSafe(FName(*BrushName));
 }
 
 
