Add "Enable" input pins to Multiband Analyzer node and Compressor node.

diff --git a/Engine/Plugins/Runtime/Metasound/Source/MetasoundStandardNodes/Private/MetasoundCompressorNode.cpp b/Engine/Plugins/Runtime/Metasound/Source/MetasoundStandardNodes/Private/MetasoundCompressorNode.cpp
--- a/Engine/Plugins/Runtime/Metasound/Source/MetasoundStandardNodes/Private/MetasoundCompressorNode.cpp
+++ b/Engine/Plugins/Runtime/Metasound/Source/MetasoundStandardNodes/Private/MetasoundCompressorNode.cpp
@@ -23,6 +23,7 @@ namespace Metasound
 	/* Mid-Side Encoder */
 	namespace CompressorVertexNames
 	{
+		METASOUND_PARAM(InputIsBypassed, "Bypass", "When true no audio is processed, the input is copied to the output, and the envelope output is zero.");
 		METASOUND_PARAM(InputAudio, "Audio", "Incoming audio signal to compress.");
 		METASOUND_PARAM(InputRatio, "Ratio", "Amount of gain reduction. 1 = no reduction, higher = more reduction.");
 		METASOUND_PARAM(InputThreshold, "Threshold dB", "Amplitude threshold (dB) above which gain will be reduced.");
@@ -46,6 +47,7 @@ namespace Metasound
 	public:
 
 		FCompressorOperator(const FOperatorSettings& InSettings,
+			const FBoolReadRef& bInIsBypassed,
 			const FAudioBufferReadRef& InAudio,
 			const FFloatReadRef& InRatio,
 			const FFloatReadRef& InThresholdDb,
@@ -59,7 +61,8 @@ namespace Metasound
 			const FBoolReadRef& bInIsAnalog,
 			const FBoolReadRef& bInIsUpwards,
 			const FFloatReadRef& InWetDryMix)
-			: AudioInput(InAudio)
+			: bIsBypassedInput(bInIsBypassed)
+			, AudioInput(InAudio)
 			, RatioInput(InRatio)
 			, ThresholdDbInput(InThresholdDb)
 			, AttackTimeInput(InAttackTime)
@@ -75,6 +78,7 @@ namespace Metasound
 			, EnvelopeOutput(FAudioBufferWriteRef::CreateNew(InSettings))
 			, InputDelay(FMath::CeilToInt(InSettings.GetSampleRate() * Compressor.GetMaxLookaheadMsec() / 1000.f) + 1, InSettings.GetSampleRate() * 10.0f / 1000.0f)
 			, DelayedInputSignal(InSettings.GetNumFramesPerBlock())
+			, bEnvelopeOutputIsZero(false)
 			, bUseSidechain(bInUseSidechain)
 			, MsToSamples(InSettings.GetSampleRate() / 1000.0f)
 			, PrevAttackTime(FMath::Max(FTime::ToMilliseconds(*InAttackTime), 0.0))
@@ -149,6 +153,7 @@ namespace Metasound
 
 			static const FVertexInterface Interface(
 				FInputVertexInterface(
+					TInputDataVertex<bool>(METASOUND_GET_PARAM_NAME_AND_METADATA(InputIsBypassed), false),
 					TInputDataVertex<FAudioBuffer>(METASOUND_GET_PARAM_NAME_AND_METADATA(InputAudio)),
 					TInputDataVertex<float>(METASOUND_GET_PARAM_NAME_AND_METADATA(InputRatio), 1.5f),
 					TInputDataVertex<float>(METASOUND_GET_PARAM_NAME_AND_METADATA(InputThreshold), -6.0f),
@@ -175,6 +180,7 @@ namespace Metasound
 		{
 			using namespace CompressorVertexNames;
 
+			InOutVertexData.BindReadVertex(METASOUND_GET_PARAM_NAME(InputIsBypassed), bIsBypassedInput);
 			InOutVertexData.BindReadVertex(METASOUND_GET_PARAM_NAME(InputAudio), AudioInput);
 			InOutVertexData.BindReadVertex(METASOUND_GET_PARAM_NAME(InputRatio), RatioInput);
 			InOutVertexData.BindReadVertex(METASOUND_GET_PARAM_NAME(InputThreshold), ThresholdDbInput);
@@ -219,6 +225,7 @@ namespace Metasound
 			const FDataReferenceCollection& Inputs = InParams.InputDataReferences;
 			const FInputVertexInterface& InputInterface = DeclareVertexInterface().GetInputInterface();
 
+			FBoolReadRef bIsBypassedIn = Inputs.GetDataReadReferenceOrConstructWithVertexDefault<bool>(InputInterface, METASOUND_GET_PARAM_NAME(InputIsBypassed), InParams.OperatorSettings);
 			FAudioBufferReadRef AudioIn = Inputs.GetDataReadReferenceOrConstruct<FAudioBuffer>(METASOUND_GET_PARAM_NAME(InputAudio), InParams.OperatorSettings);
 			FFloatReadRef RatioIn = Inputs.GetDataReadReferenceOrConstructWithVertexDefault<float>(InputInterface, METASOUND_GET_PARAM_NAME(InputRatio), InParams.OperatorSettings);
 			FFloatReadRef ThresholdDbIn = Inputs.GetDataReadReferenceOrConstructWithVertexDefault<float>(InputInterface, METASOUND_GET_PARAM_NAME(InputThreshold), InParams.OperatorSettings);
@@ -234,7 +241,7 @@ namespace Metasound
 
 			bool bIsSidechainConnected = Inputs.ContainsDataReadReference<FAudioBuffer>(METASOUND_GET_PARAM_NAME(InputSidechain));
 
-			return MakeUnique<FCompressorOperator>(InParams.OperatorSettings, AudioIn, RatioIn, ThresholdDbIn, AttackTimeIn, ReleaseTimeIn, LookaheadTimeIn, KneeIn, bIsSidechainConnected, SidechainIn, EnvelopeModeIn, bIsAnalogIn, bIsUpwardsIn, WetDryMixIn);
+			return MakeUnique<FCompressorOperator>(InParams.OperatorSettings, bIsBypassedIn, AudioIn, RatioIn, ThresholdDbIn, AttackTimeIn, ReleaseTimeIn, LookaheadTimeIn, KneeIn, bIsSidechainConnected, SidechainIn, EnvelopeModeIn, bIsAnalogIn, bIsUpwardsIn, WetDryMixIn);
 		}
 
 		void Reset(const IOperator::FResetParams& InParams)
@@ -242,6 +249,7 @@ namespace Metasound
 			// Flush audio buffers
 			AudioOutput->Zero();
 			EnvelopeOutput->Zero();
+			bEnvelopeOutputIsZero = true;
 			InputDelay.Reset();
 			DelayedInputSignal.Zero();
 
@@ -286,6 +294,17 @@ namespace Metasound
 
 		void Execute()
 		{
+			if (*bIsBypassedInput)
+			{
+				FMemory::Memcpy(AudioOutput->GetData(), AudioInput->GetData(), AudioInput->Num() * sizeof(float));
+				if (!bEnvelopeOutputIsZero)
+				{
+					EnvelopeOutput->Zero();
+					bEnvelopeOutputIsZero = true;
+				}
+				return;
+			}
+
 			/* Update parameters */
 			
 			// For a compressor, ratio values should be 1 or greater
@@ -354,6 +373,7 @@ namespace Metasound
 			{
 				Compressor.ProcessAudio(AudioInput->GetData(), AudioInput->Num(), AudioOutput->GetData(), nullptr, EnvelopeOutput->GetData());
 			}
+			bEnvelopeOutputIsZero = false;
 
 			// Calculate Wet/Dry mix
 			float NewWetDryMix = FMath::Clamp(*WetDryMixInput, 0.0f, 1.0f);
@@ -366,6 +386,7 @@ namespace Metasound
 
 	private:
 		// Audio input and output
+		FBoolReadRef bIsBypassedInput;
 		FAudioBufferReadRef AudioInput;
 		FFloatReadRef RatioInput;
 		FFloatReadRef ThresholdDbInput;
@@ -390,6 +411,9 @@ namespace Metasound
 		Audio::FIntegerDelay InputDelay;
 		FAudioBuffer DelayedInputSignal;
 
+		// When bypassed this prevents continual re-zeroing of envelope buffer.
+		bool bEnvelopeOutputIsZero;
+
 		// Whether or not to use sidechain input (is false if no input pin is connected to sidechain input)
 		bool bUseSidechain;
 
