Use MetaSound proxy cache to build metasounds

diff --git a/Engine/Plugins/Runtime/Metasound/Source/MetasoundEngineTest/Private/NodeTestGraphBuilder.cpp b/Engine/Plugins/Runtime/Metasound/Source/MetasoundEngineTest/Private/NodeTestGraphBuilder.cpp
--- a/Engine/Plugins/Runtime/Metasound/Source/MetasoundEngineTest/Private/NodeTestGraphBuilder.cpp
+++ b/Engine/Plugins/Runtime/Metasound/Source/MetasoundEngineTest/Private/NodeTestGraphBuilder.cpp
@@ -88,7 +88,7 @@ namespace Metasound::Test
 		
 		const TSet<FName> TransmittableInputNames;
 		const FString UnknownAsset = TEXT("UnknownAsset");
-		if (const TUniquePtr<FFrontendGraph> Graph = FFrontendGraphBuilder::CreateGraph(Document, TransmittableInputNames, UnknownAsset))
+		if (const TUniquePtr<FFrontendGraph> Graph = FFrontendGraphBuilder::CreateGraph(Document, UnknownAsset))
 		{
 			FOperatorBuilderSettings BuilderSettings;
 			BuilderSettings.bFailOnAnyError = true;
diff --git a/Engine/Plugins/Runtime/Metasound/Source/MetasoundFrontend/Private/MetasoundAssetBase.cpp b/Engine/Plugins/Runtime/Metasound/Source/MetasoundFrontend/Private/MetasoundAssetBase.cpp
--- a/Engine/Plugins/Runtime/Metasound/Source/MetasoundFrontend/Private/MetasoundAssetBase.cpp
+++ b/Engine/Plugins/Runtime/Metasound/Source/MetasoundFrontend/Private/MetasoundAssetBase.cpp
@@ -24,6 +24,7 @@
 #include "MetasoundFrontendDocumentVersioning.h"
 #include "MetasoundFrontendGraph.h"
 #include "MetasoundFrontendNodeTemplateRegistry.h"
+#include "MetasoundFrontendProxyDataCache.h"
 #include "MetasoundFrontendRegistries.h"
 #include "MetasoundFrontendSearchEngine.h"
 #include "MetasoundFrontendTransform.h"
@@ -126,8 +127,10 @@ namespace Metasound
 
 				virtual TUniquePtr<INode> CreateNode(const FNodeInitData&) const override
 				{
-					static const TSet<FName> ReferencedGraphTransmissibleInputNames; // Empty as referenced graphs do not support transmission
-					return FFrontendGraphBuilder().CreateGraph(*PreprocessedDoc, ReferencedGraphTransmissibleInputNames, Name);
+					FProxyDataCache ProxyDataCache;
+					ProxyDataCache.CreateAndCacheProxies(*PreprocessedDoc);
+
+					return FFrontendGraphBuilder().CreateGraph(*PreprocessedDoc, ProxyDataCache, Name);
 				}
 
 				virtual TUniquePtr<INode> CreateNode(FDefaultLiteralNodeConstructorParams&&) const override { return nullptr; }
@@ -605,7 +608,7 @@ const FMetasoundFrontendDocumentModifyContext& FMetasoundAssetBase::GetModifyCon
 }
 #endif // WITH_EDITOR
 
-TSharedPtr<Metasound::FGraph, ESPMode::ThreadSafe> FMetasoundAssetBase::BuildMetasoundDocument(const FMetasoundFrontendDocument& InPreprocessedDoc, const TSet<FName>& InTransmittableInputNames) const
+TSharedPtr<Metasound::FGraph, ESPMode::ThreadSafe> FMetasoundAssetBase::BuildMetasoundDocument(const FMetasoundFrontendDocument& InPreprocessedDoc, const Metasound::Frontend::FProxyDataCache& InProxies) const
 {
 	using namespace Metasound;
 	using namespace Metasound::Frontend;
@@ -613,7 +616,7 @@ TSharedPtr<Metasound::FGraph, ESPMode::ThreadSafe> FMetasoundAssetBase::BuildMet
 	METASOUND_TRACE_CPUPROFILER_EVENT_SCOPE(MetaSoundAssetBase::BuildMetasoundDocument);
 
 	// Create graph which can spawn instances. 
-	TUniquePtr<FFrontendGraph> FrontendGraph = FFrontendGraphBuilder::CreateGraph(InPreprocessedDoc, InTransmittableInputNames, GetOwningAssetName());
+	TUniquePtr<FFrontendGraph> FrontendGraph = FFrontendGraphBuilder::CreateGraph(InPreprocessedDoc, InProxies, GetOwningAssetName());
 	if (!FrontendGraph.IsValid())
 	{
 		UE_LOG(LogMetaSound, Error, TEXT("Failed to build MetaSound graph in asset '%s'"), *GetOwningAssetName());
@@ -991,9 +994,9 @@ const FMetasoundAssetBase::FRuntimeData& FMetasoundAssetBase::CacheRuntimeData(c
 	TArray<FMetasoundFrontendClassInput> PublicInputs = GetPublicClassInputs();
 	TArray<FMetasoundFrontendClassInput> TransmittableInputs = AssetBasePrivate::GetTransmittableInputsFromPublicInputs(PublicInputs);
 
-	TSet<FName> TransmittableInputNames;
-	Algo::Transform(TransmittableInputs, TransmittableInputNames, [](const FMetasoundFrontendClassInput& Input) { return Input.Name; });
-	TSharedPtr<Metasound::FGraph, ESPMode::ThreadSafe> Graph = BuildMetasoundDocument(InPreprocessedDoc, TransmittableInputNames);
+	Metasound::Frontend::FProxyDataCache ProxyDataCache;
+	ProxyDataCache.CreateAndCacheProxies(InPreprocessedDoc);
+	TSharedPtr<Metasound::FGraph, ESPMode::ThreadSafe> Graph = BuildMetasoundDocument(InPreprocessedDoc, ProxyDataCache);
 
 	CachedRuntimeData =
 	{
diff --git a/Engine/Plugins/Runtime/Metasound/Source/MetasoundFrontend/Private/MetasoundFrontendDocument.cpp b/Engine/Plugins/Runtime/Metasound/Source/MetasoundFrontend/Private/MetasoundFrontendDocument.cpp
--- a/Engine/Plugins/Runtime/Metasound/Source/MetasoundFrontend/Private/MetasoundFrontendDocument.cpp
+++ b/Engine/Plugins/Runtime/Metasound/Source/MetasoundFrontend/Private/MetasoundFrontendDocument.cpp
@@ -866,3 +866,58 @@ const TCHAR* LexToString(EMetasoundFrontendVertexAccessType InVertexAccess)
 	}
 }
 
+namespace Metasound
+{
+	void ForEachLiteral(const FMetasoundFrontendDocument& InDoc, FForEachLiteralFunctionRef OnLiteral)
+	{
+		ForEachLiteral(InDoc.RootGraph, OnLiteral);
+		for (const FMetasoundFrontendGraphClass& GraphClass : InDoc.Subgraphs)
+		{
+			ForEachLiteral(GraphClass, OnLiteral);
+		}
+		for (const FMetasoundFrontendClass& Dependency : InDoc.Dependencies)
+		{
+			ForEachLiteral(Dependency, OnLiteral);
+		}
+	}
+
+	void ForEachLiteral(const FMetasoundFrontendGraphClass& InGraphClass, FForEachLiteralFunctionRef OnLiteral)
+	{
+		ForEachLiteral(static_cast<const FMetasoundFrontendClass&>(InGraphClass), OnLiteral);
+
+		for (const FMetasoundFrontendNode& Node : InGraphClass.Graph.Nodes)
+		{
+			ForEachLiteral(Node, OnLiteral);
+		}
+
+		for (const FMetasoundFrontendVariable& Variable : InGraphClass.Graph.Variables)
+		{
+			OnLiteral(Variable.TypeName, Variable.Literal);
+		}
+	}
+
+	void ForEachLiteral(const FMetasoundFrontendClass& InClass, FForEachLiteralFunctionRef OnLiteral)
+	{
+		for (const FMetasoundFrontendClassInput& ClassInput : InClass.Interface.Inputs)
+		{
+			OnLiteral(ClassInput.TypeName, ClassInput.DefaultLiteral);
+		}
+	}
+
+	void ForEachLiteral(const FMetasoundFrontendNode& InNode, FForEachLiteralFunctionRef OnLiteral)
+	{
+		for (const FMetasoundFrontendVertexLiteral& VertexLiteral : InNode.InputLiterals)
+		{
+			auto HasEqualVertexID = [&VertexLiteral](const FMetasoundFrontendVertex& InVertex) -> bool
+			{ 
+				return InVertex.VertexID == VertexLiteral.VertexID; 
+			};
+
+			if (const FMetasoundFrontendVertex* InputVertex = InNode.Interface.Inputs.FindByPredicate(HasEqualVertexID))
+			{
+				OnLiteral(InputVertex->TypeName, VertexLiteral.Value);
+			}
+		}
+	}
+}
+
diff --git a/Engine/Plugins/Runtime/Metasound/Source/MetasoundFrontend/Private/MetasoundFrontendGraph.cpp b/Engine/Plugins/Runtime/Metasound/Source/MetasoundFrontend/Private/MetasoundFrontendGraph.cpp
--- a/Engine/Plugins/Runtime/Metasound/Source/MetasoundFrontend/Private/MetasoundFrontendGraph.cpp
+++ b/Engine/Plugins/Runtime/Metasound/Source/MetasoundFrontend/Private/MetasoundFrontendGraph.cpp
@@ -10,6 +10,7 @@
 #include "MetasoundFrontend.h"
 #include "MetasoundFrontendDataTypeRegistry.h"
 #include "MetasoundFrontendNodeTemplateRegistry.h"
+#include "MetasoundFrontendProxyDataCache.h"
 #include "MetasoundFrontendRegistries.h"
 #include "MetasoundGraph.h"
 #include "MetasoundLiteralNode.h"
@@ -94,7 +95,7 @@ namespace Metasound
 		return true;
 	}
 
-	TUniquePtr<INode> FFrontendGraphBuilder::CreateVariableNode(const FMetasoundFrontendNode& InNode, const FMetasoundFrontendGraph& InGraph)
+	TUniquePtr<INode> FFrontendGraphBuilder::CreateVariableNode(const FBuildContext& InContext, const FMetasoundFrontendNode& InNode, const FMetasoundFrontendGraph& InGraph)
 	{
 		using namespace Metasound::Frontend;
 
@@ -107,7 +108,7 @@ namespace Metasound
 
 			if (IsLiteralParsableByDataType)
 			{
-				FLiteral Literal = FrontendVariable->Literal.ToLiteral(FrontendVariable->TypeName);
+				FLiteral Literal = FrontendVariable->Literal.ToLiteral(FrontendVariable->TypeName, &InContext.DataTypeRegistry, &InContext.ProxyDataCache);
 
 				FVariableNodeConstructorParams InitParams =
 				{
@@ -131,7 +132,7 @@ namespace Metasound
 		return TUniquePtr<INode>(nullptr);
 	}
 
-	TUniquePtr<INode> FFrontendGraphBuilder::CreateInputNode(const FMetasoundFrontendNode& InNode, const FMetasoundFrontendClass& InClass, const FMetasoundFrontendClassInput& InOwningGraphClassInput, bool bEnableTransmission)
+	TUniquePtr<INode> FFrontendGraphBuilder::CreateInputNode(const FBuildContext& InContext, const FMetasoundFrontendNode& InNode, const FMetasoundFrontendClass& InClass, const FMetasoundFrontendClassInput& InOwningGraphClassInput)
 	{
 		using namespace Metasound::Frontend;
 
@@ -148,7 +149,7 @@ namespace Metasound
 
 				if (IsLiteralParsableByDataType)
 				{
-					FLiteral Literal = FrontendLiteral->ToLiteral(InputVertex.TypeName);
+					FLiteral Literal = FrontendLiteral->ToLiteral(InputVertex.TypeName, &InContext.DataTypeRegistry, &InContext.ProxyDataCache);
 
 					FInputNodeConstructorParams InitParams =
 					{
@@ -182,7 +183,7 @@ namespace Metasound
 		return TUniquePtr<INode>(nullptr);
 	}
 
-	TUniquePtr<INode> FFrontendGraphBuilder::CreateOutputNode(const FMetasoundFrontendNode& InNode, const FMetasoundFrontendClass& InClass, FBuildGraphContext& InGraphContext, const TSet<FNodeIDVertexID>& InEdgeDestinations)
+	TUniquePtr<INode> FFrontendGraphBuilder::CreateOutputNode(const FBuildContext& InContext, const FMetasoundFrontendNode& InNode, const FMetasoundFrontendClass& InClass, FBuildGraphContext& InGraphContext, const TSet<FNodeIDVertexID>& InEdgeDestinations)
 	{
 		using namespace Metasound::Frontend;
 
@@ -202,7 +203,7 @@ namespace Metasound
 
 			{
 				const FNodeInitData InitData = FrontendGraphPrivate::CreateNodeInitData(InNode);
-				TArray<FDefaultLiteralData> DefaultLiteralData = GetInputDefaultLiteralData(InNode, InitData, InEdgeDestinations);
+				TArray<FDefaultLiteralData> DefaultLiteralData = GetInputDefaultLiteralData(InContext, InNode, InitData, InEdgeDestinations);
 				for (FDefaultLiteralData& Data : DefaultLiteralData)
 				{
 					InGraphContext.DefaultInputs.Emplace(FNodeIDVertexID { InNode.GetID(), Data.DestinationVertexID }, MoveTemp(Data));
@@ -222,7 +223,7 @@ namespace Metasound
 		return TUniquePtr<INode>(nullptr);
 	}
 
-	TUniquePtr<INode> FFrontendGraphBuilder::CreateExternalNode(const FMetasoundFrontendNode& InNode, const FMetasoundFrontendClass& InClass, FBuildGraphContext& InGraphContext, const TSet<FNodeIDVertexID>& InEdgeDestinations)
+	TUniquePtr<INode> FFrontendGraphBuilder::CreateExternalNode(const FBuildContext& InContext, const FMetasoundFrontendNode& InNode, const FMetasoundFrontendClass& InClass, FBuildGraphContext& InGraphContext, const TSet<FNodeIDVertexID>& InEdgeDestinations)
 	{
 		using namespace Frontend;
 
@@ -230,7 +231,7 @@ namespace Metasound
 
 		const FNodeInitData InitData = FrontendGraphPrivate::CreateNodeInitData(InNode);
 		{
-			TArray<FDefaultLiteralData> DefaultLiteralData = GetInputDefaultLiteralData(InNode, InitData, InEdgeDestinations);
+			TArray<FDefaultLiteralData> DefaultLiteralData = GetInputDefaultLiteralData(InContext, InNode, InitData, InEdgeDestinations);
 			for (FDefaultLiteralData& Data : DefaultLiteralData)
 			{
 				InGraphContext.DefaultInputs.Emplace(FNodeIDVertexID{ InNode.GetID(), Data.DestinationVertexID }, MoveTemp(Data));
@@ -348,7 +349,7 @@ namespace Metasound
 		return Literal;
 	}
 
-	bool FFrontendGraphBuilder::AddNodesToGraph(FBuildGraphContext& InGraphContext, const TSet<FName>& InTransmittableInputNames)
+	bool FFrontendGraphBuilder::AddNodesToGraph(FBuildGraphContext& InGraphContext)
 	{
 		TSet<FNodeIDVertexID> GraphEdgeDestinations;
 		const TArray<FMetasoundFrontendEdge>& GraphEdges = InGraphContext.GraphClass.Graph.Edges;
@@ -372,8 +373,7 @@ namespace Metasound
 
 						if ((nullptr != ClassInput) && (INDEX_NONE != InputIndex))
 						{
-							const bool bEnableTransmission = InTransmittableInputNames.Contains(ClassInput->Name);
-							TSharedPtr<const INode> InputNode(CreateInputNode(Node, *NodeClass, *ClassInput, bEnableTransmission).Release());
+							TSharedPtr<const INode> InputNode(CreateInputNode(InGraphContext.BuildContext, Node, *NodeClass, *ClassInput).Release());
 							InGraphContext.Graph->AddInputNode(Node.GetID(), InputIndex, ClassInput->Name, InputNode);
 						}
 						else
@@ -391,7 +391,7 @@ namespace Metasound
 						const FMetasoundFrontendClassOutput* ClassOutput = FindClassOutputForOutputNode(InGraphContext.GraphClass, Node, OutputIndex);
 						if ((nullptr != ClassOutput) && (INDEX_NONE != OutputIndex))
 						{
-							TSharedPtr<const INode> OutputNode(CreateOutputNode(Node, *NodeClass, InGraphContext, GraphEdgeDestinations).Release());
+							TSharedPtr<const INode> OutputNode(CreateOutputNode(InGraphContext.BuildContext, Node, *NodeClass, InGraphContext, GraphEdgeDestinations).Release());
 							InGraphContext.Graph->AddOutputNode(Node.GetID(), OutputIndex, ClassOutput->Name, OutputNode);
 						}
 						else
@@ -427,7 +427,7 @@ namespace Metasound
 
 					case EMetasoundFrontendClassType::Variable:
 					{
-						TSharedPtr<const INode> VariableNode(CreateVariableNode(Node, InGraphContext.GraphClass.Graph).Release());
+						TSharedPtr<const INode> VariableNode(CreateVariableNode(InGraphContext.BuildContext, Node, InGraphContext.GraphClass.Graph).Release());
 						InGraphContext.Graph->AddNode(Node.GetID(), VariableNode);
 					}
 					break;
@@ -441,7 +441,7 @@ namespace Metasound
 					case EMetasoundFrontendClassType::External:
 					default:
 					{
-						TSharedPtr<const INode> ExternalNode(CreateExternalNode(Node, *NodeClass, InGraphContext, GraphEdgeDestinations).Release());
+						TSharedPtr<const INode> ExternalNode(CreateExternalNode(InGraphContext.BuildContext, Node, *NodeClass, InGraphContext, GraphEdgeDestinations).Release());
 						InGraphContext.Graph->AddNode(Node.GetID(), ExternalNode);
 					}
 					break;
@@ -590,7 +590,7 @@ namespace Metasound
 		return true;
 	}
 
-	TArray<FFrontendGraphBuilder::FDefaultLiteralData> FFrontendGraphBuilder::GetInputDefaultLiteralData(const FMetasoundFrontendNode& InNode, const FNodeInitData& InInitData, const TSet<FNodeIDVertexID>& InEdgeDestinations)
+	TArray<FFrontendGraphBuilder::FDefaultLiteralData> FFrontendGraphBuilder::GetInputDefaultLiteralData(const FBuildContext& InContext, const FMetasoundFrontendNode& InNode, const FNodeInitData& InInitData, const TSet<FNodeIDVertexID>& InEdgeDestinations)
 	{
 		TArray<FDefaultLiteralData> DefaultLiteralData;
 
@@ -612,7 +612,7 @@ namespace Metasound
 					bRequiresDefault = !InEdgeDestinations.Contains({InNode.GetID(), Vertex.VertexID});
 					if (bRequiresDefault)
 					{
-						InitParams.Literal = Literal.Value.ToLiteral(Vertex.TypeName);
+						InitParams.Literal = Literal.Value.ToLiteral(Vertex.TypeName, &InContext.DataTypeRegistry, &InContext.ProxyDataCache);
 						InitParams.InstanceID = FGuid::NewGuid();
 						InitParams.NodeName = "Literal";
 						TypeName = Vertex.TypeName;
@@ -737,7 +737,7 @@ namespace Metasound
 		return bSuccess;
 	}
 
-	TUniquePtr<FFrontendGraph> FFrontendGraphBuilder::CreateGraph(FBuildContext& InContext, const FMetasoundFrontendGraphClass& InGraphClass, const TSet<FName>& InTransmittableInputNames)
+	TUniquePtr<FFrontendGraph> FFrontendGraphBuilder::CreateGraph(FBuildContext& InContext, const FMetasoundFrontendGraphClass& InGraphClass)
 	{
 		const FString GraphName = InContext.DebugAssetName;
 
@@ -748,7 +748,7 @@ namespace Metasound
 			InContext
 		};
 
-		bool bSuccess = AddNodesToGraph(BuildGraphContext, InTransmittableInputNames);
+		bool bSuccess = AddNodesToGraph(BuildGraphContext);
 
 		if (bSuccess)
 		{
@@ -772,8 +772,19 @@ namespace Metasound
 
 	TUniquePtr<FFrontendGraph> FFrontendGraphBuilder::CreateGraph(const FMetasoundFrontendGraphClass& InGraph, const TArray<FMetasoundFrontendGraphClass>& InSubgraphs, const TArray<FMetasoundFrontendClass>& InDependencies, const TSet<FName>& InTransmittableInputNames, const FString& InDebugAssetName)
 	{
-		FBuildContext Context;
-		Context.DebugAssetName = InDebugAssetName;
+		return CreateGraph(InGraph, InSubgraphs, InDependencies, InDebugAssetName);
+	}
+
+	TUniquePtr<FFrontendGraph> FFrontendGraphBuilder::CreateGraph(const FMetasoundFrontendGraphClass& InGraph, const TArray<FMetasoundFrontendGraphClass>& InSubgraphs, const TArray<FMetasoundFrontendClass>& InDependencies, const Frontend::FProxyDataCache& InProxyDataCache, const FString& InDebugAssetName)
+	{
+		FBuildContext Context
+		{
+			InDebugAssetName, 					// DebugAssetName
+			{}, 								// FrontendClasses
+			{}, 								// Graphs
+			Frontend::IDataTypeRegistry::Get(), // DataTypeRegistry
+			InProxyDataCache 					// ProxyDataCache
+		};
 
 		// Gather all references to node classes from external dependencies and subgraphs.
 		for (const FMetasoundFrontendClass& ExtClass : InDependencies)
@@ -799,7 +810,7 @@ namespace Metasound
 		// Create each subgraph.
 		for (const FMetasoundFrontendGraphClass* FrontendSubgraphPtr : FrontendSubgraphPtrs)
 		{
-			TSharedPtr<const INode> Subgraph(CreateGraph(Context, *FrontendSubgraphPtr, InTransmittableInputNames).Release());
+			TSharedPtr<const INode> Subgraph(CreateGraph(Context, *FrontendSubgraphPtr).Release());
 			if (!Subgraph.IsValid())
 			{
 				UE_LOG(LogMetaSound, Warning, TEXT("Failed to create subgraph [SubgraphName: %s] in asset '%s'"), *FrontendSubgraphPtr->Metadata.GetClassName().ToString(), *InDebugAssetName);
@@ -812,12 +823,43 @@ namespace Metasound
 		}
 
 		// Create parent graph.
-		return CreateGraph(Context, InGraph, InTransmittableInputNames);
+		return CreateGraph(Context, InGraph);
+	}
+
+	TUniquePtr<FFrontendGraph> FFrontendGraphBuilder::CreateGraph(const FMetasoundFrontendDocument& InDocument, const FString& InDebugAssetName)
+	{
+		// Create proxies before creating graph.
+		Frontend::FProxyDataCache ProxyDataCache;
+		ProxyDataCache.CreateAndCacheProxies(InDocument);
+		
+		return CreateGraph(InDocument, ProxyDataCache, InDebugAssetName);
+	}
+
+	TUniquePtr<FFrontendGraph> FFrontendGraphBuilder::CreateGraph(const FMetasoundFrontendGraphClass& InGraph, const TArray<FMetasoundFrontendGraphClass>& InSubgraphs, const TArray<FMetasoundFrontendClass>& InDependencies, const FString& InDebugAssetName)
+	{
+		// Create proxies before building graph
+		Frontend::FProxyDataCache ProxyDataCache;
+		ProxyDataCache.CreateAndCacheProxies(InGraph);
+		for (const FMetasoundFrontendGraphClass& SubgraphClass : InSubgraphs)
+		{
+			ProxyDataCache.CreateAndCacheProxies(SubgraphClass);
+		}
+		for (const FMetasoundFrontendClass& DependencyClass : InDependencies)
+		{
+			ProxyDataCache.CreateAndCacheProxies(DependencyClass);
+		}
+
+
+		return CreateGraph(InGraph, InSubgraphs, InDependencies, ProxyDataCache, InDebugAssetName);
 	}
 	
-	/* Metasound document should be inflated by now. */
 	TUniquePtr<FFrontendGraph> FFrontendGraphBuilder::CreateGraph(const FMetasoundFrontendDocument& InDocument, const TSet<FName>& InTransmittableInputNames, const FString& InDebugAssetName)
 	{
-		return CreateGraph(InDocument.RootGraph, InDocument.Subgraphs, InDocument.Dependencies, InTransmittableInputNames, InDebugAssetName);
+		return CreateGraph(InDocument, InDebugAssetName);
+	}
+
+	TUniquePtr<FFrontendGraph> FFrontendGraphBuilder::CreateGraph(const FMetasoundFrontendDocument& InDocument, const Frontend::FProxyDataCache& InProxyDataCache, const FString& InDebugAssetName)
+	{
+		return CreateGraph(InDocument.RootGraph, InDocument.Subgraphs, InDocument.Dependencies, InProxyDataCache, InDebugAssetName);
 	}
 }
diff --git a/Engine/Plugins/Runtime/Metasound/Source/MetasoundFrontend/Private/MetasoundFrontendGraphController.cpp b/Engine/Plugins/Runtime/Metasound/Source/MetasoundFrontend/Private/MetasoundFrontendGraphController.cpp
--- a/Engine/Plugins/Runtime/Metasound/Source/MetasoundFrontend/Private/MetasoundFrontendGraphController.cpp
+++ b/Engine/Plugins/Runtime/Metasound/Source/MetasoundFrontend/Private/MetasoundFrontendGraphController.cpp
@@ -10,6 +10,7 @@
 #include "MetasoundFrontendGraph.h"
 #include "MetasoundFrontendNodeController.h"
 #include "MetasoundFrontendNodeTemplateRegistry.h"
+#include "MetasoundFrontendProxyDataCache.h"
 #include "MetasoundFrontendSubgraphNodeController.h"
 #include "MetasoundFrontendInvalidController.h"
 #include "MetasoundFrontendVariableController.h"
@@ -1609,13 +1610,13 @@ namespace Metasound
 		{
 			if (const FMetasoundFrontendGraphClass* GraphClass = GraphClassPtr.Get())
 			{
-				// TODO: bubble up errors. 
 				const TArray<FMetasoundFrontendGraphClass>& Subgraphs = OwningDocument->GetSubgraphs();
 				const TArray<FMetasoundFrontendClass>& Dependencies = OwningDocument->GetDependencies();
 
 				FString UnknownAsset = TEXT("UnknownAsset");
-				TSet<FName> TransmittableInputNames;
-				TUniquePtr<FFrontendGraph> Graph = FFrontendGraphBuilder::CreateGraph(*GraphClass, Subgraphs, Dependencies, TransmittableInputNames, UnknownAsset);
+				FProxyDataCache ProxyCache;
+				ProxyCache.CreateAndCacheProxies(*(OwningDocument->GetDocumentPtr().Get()));
+				TUniquePtr<FFrontendGraph> Graph = FFrontendGraphBuilder::CreateGraph(*GraphClass, Subgraphs, Dependencies, ProxyCache, UnknownAsset);
 
 				if (!Graph.IsValid())
 				{
diff --git a/Engine/Plugins/Runtime/Metasound/Source/MetasoundFrontend/Private/MetasoundFrontendLiteral.cpp b/Engine/Plugins/Runtime/Metasound/Source/MetasoundFrontend/Private/MetasoundFrontendLiteral.cpp
--- a/Engine/Plugins/Runtime/Metasound/Source/MetasoundFrontend/Private/MetasoundFrontendLiteral.cpp
+++ b/Engine/Plugins/Runtime/Metasound/Source/MetasoundFrontend/Private/MetasoundFrontendLiteral.cpp
@@ -5,6 +5,7 @@
 #include "CoreMinimal.h"
 #include "MetasoundLog.h"
 #include "MetasoundFrontendDataTypeRegistry.h"
+#include "MetasoundFrontendProxyDataCache.h"
 #include "MetasoundFrontendRegistries.h"
 
 #include UE_INLINE_GENERATED_CPP_BY_NAME(MetasoundFrontendLiteral)
@@ -645,49 +646,75 @@ bool FMetasoundFrontendLiteral::TryGet(TArray<FString>& OutValue) const
 	return false;
 }
 
-Metasound::FLiteral FMetasoundFrontendLiteral::ToLiteral(const FName& InMetasoundDataTypeName) const
+Metasound::FLiteral FMetasoundFrontendLiteral::ToLiteral(const FName& InMetasoundDataTypeName, const Metasound::Frontend::IDataTypeRegistry* InDataTypeRegistry, const Metasound::Frontend::FProxyDataCache* InProxyDataCache) const
 {
 	using namespace Metasound;
 	using namespace Metasound::Frontend;
 
-
-	FLiteral Literal = FLiteral::CreateInvalid();
-	const IDataTypeRegistry& DataTypeRegistry = IDataTypeRegistry::Get();
-
-	const bool bIsTypeSupported = DataTypeRegistry.IsLiteralTypeSupported(InMetasoundDataTypeName, Type);
-
-	if (!bIsTypeSupported)
+	if (EMetasoundFrontendLiteralType::UObject == Type)
 	{
-		UE_LOG(LogMetaSound, Error, TEXT("Reverting to default literal type for data type. Failed to create supported Metasound::FLiteral for data type [Name:%s] with FMetasoundFrontendLiteral [Literal:%s]"), *InMetasoundDataTypeName.ToString(), *ToString());
+		if (ensure(AsUObject.Num() > 0))
+		{
+			if (InProxyDataCache)
+			{
+				if (const TSharedPtr<Audio::IProxyData>* ProxyData = InProxyDataCache->FindProxy(AsUObject[0]))
+				{
+					return FLiteral(*ProxyData);
+				}
+				else
+				{
+					return FLiteral(TSharedPtr<Audio::IProxyData>());
+				}
+			}
+			else
+			{
+				if (!InDataTypeRegistry)
+				{
+					InDataTypeRegistry = &IDataTypeRegistry::Get();
+				}
+
+				// UObject proxies must go through the registry. The registry contains the information
+				// needed to convert UObjects to proxies. 
+				return InDataTypeRegistry->CreateLiteralFromUObject(InMetasoundDataTypeName, AsUObject[0]);
+			}
+		}
 
-		Literal = DataTypeRegistry.CreateDefaultLiteral(InMetasoundDataTypeName);
+		return FLiteral(TSharedPtr<Audio::IProxyData>());
 	}
-	else if (EMetasoundFrontendLiteralType::UObject == Type)
+	else if (EMetasoundFrontendLiteralType::UObjectArray == Type)
 	{
-		if (ensure(AsUObject.Num() > 0))
+		if (InProxyDataCache)
 		{
+			TArray<TSharedPtr<Audio::IProxyData>> Proxies;
+			for (const UObject* Object : AsUObject)
+			{
+				if (const TSharedPtr<Audio::IProxyData>* ProxyData = InProxyDataCache->FindProxy(Object))
+				{
+					Proxies.Add(*ProxyData);
+				}
+				else
+				{
+					Proxies.Add(TSharedPtr<Audio::IProxyData>());
+				}
+			}
+			return FLiteral(Proxies);
+		}
+		else
+		{
+			if (!InDataTypeRegistry)
+			{
+				InDataTypeRegistry = &IDataTypeRegistry::Get();
+			}
 			// UObject proxies must go through the registry. The registry contains the information
 			// needed to convert UObjects to proxies. 
-			Literal = DataTypeRegistry.CreateLiteralFromUObject(InMetasoundDataTypeName, AsUObject[0]);
+			return InDataTypeRegistry->CreateLiteralFromUObjectArray(InMetasoundDataTypeName, AsUObject);
 		}
 	}
-	else if (EMetasoundFrontendLiteralType::UObjectArray == Type)
-	{
-		// UObject proxies must go through the registry. The registry contains the information
-		// needed to convert UObjects to proxies. 
-		Literal = DataTypeRegistry.CreateLiteralFromUObjectArray(InMetasoundDataTypeName, AsUObject);
-	}
 	else 
 	{
 		// Use default conversions for core literal types.
-		Literal = ToLiteralNoProxy();
+		return ToLiteralNoProxy();
 	}
-
-	// The support for the data type should be the same whether we pass in an
-	// Metasound::ELiteralType or a EMetasoundFrontendLiteralType
-	check(bIsTypeSupported == DataTypeRegistry.IsLiteralTypeSupported(InMetasoundDataTypeName, Literal.GetType()));
-
-	return Literal;
 }
 
 Metasound::FLiteral FMetasoundFrontendLiteral::ToLiteralNoProxy() const
diff --git a/Engine/Plugins/Runtime/Metasound/Source/MetasoundFrontend/Private/MetasoundFrontendProxyDataCache.cpp b/Engine/Plugins/Runtime/Metasound/Source/MetasoundFrontend/Private/MetasoundFrontendProxyDataCache.cpp
new file mode 100644
--- /dev/null
+++ b/Engine/Plugins/Runtime/Metasound/Source/MetasoundFrontend/Private/MetasoundFrontendProxyDataCache.cpp
@@ -0,0 +1,101 @@
+// Copyright Epic Games, Inc. All Rights Reserved.
+
+#include "MetasoundFrontendProxyDataCache.h"
+
+
+#include "Containers/Array.h"
+#include "Containers/Map.h"
+#include "CoreGlobals.h"
+#include "IAudioProxyInitializer.h"
+#include "MetasoundFrontendDataTypeRegistry.h"
+#include "MetasoundFrontendDocument.h"
+#include "MetasoundFrontendLiteral.h"
+#include "MetasoundTrace.h"
+#include "Templates/SharedPointer.h"
+#include "UObject/NameTypes.h"
+#include "UObject/Object.h"
+
+namespace Metasound
+{
+	namespace Frontend
+	{
+		namespace MetasoundFrontendProxyDataCachePrivate
+		{
+			template<typename DocElementType>
+			void CreateAndCacheProxies(FProxyDataCache& InCache, const DocElementType& InDocElement)
+			{
+				checkf(IsInGameThread() || IsInAudioThread(), TEXT("Proxies should only be created in the game thread or audio thread"));
+				IDataTypeRegistry& DataRegistry = IDataTypeRegistry::Get();
+				TArray<UObject*> UObjectArray;
+
+				ForEachLiteral(InDocElement, [&InCache, &DataRegistry, &UObjectArray](const FName& InDataType, const FMetasoundFrontendLiteral& InLiteral)
+					{
+						EMetasoundFrontendLiteralType LiteralType = InLiteral.GetType();
+						if (LiteralType == EMetasoundFrontendLiteralType::UObject)
+						{
+							UObject* Object = nullptr;
+							InLiteral.TryGet(Object);
+							if (Object)
+							{
+								if (!InCache.Contains(Object))
+								{
+									InCache.CacheProxy(Object, DataRegistry.CreateProxyFromUObject(InDataType, Object));
+								}
+							}
+						}
+						else if (LiteralType == EMetasoundFrontendLiteralType::UObjectArray)
+						{
+							FName ElementDataTypeName = CreateElementTypeNameFromArrayTypeName(InDataType);
+							UObjectArray.Reset();
+							InLiteral.TryGet(UObjectArray);
+							for (UObject* Object : UObjectArray)
+							{
+								if (Object)
+								{
+									if (!InCache.Contains(Object))
+									{
+										InCache.CacheProxy(Object, DataRegistry.CreateProxyFromUObject(ElementDataTypeName, Object));
+									}
+								}
+							}
+						}
+					}
+				);
+			}
+		}
+
+		void FProxyDataCache::CreateAndCacheProxies(const FMetasoundFrontendDocument& InDocument)
+		{
+			METASOUND_TRACE_CPUPROFILER_EVENT_SCOPE(Metasound::FProxyDataCache::CreateAndCacheProxies_Document)	
+			MetasoundFrontendProxyDataCachePrivate::CreateAndCacheProxies(*this, InDocument);
+		}
+
+		void FProxyDataCache::CreateAndCacheProxies(const FMetasoundFrontendGraphClass& InGraphClass)
+		{
+			METASOUND_TRACE_CPUPROFILER_EVENT_SCOPE(Metasound::FProxyDataCache::CreateAndCacheProxies_GraphClass)	
+			MetasoundFrontendProxyDataCachePrivate::CreateAndCacheProxies(*this, InGraphClass);
+		}
+
+		void FProxyDataCache::CreateAndCacheProxies(const FMetasoundFrontendClass& InClass)
+		{
+			METASOUND_TRACE_CPUPROFILER_EVENT_SCOPE(Metasound::FProxyDataCache::CreateAndCacheProxies_Class)
+			MetasoundFrontendProxyDataCachePrivate::CreateAndCacheProxies(*this, InClass);
+		}
+
+		bool FProxyDataCache::Contains(const UObject* InUObject) const
+		{
+			return ProxyCache.Contains(InUObject);
+		}
+
+		void FProxyDataCache::CacheProxy(const UObject* InUObject, TSharedPtr<Audio::IProxyData> InProxy)
+		{
+			ProxyCache.Add(InUObject, MoveTemp(InProxy));
+		}
+
+		const TSharedPtr<Audio::IProxyData>* FProxyDataCache::FindProxy(const UObject* InObject) const
+		{
+			return ProxyCache.Find(InObject);
+		}
+	}
+}
+
diff --git a/Engine/Plugins/Runtime/Metasound/Source/MetasoundFrontend/Private/MetasoundFrontendProxyDataCache.h b/Engine/Plugins/Runtime/Metasound/Source/MetasoundFrontend/Private/MetasoundFrontendProxyDataCache.h
new file mode 100644
--- /dev/null
+++ b/Engine/Plugins/Runtime/Metasound/Source/MetasoundFrontend/Private/MetasoundFrontendProxyDataCache.h
@@ -0,0 +1,75 @@
+// Copyright Epic Games, Inc. All Rights Reserved.
+
+#pragma once
+
+#include "Containers/Map.h"
+#include "Templates/SharedPointer.h"
+
+
+// Forward declarations
+class UObject;
+struct FMetasoundFrontendClass;
+struct FMetasoundFrontendGraphClass;
+struct FMetasoundFrontendDocument;
+
+namespace Audio
+{
+	class IProxyData;
+}
+
+namespace Metasound
+{
+	namespace Frontend
+	{
+		/** FProxyDataCache creates and holds onto Audio::IProxyData produced from UObjects.
+		 *
+		 * A proxy data cache is useful in scenarios where proxies need to be retrieved
+		 * from a thread where calling UObject functions is not safe. Additionally,
+		 * the proxy data cache only creates a single proxy per a UObject when creating
+		 * proxies from a FMetasoundFrontendDocument or any of it's subobjects. 
+		 */
+		class FProxyDataCache
+		{
+		public:
+			/** Create and cache all proxies found in a document. */
+			void CreateAndCacheProxies(const FMetasoundFrontendDocument& InDocument);
+
+			/** Create and cache all proxies found in a graph class. */
+			void CreateAndCacheProxies(const FMetasoundFrontendGraphClass& InGraphClass);
+
+			/** Create and cache all proxies found in a frontend class class. */
+			void CreateAndCacheProxies(const FMetasoundFrontendClass& InClass);
+
+			/** Cache a specific object associated with a UObject pointer.
+			 *
+			 * No methods or data are accessed from the UObject pointer passed 
+			 * into this method. This method is safe even if the underlying UObject 
+			 * is invalid or garbage collected before or during this call.
+			 */
+			void CacheProxy(const UObject* InUObject, TSharedPtr<Audio::IProxyData> InProxy);
+
+			/** Returns true if there is a proxy stored which is associated with
+			 * the UObject address. 
+			 *
+			 * No methods or data are accessed from the UObject pointer passed 
+			 * into this method. This method is safe even if the underlying UObject 
+			 * is invalid or garbage collected before or during this call.
+			 */
+			bool Contains(const UObject* InUObject) const;
+
+
+			/** Returns a pointer to an existing proxy if one has been cached. 
+			 * Returns nullptr if none exist.
+			 *
+			 * No methods or data are accessed from the UObject pointer passed 
+			 * into this method. This method is safe even if the underlying UObject 
+			 * is invalid or garbage collected before or during this call.
+			 */
+			const TSharedPtr<Audio::IProxyData>* FindProxy(const UObject* InObject) const;
+
+		private:
+
+			TMap<const void*, TSharedPtr<Audio::IProxyData>> ProxyCache;
+		};
+	}
+}
diff --git a/Engine/Plugins/Runtime/Metasound/Source/MetasoundFrontend/Public/MetasoundAssetBase.h b/Engine/Plugins/Runtime/Metasound/Source/MetasoundFrontend/Public/MetasoundAssetBase.h
--- a/Engine/Plugins/Runtime/Metasound/Source/MetasoundFrontend/Public/MetasoundAssetBase.h
+++ b/Engine/Plugins/Runtime/Metasound/Source/MetasoundFrontend/Public/MetasoundAssetBase.h
@@ -29,6 +29,8 @@ namespace Metasound
 		class IInterfaceRegistryEntry;
 
 		METASOUNDFRONTEND_API float GetDefaultBlockRate();
+
+		class FProxyDataCache;
 	} // namespace Frontend
 } // namespace Metasound
 
@@ -247,5 +249,5 @@ private:
 	FGuid CurrentCachedRuntimeDataChangeID;
 	FRuntimeData CachedRuntimeData;
 
-	TSharedPtr<Metasound::FGraph, ESPMode::ThreadSafe> BuildMetasoundDocument(const FMetasoundFrontendDocument& InPreprocessDoc, const TSet<FName>& InTransmittableInputNames) const;
+	TSharedPtr<Metasound::FGraph, ESPMode::ThreadSafe> BuildMetasoundDocument(const FMetasoundFrontendDocument& InPreprocessDoc, const Metasound::Frontend::FProxyDataCache& InProxies) const;
 };
diff --git a/Engine/Plugins/Runtime/Metasound/Source/MetasoundFrontend/Public/MetasoundFrontendDocument.h b/Engine/Plugins/Runtime/Metasound/Source/MetasoundFrontend/Public/MetasoundFrontendDocument.h
--- a/Engine/Plugins/Runtime/Metasound/Source/MetasoundFrontend/Public/MetasoundFrontendDocument.h
+++ b/Engine/Plugins/Runtime/Metasound/Source/MetasoundFrontend/Public/MetasoundFrontendDocument.h
@@ -1652,3 +1652,22 @@ public:
 
 METASOUNDFRONTEND_API const TCHAR* LexToString(EMetasoundFrontendClassType InClassType);
 METASOUNDFRONTEND_API const TCHAR* LexToString(EMetasoundFrontendVertexAccessType InVertexAccess);
+
+namespace Metasound
+{
+	/** Signature of function called for each found literal. */
+	using FForEachLiteralFunctionRef = TFunctionRef<void(const FName& InDataTypeName, const FMetasoundFrontendLiteral&)>; 
+
+	/** Execute the provided function for each literal on a FMetasoundFrontendDocument.*/
+	METASOUNDFRONTEND_API void ForEachLiteral(const FMetasoundFrontendDocument& InDoc, FForEachLiteralFunctionRef OnLiteral);
+
+	/** Execute the provided function for each literal on a FMetasoundFrontendGraphClass.*/
+	METASOUNDFRONTEND_API void ForEachLiteral(const FMetasoundFrontendGraphClass& InGraphClass, FForEachLiteralFunctionRef OnLiteral);
+
+	/** Execute the provided function for each literal on a FMetasoundFrontendClass.*/
+	METASOUNDFRONTEND_API void ForEachLiteral(const FMetasoundFrontendClass& InClass, FForEachLiteralFunctionRef OnLiteral);
+
+	/** Execute the provided function for each literal on a FMetasoundFrontendNode.*/
+	METASOUNDFRONTEND_API void ForEachLiteral(const FMetasoundFrontendNode& InNode, FForEachLiteralFunctionRef OnLiteral);
+}
+
diff --git a/Engine/Plugins/Runtime/Metasound/Source/MetasoundFrontend/Public/MetasoundFrontendGraph.h b/Engine/Plugins/Runtime/Metasound/Source/MetasoundFrontend/Public/MetasoundFrontendGraph.h
--- a/Engine/Plugins/Runtime/Metasound/Source/MetasoundFrontend/Public/MetasoundFrontendGraph.h
+++ b/Engine/Plugins/Runtime/Metasound/Source/MetasoundFrontend/Public/MetasoundFrontendGraph.h
@@ -93,12 +93,26 @@ namespace Metasound
 
 		static bool IsFlat(const FMetasoundFrontendGraphClass& InRoot, const TArray<FMetasoundFrontendClass>& InDependencies);
 
-		/* Metasound document should be in order to create this graph. */
+		/* Create a FFrontendGraph from a FMetasoundFrontendDocument.*/
+		UE_DEPRECATED(5.4, "Use the version of CreateGraph(...) which does not include \"TransmittableInputNames\".")
 		static TUniquePtr<FFrontendGraph> CreateGraph(const FMetasoundFrontendDocument& InDocument, const TSet<FName>& TransmittableInputNames, const FString& InDebugAssetName);
 
-		/* Metasound document should be in order to create this graph. */
+		/* Create a FFrontendGraph from a FMetasoundFrontendDocument subobjects.*/
+		UE_DEPRECATED(5.4, "Use the version of CreateGraph(...) which does not include \"TransmittableInputNames\".")
 		static TUniquePtr<FFrontendGraph> CreateGraph(const FMetasoundFrontendGraphClass& InGraphClass, const TArray<FMetasoundFrontendGraphClass>& InSubgraphs, const TArray<FMetasoundFrontendClass>& InDependencies, const TSet<FName>& TransmittableInputNames, const FString& InDebugAssetName);
 
+		/* Create a FFrontendGraph from a FMetasoundFrontendDocument.*/
+		static TUniquePtr<FFrontendGraph> CreateGraph(const FMetasoundFrontendDocument& InDocument, const FString& InDebugAssetName);
+
+		/* Create a FFrontendGraph from a FMetasoundFrontendDocument retrieving proxies from a FProxyDataCache.*/
+		static TUniquePtr<FFrontendGraph> CreateGraph(const FMetasoundFrontendDocument& InDocument, const Frontend::FProxyDataCache& InProxies, const FString& InDebugAssetName);
+
+		/* Create a FFrontendGraph from a FMetasoundFrontendDocument subobjects retrieving proxies from a FProxyDataCache.*/
+		static TUniquePtr<FFrontendGraph> CreateGraph(const FMetasoundFrontendGraphClass& InGraph, const TArray<FMetasoundFrontendGraphClass>& InSubgraphs, const TArray<FMetasoundFrontendClass>& InDependencies, const FString& InDebugAssetName);
+
+		/* Create a FFrontendGraph from a FMetasoundFrontendDocument subobjects retrieving proxies from a FProxyDataCache.*/
+		static TUniquePtr<FFrontendGraph> CreateGraph(const FMetasoundFrontendGraphClass& InGraph, const TArray<FMetasoundFrontendGraphClass>& InSubgraphs, const TArray<FMetasoundFrontendClass>& InDependencies, const Frontend::FProxyDataCache& InProxyDataCache, const FString& InDebugAssetName);
+
 	private:
 		struct FDefaultLiteralData
 		{
@@ -122,6 +136,8 @@ namespace Metasound
 			FString DebugAssetName;
 			FDependencyByIDMap FrontendClasses;
 			FSharedNodeByIDMap Graphs;
+			const Frontend::IDataTypeRegistry& DataTypeRegistry;
+			const Frontend::FProxyDataCache& ProxyDataCache;
 		};
 
 		// Transient context used for building a specific graph
@@ -133,24 +149,25 @@ namespace Metasound
 			FDefaultInputByIDMap DefaultInputs;
 		};
 
-		static TArray<FDefaultLiteralData> GetInputDefaultLiteralData(const FMetasoundFrontendNode& InNode, const FNodeInitData& InInitData, const TSet<FNodeIDVertexID>& InEdgeDestinations);
+
+		static TArray<FDefaultLiteralData> GetInputDefaultLiteralData(const FBuildContext& InContext, const FMetasoundFrontendNode& InNode, const FNodeInitData& InInitData, const TSet<FNodeIDVertexID>& InEdgeDestinations);
 
 		static bool SortSubgraphDependencies(TArray<const FMetasoundFrontendGraphClass*>& Subgraphs);
 
-		static TUniquePtr<FFrontendGraph> CreateGraph(FBuildContext& InContext, const FMetasoundFrontendGraphClass& InSubgraph, const TSet<FName>& TransmittableInputNames);
+		static TUniquePtr<FFrontendGraph> CreateGraph(FBuildContext& InContext, const FMetasoundFrontendGraphClass& InSubgraph);
 
 		static const FMetasoundFrontendClassInput* FindClassInputForInputNode(const FMetasoundFrontendGraphClass& InOwningGraph, const FMetasoundFrontendNode& InInputNode, int32& OutClassInputIndex);
 		static const FMetasoundFrontendClassOutput* FindClassOutputForOutputNode(const FMetasoundFrontendGraphClass& InOwningGraph, const FMetasoundFrontendNode& InOutputNode, int32& OutClassOutputIndex);
 		static const FMetasoundFrontendLiteral* FindInputLiteralForInputNode(const FMetasoundFrontendNode& InInputNode, const FMetasoundFrontendClass& InInputNodeClass, const FMetasoundFrontendClassInput& InOwningGraphClassInput);
 		static const FMetasoundFrontendVariable* FindVariableForVariableNode(const FMetasoundFrontendNode& InVariableNode, const FMetasoundFrontendGraph& InGraph);
 
-		static TUniquePtr<INode> CreateInputNode(const FMetasoundFrontendNode& InNode, const FMetasoundFrontendClass& InClass, const FMetasoundFrontendClassInput& InOwningGraphClassInput, bool bEnableTransmission);
-		static TUniquePtr<INode> CreateOutputNode(const FMetasoundFrontendNode& InNode, const FMetasoundFrontendClass& InClass, FBuildGraphContext& InGraphContext, const TSet<FNodeIDVertexID>& InEdgeDestinations);
-		static TUniquePtr<INode> CreateVariableNode(const FMetasoundFrontendNode& InNode, const FMetasoundFrontendGraph& InGraph);
-		static TUniquePtr<INode> CreateExternalNode(const FMetasoundFrontendNode& InNode, const FMetasoundFrontendClass& InClass, FBuildGraphContext& InGraphContext, const TSet<FNodeIDVertexID>& InEdgeDestinations);
+		static TUniquePtr<INode> CreateInputNode(const FBuildContext& InContext, const FMetasoundFrontendNode& InNode, const FMetasoundFrontendClass& InClass, const FMetasoundFrontendClassInput& InOwningGraphClassInput);
+		static TUniquePtr<INode> CreateOutputNode(const FBuildContext& InContext, const FMetasoundFrontendNode& InNode, const FMetasoundFrontendClass& InClass, FBuildGraphContext& InGraphContext, const TSet<FNodeIDVertexID>& InEdgeDestinations);
+		static TUniquePtr<INode> CreateVariableNode(const FBuildContext& InContext, const FMetasoundFrontendNode& InNode, const FMetasoundFrontendGraph& InGraph);
+		static TUniquePtr<INode> CreateExternalNode(const FBuildContext& InContext, const FMetasoundFrontendNode& InNode, const FMetasoundFrontendClass& InClass, FBuildGraphContext& InGraphContext, const TSet<FNodeIDVertexID>& InEdgeDestinations);
 
 		// Returns false on error
-		static bool AddNodesToGraph(FBuildGraphContext& InGraphContext, const TSet<FName>& TransmittableInputNames);
+		static bool AddNodesToGraph(FBuildGraphContext& InGraphContext);
 
 		// Returns false on error
 		static bool AddEdgesToGraph(FBuildGraphContext& InGraphContext);
diff --git a/Engine/Plugins/Runtime/Metasound/Source/MetasoundFrontend/Public/MetasoundFrontendLiteral.h b/Engine/Plugins/Runtime/Metasound/Source/MetasoundFrontend/Public/MetasoundFrontendLiteral.h
--- a/Engine/Plugins/Runtime/Metasound/Source/MetasoundFrontend/Public/MetasoundFrontendLiteral.h
+++ b/Engine/Plugins/Runtime/Metasound/Source/MetasoundFrontend/Public/MetasoundFrontendLiteral.h
@@ -53,6 +53,15 @@ static_assert(static_cast<uint8>(EMetasoundFrontendLiteralType::UObjectArray) ==
 static_assert(std::is_same<uint8, std::underlying_type<EMetasoundFrontendLiteralType>::type>::value, "Update type in static_cast<TYPE> from Metasound::ELiteralType to EMetasoundFrontendLiteralType in EMetasoundFrontendLiteralType declaration.");
 static_assert(std::is_same<std::underlying_type<Metasound::ELiteralType>::type, std::underlying_type<EMetasoundFrontendLiteralType>::type>::value, "EMetasoundFrontendLiteralType and Metasound::ELiteralType must have matching underlying types to support conversion.");
 
+// Forward Declare
+namespace Metasound
+{
+	namespace Frontend
+	{
+		class IDataTypeRegistry;
+		class FProxyDataCache;
+	}
+}
 
 // Represents the serialized version of variant literal types. 
 USTRUCT(BlueprintType, meta = (DisplayName = "MetaSound Literal"))
@@ -131,11 +140,16 @@ public:
 
 	EMetasoundFrontendLiteralType GetType() const;
 
-	// Return the literal description parsed into a init param. 
-	// @Returns an invalid init param if the data type couldn't be found, or if the literal type was incompatible with the data type.
-	Metasound::FLiteral ToLiteral(const FName& InMetasoundDataType) const;
+	// Return a Metasound::FLiteral representation of this object. 
+	//
+	// @param InMetaSoundDataType - The name of the MetaSound data type
+	// @param InDataTypeRegistry - A pointer to an existing data type registry. If null, the data type registry will be retrieved within this function.
+	// @param InProxyDataCache - A pointer to an existing proxy data cache. If not null, UObject proxies will be retrieved from the cache. If null, UObject proxies will be created in this function.
+	//
+	// @returns An Fliteral. If the data type couldn't be found, or if the literal type was incompatible with the data type, then an invalid FLiteral is returned.
+	Metasound::FLiteral ToLiteral(const FName& InMetaSoundDataType, const Metasound::Frontend::IDataTypeRegistry* InDataTypeRegistry=nullptr, const Metasound::Frontend::FProxyDataCache* InProxyDataCache=nullptr) const;
 
-	// Does not do type checking or handle proxy
+	// Return a Metasound::Fliteral representation of this object, excluding UObject proxies.
 	Metasound::FLiteral ToLiteralNoProxy() const;
 
 	// Convert the value to a string for printing. 
