[Metasound] Notify users when a generator's I/O changes

diff --git a/Engine/Plugins/Runtime/Metasound/Source/MetasoundEngineTest/Private/MetasoundDynamicGeneratorTests.cpp b/Engine/Plugins/Runtime/Metasound/Source/MetasoundEngineTest/Private/MetasoundDynamicGeneratorTests.cpp
new file mode 100644
--- /dev/null
+++ b/Engine/Plugins/Runtime/Metasound/Source/MetasoundEngineTest/Private/MetasoundDynamicGeneratorTests.cpp
@@ -0,0 +1,225 @@
+// Copyright Epic Games, Inc. All Rights Reserved.
+
+#include "MetasoundDynamicOperatorTransactor.h"
+#include "MetasoundFrontendGraph.h"
+#include "MetasoundGenerator.h"
+
+#include "Interfaces/MetasoundOutputFormatInterfaces.h"
+
+#include "Misc/AutomationTest.h"
+
+#if WITH_DEV_AUTOMATION_TESTS
+
+namespace Metasound::Test::Generator::Dynamic
+{
+	TArray<FVertexName> GetAudioOutputVertexNames(const EMetaSoundOutputAudioFormat Format)
+	{
+		const Engine::FOutputAudioFormatInfo* FormatInfo = Engine::GetOutputAudioFormatInfo().Find(Format);
+		return FormatInfo != nullptr ? FormatInfo->OutputVertexChannelOrder : TArray<FVertexName>{};
+	}
+	
+	class FDynamicGeneratorBuilder
+	{
+	public:
+		FDynamicGeneratorBuilder(FSampleRate SampleRate, int32 BlockSize)
+			: OperatorSettings(SampleRate, static_cast<float>(SampleRate) / BlockSize)
+			, Generator(OperatorSettings)
+			, RenderBuffer(BlockSize)
+		{
+			check(OperatorSettings.GetNumFramesPerBlock() == BlockSize);
+
+			// Add the minimum required interfaces so we don't get warnings when we run the test(s)
+			// NB: if you start getting warnings, check FMetasoundDynamicGraphGenerator to see if the required
+			// I/O has changed.
+			// TODO: Add a future-proof way to do this
+			AddInput<FTrigger>(Frontend::SourceInterface::Inputs::OnPlay, FGuid::NewGuid(), {});
+			AddOutput<FAudioBuffer>(Engine::OutputFormatMonoInterface::Outputs::MonoOut, FGuid::NewGuid());
+
+			// Make the generator
+			FOperatorBuilderSettings BuilderSettings = FOperatorBuilderSettings::GetDefaultSettings();
+			BuilderSettings.bEnableOperatorRebind = true;
+
+			FMetasoundEnvironment Environment;
+			Environment.SetValue<uint64>(Frontend::SourceInterface::Environment::TransmitterID, 123);
+			
+			FMetasoundDynamicGraphGeneratorInitParams InitParams
+			{
+				{
+					OperatorSettings,
+					MoveTemp(BuilderSettings),
+					MakeShared<FGraph>(Transactor.GetGraph()),
+					Environment,
+					"TestMetaSoundGenerator",
+					GetAudioOutputVertexNames(EMetaSoundOutputAudioFormat::Mono),
+					{},
+					true
+				},
+				Transactor.CreateTransformQueue(OperatorSettings, Environment) // Create transaction queue
+			};
+			
+			Generator.Init(MoveTemp(InitParams));
+		}
+		
+		template<typename DataType>
+		bool AddInput(const FVertexName& Name, const FGuid& NodeGuid, const FLiteral& DefaultLiteral)
+		{
+			FInputNodeConstructorParams InputNodeParams
+			{
+				Name,
+				NodeGuid,
+				Name,
+				DefaultLiteral
+			};
+			
+			TUniquePtr<INode> Node = DataRegistry.CreateInputNode(GetMetasoundDataTypeName<DataType>(), MoveTemp(InputNodeParams));
+
+			if (!Node.IsValid())
+			{
+				return false;
+			}
+
+			Transactor.AddNode(InputNodeParams.InstanceID, MoveTemp(Node));
+
+			const auto CreateDataReference = [](
+				const FOperatorSettings& InSettings,
+				const FName InDataType, 
+				const FLiteral& InLiteral,
+				const EDataReferenceAccessType InAccessType)
+			{
+				const Frontend::IDataTypeRegistry& DataRegistry2 = Frontend::IDataTypeRegistry::Get();
+				return DataRegistry2.CreateDataReference(InDataType, InAccessType, InLiteral, InSettings);
+			};
+			
+			Transactor.AddInputDataDestination(
+				InputNodeParams.InstanceID,
+				Name,
+				DefaultLiteral,
+				CreateDataReference);
+
+			return true;
+		}
+
+		void RemoveInput(const FVertexName& Name, const FGuid& NodeGuid)
+		{
+			Transactor.RemoveInputDataDestination(Name);
+			Transactor.RemoveNode(NodeGuid);
+		}
+
+		template<typename DataType>
+		bool AddOutput(const FVertexName& Name, const FGuid& NodeGuid)
+		{
+			FOutputNodeConstructorParams OutputNodeParams
+			{
+				Name,
+				NodeGuid,
+				Name
+			};
+
+			TUniquePtr<INode> Node = DataRegistry.CreateOutputNode(GetMetasoundDataTypeName<DataType>(), MoveTemp(OutputNodeParams));
+
+			if (!Node.IsValid())
+			{
+				return false;
+			}
+
+			Transactor.AddNode(OutputNodeParams.InstanceID, MoveTemp(Node));
+			Transactor.AddOutputDataSource(OutputNodeParams.InstanceID, Name);
+			return true;
+		}
+
+		void RemoveOutput(const FVertexName& Name, const FGuid& NodeGuid)
+		{
+			Transactor.RemoveOutputDataSource(Name);
+			Transactor.RemoveNode(NodeGuid);
+		}
+
+		void Execute()
+		{
+			Generator.OnGenerateAudio(RenderBuffer.GetData(), RenderBuffer.Num());
+		}
+
+		const FOperatorSettings OperatorSettings;
+		FMetasoundDynamicGraphGenerator Generator;
+
+	private:
+		DynamicGraph::FDynamicOperatorTransactor Transactor;
+		Frontend::IDataTypeRegistry& DataRegistry = Frontend::IDataTypeRegistry::Get();
+		FAudioBuffer RenderBuffer;
+	};
+	
+	IMPLEMENT_SIMPLE_AUTOMATION_TEST(
+		FMetasoundGeneratorDynamicVertexInterfaceUpdatedTest,
+		"Audio.Metasound.Generator.Dynamic.VertexInterfaceUpdated",
+		EAutomationTestFlags::EditorContext | EAutomationTestFlags::EngineFilter)
+	bool FMetasoundGeneratorDynamicVertexInterfaceUpdatedTest::RunTest(const FString& Parameters)
+	{
+		// Make a dynamic generator
+		FDynamicGeneratorBuilder GeneratorBuilder{ 48000, 480 };
+
+		// Register for vertex interface updates
+		FVertexInterfaceData LatestInterfaceData;
+		GeneratorBuilder.Generator.OnVertexInterfaceDataUpdated.AddLambda([&LatestInterfaceData](FVertexInterfaceData VertexInterfaceData)
+		{
+			LatestInterfaceData = MoveTemp(VertexInterfaceData);
+		});
+
+		// Add an input
+		const FVertexName InputName = "SomeInput";
+		const FGuid InputGuid = FGuid::NewGuid();
+		{
+			// Add the input
+			constexpr float DefaultValue = 123.456f;
+			UTEST_TRUE("Added input", GeneratorBuilder.AddInput<float>(InputName, InputGuid, DefaultValue));
+
+			// Render to flush the transaction queue
+			GeneratorBuilder.Execute();
+
+			// Check that the input actually got added with the default
+			const FAnyDataReference* InputRef = LatestInterfaceData.GetInputs().FindDataReference(InputName);
+			UTEST_NOT_NULL("Vertex data contains input", InputRef);
+			const float* Value = InputRef->GetValue<float>();
+			UTEST_NOT_NULL("Value exists", Value)
+			UTEST_EQUAL("Value is default", *Value, DefaultValue);
+		}
+
+		// Remove the input
+		GeneratorBuilder.RemoveInput(InputName, InputGuid);
+		{
+			// Render to flush the transaction queue
+			GeneratorBuilder.Execute();
+
+			// Check that the input actually got removed
+			const FAnyDataReference* InputRef = LatestInterfaceData.GetInputs().FindDataReference(InputName);
+			UTEST_NULL("Vertex data does not contain input", InputRef);
+		}
+
+		// Add an output
+		const FVertexName OutputName = "SomeOutput";
+		const FGuid OutputGuid = FGuid::NewGuid();
+		{
+			UTEST_TRUE("Added output", GeneratorBuilder.AddOutput<int32>(OutputName, OutputGuid));
+
+			// Render to flush the transaction queue
+			GeneratorBuilder.Execute();
+			
+			// check that the output actually got added
+			const FAnyDataReference* OutputRef = LatestInterfaceData.GetOutputs().FindDataReference(OutputName);
+			UTEST_NOT_NULL("Vertex data contains output", OutputRef);
+		}
+
+		// Remove the output
+		GeneratorBuilder.RemoveOutput(OutputName, OutputGuid);
+		{
+			// Render to flush the transaction queue
+			GeneratorBuilder.Execute();
+
+			// Check that the output actually got removed
+			const FAnyDataReference* OutputRef = LatestInterfaceData.GetOutputs().FindDataReference(OutputName);
+			UTEST_NULL("Vertex data does not contain output", OutputRef);
+		}
+		
+		return true;
+	}
+}
+
+#endif
\ No newline at end of file
diff --git a/Engine/Plugins/Runtime/Metasound/Source/MetasoundGenerator/Private/MetasoundGenerator.cpp b/Engine/Plugins/Runtime/Metasound/Source/MetasoundGenerator/Private/MetasoundGenerator.cpp
--- a/Engine/Plugins/Runtime/Metasound/Source/MetasoundGenerator/Private/MetasoundGenerator.cpp
+++ b/Engine/Plugins/Runtime/Metasound/Source/MetasoundGenerator/Private/MetasoundGenerator.cpp
@@ -425,6 +425,7 @@ PRAGMA_ENABLE_DEPRECATION_WARNINGS
 
 		bIsWaitingForFirstGraph = false;
 		OnSetGraph.Broadcast();
+		OnVertexInterfaceDataUpdated.Broadcast(VertexInterfaceData);
 	}
 
 	TUniquePtr<IOperator> FMetasoundGenerator::ReleaseGraphOperator()
@@ -446,6 +447,8 @@ PRAGMA_ENABLE_DEPRECATION_WARNINGS
 		GraphAnalyzer.Reset();
 		ParameterSetters.Reset();
 		ParameterPackSetters.Reset();
+
+		OnVertexInterfaceDataUpdated.Broadcast(VertexInterfaceData);
 	}
 
 	void FMetasoundGenerator::QueueParameterPack(TSharedPtr<FMetasoundParameterPackStorage> ParameterPack)
@@ -963,6 +966,8 @@ PRAGMA_ENABLE_DEPRECATION_WARNINGS
 		// Update vertex data and set parameters. 
 		VertexInterfaceData.GetInputs() = InGraphInputs;
 		GeneratorBuilder::AddParameterSetterIfWritable(InVertexName, VertexInterfaceData.GetInputs(), ParameterSetters, ParameterPackSetters);
+
+		OnVertexInterfaceDataUpdated.Broadcast(VertexInterfaceData);
 	}
 
 	void FMetasoundDynamicGraphGenerator::OnInputRemoved(const FVertexName& InVertexName, const FInputVertexInterfaceData& InGraphInputs)
@@ -970,6 +975,7 @@ PRAGMA_ENABLE_DEPRECATION_WARNINGS
 		VertexInterfaceData.GetInputs() = InGraphInputs;
 		ParameterSetters.Remove(InVertexName);
 		ParameterPackSetters.Remove(InVertexName);
+		OnVertexInterfaceDataUpdated.Broadcast(VertexInterfaceData);
 	}
 
 	void FMetasoundDynamicGraphGenerator::OnOutputAdded(const FVertexName& InVertexName, const FOutputVertexInterfaceData& InGraphOutputs)
@@ -981,6 +987,7 @@ PRAGMA_ENABLE_DEPRECATION_WARNINGS
 		{
 			OnFinishedTriggerRef = InGraphOutputs.GetDataReadReference<FTrigger>(InVertexName);
 		}
+		OnVertexInterfaceDataUpdated.Broadcast(VertexInterfaceData);
 	}
 
 	void FMetasoundDynamicGraphGenerator::OnOutputUpdated(const FVertexName& InVertexName, const FOutputVertexInterfaceData& InGraphOutputs)
@@ -998,6 +1005,7 @@ PRAGMA_ENABLE_DEPRECATION_WARNINGS
 		{
 			OnFinishedTriggerRef = InGraphOutputs.GetDataReadReference<FTrigger>(InVertexName);
 		}
+		OnVertexInterfaceDataUpdated.Broadcast(VertexInterfaceData);
 	}
 
 	void FMetasoundDynamicGraphGenerator::OnOutputRemoved(const FVertexName& InVertexName, const FOutputVertexInterfaceData& InGraphOutputs)
@@ -1005,6 +1013,7 @@ PRAGMA_ENABLE_DEPRECATION_WARNINGS
 		MetasoundGeneratorPrivate::LogErrorIfIsAudioVertex(InVertexName, TEXT("The MetaSound Generator cannot dynamically remove audio outputs"));
 
 		VertexInterfaceData.GetOutputs() = InGraphOutputs;
+		OnVertexInterfaceDataUpdated.Broadcast(VertexInterfaceData);
 	}
 
 	TUniquePtr<IOperator> FMetasoundDynamicGraphGenerator::ReleaseGraphOperator()
diff --git a/Engine/Plugins/Runtime/Metasound/Source/MetasoundGenerator/Public/MetasoundGenerator.h b/Engine/Plugins/Runtime/Metasound/Source/MetasoundGenerator/Public/MetasoundGenerator.h
--- a/Engine/Plugins/Runtime/Metasound/Source/MetasoundGenerator/Public/MetasoundGenerator.h
+++ b/Engine/Plugins/Runtime/Metasound/Source/MetasoundGenerator/Public/MetasoundGenerator.h
@@ -184,6 +184,9 @@ namespace Metasound
 			return ReadRef;
 		}
 
+		DECLARE_TS_MULTICAST_DELEGATE_OneParam(FOnVertexInterfaceDataUpdated, FVertexInterfaceData);
+		FOnVertexInterfaceDataUpdated OnVertexInterfaceDataUpdated;
+
 		/**
 		 * Add a vertex analyzer for a named output with the given address info.
 		 *
@@ -244,8 +247,6 @@ namespace Metasound
 		void ClearGraph();
 		bool UpdateGraphIfPending();
 
-
-
 	private:
 
 		friend class FAsyncMetaSoundBuilderBase;
