- Make UMetaSoundSource & UMetaSoundPatch the direct document manipulated object for subsystem builders
- Remove specialized Audition object in favor of directly using UMetaSoundSource for Source builders

diff --git a/Engine/Plugins/Runtime/Metasound/Source/MetasoundEditor/Private/MetasoundEditorGraphSchema.cpp b/Engine/Plugins/Runtime/Metasound/Source/MetasoundEditor/Private/MetasoundEditorGraphSchema.cpp
--- a/Engine/Plugins/Runtime/Metasound/Source/MetasoundEditor/Private/MetasoundEditorGraphSchema.cpp
+++ b/Engine/Plugins/Runtime/Metasound/Source/MetasoundEditor/Private/MetasoundEditorGraphSchema.cpp
@@ -1782,8 +1782,8 @@ void UMetasoundEditorGraphSchema::DroppedAssetsOnGraph(const TArray<FAssetData>&
 			}
 
 			// This may not be necessary as dropping an asset on the graph may load it, thus triggering the registration from the MetaSoundAssetManager.
-			const FNodeClassInfo ClassInfo = DroppedMetaSoundAsset->GetAssetClassInfo();
-			const FNodeRegistryKey RegistryKey = NodeRegistryKey::CreateKey(ClassInfo);
+			const FMetasoundFrontendDocument& DroppedDoc = DroppedMetaSoundAsset->GetDocumentChecked();
+			const FNodeRegistryKey RegistryKey = NodeRegistryKey::CreateKey(DroppedDoc.RootGraph);
 			if (ensure(NodeRegistryKey::IsValid(RegistryKey)))
 			{
 				FMetaSoundAssetRegistrationOptions RegOptions;
@@ -1791,10 +1791,8 @@ void UMetasoundEditorGraphSchema::DroppedAssetsOnGraph(const TArray<FAssetData>&
 				DroppedMetaSoundAsset->RegisterGraphWithFrontend(RegOptions);
 			}
 
-			const FMetasoundFrontendClassName ClassName = DroppedMetaSoundAsset->GetAssetClassInfo().ClassName;
-
 			FMetasoundFrontendClass Class;
-			if (ensure(ISearchEngine::Get().FindClassWithHighestVersion(ClassName.ToNodeClassName(), Class)))
+			if (ensure(ISearchEngine::Get().FindClassWithHighestVersion(DroppedDoc.RootGraph.Metadata.GetClassName(), Class)))
 			{
 				Metasound::Editor::FGraphBuilder::AddExternalNode(MetaSound, Class.Metadata, GraphPosition);
 				bTransactionSucceeded = true;
diff --git a/Engine/Plugins/Runtime/Metasound/Source/MetasoundEngine/Private/Metasound.cpp b/Engine/Plugins/Runtime/Metasound/Source/MetasoundEngine/Private/Metasound.cpp
--- a/Engine/Plugins/Runtime/Metasound/Source/MetasoundEngine/Private/Metasound.cpp
+++ b/Engine/Plugins/Runtime/Metasound/Source/MetasoundEngine/Private/Metasound.cpp
@@ -148,11 +148,6 @@ void UMetaSoundPatch::PostLoad()
 	Metasound::FMetaSoundEngineAssetHelper::PostLoad(*this);
 }
 
-Metasound::Frontend::FNodeClassInfo UMetaSoundPatch::GetAssetClassInfo() const
-{
-	return { GetDocumentChecked().RootGraph, FSoftObjectPath(this) };
-}
-
 #if WITH_EDITOR
 void UMetaSoundPatch::SetReferencedAssetClasses(TSet<Metasound::Frontend::IMetaSoundAssetManager::FAssetInfo>&& InAssetClasses)
 {
diff --git a/Engine/Plugins/Runtime/Metasound/Source/MetasoundEngine/Private/MetasoundAssetSubsystem.cpp b/Engine/Plugins/Runtime/Metasound/Source/MetasoundEngine/Private/MetasoundAssetSubsystem.cpp
--- a/Engine/Plugins/Runtime/Metasound/Source/MetasoundEngine/Private/MetasoundAssetSubsystem.cpp
+++ b/Engine/Plugins/Runtime/Metasound/Source/MetasoundEngine/Private/MetasoundAssetSubsystem.cpp
@@ -142,8 +142,8 @@ void UMetaSoundAssetSubsystem::AddAssetReferences(FMetasoundAssetBase& InAssetBa
 	using namespace Metasound;
 	using namespace Metasound::Frontend;
 
-	const FNodeClassInfo AssetClassInfo = InAssetBase.GetAssetClassInfo();
-	const FNodeRegistryKey AssetClassKey = NodeRegistryKey::CreateKey(AssetClassInfo);
+	const FMetasoundFrontendDocument& Document = InAssetBase.GetDocumentChecked();
+	const FNodeRegistryKey AssetClassKey = NodeRegistryKey::CreateKey(Document.RootGraph);
 
 	if (!ContainsKey(AssetClassKey))
 	{
@@ -175,8 +175,8 @@ void UMetaSoundAssetSubsystem::AddAssetReferences(FMetasoundAssetBase& InAssetBa
 	{
 		if (Asset)
 		{
-			FNodeClassInfo ClassInfo = Asset->GetAssetClassInfo();
-			const FNodeRegistryKey ClassKey = NodeRegistryKey::CreateKey(ClassInfo);
+			const FMetasoundFrontendDocument& RefDocument = Asset->GetDocumentChecked();
+			const FNodeRegistryKey ClassKey = NodeRegistryKey::CreateKey(RefDocument.RootGraph);
 			if (!ContainsKey(ClassKey))
 			{
 				UE_LOG(LogMetaSound, Verbose,
@@ -210,8 +210,8 @@ Metasound::Frontend::FNodeRegistryKey UMetaSoundAssetSubsystem::AddOrUpdateAsset
 	const FMetasoundAssetBase* MetaSoundAsset = Metasound::IMetasoundUObjectRegistry::Get().GetObjectAsAssetBase(&InObject);
 	check(MetaSoundAsset);
 
-	FNodeClassInfo ClassInfo = MetaSoundAsset->GetAssetClassInfo();
-	const FNodeRegistryKey RegistryKey = NodeRegistryKey::CreateKey(ClassInfo);
+	const FMetasoundFrontendDocument& Document = MetaSoundAsset->GetDocumentChecked();
+	const FNodeRegistryKey RegistryKey = NodeRegistryKey::CreateKey(Document.RootGraph);
 
 	if (NodeRegistryKey::IsValid(RegistryKey))
 	{
@@ -224,13 +224,12 @@ Metasound::Frontend::FNodeRegistryKey UMetaSoundAssetSubsystem::AddOrUpdateAsset
 Metasound::Frontend::FNodeRegistryKey UMetaSoundAssetSubsystem::AddOrUpdateAsset(const FAssetData& InAssetData)
 {
 	using namespace Metasound;
-	using namespace Metasound::AssetSubsystemPrivate;
 	using namespace Metasound::Frontend;
 
 	METASOUND_TRACE_CPUPROFILER_EVENT_SCOPE(UMetaSoundAssetSubsystem::AddOrUpdateAsset);
 
 	FNodeClassInfo ClassInfo;
-	bool bClassInfoFound = GetAssetClassInfo(InAssetData, ClassInfo);
+	bool bClassInfoFound = AssetSubsystemPrivate::GetAssetClassInfo(InAssetData, ClassInfo);
 	if (!bClassInfoFound)
 	{
 		UObject* Object = nullptr;
@@ -598,12 +597,12 @@ void UMetaSoundAssetSubsystem::RemoveAsset(const UObject& InObject)
 
 	if (const FMetasoundAssetBase* MetaSoundAsset = Metasound::IMetasoundUObjectRegistry::Get().GetObjectAsAssetBase(&InObject))
 	{
-		const FNodeClassInfo ClassInfo = MetaSoundAsset->GetAssetClassInfo();
-		FNodeRegistryKey RegistryKey = FMetasoundFrontendRegistryContainer::Get()->GetRegistryKey(ClassInfo);
 		const FSoftObjectPath ObjectPath(&InObject);
+		const FMetasoundFrontendDocument& Document = MetaSoundAsset->GetDocumentChecked();
+		FNodeRegistryKey RegistryKey = NodeRegistryKey::CreateKey(Document.RootGraph);
 		if (AssetSubsystemPrivate::RemoveIfExactMatch(PathMap, RegistryKey, ObjectPath))
 		{
-			UMetaSoundBuilderSubsystem::GetChecked().DetachBuilderFromAsset(ClassInfo.ClassName);
+			UMetaSoundBuilderSubsystem::GetChecked().DetachBuilderFromAsset(Document.RootGraph.Metadata.GetClassName());
 		}
 	}
 }
diff --git a/Engine/Plugins/Runtime/Metasound/Source/MetasoundEngine/Private/MetasoundBuilderSubsystem.cpp b/Engine/Plugins/Runtime/Metasound/Source/MetasoundEngine/Private/MetasoundBuilderSubsystem.cpp
--- a/Engine/Plugins/Runtime/Metasound/Source/MetasoundEngine/Private/MetasoundBuilderSubsystem.cpp
+++ b/Engine/Plugins/Runtime/Metasound/Source/MetasoundEngine/Private/MetasoundBuilderSubsystem.cpp
@@ -52,18 +52,6 @@ namespace Metasound::Engine
 			FLiteralNodeConstructorParams Params { "Literal", FGuid::NewGuid(), MoveTemp(ValueLiteral) };
 			return IDataTypeRegistry::Get().CreateLiteralNode(DataType, MoveTemp(Params));
 		}
-
-		template <typename BuilderClass>
-		BuilderClass& CreateTransientBuilder(FName BuilderName = { })
-		{
-			const EObjectFlags NewObjectFlags = RF_Public | RF_Transient;
-			UPackage* TransientPackage = GetTransientPackage();
-			const FName ObjectName = MakeUniqueObjectName(TransientPackage, BuilderClass::StaticClass(), BuilderName);
-			TObjectPtr<BuilderClass> NewBuilder = NewObject<BuilderClass>(TransientPackage, ObjectName, NewObjectFlags);
-			check(NewBuilder);
-			NewBuilder->InitFrontendBuilder();
-			return *NewBuilder.Get();
-		}
 	} // namespace BuilderSubsystemPrivate
 } // namespace Metasound::Engine
 
@@ -899,6 +887,17 @@ void UMetaSoundBuilderBase::UpdateDependencyClassNames(const TMap<FMetasoundFron
 	Builder.UpdateDependencyClassNames(OldToNewReferencedClassNames);
 }
 
+void UMetaSoundPatchBuilder::CreateTransientBuilder()
+{
+	using namespace Metasound::Frontend;
+
+	UMetaSoundPatch* Patch = NewObject<UMetaSoundPatch>();
+	check(Patch);
+
+	Builder = FMetaSoundFrontendDocumentBuilder(Patch);
+	Builder.InitDocument();
+}
+
 TScriptInterface<IMetaSoundDocumentInterface> UMetaSoundPatchBuilder::Build(UObject* Parent, const FMetaSoundBuilderOptions& InBuilderOptions) const
 {
 	return &BuildInternal<UMetaSoundPatch>(Parent, InBuilderOptions);
@@ -925,7 +924,7 @@ UMetaSoundBuilderBase& UMetaSoundBuilderSubsystem::AttachBuilderToAssetChecked(U
 	else
 	{
 		checkf(false, TEXT("UClass '%s' is not a base MetaSound that supports attachment via the MetaSoundBuilderSubsystem"), *BaseClass->GetFullName());
-		return Metasound::Engine::BuilderSubsystemPrivate::CreateTransientBuilder<UMetaSoundPatchBuilder>();
+		return *NewObject<UMetaSoundPatchBuilder>();
 	}
 }
 
@@ -944,7 +943,6 @@ UMetaSoundSourceBuilder* UMetaSoundBuilderSubsystem::AttachSourceBuilderToAsset(
 	if (InSource)
 	{
 		UMetaSoundSourceBuilder& SourceBuilder = AttachBuilderToAssetCheckedPrivate<UMetaSoundSourceBuilder>(InSource);
-		SourceBuilder.AuditionSound = InSource;
 		return &SourceBuilder;
 	}
 
@@ -953,8 +951,9 @@ UMetaSoundSourceBuilder* UMetaSoundBuilderSubsystem::AttachSourceBuilderToAsset(
 
 void UMetaSoundSourceBuilder::Audition(UObject* Parent, UAudioComponent* AudioComponent, FOnCreateAuditionGeneratorHandleDelegate CreateGenerator, bool bLiveUpdatesEnabled)
 {
+	using namespace Metasound;
+	using namespace Metasound::DynamicGraph;
 	using namespace Metasound::Engine;
-	using namespace Metasound::Engine::BuilderSubsystemPrivate;
 	using namespace Metasound::Frontend;
 
 	METASOUND_TRACE_CPUPROFILER_EVENT_SCOPE(UMetaSoundSourceBuilder::Audition);
@@ -965,36 +964,17 @@ void UMetaSoundSourceBuilder::Audition(UObject* Parent, UAudioComponent* AudioCo
 		return;
 	}
 
-	if (!bIsAttached)
-	{
-		FMetaSoundBuilderOptions BuilderOptions;
-		BuilderOptions.bAddToRegistry = true;
+	UMetaSoundSource& MetaSoundSource = GetMetaSoundSource();
 
-		if (AuditionSound.IsValid())
-		{
-			BuilderOptions.ExistingMetaSound = AuditionSound.Get();
-			if (USoundBase* Sound = AudioComponent->GetSound())
-			{
-				if (Sound != AuditionSound)
-				{
-					UE_LOG(LogMetaSound, Warning, TEXT("MetaSoundBuilder '%s' supplied AudioComponent with unlinked sound '%s'. Stopping sound and replacing with builder's audition sound."),
-						*GetFullName(),
-						*Sound->GetFullName());
-				}
-			}
-		}
-		else
-		{
-			BuilderOptions.Name = MakeUniqueObjectName(nullptr, UMetaSoundSource::StaticClass(), FName(GetName() + TEXT("_Audition")));
-		}
-
-		AuditionSound = &BuildInternal<UMetaSoundSource>(Parent, BuilderOptions);
-	}
+	FMetasoundAssetBase* MetaSoundAsset = IMetasoundUObjectRegistry::Get().GetObjectAsAssetBase(&MetaSoundSource);
+	check(MetaSoundAsset);
+	MetaSoundAsset->RegisterGraphWithFrontend();
 
-	AudioComponent->SetSound(AuditionSound.Get());
+	// Must be called post register as register ensures cached runtime data passed to transactor is up-to-date
+	MetaSoundSource.SetDynamicGeneratorEnabled(bLiveUpdatesEnabled);
+	MetaSoundSource.ConformObjectDataToInterfaces();
 
-	const bool bEnableDynamicGenerators = bLiveUpdatesEnabled;
-	AuditionSound->SetDynamicGeneratorEnabled(bEnableDynamicGenerators);
+	AudioComponent->SetSound(&MetaSoundSource);
 
 	if (CreateGenerator.IsBound())
 	{
@@ -1013,13 +993,10 @@ bool UMetaSoundSourceBuilder::ExecuteAuditionableTransaction(FAuditionableTransa
 
 	METASOUND_TRACE_CPUPROFILER_EVENT_SCOPE(UMetaSoundSourceBuilder::ExecuteAuditionableTransaction);
 
-	if (AuditionSound.IsValid())
+	TSharedPtr<FDynamicOperatorTransactor> Transactor = GetMetaSoundSource().GetDynamicGeneratorTransactor();
+	if (Transactor.IsValid())
 	{
-		TSharedPtr<FDynamicOperatorTransactor> Transactor = AuditionSound->GetDynamicGeneratorTransactor();
-		if (Transactor.IsValid())
-		{
-			return Transaction(*Transactor);
-		}
+		return Transaction(*Transactor);
 	}
 
 	return false;
@@ -1030,15 +1007,17 @@ TScriptInterface<IMetaSoundDocumentInterface> UMetaSoundSourceBuilder::Build(UOb
 	return &BuildInternal<UMetaSoundSource>(Parent, InBuilderOptions);
 }
 
-void UMetaSoundSourceBuilder::InitFrontendBuilder()
+void UMetaSoundSourceBuilder::CreateTransientBuilder()
 {
 	using namespace Metasound::Frontend;
 
 	TSharedRef<FDocumentModifyDelegates> DocumentDelegates = MakeShared<FDocumentModifyDelegates>();
 	InitDelegates(*DocumentDelegates);
 
-	UMetaSoundBuilderDocument& DocObject = UMetaSoundBuilderDocument::Create(GetBuilderUClass());
-	Builder = FMetaSoundFrontendDocumentBuilder(&DocObject, DocumentDelegates);
+	UMetaSoundSource* Source = NewObject<UMetaSoundSource>();
+	check(Source);
+
+	Builder = FMetaSoundFrontendDocumentBuilder(Source, DocumentDelegates);
 	Builder.InitDocument();
 }
 
@@ -1064,14 +1043,17 @@ const UClass& UMetaSoundSourceBuilder::GetBuilderUClass() const
 
 bool UMetaSoundSourceBuilder::GetLiveUpdatesEnabled() const
 {
-	using namespace Metasound::Engine::BuilderSubsystemPrivate;
+	return GetMetaSoundSource().GetDynamicGeneratorTransactor().IsValid();
+}
 
-	if (!AuditionSound.IsValid())
-	{
-		return false;
-	}
-	
-	return AuditionSound->GetDynamicGeneratorTransactor().IsValid();
+const UMetaSoundSource& UMetaSoundSourceBuilder::GetMetaSoundSource() const
+{
+	return GetConstBuilder().CastDocumentObjectChecked<UMetaSoundSource>();
+}
+
+UMetaSoundSource& UMetaSoundSourceBuilder::GetMetaSoundSource()
+{
+	return Builder.CastDocumentObjectChecked<UMetaSoundSource>();
 }
 
 void UMetaSoundSourceBuilder::InitDelegates(Metasound::Frontend::FDocumentModifyDelegates& OutDocumentDelegates) const
@@ -1444,7 +1426,7 @@ void UMetaSoundSourceBuilder::SetFormat(EMetaSoundOutputAudioFormat OutputFormat
 UMetaSoundPatchBuilder* UMetaSoundBuilderSubsystem::CreatePatchBuilder(FName BuilderName, EMetaSoundBuilderResult& OutResult)
 {
 	OutResult = EMetaSoundBuilderResult::Succeeded;
-	return &Metasound::Engine::BuilderSubsystemPrivate::CreateTransientBuilder<UMetaSoundPatchBuilder>(BuilderName);
+	return &CreateTransientBuilder<UMetaSoundPatchBuilder>(BuilderName);
 }
 
 UMetaSoundSourceBuilder* UMetaSoundBuilderSubsystem::CreateSourceBuilder(
@@ -1559,7 +1541,7 @@ UMetaSoundPatchBuilder* UMetaSoundBuilderSubsystem::CreatePatchPresetBuilder(FNa
 {
 	if (ReferencedNodeClass)
 	{
-		UMetaSoundPatchBuilder& Builder = Metasound::Engine::BuilderSubsystemPrivate::CreateTransientBuilder<UMetaSoundPatchBuilder>(BuilderName);
+		UMetaSoundPatchBuilder& Builder = CreateTransientBuilder<UMetaSoundPatchBuilder>(BuilderName);
 		Builder.ConvertToPreset(ReferencedNodeClass, OutResult);
 		return &Builder;
 	}
@@ -1582,7 +1564,7 @@ UMetaSoundBuilderBase& UMetaSoundBuilderSubsystem::CreatePresetBuilder(FName Bui
 	else
 	{
 		checkf(false, TEXT("UClass '%s' cannot be built to a MetaSound preset"), *Class.GetFullName());
-		return Metasound::Engine::BuilderSubsystemPrivate::CreateTransientBuilder<UMetaSoundPatchBuilder>(BuilderName);
+		return *NewObject<UMetaSoundPatchBuilder>();
 	}
 }
 
@@ -1590,7 +1572,7 @@ UMetaSoundSourceBuilder* UMetaSoundBuilderSubsystem::CreateSourcePresetBuilder(F
 {
 	if (ReferencedNodeClass)
 	{
-		UMetaSoundSourceBuilder& Builder = Metasound::Engine::BuilderSubsystemPrivate::CreateTransientBuilder<UMetaSoundSourceBuilder>();
+		UMetaSoundSourceBuilder& Builder = CreateTransientBuilder<UMetaSoundSourceBuilder>();
 		Builder.ConvertToPreset(ReferencedNodeClass, OutResult);
 		return &Builder;
 	}
diff --git a/Engine/Plugins/Runtime/Metasound/Source/MetasoundEngine/Private/MetasoundSource.cpp b/Engine/Plugins/Runtime/Metasound/Source/MetasoundEngine/Private/MetasoundSource.cpp
--- a/Engine/Plugins/Runtime/Metasound/Source/MetasoundEngine/Private/MetasoundSource.cpp
+++ b/Engine/Plugins/Runtime/Metasound/Source/MetasoundEngine/Private/MetasoundSource.cpp
@@ -552,11 +552,6 @@ void UMetaSoundSource::InitResources()
 	RegisterGraphWithFrontend(GetInitRegistrationOptions());
 }
 
-Metasound::Frontend::FNodeClassInfo UMetaSoundSource::GetAssetClassInfo() const
-{
-	return { GetDocumentChecked().RootGraph, FSoftObjectPath(this) };
-}
-
 bool UMetaSoundSource::IsPlayable() const
 {
 	// todo: cache off whether this metasound is buildable to an operator.
diff --git a/Engine/Plugins/Runtime/Metasound/Source/MetasoundEngine/Public/Metasound.h b/Engine/Plugins/Runtime/Metasound/Source/MetasoundEngine/Public/Metasound.h
--- a/Engine/Plugins/Runtime/Metasound/Source/MetasoundEngine/Public/Metasound.h
+++ b/Engine/Plugins/Runtime/Metasound/Source/MetasoundEngine/Public/Metasound.h
@@ -143,9 +143,6 @@ public:
 	virtual void Serialize(FArchive& InArchive) override;
 	virtual void PostLoad() override;
 
-	// Returns Asset Metadata associated with this MetaSound
-	virtual Metasound::Frontend::FNodeClassInfo GetAssetClassInfo() const override;
-
 	virtual bool ConformObjectDataToInterfaces() override { return false; }
 
 	virtual const TSet<FString>& GetReferencedAssetClassKeys() const override
diff --git a/Engine/Plugins/Runtime/Metasound/Source/MetasoundEngine/Public/MetasoundBuilderSubsystem.h b/Engine/Plugins/Runtime/Metasound/Source/MetasoundEngine/Public/MetasoundBuilderSubsystem.h
--- a/Engine/Plugins/Runtime/Metasound/Source/MetasoundEngine/Public/MetasoundBuilderSubsystem.h
+++ b/Engine/Plugins/Runtime/Metasound/Source/MetasoundEngine/Public/MetasoundBuilderSubsystem.h
@@ -101,9 +101,9 @@ struct METASOUNDENGINE_API FMetaSoundBuilderOptions
 	UPROPERTY(EditAnywhere, BlueprintReadWrite, Category = "MetaSound|Builder")
 	FName Name;
 
-	// If true, this will force regeneration of the class identifier. If the resulting MetaSound is building over
-	// an existing document, this will effectively invalidate any referencing MetaSounds and register the MetaSound
-	// as a new entry in the Frontend.
+	// If the resulting MetaSound is building over an existing document, a unique class name will be generated,
+	// invalidating any referencing MetaSounds and registering the MetaSound as a new entry in the Frontend. If
+	// building a new document, option is ignored.
 	UPROPERTY(EditAnywhere, BlueprintReadWrite, Category = "MetaSound|Builder", meta = (AdvancedDisplay))
 	bool bForceUniqueClassName = false;
 
@@ -372,19 +372,24 @@ public:
 
 	virtual TScriptInterface<IMetaSoundDocumentInterface> Build(UObject* Parent, const FMetaSoundBuilderOptions& Options) const PURE_VIRTUAL(UMetaSoundBuilderBase::Build, return { }; );
 
-	// Returns the base MetaSound UClass the builder is operating on (ex. MetaSoundSource, MetaSoundPatch, etc. not to be confused with the underlying document builder's MetaSound class type)
+	// Returns the base MetaSound UClass the builder is operating on
 	virtual const UClass& GetBuilderUClass() const PURE_VIRTUAL(UMetaSoundBuilderBase::Build, return *UClass::StaticClass(); );
 
+	UE_DEPRECATED(5.4, "Moved to protected pure virtual 'CreateTransientBuilder'")
+	virtual void InitFrontendBuilder();
+
 	// Initializes and ensures all nodes have a position (required prior to exporting to an asset if expected to be viewed in the editor).
 	void InitNodeLocations();
 
-	virtual void InitFrontendBuilder();
-
 #if WITH_EDITOR
 	void SetNodeLocation(const FMetaSoundNodeHandle& InNodeHandle, const FVector2D& InLocation, EMetaSoundBuilderResult& OutResult);
 #endif // WITH_EDITOR
 
 protected:
+	// Creates a FrontendBuilder wrapping the transient MetaSound object of the class supported by the given subsystem builder (See GetBuilderUClass).
+	// Should only be used when builders is not being attached to an existing, serialized MetaSound asset.
+	virtual void CreateTransientBuilder() PURE_VIRTUAL(UMetaSoundBuilderBase::CreateTransientBuilder, );
+
 	const FMetaSoundFrontendDocumentBuilder& GetConstBuilder() const
 	{
 		return Builder;
@@ -469,11 +474,10 @@ protected:
 	UPROPERTY()
 	FMetaSoundFrontendDocumentBuilder Builder;
 
-	// If true, builder is attached to an existing asset (directly making changes to its document). If false, builder is
-	// operating on a transient document which must be built to an asset prior to use by the MetaSound Frontend where it
-	// is then available for referencing, playback, auditioning, etc. by core execution.
-	UPROPERTY()
+#if WITH_EDITORONLY_DATA
+	UPROPERTY(meta = (DeprecatedProperty, DeprecationMessage = "5.4 - All source builders now operate on an underlying document source document that is also used to audition."))
 	bool bIsAttached = false;
+#endif // WITH_EDITORONLY_DATA
 
 	// Friending allows for swapping the builder in certain circumstances where desired (eg. attaching a builder to an existing asset)
 	friend class UMetaSoundBuilderSubsystem;
@@ -490,6 +494,11 @@ public:
 	virtual UPARAM(DisplayName = "MetaSound") TScriptInterface<IMetaSoundDocumentInterface> Build(UObject* Parent, const FMetaSoundBuilderOptions& Options) const override;
 
 	virtual const UClass& GetBuilderUClass() const override;
+
+protected:
+	virtual void CreateTransientBuilder() override;
+
+	friend class UMetaSoundBuilderSubsystem;
 };
 
 /** Builder in charge of building a MetaSound Source */
@@ -516,12 +525,19 @@ public:
 	const Metasound::Engine::FOutputAudioFormatInfoPair* FindOutputAudioFormatInfo() const;
 
 	virtual const UClass& GetBuilderUClass() const override;
-	virtual void InitFrontendBuilder() override;
+
+protected:
+	virtual void CreateTransientBuilder() override;
+
 
 private:
 	static TOptional<Metasound::FAnyDataReference> CreateDataReference(const Metasound::FOperatorSettings& InOperatorSettings, FName DataType, const Metasound::FLiteral& InLiteral, Metasound::EDataReferenceAccessType AccessType);
 
+	const UMetaSoundSource& GetMetaSoundSource() const;
+	UMetaSoundSource& GetMetaSoundSource();
+
 	void InitDelegates(Metasound::Frontend::FDocumentModifyDelegates& OutDocumentDelegates) const;
+
 	void OnEdgeAdded(int32 EdgeIndex) const;
 	void OnInputAdded(int32 InputIndex) const;
 	void OnNodeAdded(int32 NodeIndex) const;
@@ -533,8 +549,6 @@ private:
 	void OnRemovingNodeInputLiteral(int32 NodeIndex, int32 VertexIndex, int32 LiteralIndex) const;
 	void OnRemovingOutput(int32 OutputIndex) const;
 
-	TWeakObjectPtr<UMetaSoundSource> AuditionSound;
-
 	using FAuditionableTransaction = TFunctionRef<bool(Metasound::DynamicGraph::FDynamicOperatorTransactor&)>;
 	bool ExecuteAuditionableTransaction(FAuditionableTransaction Transaction) const;
 
@@ -665,6 +679,18 @@ public:
 	UPARAM(DisplayName = "Unregistered") bool UnregisterSourceBuilder(FName BuilderName);
 
 private:
+	template <typename BuilderClass>
+	BuilderClass& CreateTransientBuilder(FName BuilderName = FName())
+	{
+		const EObjectFlags NewObjectFlags = RF_Public | RF_Transient;
+		UPackage* TransientPackage = GetTransientPackage();
+		const FName ObjectName = MakeUniqueObjectName(TransientPackage, BuilderClass::StaticClass(), BuilderName);
+		TObjectPtr<BuilderClass> NewBuilder = NewObject<BuilderClass>(TransientPackage, ObjectName, NewObjectFlags);
+		check(NewBuilder);
+		NewBuilder->CreateTransientBuilder();
+		return *NewBuilder.Get();
+	}
+
 	template <typename BuilderClass>
 	BuilderClass& AttachBuilderToAssetCheckedPrivate(UObject* InMetaSoundObject) const
 	{
@@ -683,7 +709,6 @@ private:
 		TObjectPtr<BuilderClass> NewBuilder = NewObject<BuilderClass>(InMetaSoundObject);
 		check(NewBuilder);
 		NewBuilder->Builder = FMetaSoundFrontendDocumentBuilder(DocInterface);
-		NewBuilder->bIsAttached = true;
 		TObjectPtr<UMetaSoundBuilderBase> NewBuilderBase = CastChecked<UMetaSoundBuilderBase>(NewBuilder);
 		AssetBuilders.Add(FullClassName, NewBuilderBase);
 		return *NewBuilder;
diff --git a/Engine/Plugins/Runtime/Metasound/Source/MetasoundEngine/Public/MetasoundSource.h b/Engine/Plugins/Runtime/Metasound/Source/MetasoundEngine/Public/MetasoundSource.h
--- a/Engine/Plugins/Runtime/Metasound/Source/MetasoundEngine/Public/MetasoundSource.h
+++ b/Engine/Plugins/Runtime/Metasound/Source/MetasoundEngine/Public/MetasoundSource.h
@@ -169,10 +169,6 @@ public:
 	virtual void Serialize(FArchive& Ar) override;
 	virtual void PostLoad() override;
 
-	// Returns Asset Metadata associated with this MetaSoundSource
-	virtual Metasound::Frontend::FNodeClassInfo GetAssetClassInfo() const override;
-
-
 	virtual bool ConformObjectDataToInterfaces() override;
 
 	UObject* GetOwningAsset() override
diff --git a/Engine/Plugins/Runtime/Metasound/Source/MetasoundEngineTest/Private/EngineTestMetasoundSource.cpp b/Engine/Plugins/Runtime/Metasound/Source/MetasoundEngineTest/Private/EngineTestMetasoundSource.cpp
--- a/Engine/Plugins/Runtime/Metasound/Source/MetasoundEngineTest/Private/EngineTestMetasoundSource.cpp
+++ b/Engine/Plugins/Runtime/Metasound/Source/MetasoundEngineTest/Private/EngineTestMetasoundSource.cpp
@@ -116,7 +116,7 @@ namespace EngineTestMetaSoundSourcePrivate
 		constexpr EMetaSoundOutputAudioFormat OutputFormat = EMetaSoundOutputAudioFormat::Mono;
 		constexpr bool bIsOneShot = false;
 		FInitTestBuilderSourceOutput Output;
-		UMetaSoundSourceBuilder& Builder = CreateSourceBuilder(Test, EMetaSoundOutputAudioFormat::Mono, bIsOneShot, Output);
+		UMetaSoundSourceBuilder& Builder = CreateSourceBuilder(Test, OutputFormat, bIsOneShot, Output);
 
 		EMetaSoundBuilderResult Result = EMetaSoundBuilderResult::Failed;
 		if (MonoOutNodeInput)
diff --git a/Engine/Plugins/Runtime/Metasound/Source/MetasoundFrontend/Private/MetasoundAssetBase.cpp b/Engine/Plugins/Runtime/Metasound/Source/MetasoundFrontend/Private/MetasoundAssetBase.cpp
--- a/Engine/Plugins/Runtime/Metasound/Source/MetasoundFrontend/Private/MetasoundAssetBase.cpp
+++ b/Engine/Plugins/Runtime/Metasound/Source/MetasoundFrontend/Private/MetasoundAssetBase.cpp
@@ -34,7 +34,6 @@
 #include "MetasoundParameterTransmitter.h"
 #include "MetasoundTrace.h"
 #include "MetasoundVertex.h"
-#include "NodeTemplates/MetasoundFrontendDocumentTemplatePreprocessor.h"
 #include "StructSerializer.h"
 #include "Templates/SharedPointer.h"
 #include "UObject/MetaData.h"
@@ -49,6 +48,60 @@ namespace Metasound
 		{
 			static float BlockRate = 100.f;
 
+			TScriptInterface<IMetaSoundDocumentInterface> BuildRegistryDocument(TScriptInterface<IMetaSoundDocumentInterface> DocumentInterface)
+			{
+				using namespace Metasound::Frontend;
+
+				METASOUND_TRACE_CPUPROFILER_EVENT_SCOPE(BuildRegistryDocument);
+
+				const IMetaSoundDocumentInterface* ConstInterface = DocumentInterface.GetInterface();
+				check(ConstInterface != nullptr);
+
+#if WITH_EDITOR
+				// Node template transform is performed on copy of local document to avoid overwriting editable data
+				constexpr bool bTransformDocumentBeforeRegistering = true;
+#else // !WITH_EDITOR
+				// Node template transform is performed on local document only if cook determinism ID generation
+				// is enabled to avoid transforms potentially creating new edges with non-deterministic IDs.
+				const bool bTransformDocumentBeforeRegistering = MetaSoundEnableCookDeterministicIDGeneration == 0;
+#endif // WITH_EDITOR
+
+				const FMetasoundFrontendDocument& Document = ConstInterface->GetDocument();
+				const TArray<FMetasoundFrontendClass>& Dependencies = Document.Dependencies;
+
+				if (bTransformDocumentBeforeRegistering)
+				{
+					// 1. Find template dependencies to build prior to making new document/builder as an optimization
+					// (no sense in creating new document/builder if no templates need processing)
+					FMetaSoundFrontendDocumentBuilder OriginalDocBuilder(DocumentInterface);
+					const bool bContainsTemplateDependency = OriginalDocBuilder.ContainsDependencyOfType(EMetasoundFrontendClassType::Template);
+					if (bContainsTemplateDependency)
+					{
+						UMetaSoundBuilderDocument& RegistryDocObject = UMetaSoundBuilderDocument::Create(*DocumentInterface.GetInterface());
+						FMetaSoundFrontendDocumentBuilder RegistryDocBuilder(&RegistryDocObject);
+
+						RegistryDocBuilder.TransformTemplateNodes();
+
+						return &RegistryDocObject;
+					}
+				}
+#if !NO_LOGGING
+				else
+				{
+					FMetaSoundFrontendDocumentBuilder OriginalDocBuilder(DocumentInterface);
+					const bool bContainsTemplateDependency = OriginalDocBuilder.ContainsDependencyOfType(EMetasoundFrontendClassType::Template);
+					if (bContainsTemplateDependency)
+					{
+						UE_LOG(LogMetaSound, Error,
+							TEXT("Template node processing disabled but provided asset class at '%s' to register contains template nodes. Runtime graph will fail to build."),
+							*OriginalDocBuilder.GetDebugName());
+					}
+				}
+#endif // !NO_LOGGING
+
+				return DocumentInterface;
+			}
+
 			void DepthFirstTraversal(const FMetasoundAssetBase& InInitAsset, TFunctionRef<TSet<const FMetasoundAssetBase*>(const FMetasoundAssetBase&)> InVisitFunction)
 			{
 				// Non recursive depth first traversal.
@@ -101,21 +154,29 @@ namespace Metasound
 
 				return TransmittableInputs;
 			}
-			
 
 			// Registers node by copying document. Updates to document require re-registration.
 			class FNodeRegistryEntry : public INodeRegistryEntry
 			{
 			public:
-				FNodeRegistryEntry(const FString& InName, TSharedPtr<FMetasoundFrontendDocument> InPreprocessedDoc, const FSoftObjectPath& InAssetPath)
-				: Name(InName)
-				, PreprocessedDoc(InPreprocessedDoc)
+				FNodeRegistryEntry(const FString& InName, const FMetasoundFrontendDocument& InDocument, const FSoftObjectPath& InAssetPath)
+					: Name(InName)
+					, Document(InDocument)
 				{
 					// Copy FrontendClass to preserve original document.
-					FrontendClass = PreprocessedDoc->RootGraph;
+					FrontendClass = Document.RootGraph;
 					FrontendClass.Metadata.SetType(EMetasoundFrontendClassType::External);
+					ClassInfo = FNodeClassInfo(Document.RootGraph, InAssetPath);
+				}
 
-					ClassInfo = FNodeClassInfo(PreprocessedDoc->RootGraph, InAssetPath);
+				FNodeRegistryEntry(const FString& InName, FMetasoundFrontendDocument&& InDocument, const FSoftObjectPath& InAssetPath)
+					: Name(InName)
+					, Document(MoveTemp(InDocument))
+				{
+					// Copy FrontendClass to preserve original document.
+					FrontendClass = Document.RootGraph;
+					FrontendClass.Metadata.SetType(EMetasoundFrontendClassType::External);
+					ClassInfo = FNodeClassInfo(Document.RootGraph, InAssetPath);
 				}
 
 				virtual ~FNodeRegistryEntry() = default;
@@ -128,9 +189,9 @@ namespace Metasound
 				virtual TUniquePtr<INode> CreateNode(const FNodeInitData&) const override
 				{
 					FProxyDataCache ProxyDataCache;
-					ProxyDataCache.CreateAndCacheProxies(*PreprocessedDoc);
+					ProxyDataCache.CreateAndCacheProxies(Document);
 
-					return FFrontendGraphBuilder().CreateGraph(*PreprocessedDoc, ProxyDataCache, Name);
+					return FFrontendGraphBuilder().CreateGraph(Document, ProxyDataCache, Name);
 				}
 
 				virtual TUniquePtr<INode> CreateNode(FDefaultLiteralNodeConstructorParams&&) const override { return nullptr; }
@@ -144,12 +205,12 @@ namespace Metasound
 
 				virtual TUniquePtr<INodeRegistryEntry> Clone() const override
 				{
-					return MakeUnique<FNodeRegistryEntry>(Name, PreprocessedDoc, ClassInfo.AssetPath);
+					return MakeUnique<FNodeRegistryEntry>(Name, Document, ClassInfo.AssetPath);
 				}
 
-				virtual TSet<FMetasoundFrontendVersion>* GetImplementedInterfaces() const override
+				virtual const TSet<FMetasoundFrontendVersion>* GetImplementedInterfaces() const override
 				{
-					return &PreprocessedDoc->Interfaces;
+					return &Document.Interfaces;
 				}
 
 				virtual bool IsNative() const override
@@ -159,7 +220,7 @@ namespace Metasound
 
 			private:
 				FString Name;
-				TSharedPtr<FMetasoundFrontendDocument> PreprocessedDoc;
+				FMetasoundFrontendDocument Document;
 				FMetasoundFrontendClass FrontendClass;
 				FNodeClassInfo ClassInfo;
 			};
@@ -236,40 +297,18 @@ void FMetasoundAssetBase::RegisterGraphWithFrontend(Metasound::Frontend::FMetaSo
 	CacheRegistryMetadata();
 #endif // WITH_EDITOR
 
-	TScriptInterface<IMetaSoundDocumentInterface> DocInterface = GetOwningAsset();
-	check(DocInterface.GetObject());
-	TSharedPtr<FMetasoundFrontendDocument> DocumentToRegister;
-	const IMetaSoundDocumentInterface& Interface = *DocInterface.GetInterface();
-
-	auto CreateTransformedDocumentToRegister = [&Interface]()
-	{
-		// Node template transform is performed on copy of local document to avoid overwriting editable data
-		UMetaSoundBuilderDocument& DocObject = UMetaSoundBuilderDocument::Create(Interface);
-		FMetaSoundFrontendDocumentBuilder DocBuilder(&DocObject);
-		DocBuilder.TransformTemplateNodes();
-		return MakeShared<FMetasoundFrontendDocument>(MoveTemp(DocObject.GetDocument()));
-	};
+	UObject* Owner = GetOwningAsset();
+	check(Owner);
+	const FString AssetName = Owner->GetName();
 
-#if WITH_EDITOR
-	DocumentToRegister = CreateTransformedDocumentToRegister();
-#else // !WITH_EDITOR
-	if (Frontend::MetaSoundEnableCookDeterministicIDGeneration != 0)
-	{
-		DocumentToRegister = MakeShared<FMetasoundFrontendDocument>(Interface.GetDocument());
-	}
-	else
 	{
-		DocumentToRegister = CreateTransformedDocumentToRegister();
+		TScriptInterface<IMetaSoundDocumentInterface> RegistryDocInterface = AssetBasePrivate::BuildRegistryDocument(Owner);
+		const FMetasoundFrontendDocument& RegistryDocument = static_cast<const IMetaSoundDocumentInterface*>(RegistryDocInterface.GetInterface())->GetDocument();
+		CacheRuntimeData(RegistryDocument);
+		TUniquePtr<INodeRegistryEntry> RegistryEntry = MakeUnique<AssetBasePrivate::FNodeRegistryEntry>(AssetName, RegistryDocument, FSoftObjectPath(Owner));
+		UnregisterGraphWithFrontend();
+		RegistryKey = FMetasoundFrontendRegistryContainer::Get()->RegisterNode(MoveTemp(RegistryEntry));
 	}
-#endif // !WITH_EDITOR
-
-	CacheRuntimeData(*DocumentToRegister);
-	FNodeClassInfo AssetClassInfo = GetAssetClassInfo();
-	const FString AssetName = DocInterface.GetObject()->GetName();
-	TUniquePtr<INodeRegistryEntry> RegistryEntry = MakeUnique<AssetBasePrivate::FNodeRegistryEntry>(AssetName, DocumentToRegister, AssetClassInfo.AssetPath);
-
-	UnregisterGraphWithFrontend();
-	RegistryKey = FMetasoundFrontendRegistryContainer::Get()->RegisterNode(MoveTemp(RegistryEntry));
 
 	if (NodeRegistryKey::IsValid(RegistryKey) && FMetasoundFrontendRegistryContainer::Get()->IsNodeRegistered(RegistryKey))
 	{
@@ -279,7 +318,7 @@ void FMetasoundAssetBase::RegisterGraphWithFrontend(Metasound::Frontend::FMetaSo
 	}
 	else
 	{
-		UClass* Class = DocInterface.GetObject()->GetClass();
+		UClass* Class = Owner->GetClass();
 		check(Class);
 		const FString ClassName = Class->GetName();
 		UE_LOG(LogMetaSound, Error, TEXT("Registration failed for MetaSound node class '%s' of UObject class '%s'"), *AssetName, *ClassName);
@@ -315,23 +354,25 @@ void FMetasoundAssetBase::CookMetaSound()
 	CacheRegistryMetadata();
 #endif // WITH_EDITOR
 
-	TScriptInterface<IMetaSoundDocumentInterface> DocInterface = GetOwningAsset();
-	check(DocInterface.GetObject());
-	TSharedPtr<FMetasoundFrontendDocument> DocumentToRegister;
-	const IMetaSoundDocumentInterface& Interface = *DocInterface.GetInterface();
+	UObject* Owner = GetOwningAsset();
+	check(Owner);
+	const FString AssetName = Owner->GetName();
 
-	// Performs document transforms on local copy, which reduces document footprint & renders resolve unnecessary at runtime upon registration
-	FMetaSoundFrontendDocumentBuilder DocBuilder(DocInterface);
-	DocBuilder.TransformTemplateNodes();
-	const FMetaSoundFrontendDocumentBuilder& ConstBuilder = DocBuilder;
-	TSharedRef<FMetasoundFrontendDocument> ResolvedDocument = MakeShared<FMetasoundFrontendDocument>(ConstBuilder.GetDocument());
-
-	FNodeClassInfo AssetClassInfo = GetAssetClassInfo();
-	const FString AssetName = DocInterface.GetObject()->GetName();
-	TUniquePtr<INodeRegistryEntry> RegistryEntry = MakeUnique<AssetBasePrivate::FNodeRegistryEntry>(AssetName, ResolvedDocument, AssetClassInfo.AssetPath);
+	{
+	// Performs document transforms on local copy, which reduces document footprint & renders transforming unnecessary unless altered at runtime when registering
+		FMetaSoundFrontendDocumentBuilder DocBuilder(Owner);
+		const bool bContainsTemplateDependency = DocBuilder.ContainsDependencyOfType(EMetasoundFrontendClassType::Template);
+		if (bContainsTemplateDependency)
+		{
+			DocBuilder.TransformTemplateNodes();
+		}
 
-	UnregisterGraphWithFrontend();
-	RegistryKey = FMetasoundFrontendRegistryContainer::Get()->RegisterNode(MoveTemp(RegistryEntry));
+		const IMetaSoundDocumentInterface& Interface = *TScriptInterface<IMetaSoundDocumentInterface>(Owner).GetInterface();
+		CacheRuntimeData(Interface.GetDocument());
+		TUniquePtr<INodeRegistryEntry> RegistryEntry = MakeUnique<AssetBasePrivate::FNodeRegistryEntry>(AssetName, Interface.GetDocument(), FSoftObjectPath(Owner));
+		UnregisterGraphWithFrontend();
+		RegistryKey = FMetasoundFrontendRegistryContainer::Get()->RegisterNode(MoveTemp(RegistryEntry));
+	}
 
 	if (NodeRegistryKey::IsValid(RegistryKey) && FMetasoundFrontendRegistryContainer::Get()->IsNodeRegistered(RegistryKey))
 	{
@@ -341,7 +382,7 @@ void FMetasoundAssetBase::CookMetaSound()
 	}
 	else
 	{
-		const UClass* Class = DocInterface.GetObject()->GetClass();
+		const UClass* Class = Owner->GetClass();
 		check(Class);
 		const FString ClassName = Class->GetName();
 		UE_LOG(LogMetaSound, Error, TEXT("Registration failed during cook for MetaSound node class '%s' of UObject class '%s'"), *AssetName, *ClassName);
@@ -390,6 +431,16 @@ bool FMetasoundAssetBase::IsInterfaceDeclared(const FMetasoundFrontendVersion& I
 	return GetDocumentChecked().Interfaces.Contains(InVersion);
 }
 
+Metasound::Frontend::FNodeClassInfo FMetasoundAssetBase::GetAssetClassInfo() const
+{
+	using namespace Metasound::Frontend;
+
+	const UObject* Owner = GetOwningAsset();
+	check(Owner);
+
+	return FNodeClassInfo { GetDocumentChecked().RootGraph, FSoftObjectPath(Owner) };
+}
+
 void FMetasoundAssetBase::SetDocument(const FMetasoundFrontendDocument& InDocument, bool bMarkDirty)
 {
 	FMetasoundFrontendDocument& Document = GetDocumentChecked();
@@ -608,7 +659,7 @@ const FMetasoundFrontendDocumentModifyContext& FMetasoundAssetBase::GetModifyCon
 }
 #endif // WITH_EDITOR
 
-TSharedPtr<Metasound::FGraph, ESPMode::ThreadSafe> FMetasoundAssetBase::BuildMetasoundDocument(const FMetasoundFrontendDocument& InPreprocessedDoc, const Metasound::Frontend::FProxyDataCache& InProxies) const
+TSharedPtr<Metasound::FGraph, ESPMode::ThreadSafe> FMetasoundAssetBase::BuildMetasoundDocument(const FMetasoundFrontendDocument& InDocument, const Metasound::Frontend::FProxyDataCache& InProxies) const
 {
 	using namespace Metasound;
 	using namespace Metasound::Frontend;
@@ -616,7 +667,7 @@ TSharedPtr<Metasound::FGraph, ESPMode::ThreadSafe> FMetasoundAssetBase::BuildMet
 	METASOUND_TRACE_CPUPROFILER_EVENT_SCOPE(MetaSoundAssetBase::BuildMetasoundDocument);
 
 	// Create graph which can spawn instances. 
-	TUniquePtr<FFrontendGraph> FrontendGraph = FFrontendGraphBuilder::CreateGraph(InPreprocessedDoc, InProxies, GetOwningAssetName());
+	TUniquePtr<FFrontendGraph> FrontendGraph = FFrontendGraphBuilder::CreateGraph(InDocument, InProxies, GetOwningAssetName());
 	if (!FrontendGraph.IsValid())
 	{
 		UE_LOG(LogMetaSound, Error, TEXT("Failed to build MetaSound graph in asset '%s'"), *GetOwningAssetName());
@@ -957,7 +1008,10 @@ void FMetasoundAssetBase::UpdateAssetRegistry()
 	using namespace Metasound;
 	using namespace Metasound::Frontend;
 
-	FNodeClassInfo AssetClassInfo = GetAssetClassInfo();
+	const UObject* Owner = GetOwningAsset();
+	check(Owner);
+	FNodeClassInfo AssetClassInfo { GetDocumentChecked().RootGraph, FSoftObjectPath(Owner) };
+
 	// Refresh Asset Registry Info if successfully registered with Frontend
 	const FMetasoundFrontendGraphClass& DocumentClassGraph = GetDocumentHandle()->GetRootGraphClass();
 	const FMetasoundFrontendClassMetadata& DocumentClassMetadata = DocumentClassGraph.Metadata;
@@ -982,7 +1036,7 @@ TSharedPtr<FMetasoundFrontendDocument> FMetasoundAssetBase::PreprocessDocument()
 	return nullptr;
 }
 
-const FMetasoundAssetBase::FRuntimeData& FMetasoundAssetBase::CacheRuntimeData(const FMetasoundFrontendDocument& InPreprocessedDoc)
+const FMetasoundAssetBase::FRuntimeData& FMetasoundAssetBase::CacheRuntimeData(const FMetasoundFrontendDocument& InDocument)
 {
 	using namespace Metasound::Frontend;
 
@@ -995,8 +1049,8 @@ const FMetasoundAssetBase::FRuntimeData& FMetasoundAssetBase::CacheRuntimeData(c
 	TArray<FMetasoundFrontendClassInput> TransmittableInputs = AssetBasePrivate::GetTransmittableInputsFromPublicInputs(PublicInputs);
 
 	Metasound::Frontend::FProxyDataCache ProxyDataCache;
-	ProxyDataCache.CreateAndCacheProxies(InPreprocessedDoc);
-	TSharedPtr<Metasound::FGraph, ESPMode::ThreadSafe> Graph = BuildMetasoundDocument(InPreprocessedDoc, ProxyDataCache);
+	ProxyDataCache.CreateAndCacheProxies(InDocument);
+	TSharedPtr<Metasound::FGraph, ESPMode::ThreadSafe> Graph = BuildMetasoundDocument(InDocument, ProxyDataCache);
 
 	CachedRuntimeData =
 	{
diff --git a/Engine/Plugins/Runtime/Metasound/Source/MetasoundFrontend/Private/MetasoundFrontendDataTypeRegistry.cpp b/Engine/Plugins/Runtime/Metasound/Source/MetasoundFrontend/Private/MetasoundFrontendDataTypeRegistry.cpp
--- a/Engine/Plugins/Runtime/Metasound/Source/MetasoundFrontend/Private/MetasoundFrontendDataTypeRegistry.cpp
+++ b/Engine/Plugins/Runtime/Metasound/Source/MetasoundFrontend/Private/MetasoundFrontendDataTypeRegistry.cpp
@@ -142,7 +142,7 @@ namespace Metasound
 					return nullptr;
 				}
 
-				virtual TSet<FMetasoundFrontendVersion>* GetImplementedInterfaces() const override
+				virtual const TSet<FMetasoundFrontendVersion>* GetImplementedInterfaces() const override
 				{
 					return nullptr;
 				}
diff --git a/Engine/Plugins/Runtime/Metasound/Source/MetasoundFrontend/Private/MetasoundFrontendDocumentBuilder.cpp b/Engine/Plugins/Runtime/Metasound/Source/MetasoundFrontend/Private/MetasoundFrontendDocumentBuilder.cpp
--- a/Engine/Plugins/Runtime/Metasound/Source/MetasoundFrontend/Private/MetasoundFrontendDocumentBuilder.cpp
+++ b/Engine/Plugins/Runtime/Metasound/Source/MetasoundFrontend/Private/MetasoundFrontendDocumentBuilder.cpp
@@ -98,6 +98,19 @@ namespace Metasound::Frontend
 			}
 		}
 
+		const FString GetDebugName(const IMetaSoundDocumentInterface& DocumentInterface)
+		{
+			const FMetasoundFrontendClassMetadata& Metadata = DocumentInterface.GetDocument().RootGraph.Metadata;
+			const FMetasoundFrontendVersionNumber& Version = Metadata.GetVersion();
+			const FNodeRegistryKey RegKey = NodeRegistryKey::CreateKey(EMetasoundFrontendClassType::External, Metadata.GetClassName().ToString(), Version.Major, Version.Minor);
+			if (const FMetasoundAssetBase* Asset = IMetaSoundAssetManager::GetChecked().TryLoadAssetFromKey(RegKey))
+			{
+				return Asset->GetOwningAssetName();
+			}
+
+			return Metadata.GetClassName().ToString();
+		}
+
 		class FModifyInterfacesImpl
 		{
 		public:
@@ -1197,6 +1210,11 @@ void FMetaSoundFrontendDocumentBuilder::ClearGraph()
 	ReloadCacheInternal();
 }
 
+bool FMetaSoundFrontendDocumentBuilder::ContainsDependencyOfType(EMetasoundFrontendClassType ClassType) const
+{
+	return DocumentCache->ContainsDependencyOfType(ClassType);
+}
+
 bool FMetaSoundFrontendDocumentBuilder::ContainsEdge(const FMetasoundFrontendEdge& InEdge) const
 {
 	using namespace Metasound::Frontend;
@@ -1559,15 +1577,7 @@ const FString FMetaSoundFrontendDocumentBuilder::GetDebugName() const
 {
 	using namespace Metasound::Frontend;
 
-	const FMetasoundFrontendClassMetadata& Metadata = GetDocument().RootGraph.Metadata;
-	const FMetasoundFrontendVersionNumber& Version = Metadata.GetVersion();
-	const FNodeRegistryKey RegKey = NodeRegistryKey::CreateKey(EMetasoundFrontendClassType::External, Metadata.GetClassName().ToString(), Version.Major, Version.Minor);
-	if (const FMetasoundAssetBase* Asset = IMetaSoundAssetManager::GetChecked().TryLoadAssetFromKey(RegKey))
-	{
-		return Asset->GetOwningAssetName();
-	}
-
-	return Metadata.GetClassName().ToString();
+	return DocumentBuilderPrivate::GetDebugName(GetDocumentInterface());
 }
 
 FMetasoundFrontendDocument& FMetaSoundFrontendDocumentBuilder::GetDocument()
@@ -1968,84 +1978,76 @@ bool FMetaSoundFrontendDocumentBuilder::TransformTemplateNodes()
 
 	METASOUND_TRACE_CPUPROFILER_EVENT_SCOPE(FMetaSoundFrontendDocumentBuilder::TransformTemplateNodes);
 
-	bool bModified = false;
+	struct FTemplateTransformParams
+	{
+		const Metasound::Frontend::INodeTemplate* Template;
+		TArray<FGuid> NodeIDs;
+	};
+	using FTemplateTransformParamsMap = TSortedMap<FGuid, FTemplateTransformParams>;
 
 	FMetasoundFrontendDocument& Document = GetDocument();
 	FMetasoundFrontendGraph& Graph = Document.RootGraph.Graph;
 	TArray<FMetasoundFrontendClass>& Dependencies = Document.Dependencies;
-	const TArray<FMetasoundFrontendEdge>& GraphEdges = Graph.Edges;
 
-	struct FTemplateParams
+	FTemplateTransformParamsMap TemplateParams;
+	for (const FMetasoundFrontendClass& Dependency : Dependencies)
 	{
-		FGuid ClassID;
-		const INodeTemplate* Template = nullptr;
-	};
-
-	// 1. Find template dependencies to build
-	TArray<FTemplateParams> TemplateParams;
-	Algo::TransformIf(Dependencies, TemplateParams,
-		[](const FMetasoundFrontendClass& Class) { return Class.Metadata.GetType() == EMetasoundFrontendClassType::Template; },
-		[](const FMetasoundFrontendClass& Class)
+		if (Dependency.Metadata.GetType() == EMetasoundFrontendClassType::Template)
 		{
-			const FNodeRegistryKey Key = NodeRegistryKey::CreateKey(Class.Metadata);
-			FGuid ClassID = Class.ID;
+			const FNodeRegistryKey Key = NodeRegistryKey::CreateKey(Dependency.Metadata);
 			const INodeTemplate* Template = INodeTemplateRegistry::Get().FindTemplate(Key);
-			ensureMsgf(Template, TEXT("Template not found for template class reference '%s'"), *Class.Metadata.GetClassName().ToString());
-			return FTemplateParams { ClassID, Template };
+			ensureMsgf(Template, TEXT("Template not found for template class reference '%s'"), *Dependency.Metadata.GetClassName().ToString());
+			TemplateParams.Add(Dependency.ID, FTemplateTransformParams { Template });
 		}
-	);
+	}
 
-	for (FTemplateParams& Params : TemplateParams)
+	if (TemplateParams.IsEmpty())
 	{
-		if (!Params.Template)
-		{
-			continue;
-		}
-
-		TUniquePtr<INodeTransform> NodeTransform = Params.Template->GenerateNodeTransform();
-		check(NodeTransform.IsValid());
-
-		TSet<const FMetasoundFrontendNode*> NodesToRemove;
-
-		// 2. Execute generated template node transform on copy of node array,
-		// which allows for addition/removal of nodes to/from original array container
-		// without template transform having to worry about mutation while iterating
-		TArray<FGuid> TemplateNodeIDs;
-		Algo::TransformIf(Graph.Nodes, TemplateNodeIDs,
-			[ClassID = Params.ClassID](const FMetasoundFrontendNode& Node) { return ClassID == Node.ClassID; },
-			[](const FMetasoundFrontendNode& Node) { return Node.GetID(); });
+		return false;
+	}
 
-		for (const FGuid& NodeID : TemplateNodeIDs)
+	// 1. Execute generated template node transform on copy of node array,
+	// which allows for addition/removal of nodes to/from original array container
+	// without template transform having to worry about mutation while iterating
+	TArray<FGuid> TemplateNodeIDs;
+	for (const FMetasoundFrontendNode& Node : Graph.Nodes)
+	{
+		if (FTemplateTransformParams* Params = TemplateParams.Find(Node.ClassID))
 		{
-			bModified = true;
-			NodeTransform->Transform(NodeID, *this);
+			Params->NodeIDs.Add(Node.GetID());
 		}
 	}
 
-	// 4. Remove template classes from dependency list
+	// 2. Transform nodes
+	bool bModified = false;
+	for (const TPair<FGuid, FTemplateTransformParams>& Pair : TemplateParams)
 	{
-		TSet<FString> TemplateKeys;
-		Algo::Transform(TemplateParams, TemplateKeys, [](const FTemplateParams& Params)
+		const FTemplateTransformParams& Params = Pair.Value;
+		if (Params.Template)
 		{
-			if (Params.Template)
+			TUniquePtr<INodeTransform> NodeTransform = Params.Template->GenerateNodeTransform();
+			check(NodeTransform.IsValid());
+
+			for (const FGuid& NodeID : Params.NodeIDs)
 			{
-				return NodeRegistryKey::CreateKey(Params.Template->GetFrontendClass().Metadata);
+				bModified = true;
+				NodeTransform->Transform(NodeID, *this);
 			}
+		}
+	}
 
-			return FString();
-		});
-		constexpr bool bAllowShrinking = false;
-		for (int32 i = Dependencies.Num() - 1; i >= 0; --i)
+	// 3. Remove template classes from dependency list
+	constexpr bool bAllowShrinking = false;
+	for (int32 i = Dependencies.Num() - 1; i >= 0; --i)
+	{
+		const FMetasoundFrontendClass& Class = Dependencies[i];
+		if (TemplateParams.Contains(Class.ID))
 		{
-			const FMetasoundFrontendClass& Class = Dependencies[i];
-			if (TemplateKeys.Contains(NodeRegistryKey::CreateKey(Class.Metadata)))
-			{
-				DocumentDelegates->OnRemoveSwappingDependency.Broadcast(i, Dependencies.Num() - 1);
-				Dependencies.RemoveAtSwap(i, 1, bAllowShrinking);
-			}
+			DocumentDelegates->OnRemoveSwappingDependency.Broadcast(i, Dependencies.Num() - 1);
+			Dependencies.RemoveAtSwap(i, 1, bAllowShrinking);
 		}
-		Dependencies.Shrink();
 	}
+	Dependencies.Shrink();
 
 	return bModified;
 }
diff --git a/Engine/Plugins/Runtime/Metasound/Source/MetasoundFrontend/Private/MetasoundFrontendDocumentCache.cpp b/Engine/Plugins/Runtime/Metasound/Source/MetasoundFrontend/Private/MetasoundFrontendDocumentCache.cpp
--- a/Engine/Plugins/Runtime/Metasound/Source/MetasoundFrontend/Private/MetasoundFrontendDocumentCache.cpp
+++ b/Engine/Plugins/Runtime/Metasound/Source/MetasoundFrontend/Private/MetasoundFrontendDocumentCache.cpp
@@ -2,6 +2,7 @@
 
 #include "MetasoundFrontendDocumentCache.h"
 
+#include "Algo/AnyOf.h"
 #include "Algo/ForEach.h"
 #include "MetasoundFrontendDocument.h"
 #include "MetasoundFrontendDocumentBuilder.h"
@@ -130,6 +131,13 @@ namespace Metasound::Frontend
 		return KeyToIndex.Contains(InClassKey);
 	}
 
+	bool FDocumentCache::ContainsDependencyOfType(EMetasoundFrontendClassType ClassType) const
+	{
+		const TArray<FMetasoundFrontendClass>& Dependencies = GetDocument().Dependencies;
+		auto IsTemplateDependency = [&ClassType](const FMetasoundFrontendClass& Class) { return Class.Metadata.GetType() == ClassType; };
+		return Algo::AnyOf(Dependencies, IsTemplateDependency);
+	}
+
 	const FMetasoundFrontendClass* FDocumentCache::FindDependency(const FNodeRegistryKey& InClassKey) const
 	{
 		if (const int32* Index = KeyToIndex.Find(InClassKey))
diff --git a/Engine/Plugins/Runtime/Metasound/Source/MetasoundFrontend/Private/MetasoundFrontendDocumentCache.h b/Engine/Plugins/Runtime/Metasound/Source/MetasoundFrontend/Private/MetasoundFrontendDocumentCache.h
--- a/Engine/Plugins/Runtime/Metasound/Source/MetasoundFrontend/Private/MetasoundFrontendDocumentCache.h
+++ b/Engine/Plugins/Runtime/Metasound/Source/MetasoundFrontend/Private/MetasoundFrontendDocumentCache.h
@@ -141,6 +141,7 @@ namespace Metasound::Frontend
 		virtual ~FDocumentCache() = default;
 
 		virtual bool ContainsDependency(const FNodeRegistryKey& InClassKey) const override;
+		virtual bool ContainsDependencyOfType(EMetasoundFrontendClassType ClassType) const override;
 
 		virtual const FMetasoundFrontendClass* FindDependency(const FNodeRegistryKey& InClassKey) const override;
 		virtual const FMetasoundFrontendClass* FindDependency(const FGuid& InClassID) const override;
diff --git a/Engine/Plugins/Runtime/Metasound/Source/MetasoundFrontend/Private/MetasoundFrontendDocumentTemplatePreprocessor.cpp b/Engine/Plugins/Runtime/Metasound/Source/MetasoundFrontend/Private/MetasoundFrontendDocumentTemplatePreprocessor.cpp
deleted file mode 100644
diff --git a/Engine/Plugins/Runtime/Metasound/Source/MetasoundFrontend/Private/MetasoundFrontendRegistries.cpp b/Engine/Plugins/Runtime/Metasound/Source/MetasoundFrontend/Private/MetasoundFrontendRegistries.cpp
--- a/Engine/Plugins/Runtime/Metasound/Source/MetasoundFrontend/Private/MetasoundFrontendRegistries.cpp
+++ b/Engine/Plugins/Runtime/Metasound/Source/MetasoundFrontend/Private/MetasoundFrontendRegistries.cpp
@@ -624,11 +624,23 @@ namespace Metasound
 
 			FNodeRegistryKey CreateKey(const FMetasoundFrontendClassMetadata& InNodeMetadata)
 			{
+				checkf(InNodeMetadata.GetType() != EMetasoundFrontendClassType::Graph, TEXT("Cannot create key from 'graph' type. Likely meant to use CreateKey overload that is provided FMetasoundFrontendGraphClass"));
 				return CreateKey(InNodeMetadata.GetType(), InNodeMetadata.GetClassName().GetFullName().ToString(), InNodeMetadata.GetVersion().Major, InNodeMetadata.GetVersion().Minor);
 			}
 
+			FNodeRegistryKey CreateKey(const FMetasoundFrontendGraphClass& InGraphClass)
+			{
+				return CreateKey(
+					EMetasoundFrontendClassType::External, // Overridden as all graphs are considered the same as an external in the registry
+					InGraphClass.Metadata.GetClassName().GetFullName().ToString(),
+					InGraphClass.Metadata.GetVersion().Major,
+					InGraphClass.Metadata.GetVersion().Minor
+				);
+			}
+
 			FNodeRegistryKey CreateKey(const FNodeClassInfo& InClassInfo)
 			{
+				checkf(InClassInfo.Type != EMetasoundFrontendClassType::Graph, TEXT("Cannot create key from 'graph' type. Likely meant to use CreateKey overload that is provided FMetasoundFrontendGraphClass"));
 				return CreateKey(InClassInfo.Type, InClassInfo.ClassName.GetFullName().ToString(), InClassInfo.Version.Major, InClassInfo.Version.Minor);
 			}
 		}
diff --git a/Engine/Plugins/Runtime/Metasound/Source/MetasoundFrontend/Private/NodeTemplates/MetasoundFrontendDocumentTemplatePreprocessor.h b/Engine/Plugins/Runtime/Metasound/Source/MetasoundFrontend/Private/NodeTemplates/MetasoundFrontendDocumentTemplatePreprocessor.h
deleted file mode 100644
diff --git a/Engine/Plugins/Runtime/Metasound/Source/MetasoundFrontend/Private/NodeTemplates/MetasoundFrontendNodeTemplateReroute.cpp b/Engine/Plugins/Runtime/Metasound/Source/MetasoundFrontend/Private/NodeTemplates/MetasoundFrontendNodeTemplateReroute.cpp
--- a/Engine/Plugins/Runtime/Metasound/Source/MetasoundFrontend/Private/NodeTemplates/MetasoundFrontendNodeTemplateReroute.cpp
+++ b/Engine/Plugins/Runtime/Metasound/Source/MetasoundFrontend/Private/NodeTemplates/MetasoundFrontendNodeTemplateReroute.cpp
@@ -8,7 +8,6 @@
 #include "MetasoundFrontendDocumentBuilder.h"
 #include "MetasoundFrontendNodeTemplateRegistry.h"
 #include "MetasoundFrontendRegistries.h"
-#include "NodeTemplates/MetasoundFrontendDocumentTemplatePreprocessor.h"
 
 
 namespace Metasound::Frontend
diff --git a/Engine/Plugins/Runtime/Metasound/Source/MetasoundFrontend/Public/MetasoundAssetBase.h b/Engine/Plugins/Runtime/Metasound/Source/MetasoundFrontend/Public/MetasoundAssetBase.h
--- a/Engine/Plugins/Runtime/Metasound/Source/MetasoundFrontend/Public/MetasoundAssetBase.h
+++ b/Engine/Plugins/Runtime/Metasound/Source/MetasoundFrontend/Public/MetasoundAssetBase.h
@@ -100,7 +100,8 @@ public:
 	bool IsInterfaceDeclared(const FMetasoundFrontendVersion& InVersion) const;
 
 	// Gets the asset class info.
-	virtual Metasound::Frontend::FNodeClassInfo GetAssetClassInfo() const = 0;
+	UE_DEPRECATED(5.4, "NodeClassInfo can be constructed directly from document's root graph & asset's path and requires no specialized virtual getter.")
+	virtual Metasound::Frontend::FNodeClassInfo GetAssetClassInfo() const;
 
 	// Returns all the class keys of this asset's referenced assets
 	virtual const TSet<FString>& GetReferencedAssetClassKeys() const = 0;
diff --git a/Engine/Plugins/Runtime/Metasound/Source/MetasoundFrontend/Public/MetasoundAssetManager.h b/Engine/Plugins/Runtime/Metasound/Source/MetasoundFrontend/Public/MetasoundAssetManager.h
--- a/Engine/Plugins/Runtime/Metasound/Source/MetasoundFrontend/Public/MetasoundAssetManager.h
+++ b/Engine/Plugins/Runtime/Metasound/Source/MetasoundFrontend/Public/MetasoundAssetManager.h
@@ -55,13 +55,13 @@ namespace Metasound
 			// If true, warnings will be logged if updating a node results in existing connections being discarded.
 			bool bAutoUpdateLogWarningOnDroppedConnection = false;
 
-			// If true, the preprocessed document will be registered 
+#if WITH_EDITOR
+			// Soft deprecated. Preprocessing now handled contextually if cooking or serializing.
 			bool bPreprocessDocument = true;
 
-#if WITH_EDITOR
 			// Attempt to rebuild referenced classes (only run if class not registered or set to force re-register)
 			bool bRebuildReferencedAssetClasses = true;
-#endif
+#endif // WITH_EDITOR
 		};
 
 		class METASOUNDFRONTEND_API IMetaSoundAssetManager
diff --git a/Engine/Plugins/Runtime/Metasound/Source/MetasoundFrontend/Public/MetasoundFrontendDocumentBuilder.h b/Engine/Plugins/Runtime/Metasound/Source/MetasoundFrontend/Public/MetasoundFrontendDocumentBuilder.h
--- a/Engine/Plugins/Runtime/Metasound/Source/MetasoundFrontend/Public/MetasoundFrontendDocumentBuilder.h
+++ b/Engine/Plugins/Runtime/Metasound/Source/MetasoundFrontend/Public/MetasoundFrontendDocumentBuilder.h
@@ -22,6 +22,9 @@ class FMetasoundAssetBase;
 
 namespace Metasound::Frontend
 {
+	// Forward Declarations
+	class INodeTemplate;
+
 	using FFinalizeNodeFunctionRef = TFunctionRef<void(FMetasoundFrontendNode&, const Metasound::Frontend::FNodeRegistryKey&)>;
 
 	enum class EInvalidEdgeReason : uint8
@@ -150,6 +153,8 @@ public:
 	bool CanAddEdge(const FMetasoundFrontendEdge& InEdge) const;
 
 	void ClearGraph();
+
+	bool ContainsDependencyOfType(EMetasoundFrontendClassType ClassType) const;
 	bool ContainsEdge(const FMetasoundFrontendEdge& InEdge) const;
 	bool ContainsNode(const FGuid& InNodeID) const;
 
@@ -185,6 +190,21 @@ public:
 
 	const FString GetDebugName() const;
 	const FMetasoundFrontendDocument& GetDocument() const;
+
+	template<typename TObjectType>
+	const TObjectType& CastDocumentObjectChecked() const
+	{
+		const UObject* Owner = DocumentInterface.GetObject();
+		return *CastChecked<TObjectType>(Owner);
+	}
+
+	template<typename TObjectType>
+	TObjectType& CastDocumentObjectChecked()
+	{
+		UObject* Owner = DocumentInterface.GetObject();
+		return *CastChecked<TObjectType>(Owner);
+	}
+
 	const Metasound::Frontend::FDocumentModifyDelegates& GetDocumentDelegates() const;
 	const IMetaSoundDocumentInterface& GetDocumentInterface() const;
 	EMetasoundFrontendVertexAccessType GetNodeInputAccessType(const FGuid& InNodeID, const FGuid& InVertexID) const;
@@ -246,6 +266,11 @@ public:
 	bool SwapGraphOutput(const FMetasoundFrontendClassVertex& InExistingOutputVertex, const FMetasoundFrontendClassVertex& NewOutputVertex);
 	bool UpdateDependencyClassNames(const TMap<FMetasoundFrontendClassName, FMetasoundFrontendClassName>& OldToNewReferencedClassNames);
 
+	// Determines if the provided interface requires further processing or if its data  is ready as is for insertion into the FrontendRegistry.
+	// If it requires further processing, returns a transient document interface to be used by the registry with an optimized, transformed version
+	// of the given interface's document.
+	static TScriptInterface<IMetaSoundDocumentInterface> BuildRegistryDocument(TScriptInterface<IMetaSoundDocumentInterface> DocumentInterface);
+
 	// Transforms template nodes within the given builder's document, which can include swapping associated edges and/or
 	// replacing nodes with other, registry-defined concrete node class instances. Returns true if any template nodes were processed.
 	bool TransformTemplateNodes();
diff --git a/Engine/Plugins/Runtime/Metasound/Source/MetasoundFrontend/Public/MetasoundFrontendDocumentCacheInterface.h b/Engine/Plugins/Runtime/Metasound/Source/MetasoundFrontend/Public/MetasoundFrontendDocumentCacheInterface.h
--- a/Engine/Plugins/Runtime/Metasound/Source/MetasoundFrontend/Public/MetasoundFrontendDocumentCacheInterface.h
+++ b/Engine/Plugins/Runtime/Metasound/Source/MetasoundFrontend/Public/MetasoundFrontendDocumentCacheInterface.h
@@ -99,7 +99,7 @@ namespace Metasound::Frontend
 		virtual ~IDocumentCache() = default;
 
 		virtual bool ContainsDependency(const FNodeRegistryKey& InClassKey) const = 0;
-
+		virtual bool ContainsDependencyOfType(EMetasoundFrontendClassType ClassType) const = 0;
 		virtual const FMetasoundFrontendClass* FindDependency(const Metasound::Frontend::FNodeRegistryKey& InClassKey) const = 0;
 		virtual const FMetasoundFrontendClass* FindDependency(const FGuid& InClassID) const = 0;
 		virtual const int32* FindDependencyIndex(const Metasound::Frontend::FNodeRegistryKey& InClassKey) const = 0;
diff --git a/Engine/Plugins/Runtime/Metasound/Source/MetasoundFrontend/Public/MetasoundFrontendRegistries.h b/Engine/Plugins/Runtime/Metasound/Source/MetasoundFrontend/Public/MetasoundFrontendRegistries.h
--- a/Engine/Plugins/Runtime/Metasound/Source/MetasoundFrontend/Public/MetasoundFrontendRegistries.h
+++ b/Engine/Plugins/Runtime/Metasound/Source/MetasoundFrontend/Public/MetasoundFrontendRegistries.h
@@ -127,7 +127,7 @@ namespace Metasound
 			  *
 			  * Returns nullptr if node class implementation does not support interface implementation.
 			  */
-			virtual TSet<FMetasoundFrontendVersion>* GetImplementedInterfaces() const = 0;
+			virtual const TSet<FMetasoundFrontendVersion>* GetImplementedInterfaces() const = 0;
 
 			/** Whether or not the node is natively defined */
 			virtual bool IsNative() const = 0;
@@ -251,6 +251,7 @@ namespace Metasound
 			METASOUNDFRONTEND_API FNodeRegistryKey CreateKey(EMetasoundFrontendClassType InType, const FString& InFullClassName, int32 InMajorVersion, int32 InMinorVersion);
 			METASOUNDFRONTEND_API FNodeRegistryKey CreateKey(const FNodeClassMetadata& InNodeMetadata);
 			METASOUNDFRONTEND_API FNodeRegistryKey CreateKey(const FMetasoundFrontendClassMetadata& InNodeMetadata);
+			METASOUNDFRONTEND_API FNodeRegistryKey CreateKey(const FMetasoundFrontendGraphClass& InNodeMetadata);
 			METASOUNDFRONTEND_API FNodeRegistryKey CreateKey(const FNodeClassInfo& ClassInfo);
 		}
 	} // namespace Frontend
diff --git a/Engine/Plugins/Runtime/Metasound/Source/MetasoundFrontend/Public/MetasoundNodeRegistrationMacro.h b/Engine/Plugins/Runtime/Metasound/Source/MetasoundFrontend/Public/MetasoundNodeRegistrationMacro.h
--- a/Engine/Plugins/Runtime/Metasound/Source/MetasoundFrontend/Public/MetasoundNodeRegistrationMacro.h
+++ b/Engine/Plugins/Runtime/Metasound/Source/MetasoundFrontend/Public/MetasoundNodeRegistrationMacro.h
@@ -82,7 +82,7 @@ namespace Metasound
 
 			virtual TUniquePtr<INodeRegistryEntry> Clone() const override = 0;
 
-			virtual TSet<FMetasoundFrontendVersion>* GetImplementedInterfaces() const override
+			virtual const TSet<FMetasoundFrontendVersion>* GetImplementedInterfaces() const override
 			{
 				return nullptr;
 			}
