[Metasound] Refactor UMetasoundGeneratorHandle to work with Live Update

diff --git a/Engine/Plugins/Runtime/Metasound/Source/MetasoundEngine/Private/MetasoundGeneratorHandle.cpp b/Engine/Plugins/Runtime/Metasound/Source/MetasoundEngine/Private/MetasoundGeneratorHandle.cpp
--- a/Engine/Plugins/Runtime/Metasound/Source/MetasoundEngine/Private/MetasoundGeneratorHandle.cpp
+++ b/Engine/Plugins/Runtime/Metasound/Source/MetasoundEngine/Private/MetasoundGeneratorHandle.cpp
@@ -3,558 +3,687 @@
 #include "MetasoundGeneratorHandle.h"
 
 #include "MetasoundGenerator.h"
-#include "MetasoundParameterPack.h"
 #include "MetasoundSource.h"
 #include "MetasoundTrace.h"
+
+#include "Analysis/MetasoundFrontendAnalyzerFactory.h"
 #include "Analysis/MetasoundFrontendAnalyzerRegistry.h"
-#include "Components/AudioComponent.h"
 
-TMap<FName, UMetasoundGeneratorHandle::FPassthroughAnalyzerInfo> UMetasoundGeneratorHandle::PassthroughAnalyzers{};
+#include "Async/Async.h"
 
-UMetasoundGeneratorHandle* UMetasoundGeneratorHandle::CreateMetaSoundGeneratorHandle(UAudioComponent* OnComponent)
+#include "Components/AudioComponent.h"
+
+namespace Metasound
 {
-	METASOUND_LLM_SCOPE;
-	METASOUND_TRACE_CPUPROFILER_EVENT_SCOPE(UMetasoundGeneratorHandle::CreateMetaSoundGeneratorHandle);
+	TMap<FName, FMetasoundGeneratorHandle::FPassthroughAnalyzerInfo> FMetasoundGeneratorHandle::PassthroughAnalyzers{};
 	
-	if (!OnComponent)
-	{
-		return nullptr;
+	FMetasoundGeneratorHandle::FMetasoundGeneratorHandle(FPrivateToken,
+	TWeakObjectPtr<UAudioComponent>&& InAudioComponent)
+		: AudioComponent(MoveTemp(InAudioComponent))
+		, AudioComponentId(AudioComponent.IsValid() ? AudioComponent->GetAudioComponentID() : INDEX_NONE)
+	{
+		METASOUND_LLM_SCOPE;
+		METASOUND_TRACE_CPUPROFILER_EVENT_SCOPE(FMetasoundGeneratorHandle::FMetasoundGeneratorHandle);
+		
+		if (!AudioComponent.IsValid())
+		{
+			UE_LOG(LogMetaSound, Error, TEXT("Created a FMetaSoundGeneratorHandle with an invalid UAudioComponent."));
+			return;
+		}
+		
+		if (AudioComponent->bCanPlayMultipleInstances)
+		{
+			UE_LOG(
+				LogMetaSound,
+				Warning,
+				TEXT("Created a FMetaSoundGeneratorHandle for a UAudioComponent that is allowed to play multiple instances. This may not work as expected."))
+		}
 	}
-
-	if (OnComponent->bCanPlayMultipleInstances)
+	
+	TSharedPtr<FMetasoundGeneratorHandle> FMetasoundGeneratorHandle::Create(
+		TWeakObjectPtr<UAudioComponent>&& InAudioComponent)
 	{
-		UE_LOG(
-			LogMetaSound,
-			Warning,
-			TEXT("Created a UMetaSoundGeneratorHandle for a UAudioComponent that is allowed to play multiple instances. This may not work as expected."))
-	}
+		TSharedRef<FMetasoundGeneratorHandle> Handle = MakeShared<FMetasoundGeneratorHandle>(
+			FPrivateToken{},
+			MoveTemp(InAudioComponent));
 
-	UMetasoundGeneratorHandle* Result = NewObject<UMetasoundGeneratorHandle>();
-	Result->SetAudioComponent(OnComponent);
-	return Result;
-}
+		if (Handle->IsValid())
+		{
+			const TWeakObjectPtr<UMetaSoundSource> Source = Handle->GetMetaSoundSource();
 
-void UMetasoundGeneratorHandle::BeginDestroy()
-{
-	Super::BeginDestroy();
-	DetachGeneratorDelegates();
+			if (!Source.IsValid())
+			{
+				UE_LOG(LogMetaSound, Error, TEXT("Couldn't get the source for the audio component."));
+				return nullptr;
+			}
+			
+			const uint64 AudioComponentId = Handle->GetAudioComponentId();
+			TWeakPtr<FMetasoundGenerator> GeneratorForComponent = Source->GetGeneratorForAudioComponent(AudioComponentId);
 
-	if (TSharedPtr<Metasound::FMetasoundGenerator> PinnedGenerator = CachedGeneratorPtr.Pin())
-	{
-		PinnedGenerator->OnOutputChanged.Remove(OutputChangedDelegateHandle);
-	}
-}
+			// If we have a generator already, set it.
+			if (GeneratorForComponent.IsValid())
+			{
+				Handle->SetGenerator(MoveTemp(GeneratorForComponent));
+			}
+			
+			// Listen for the source creating a new generator
+			Handle->GeneratorCreatedDelegateHandle = Source->OnGeneratorInstanceCreated.AddSP(
+				Handle,
+				&FMetasoundGeneratorHandle::HandleGeneratorCreated);
 
-bool UMetasoundGeneratorHandle::IsValid() const
-{
-	return AudioComponent.IsValid();
-}
+			// Listen for the generator being destroyed
+			Handle->GeneratorDestroyedDelegateHandle = Source->OnGeneratorInstanceDestroyed.AddSP(
+				Handle,
+				&FMetasoundGeneratorHandle::HandleGeneratorDestroyed);
 
-uint64 UMetasoundGeneratorHandle::GetAudioComponentId() const
-{
-	if (AudioComponent.IsValid())
-	{
-		return AudioComponent->GetAudioComponentID();
+			return Handle;
+		}
+
+		return nullptr;
 	}
 
-	return INDEX_NONE;
-}
+	FMetasoundGeneratorHandle::~FMetasoundGeneratorHandle()
+	{
+		check(IsInGameThread());
+		
+		// unsubscribe from source events
+		{
+			const TWeakObjectPtr<UMetaSoundSource> Source = GetMetaSoundSource();
+			
+			if (Source.IsValid())
+			{
+				Source->OnGeneratorInstanceCreated.Remove(GeneratorCreatedDelegateHandle);
+				Source->OnGeneratorInstanceDestroyed.Remove(GeneratorDestroyedDelegateHandle);
+			}
+		}
 
-void UMetasoundGeneratorHandle::ClearCachedData()
-{
-	DetachGeneratorDelegates();
-	AudioComponent        = nullptr;
-	CachedMetasoundSource = nullptr;
-	CachedGeneratorPtr    = nullptr;
-	CachedParameterPack   = nullptr;
-}
+		// unset the generator and clean up
+		SetGenerator(nullptr);
+	}
 
-void UMetasoundGeneratorHandle::SetAudioComponent(UAudioComponent* InAudioComponent)
-{
-	if (InAudioComponent != AudioComponent)
+	bool FMetasoundGeneratorHandle::IsValid() const
 	{
-		ClearCachedData();
-		AudioComponent   = InAudioComponent;
+		return AudioComponent.IsValid();
 	}
-}
 
-void UMetasoundGeneratorHandle::CacheMetasoundSource()
-{
-	if (!AudioComponent.IsValid())
+	uint64 FMetasoundGeneratorHandle::GetAudioComponentId() const
 	{
-		return;
+		return AudioComponentId;
 	}
 
-	UMetaSoundSource* CurrentMetasoundSource = Cast<UMetaSoundSource>(AudioComponent->GetSound());
-	if (CachedMetasoundSource == CurrentMetasoundSource)
+	TSharedPtr<FMetasoundGenerator> FMetasoundGeneratorHandle::GetGenerator() const
 	{
-		return;
+		return Generator.Pin();
 	}
 
-	DetachGeneratorDelegates();
-	CachedGeneratorPtr    = nullptr;
-	CachedMetasoundSource = CurrentMetasoundSource;
-
-	if (CachedMetasoundSource.IsValid())
+	void FMetasoundGeneratorHandle::UpdateParameters(const UMetasoundParameterPack& ParameterPack)
 	{
-		AttachGeneratorDelegates();
+		METASOUND_LLM_SCOPE;
+		METASOUND_TRACE_CPUPROFILER_EVENT_SCOPE(FMetasoundGeneratorHandle::UpdateParameters);
+
+		// Update the latest state
+		LatestParameterState = ParameterPack.GetCopyOfParameterStorage();
+
+		// Try to send to the generator
+		SendParametersToGenerator();
 	}
-}
 
-void UMetasoundGeneratorHandle::AttachGeneratorDelegates()
-{
-	// These delegates can be called on a separate thread, so we need to take steps
-	// to assure this UObject hasn't been garbage collected before trying to dereference. 
-	// That is why we capture a TWeakObjectPtr and try to dereference that later!
-	GeneratorCreatedDelegateHandle = CachedMetasoundSource->OnGeneratorInstanceCreated.AddLambda(
-		[WeakGeneratorHandlePtr = TWeakObjectPtr<UMetasoundGeneratorHandle>(this),
-		StatId = this->GetStatID(true)]
-		(uint64 InAudioComponentId, TSharedPtr<Metasound::FMetasoundGenerator> InGenerator)
+	bool FMetasoundGeneratorHandle::WatchOutput(
+		const FName OutputName,
+		const FOnMetasoundOutputValueChanged& OnOutputValueChanged,
+		const FName AnalyzerName,
+		const FName AnalyzerOutputName)
+	{
+		METASOUND_LLM_SCOPE;
+		METASOUND_TRACE_CPUPROFILER_EVENT_SCOPE(FMetasoundGeneratorHandle::WatchOutput);
+
+		check(IsInGameThread());
+
+		if (!IsValid())
 		{
-			// We are in the audio render (or control) thread here, so create a "dispatch task" to be
-			// executed later on the game thread...
-			FFunctionGraphTask::CreateAndDispatchWhenReady(
-				[WeakGeneratorHandlePtr, InAudioComponentId, InGenerator]()
-				{
-					METASOUND_TRACE_CPUPROFILER_EVENT_SCOPE(UMetasoundGeneratorHandle::CallingGeneratorAttachedDelegates);
-					check(IsInGameThread());
-					// Now, since we are in the game thread, try to dereference the pointer to
-					// to the UMetasoundGeneratorHandle. This should only succeed if the UObject
-					// hasn't been garbage collected.
-					if (UMetasoundGeneratorHandle* TheHandle = WeakGeneratorHandlePtr.Get())
-					{
-						TheHandle->OnSourceCreatedAGenerator(InAudioComponentId, InGenerator);
-					}
-				}, 
-				StatId, nullptr, ENamedThreads::GameThread);
-		});
-	GeneratorDestroyedDelegateHandle = CachedMetasoundSource->OnGeneratorInstanceDestroyed.AddLambda(
-		[WeakGeneratorHandlePtr = TWeakObjectPtr<UMetasoundGeneratorHandle>(this),
-		StatId = this->GetStatID(true)]
-		(uint64 InAudioComponentId, TSharedPtr<Metasound::FMetasoundGenerator> InGenerator)
+			return false;
+		}
+
+		// Make the analyzer address.
+		Frontend::FAnalyzerAddress AnalyzerAddress;
+		AnalyzerAddress.InstanceID = GetAudioComponentId();
+		AnalyzerAddress.OutputName = OutputName;
+		AnalyzerAddress.AnalyzerName = AnalyzerName;
+		AnalyzerAddress.AnalyzerMemberName = AnalyzerOutputName;
+		AnalyzerAddress.AnalyzerInstanceID = FGuid::NewGuid();
+		
+		// Find the output node and get the data type/node id from that
 		{
-			// We are in the audio render (or control) thread here, so create a "dispatch task" to be
-			// executed later on the game thread...
-			FFunctionGraphTask::CreateAndDispatchWhenReady(
-				[WeakGeneratorHandlePtr, InAudioComponentId]()
-				{
-					METASOUND_TRACE_CPUPROFILER_EVENT_SCOPE(UMetasoundGeneratorHandle::CallingGeneratorDetachedDelegates);
-					check(IsInGameThread());
-					// Now, since we are in the game thread, try to dereference the pointer to
-					// to the UMetasoundGeneratorHandle. This should only succeed if the UObject
-					// hasn't been garbage collected.
-					if (UMetasoundGeneratorHandle* TheHandle = WeakGeneratorHandlePtr.Get())
-					{
-						TheHandle->OnSourceDestroyedAGenerator(InAudioComponentId);
-					}
-				},
-				StatId, nullptr, ENamedThreads::GameThread);
-		});
-}
+			const TWeakObjectPtr<UMetaSoundSource> Source = GetMetaSoundSource();
 
-void UMetasoundGeneratorHandle::AttachGraphChangedDelegate()
-{
-	if (!CachedGeneratorPtr.IsValid() || !OnGeneratorsGraphChanged.IsBound())
-		return;
-
-	TSharedPtr<Metasound::FMetasoundGenerator> Generator = CachedGeneratorPtr.Pin();
-	if (Generator)
-	{
-		// We're about to add a delegate to the generator. This delegate will be called on 
-		// the audio render thread so we need to take steps to assure this UMetasoundGeneratorHandle
-		// hasn't been garbage collected before trying to dereference it later when the delegate "fires".
-		// That is why we capture a TWeakObjectPtr and try to dereference that later!
-		Metasound::FOnSetGraph::FDelegate NewSetGraphDelegate;
-		NewSetGraphDelegate.BindLambda([WeakGeneratorHandlePtr = TWeakObjectPtr<UMetasoundGeneratorHandle>(this), StatId = this->GetStatID(true)]
-			()
+			if (!Source.IsValid())
 			{
-				// We are in the audio render thread here, so create a "dispatch task" to be
-				// executed later on the game thread...
-				FFunctionGraphTask::CreateAndDispatchWhenReady(
-					[WeakGeneratorHandlePtr]()
-					{
-						METASOUND_TRACE_CPUPROFILER_EVENT_SCOPE(UMetasoundGeneratorHandle::CallingGraphChangedDelegates);
-						check(IsInGameThread());
-						// Now, since we are in the game thread, try to dereference the pointer to
-						// to the UMetasoundGeneratorHandle. This should only succeed if the UObject
-						// hasn't been garbage collected.
-						if (UMetasoundGeneratorHandle* TheHandle = WeakGeneratorHandlePtr.Get())
-						{
-							TheHandle->OnGeneratorsGraphChanged.Broadcast();
-						}
-					},
-					StatId, nullptr, ENamedThreads::GameThread);
-			});
+				UE_LOG(LogMetaSound, Warning, TEXT("Couldn't find the MetaSound Source"));
+				return false;
+			}
+		
+			// Find the node id and type name
+			const Frontend::FNodeHandle Node =
+				Source->GetRootGraphHandle()->GetOutputNodeWithName(AnalyzerAddress.OutputName);
 
-		GeneratorGraphChangedDelegateHandle = Generator->AddGraphSetCallback(MoveTemp(NewSetGraphDelegate));
-	}
-}
+			if (!Node->IsValid())
+			{
+				return false;
+			}
 
-void UMetasoundGeneratorHandle::DetachGeneratorDelegates()
-{
-	// First detach any callbacks that tell us when a generator is created or destroyed
-	// for the UMetasoundSource of interest...
-	if (CachedMetasoundSource.IsValid())
-	{
-	CachedMetasoundSource->OnGeneratorInstanceCreated.Remove(GeneratorCreatedDelegateHandle);
-	GeneratorCreatedDelegateHandle.Reset();
-	CachedMetasoundSource->OnGeneratorInstanceDestroyed.Remove(GeneratorDestroyedDelegateHandle);
-	GeneratorDestroyedDelegateHandle.Reset();
-}
-	// Now detach any callback we may have registered to get callbacks when the generator's
-	// graph has been changed...
-	TSharedPtr<Metasound::FMetasoundGenerator> PinnedGenerator = CachedGeneratorPtr.Pin();
-	if (PinnedGenerator)
-	{
-		PinnedGenerator->RemoveGraphSetCallback(GeneratorGraphChangedDelegateHandle);
-		GeneratorGraphChangedDelegateHandle.Reset();
+			AnalyzerAddress.NodeID = Node->GetID();
+
+			// We expect output nodes to have only one output
+			if (!ensure(Node->GetNumOutputs() == 1))
+			{
+				return false;
+			}
+
+			AnalyzerAddress.DataType = Node->GetOutputs()[0]->GetDataType();
+		}
+		
+		// If no analyzer name was provided, try to find a passthrough analyzer
+		if (AnalyzerAddress.AnalyzerName.IsNone())
+		{
+			if (!PassthroughAnalyzers.Contains(AnalyzerAddress.DataType))
+			{
+				return false;
+			}
+
+			AnalyzerAddress.AnalyzerName = PassthroughAnalyzers[AnalyzerAddress.DataType].AnalyzerName;
+			AnalyzerAddress.AnalyzerMemberName = PassthroughAnalyzers[AnalyzerAddress.DataType].OutputName;
+		}
+
+		// Check to see if the analyzer exists
+		{
+			using namespace Metasound::Frontend;
+			const IVertexAnalyzerFactory* Factory =
+				IVertexAnalyzerRegistry::Get().FindAnalyzerFactory(AnalyzerAddress.AnalyzerName);
+			
+			if (nullptr == Factory)
+			{
+				return false;
+			}
+		}
+
+		// Create the watcher
+		CreateOutputWatcher(AnalyzerAddress, OnOutputValueChanged);
+
+		// Update the generator's analyzers if necessary
+		FixUpOutputWatchers();
+
+		return true;
 	}
-}
 
-TSharedPtr<Metasound::FMetasoundGenerator> UMetasoundGeneratorHandle::PinGenerator()
-{
-	TSharedPtr<Metasound::FMetasoundGenerator> PinnedGenerator = CachedGeneratorPtr.Pin();
-	if (PinnedGenerator.IsValid() || !CachedMetasoundSource.IsValid())
+	void FMetasoundGeneratorHandle::UpdateOutputWatchers()
 	{
-		return PinnedGenerator;
-	}
+		METASOUND_LLM_SCOPE;
+		METASOUND_TRACE_CPUPROFILER_EVENT_SCOPE(FMetasoundGeneratorHandle::UpdateOutputWatchers);
 
-	// The first attempt to pin failed, so reach out to the MetaSoundSource and see if it has a 
-	// generator for our AudioComponent...
-	check(AudioComponent.IsValid()); // expect the audio component to still be valid if the generator is.
-	CachedGeneratorPtr = CachedMetasoundSource->GetGeneratorForAudioComponent(AudioComponent->GetAudioComponentID());
-	PinnedGenerator    = CachedGeneratorPtr.Pin();
-	return PinnedGenerator;
-}
+		check(IsInGameThread());
 
-bool UMetasoundGeneratorHandle::ApplyParameterPack(UMetasoundParameterPack* Pack)
-{
-	METASOUND_LLM_SCOPE;
-	METASOUND_TRACE_CPUPROFILER_EVENT_SCOPE(UMetasoundGeneratorHandle::ApplyParameterPack);
-	
-	if (!Pack)
-	{
-		return false;
+		while (TOptional<FOutputPayload> ChangedOutput = ChangedOutputs.Dequeue())
+		{
+			if (const FOutputWatcher* OutputListener = OutputWatchers.FindByPredicate(
+				[&ChangedOutput](const FOutputWatcher& ExistingListener)
+				{
+					return ChangedOutput->AnalyzerName == ExistingListener.AnalyzerAddress.AnalyzerName
+					&& ChangedOutput->OutputName == ExistingListener.AnalyzerAddress.OutputName
+					&& ChangedOutput->OutputValue.Name == ExistingListener.AnalyzerAddress.AnalyzerMemberName;
+				}))
+			{
+				OutputListener->OnOutputValueChanged.Broadcast(ChangedOutput->OutputName, ChangedOutput->OutputValue);
+			}
+		}
 	}
 
-	// Create a copy of the parameter pack and cache it.
-	CachedParameterPack = Pack->GetCopyOfParameterStorage();
-
-	// No point in continuing if the parameter pack is not valid for any reason.
-	if (!CachedParameterPack.IsValid())
+	void FMetasoundGeneratorHandle::RegisterPassthroughAnalyzerForType(
+		const FName TypeName,
+		const FName AnalyzerName,
+		const FName OutputName)
 	{
-		return false;
+		check(!PassthroughAnalyzers.Contains(TypeName));
+		PassthroughAnalyzers.Add(TypeName, { AnalyzerName, OutputName });
 	}
 
-	// Assure that our MetaSoundSource is up to date. It is possible that this has been 
-	// changed via script since we were first created.
-	CacheMetasoundSource();
+	void FMetasoundGeneratorHandle::EnableRuntimeRenderTiming(bool Enable)
+	{
+		bRuntimeRenderTimingShouldBeEnabled = Enable;
 
-	// Now we can try to pin the generator.
-	TSharedPtr<Metasound::FMetasoundGenerator> PinnedGenerator = PinGenerator();
+		if (const TSharedPtr<FMetasoundGenerator> PinnedGenerator = Generator.Pin())
+		{
+			PinnedGenerator->EnableRuntimeRenderTiming(bRuntimeRenderTimingShouldBeEnabled);
+		}
+	}
 
-	if (!PinnedGenerator.IsValid())
+	double FMetasoundGeneratorHandle::GetCPUCoreUtilization() const
 	{
-		// Failed to pin the generator, but we have cached the parameter pack,
-		// so if our delegate gets called when a new generator is created we can 
-		// apply the cached parameters then.
-		return false;
+		if (const TSharedPtr<FMetasoundGenerator> PinnedGenerator = Generator.Pin())
+		{
+			return PinnedGenerator->GetCPUCoreUtilization();
+		}
+
+		return 0;
 	}
 
-	// Finally... send down the parameter pack.
-	PinnedGenerator->QueueParameterPack(CachedParameterPack);
-	return true;
-}
+	void FMetasoundGeneratorHandle::SetGenerator(TWeakPtr<FMetasoundGenerator>&& InGenerator)
+	{
+		METASOUND_LLM_SCOPE;
+		METASOUND_TRACE_CPUPROFILER_EVENT_SCOPE(FMetasoundGeneratorHandle::SetGenerator);
+		
+		check(IsInGameThread());
 
-TSharedPtr<Metasound::FMetasoundGenerator> UMetasoundGeneratorHandle::GetGenerator()
-{
-	// Attach if we aren't attached, check for changes, etc...
-	CacheMetasoundSource();
+		// early out if the incoming generator is null and the current generator is null
+		if (!Generator.IsValid() && !InGenerator.IsValid())
+		{
+			return;
+		}
+		
+		if (const TSharedPtr<FMetasoundGenerator> PinnedGenerator = Generator.Pin())
+		{
+			// skip the below logic if we are setting the same generator
+			if (InGenerator.HasSameObject(PinnedGenerator.Get()))
+			{
+				return;
+			}
 
-	return PinGenerator();
-}
+			// clean up if we had another generator
+			UnregisterGeneratorEvents();
+		}
 
-FDelegateHandle UMetasoundGeneratorHandle::AddGraphSetCallback(UMetasoundGeneratorHandle::FOnSetGraph::FDelegate&& Delegate)
-{
-	
-	FDelegateHandle Handle = OnGeneratorsGraphChanged.Add(Delegate);
-	CacheMetasoundSource();
-	AttachGraphChangedDelegate();
-	return Handle;
-}
+		// set the cached generator
+		Generator = MoveTemp(InGenerator);
+		
+		// Notify the generator has changed
+		if (OnGeneratorSet.IsBound())
+		{
+			OnGeneratorSet.Execute(TWeakPtr<FMetasoundGenerator>(Generator));
+		}
 
-bool UMetasoundGeneratorHandle::RemoveGraphSetCallback(const FDelegateHandle& Handle)
-{
-	return OnGeneratorsGraphChanged.Remove(Handle);
-}
+		// We're setting a new generator, so do the setup stuff
+		if (const TSharedPtr<FMetasoundGenerator> PinnedGenerator = Generator.Pin())
+		{
+			// Subscribe to generator events
+			RegisterGeneratorEvents();
+			
+			// Update params on the generator
+			SendParametersToGenerator();
 
-bool AnalyzerAddressesReferToSameGeneratorOutput(const Metasound::Frontend::FAnalyzerAddress& Lhs, const Metasound::Frontend::FAnalyzerAddress& Rhs)
-{
-	return Lhs.OutputName == Rhs.OutputName && Lhs.AnalyzerName == Rhs.AnalyzerName && Lhs.AnalyzerMemberName == Rhs.AnalyzerMemberName;
-}
+			// Attach any output watchers we might have
+			FixUpOutputWatchers();
 
-bool UMetasoundGeneratorHandle::WatchOutput(
-	FName OutputName,
-	const FOnMetasoundOutputValueChanged& OnOutputValueChanged,
-	FName AnalyzerName,
-	FName AnalyzerOutputName)
-{
-	METASOUND_LLM_SCOPE;
-	METASOUND_TRACE_CPUPROFILER_EVENT_SCOPE(UMetasoundGeneratorHandle::WatchOutput);
+			// Enable render timing if appropriate
+			PinnedGenerator->EnableRuntimeRenderTiming(bRuntimeRenderTimingShouldBeEnabled);
+		}
+	}
 
-	if (!IsValid())
+	void FMetasoundGeneratorHandle::RegisterGeneratorEvents()
 	{
-		return false;
-	}
+		METASOUND_LLM_SCOPE;
+		METASOUND_TRACE_CPUPROFILER_EVENT_SCOPE(FMetasoundGeneratorHandle::RegisterGeneratorEvents);
+		
+		check(IsInGameThread());
+		
+		if (const TSharedPtr<FMetasoundGenerator> PinnedGenerator = Generator.Pin())
+		{
+			// Output watchers
+			GeneratorOutputChangedDelegateHandle = PinnedGenerator->OnOutputChanged.AddSP(
+				AsShared(),
+				&FMetasoundGeneratorHandle::HandleOutputChanged);
 
-	CacheMetasoundSource();
+			// Graph updated
+			{
+				FOnSetGraph::FDelegate GraphSetDelegate;
+				GraphSetDelegate.BindSP(AsShared(), &FMetasoundGeneratorHandle::HandleGeneratorGraphSet);
+				GeneratorGraphSetDelegateHandle = PinnedGenerator->AddGraphSetCallback(MoveTemp(GraphSetDelegate));
+			}
 
-	if (!CachedMetasoundSource.IsValid())
-	{
-		return false;
+			// Vertex interface updated (Live Update support)
+			GeneratorVertexInterfaceChangedDelegateHandle = PinnedGenerator->OnVertexInterfaceDataUpdated.AddSP(
+				AsShared(),
+				&FMetasoundGeneratorHandle::HandleGeneratorVertexInterfaceChanged);
+		}
 	}
 
-	// Find the node id and type name
-	const Metasound::Frontend::FNodeHandle Node = CachedMetasoundSource->GetRootGraphHandle()->GetOutputNodeWithName(OutputName);
-
-	if (!Node->IsValid())
+	void FMetasoundGeneratorHandle::UnregisterGeneratorEvents() const
 	{
-		return false;
+		METASOUND_LLM_SCOPE;
+		METASOUND_TRACE_CPUPROFILER_EVENT_SCOPE(FMetasoundGeneratorHandle::UnregisterGeneratorEvents);
+		
+		check(IsInGameThread());
+		
+		if (const TSharedPtr<FMetasoundGenerator> PinnedGenerator = Generator.Pin())
+		{
+			PinnedGenerator->OnOutputChanged.Remove(GeneratorOutputChangedDelegateHandle);
+			PinnedGenerator->RemoveGraphSetCallback(GeneratorGraphSetDelegateHandle);
+			PinnedGenerator->OnVertexInterfaceDataUpdated.Remove(GeneratorVertexInterfaceChangedDelegateHandle);
+		}
 	}
 
-	const FGuid NodeId = Node->GetID();
-
-	// We expect an output node to have exactly one output
-	if (!ensure(Node->GetNumOutputs() == 1))
+	TWeakObjectPtr<UMetaSoundSource> FMetasoundGeneratorHandle::GetMetaSoundSource() const
 	{
-		return false;
+		METASOUND_LLM_SCOPE;
+		METASOUND_TRACE_CPUPROFILER_EVENT_SCOPE(FMetasoundGeneratorHandle::GetMetaSoundSource);
+
+		check(IsInGameThread()); // UAudioComponent::GetSound() isn't thread-safe.
+
+		if (!IsValid())
+		{
+			return nullptr;
+		}
+
+		return Cast<UMetaSoundSource>(AudioComponent->GetSound());
 	}
 
-	const FName TypeName = Node->GetOutputs()[0]->GetDataType();
+	void FMetasoundGeneratorHandle::SendParametersToGenerator() const
+	{
+		METASOUND_LLM_SCOPE;
+		METASOUND_TRACE_CPUPROFILER_EVENT_SCOPE(FMetasoundGeneratorHandle::SendParametersToGenerator);
 
-	// Make the analyzer address
-	Metasound::Frontend::FAnalyzerAddress AnalyzerAddress;
-	AnalyzerAddress.DataType = TypeName;
-	AnalyzerAddress.InstanceID = AudioComponent->GetAudioComponentID();
-	AnalyzerAddress.OutputName = OutputName;
+		if (!LatestParameterState.IsValid())
+		{
+			return;
+		}
+		
+		// If we have a generator, enqueue the updated parameter state
+		if (const TSharedPtr<FMetasoundGenerator> PinnedGenerator = Generator.Pin())
+		{
+			PinnedGenerator->QueueParameterPack(LatestParameterState);
+		}
+	}
 
-	// If no analyzer name was provided, try to find a passthrough analyzer
-	if (AnalyzerName.IsNone())
+	void FMetasoundGeneratorHandle::FixUpOutputWatchers()
 	{
-		if (!PassthroughAnalyzers.Contains(TypeName))
+		METASOUND_LLM_SCOPE;
+		METASOUND_TRACE_CPUPROFILER_EVENT_SCOPE(FMetasoundGeneratorHandle::FixUpOutputWatchers);
+
+		check(IsInGameThread());
+
+		if (!IsValid())
 		{
-			return false;
+			return;
 		}
 
-		AnalyzerName = PassthroughAnalyzers[TypeName].AnalyzerName;
-		AnalyzerOutputName = PassthroughAnalyzers[TypeName].OutputName;
+		if (const TSharedPtr<FMetasoundGenerator> PinnedGenerator = Generator.Pin())
+		{
+			// For each watcher, make sure the generator has a corresponding analyzer
+			// (will fail gracefully on duplicates or non-existent outputs)
+			for (const FOutputWatcher& Watcher : OutputWatchers)
+			{
+				PinnedGenerator->AddOutputVertexAnalyzer(Watcher.AnalyzerAddress);
+			}
+		}
 	}
-	AnalyzerAddress.AnalyzerName = AnalyzerName;
-	AnalyzerAddress.AnalyzerMemberName = AnalyzerOutputName;
 	
-	AnalyzerAddress.AnalyzerInstanceID = FGuid::NewGuid();
-	AnalyzerAddress.NodeID = NodeId;
+	void FMetasoundGeneratorHandle::CreateOutputWatcher(
+		const Frontend::FAnalyzerAddress& AnalyzerAddress,
+		const FOnMetasoundOutputValueChanged& OnOutputValueChanged)
+	{
+		METASOUND_LLM_SCOPE;
+		METASOUND_TRACE_CPUPROFILER_EVENT_SCOPE(FMetasoundGeneratorHandle::CreateOutputWatcher);
+
+		check(IsInGameThread()); // modifying watchers isn't thread-safe
+		
+		// If we already have a watcher for this output, just add the delegate to that one
+		if (FOutputWatcher* Watcher = OutputWatchers.FindByPredicate(
+			[&AnalyzerAddress](const FOutputWatcher& ExistingListener)
+			{
+				return AnalyzerAddress.OutputName == ExistingListener.AnalyzerAddress.OutputName
+					&& AnalyzerAddress.AnalyzerName == ExistingListener.AnalyzerAddress.AnalyzerName
+					&& AnalyzerAddress.AnalyzerMemberName == ExistingListener.AnalyzerAddress.AnalyzerMemberName;
+			}))
+		{
+			Watcher->OnOutputValueChanged.AddUnique(OnOutputValueChanged);
+		}
+		// Otherwise add a new watcher
+		else
+		{
+			OutputWatchers.Emplace(AnalyzerAddress, OnOutputValueChanged);
+		}
+	}
 
-	// Check to see if the analyzer exists
+	void FMetasoundGeneratorHandle::HandleGeneratorCreated(
+		const uint64 InAudioComponentId,
+		TSharedPtr<FMetasoundGenerator> InGenerator)
 	{
-		using namespace Metasound::Frontend;
-		const IVertexAnalyzerFactory* Factory =
-			IVertexAnalyzerRegistry::Get().FindAnalyzerFactory(AnalyzerAddress.AnalyzerName);
-		if (nullptr == Factory)
+		METASOUND_LLM_SCOPE;
+		METASOUND_TRACE_CPUPROFILER_EVENT_SCOPE(FMetasoundGeneratorHandle::HandleGeneratorCreated);
+		
+		check(InAudioComponentId != INDEX_NONE);
+					
+		if (InAudioComponentId == GetAudioComponentId())
 		{
-			return false;
+			// Set the generator on the game thread. We grab a weak pointer in case this gets destroyed while we wait.
+			AsyncTask(ENamedThreads::GameThread, [WeakThis = AsWeak(), WeakGenerator = InGenerator.ToWeakPtr()]()
+			{
+				if (const TSharedPtr<FMetasoundGeneratorHandle> PinnedThis = WeakThis.Pin())
+				{
+					PinnedThis->SetGenerator(TWeakPtr<FMetasoundGenerator>(WeakGenerator));
+				}
+			});
 		}
 	}
 
-	// if we already have a generator, go ahead and make the analyzer
-	if (const TSharedPtr<Metasound::FMetasoundGenerator> PinnedGenerator = PinGenerator())
+	void FMetasoundGeneratorHandle::HandleGeneratorDestroyed(
+		const uint64 InAudioComponentId,
+		TSharedPtr<FMetasoundGenerator>)
 	{
-		CreateAnalyzer(AnalyzerAddress, PinnedGenerator);
-	}
+		METASOUND_LLM_SCOPE;
+		METASOUND_TRACE_CPUPROFILER_EVENT_SCOPE(FMetasoundGeneratorHandle::HandleGeneratorDestroyed);
+		
+		check(InAudioComponentId != INDEX_NONE);
 
-	// Create the listener
-	CreateListener(AnalyzerAddress, OnOutputValueChanged);
+		if (InAudioComponentId == GetAudioComponentId())
+		{
+			// Unset the generator on the game thread. We grab a weak pointer in case this gets destroyed while we wait.
+			AsyncTask(ENamedThreads::GameThread, [WeakThis = AsWeak()]()
+			{
+				if (const TSharedPtr<FMetasoundGeneratorHandle> PinnedThis = WeakThis.Pin())
+				{
+					PinnedThis->SetGenerator(nullptr);
+				}
+			});
+		}
+	}
 
-	return true;
-}
+	void FMetasoundGeneratorHandle::HandleGeneratorGraphSet()
+	{
+		METASOUND_LLM_SCOPE;
+		METASOUND_TRACE_CPUPROFILER_EVENT_SCOPE(FMetasoundGeneratorHandle::HandleGeneratorGraphSet);
 
-void UMetasoundGeneratorHandle::EnableRuntimeRenderTiming(bool Enable)
-{
-	bRuntimeRenderTimingShouldBeEnabled = Enable;
+		// Defer to the game thread. We grab a weak pointer in case this gets destroyed while we wait.
+		AsyncTask(ENamedThreads::GameThread, [WeakThis = AsWeak()]()
+		{
+			if (const TSharedPtr<FMetasoundGeneratorHandle> PinnedThis = WeakThis.Pin())
+			{
+				PinnedThis->SendParametersToGenerator();
+				PinnedThis->FixUpOutputWatchers();
 
-	if (!IsValid())
-	{
-		return;
+				if (PinnedThis->OnGraphUpdated.IsBound())
+				{
+					PinnedThis->OnGraphUpdated.Execute();
+				}
+			}
+		});
 	}
 
-	if (!CachedMetasoundSource.IsValid())
+	void FMetasoundGeneratorHandle::HandleGeneratorVertexInterfaceChanged(FVertexInterfaceData)
 	{
-		CacheMetasoundSource();
+		METASOUND_LLM_SCOPE;
+		METASOUND_TRACE_CPUPROFILER_EVENT_SCOPE(FMetasoundGeneratorHandle::HandleGeneratorVertexInterfaceChanged);
+
+		// Defer to the game thread. We grab a weak pointer in case this gets destroyed while we wait.
+		AsyncTask(ENamedThreads::GameThread, [WeakThis = AsWeak()]()
+		{
+			if (const TSharedPtr<FMetasoundGeneratorHandle> PinnedThis = WeakThis.Pin())
+			{
+				PinnedThis->SendParametersToGenerator();
+				PinnedThis->FixUpOutputWatchers();
+			}
+		});
 	}
 
-	TSharedPtr<Metasound::FMetasoundGenerator> Generator = PinGenerator();
-	if (Generator)
+	void FMetasoundGeneratorHandle::HandleOutputChanged(
+		FName AnalyzerName,
+		FName OutputName,
+		FName AnalyzerOutputName,
+		TSharedPtr<IOutputStorage> OutputData)
 	{
-		Generator->EnableRuntimeRenderTiming(Enable);
+		METASOUND_LLM_SCOPE;
+		METASOUND_TRACE_CPUPROFILER_EVENT_SCOPE(FMetasoundGeneratorHandle::HandleOutputChanged);
+		
+		ChangedOutputs.Enqueue(AnalyzerName, OutputName, AnalyzerOutputName, OutputData);
 	}
 }
 
-double UMetasoundGeneratorHandle::GetCPUCoreUtilization()
+UMetasoundGeneratorHandle* UMetasoundGeneratorHandle::CreateMetaSoundGeneratorHandle(UAudioComponent* OnComponent)
 {
-	if (!IsValid())
+	METASOUND_LLM_SCOPE;
+	METASOUND_TRACE_CPUPROFILER_EVENT_SCOPE(UMetasoundGeneratorHandle::CreateMetaSoundGeneratorHandle);
+	
+	if (!OnComponent)
 	{
-		return 0.0;
+		return nullptr;
 	}
 
-	if (!CachedMetasoundSource.IsValid())
+	UMetasoundGeneratorHandle* Handle = NewObject<UMetasoundGeneratorHandle>();
+	return Handle->InitGeneratorHandle(OnComponent) ? Handle : nullptr;
+}
+
+void UMetasoundGeneratorHandle::BeginDestroy()
+{
+	Super::BeginDestroy();
+
+	GeneratorHandle.Reset();
+}
+
+bool UMetasoundGeneratorHandle::IsValid() const
+{
+	return GeneratorHandle.IsValid() && GeneratorHandle->IsValid();
+}
+
+uint64 UMetasoundGeneratorHandle::GetAudioComponentId() const
+{
+	return IsValid() ? GeneratorHandle->GetAudioComponentId() : INDEX_NONE;
+}
+
+bool UMetasoundGeneratorHandle::ApplyParameterPack(UMetasoundParameterPack* Pack)
+{
+	METASOUND_LLM_SCOPE;
+	METASOUND_TRACE_CPUPROFILER_EVENT_SCOPE(UMetasoundGeneratorHandle::ApplyParameterPack);
+	
+	if (nullptr == Pack)
 	{
-		CacheMetasoundSource();
+		return false;
 	}
 
-	TSharedPtr<Metasound::FMetasoundGenerator> Generator = PinGenerator();
-	if (!Generator)
+	if (IsValid())
 	{
-		return 0.0;
+		GeneratorHandle->UpdateParameters(*Pack);
+		return true;
 	}
-	return Generator->GetCPUCoreUtilization();
+
+	return false;
 }
 
-void UMetasoundGeneratorHandle::RegisterPassthroughAnalyzerForType(FName TypeName, FName AnalyzerName, FName OutputName)
+TSharedPtr<Metasound::FMetasoundGenerator> UMetasoundGeneratorHandle::GetGenerator() const
 {
-	check(!PassthroughAnalyzers.Contains(TypeName));
-	PassthroughAnalyzers.Add(TypeName, { AnalyzerName, OutputName });
+	return IsValid() ? GeneratorHandle->GetGenerator() : nullptr;
 }
 
-void UMetasoundGeneratorHandle::UpdateWatchers()
+FDelegateHandle UMetasoundGeneratorHandle::AddGraphSetCallback(FOnSetGraph::FDelegate&& Delegate)
 {
-	METASOUND_LLM_SCOPE;
-	METASOUND_TRACE_CPUPROFILER_EVENT_SCOPE(UMetasoundGeneratorHandle::UpdateWatchers);
+	return OnGeneratorsGraphChanged.Add(MoveTemp(Delegate));
+}
 
-	while (TOptional<FOutputPayload> ChangedOutput = ChangedOutputs->Dequeue())
-	{
-		if (const FOutputListener* OutputListener = OutputListeners.FindByPredicate(
-			[&ChangedOutput](const FOutputListener& ExistingListener)
-			{
-				return (*ChangedOutput).AnalyzerName == ExistingListener.AnalyzerAddress.AnalyzerName
-				&& (*ChangedOutput).OutputName == ExistingListener.AnalyzerAddress.OutputName
-				&& (*ChangedOutput).OutputValue.Name == ExistingListener.AnalyzerAddress.AnalyzerMemberName;
-			}))
-		{
-			OutputListener->OnOutputValueChanged.Broadcast((*ChangedOutput).OutputName, (*ChangedOutput).OutputValue);
-		}
-	}
+bool UMetasoundGeneratorHandle::RemoveGraphSetCallback(const FDelegateHandle& Handle)
+{
+	return OnGeneratorsGraphChanged.Remove(Handle);
 }
 
-void UMetasoundGeneratorHandle::OnSourceCreatedAGenerator(uint64 InAudioComponentId, TSharedPtr<Metasound::FMetasoundGenerator> InGenerator)
+bool UMetasoundGeneratorHandle::WatchOutput(
+	const FName OutputName,
+	const FOnMetasoundOutputValueChanged& OnOutputValueChanged,
+	const FName AnalyzerName,
+	const FName AnalyzerOutputName)
 {
-	check(IsInGameThread());
+	METASOUND_LLM_SCOPE;
+	METASOUND_TRACE_CPUPROFILER_EVENT_SCOPE(UMetasoundGeneratorHandle::WatchOutput);
 
-	if (!AudioComponent.IsValid())
+	if (!IsValid())
 	{
-		return;
+		return false;
 	}
-	
-	if (InAudioComponentId == AudioComponent->GetAudioComponentID())
-	{
-		CachedGeneratorPtr = InGenerator;
-		if (InGenerator)
-		{
-			// If there is a parameter pack to apply, apply it
-			if (CachedParameterPack)
-			{
-				InGenerator->QueueParameterPack(CachedParameterPack);
-			}
 
-			// If we have listeners, add the analyzers for them
-			{
-				for (const FOutputListener& Listener : OutputListeners)
-				{
-					CreateAnalyzer(Listener.AnalyzerAddress, InGenerator);
-				}
-			}
-			
-			InGenerator->EnableRuntimeRenderTiming(bRuntimeRenderTimingShouldBeEnabled);
+	return GeneratorHandle->WatchOutput(OutputName, OnOutputValueChanged, AnalyzerName, AnalyzerOutputName);
+}
 
-			OnGeneratorHandleAttached.Broadcast();
-			// If anyone has told us they are interested in being notified when a generator's 
-			// graph has changed go ahead and set that up now...
-			AttachGraphChangedDelegate();
-		}
-	}
+void UMetasoundGeneratorHandle::RegisterPassthroughAnalyzerForType(
+	const FName TypeName,
+	const FName AnalyzerName,
+	const FName OutputName)
+{
+	Metasound::FMetasoundGeneratorHandle::RegisterPassthroughAnalyzerForType(TypeName, AnalyzerName, OutputName);
 }
 
-void UMetasoundGeneratorHandle::OnSourceDestroyedAGenerator(uint64 InAudioComponentId)
+void UMetasoundGeneratorHandle::UpdateWatchers() const
 {
-	check(IsInGameThread());
+	METASOUND_LLM_SCOPE;
+	METASOUND_TRACE_CPUPROFILER_EVENT_SCOPE(UMetasoundGeneratorHandle::UpdateWatchers);
 
-	if (!AudioComponent.IsValid())
-	{
-		return;
-	}
-	
-	if (InAudioComponentId == AudioComponent->GetAudioComponentID())
+	if (IsValid())
 	{
-		if (CachedGeneratorPtr.IsValid())
-		{
-			OnGeneratorHandleDetached.Broadcast();
-		}
-		CachedGeneratorPtr = nullptr;
-
-		OutputChangedDelegateHandle.Reset();
+		GeneratorHandle->UpdateOutputWatchers();
 	}
 }
 
-void UMetasoundGeneratorHandle::CreateAnalyzer(
-	const Metasound::Frontend::FAnalyzerAddress& AnalyzerAddress,
-	const TSharedPtr<Metasound::FMetasoundGenerator> Generator)
+void UMetasoundGeneratorHandle::EnableRuntimeRenderTiming(const bool Enable) const
 {
-	if (!IsValid())
+	if (IsValid())
 	{
-		return;
+		GeneratorHandle->EnableRuntimeRenderTiming(Enable);
 	}
+}
 
-	// Create the analyzer (will skip if there's already one)
-	Generator->AddOutputVertexAnalyzer(AnalyzerAddress);
-
-	// Subscribe for output updates if we haven't already
-	if (!OutputChangedDelegateHandle.IsValid())
+double UMetasoundGeneratorHandle::GetCPUCoreUtilization() const
+{
+	if (IsValid())
 	{
-		TWeakPtr<TSpscQueue<FOutputPayload>> ChangedOutputQueue{ ChangedOutputs };
-		OutputChangedDelegateHandle = Generator->OnOutputChanged.AddLambda([ChangedOutputQueue](
-			FName AnalyzerName,
-			FName OutputName,
-			FName AnalyzerOutputName,
-			TSharedPtr<Metasound::IOutputStorage> OutputData)
-		{
-			if (const TSharedPtr<TSpscQueue<FOutputPayload>> PinnedQueue = ChangedOutputQueue.Pin())
-			{
-				PinnedQueue->Enqueue(AnalyzerName, OutputName, AnalyzerOutputName, OutputData);
-			}
-		});
+		return GeneratorHandle->GetCPUCoreUtilization();
 	}
+
+	return 0;
 }
 
-void UMetasoundGeneratorHandle::CreateListener(
-	const Metasound::Frontend::FAnalyzerAddress& AnalyzerAddress,
-	const FOnMetasoundOutputValueChanged& OnOutputValueChanged)
+bool UMetasoundGeneratorHandle::InitGeneratorHandle(TWeakObjectPtr<UAudioComponent>&& AudioComponent)
 {
-	if (!IsValid())
+	GeneratorHandle = Metasound::FMetasoundGeneratorHandle::Create(MoveTemp(AudioComponent));
+
+	if (!GeneratorHandle.IsValid())
 	{
-		return;
+		return false;
 	}
-	
-	// If we already have a listener for this output, just add the delegate to that one
-	if (FOutputListener* Listener = OutputListeners.FindByPredicate(
-		[&AnalyzerAddress](const FOutputListener& ExistingListener)
-		{
-			return AnalyzerAddressesReferToSameGeneratorOutput(AnalyzerAddress, ExistingListener.AnalyzerAddress);
-		}))
+
+	// Attach delegates
+	// NB: FMetasoundGeneratorHandle already executes these on the game thread,
+	// and its lifetime is tied to UMetasoundGeneratorHandle's lifetime,
+	// so we can guarantee the this pointer is valid when these get called.
+	GeneratorHandle->OnGeneratorSet.BindLambda([this](TWeakPtr<Metasound::FMetasoundGenerator>&& Generator)
 	{
-		Listener->OnOutputValueChanged.AddUnique(OnOutputValueChanged);
-	}
-	// Otherwise add a new listener
-	else
+		if (Generator.IsValid())
+		{
+			OnGeneratorHandleAttached.Broadcast();
+		}
+		else
+		{
+			OnGeneratorHandleDetached.Broadcast();
+		}
+	});
+
+	GeneratorHandle->OnGraphUpdated.BindLambda([this]()
 	{
-		OutputListeners.Emplace(AnalyzerAddress, OnOutputValueChanged);
-	}
+		OnGeneratorsGraphChanged.Broadcast();
+	});
+
+	return true;
 }
diff --git a/Engine/Plugins/Runtime/Metasound/Source/MetasoundEngine/Private/MetasoundOutputSubsystem.cpp b/Engine/Plugins/Runtime/Metasound/Source/MetasoundEngine/Private/MetasoundOutputSubsystem.cpp
--- a/Engine/Plugins/Runtime/Metasound/Source/MetasoundEngine/Private/MetasoundOutputSubsystem.cpp
+++ b/Engine/Plugins/Runtime/Metasound/Source/MetasoundEngine/Private/MetasoundOutputSubsystem.cpp
@@ -39,7 +39,11 @@ bool UMetaSoundOutputSubsystem::WatchOutput(
 	else
 	{
 		Handle = UMetasoundGeneratorHandle::CreateMetaSoundGeneratorHandle(AudioComponent);
-		TrackedGenerators.Add(Handle);
+		
+		if(nullptr != Handle && Handle->IsValid())
+		{
+			TrackedGenerators.Add(Handle);
+		}
 	}
 
 	if (nullptr == Handle)
diff --git a/Engine/Plugins/Runtime/Metasound/Source/MetasoundEngine/Public/MetasoundGeneratorHandle.h b/Engine/Plugins/Runtime/Metasound/Source/MetasoundEngine/Public/MetasoundGeneratorHandle.h
--- a/Engine/Plugins/Runtime/Metasound/Source/MetasoundEngine/Public/MetasoundGeneratorHandle.h
+++ b/Engine/Plugins/Runtime/Metasound/Source/MetasoundEngine/Public/MetasoundGeneratorHandle.h
@@ -12,7 +12,7 @@
 
 class UAudioComponent;
 class UMetaSoundSource;
-class UMetasoundParameterPack;
+
 namespace Metasound
 {
 	class FMetasoundGenerator;
@@ -20,13 +20,224 @@ namespace Metasound
 
 DECLARE_DYNAMIC_DELEGATE_TwoParams(FOnMetasoundOutputValueChanged, FName, OutputName, const FMetaSoundOutput&, Output);
 
+
+DECLARE_DYNAMIC_MULTICAST_DELEGATE_TwoParams(FOnMetasoundOutputValueChangedMulticast, FName, Name, const FMetaSoundOutput&, Output);
+
+namespace Metasound
+{
+	/**
+	 * User-facing interface to a FMetasoundGenerator on a UAudioComponent.
+	 */
+	class METASOUNDENGINE_API FMetasoundGeneratorHandle : public TSharedFromThis<FMetasoundGeneratorHandle>
+	{
+		struct FPrivateToken { explicit FPrivateToken() = default; };
+		
+	public:
+		/**
+		 * NOTE: This constructor is effectively private, but will work with TSharedPtr/TSharedFromThis.
+		 * (see TIntrusiveReferenceController for explanation)
+		 * Use FMetasoundGeneratorHandle::Create instead
+		 */
+		explicit FMetasoundGeneratorHandle(FPrivateToken, TWeakObjectPtr<UAudioComponent>&& InAudioComponent);
+
+		/**
+		 * Create a handle to a MetaSound generator on the given audio component
+		 * 
+		 * @param InAudioComponent The audio component to attach to
+		 * @return The generator handle, if successfully created
+		 */
+		static TSharedPtr<FMetasoundGeneratorHandle> Create(TWeakObjectPtr<UAudioComponent>&& InAudioComponent);
+		
+		~FMetasoundGeneratorHandle();
+
+		/**
+		 * Find out if this handle is still valid.
+		 * 
+		 * @return Whether this handle is still valid
+		 */
+		bool IsValid() const;
+
+		/**
+		 * Get the id for the UAudioComponent associated with this handle.
+		 * NOTE: Be sure to check IsValid() before expecting a valid return from this method.
+		 *
+		 * @returns The audio component's id, or INDEX_NONE if the component is no longer valid.
+		 */
+		uint64 GetAudioComponentId() const;
+
+		TSharedPtr<FMetasoundGenerator> GetGenerator() const;
+	
+		DECLARE_DELEGATE_OneParam(FOnGeneratorSet, TWeakPtr<FMetasoundGenerator>&&);
+
+		/**
+		 * Fires on the game thread when a generator is set or unset on this handle.
+		 */
+		FOnGeneratorSet OnGeneratorSet;
+	
+		DECLARE_DELEGATE(FOnGeneratorGraphSet);
+
+		/**
+		 * Fires on the game thread when a graph is updated on the generator.
+		 */
+		FOnGeneratorGraphSet OnGraphUpdated;
+	
+		DECLARE_DELEGATE_OneParam(FOnGeneratorVertexInterfaceUpdated, const FVertexInterfaceData&);
+
+		/**
+		 * Update the current parameter state on this handle and enqueue the changes on the generator.
+		 * 
+		 * @param ParameterPack The parameter pack to set, which will be merged with the current state.
+		 * @return true if parameters were updated, false otherwise
+		 */
+		void UpdateParameters(const UMetasoundParameterPack& ParameterPack);
+		
+		/**
+		 * Watch an output value.
+		 *
+		 * @param OutputName - The user-specified name of the output in the Metasound
+		 * @param OnOutputValueChanged - The event to fire when the output's value changes
+		 * @param AnalyzerName - (optional) The name of the analyzer to use on the output, defaults to a passthrough
+		 * @param AnalyzerOutputName - (optional) The name of the output on the analyzer to watch, defaults to the passthrough output
+		 * @returns true if the watch setup succeeded, false otherwise
+		 */
+		bool WatchOutput(
+			FName OutputName,
+			const FOnMetasoundOutputValueChanged& OnOutputValueChanged,
+			FName AnalyzerName = NAME_None,
+			FName AnalyzerOutputName = NAME_None);
+
+		/**
+		 * Update any watched outputs
+		 */
+		void UpdateOutputWatchers();
+
+		/**
+		 * Map a type name to a passthrough analyzer name to use as a default for UMetasoundOutputSubsystem::WatchOutput()
+		 *
+		 * @param TypeName - The type name returned from GetMetasoundDataTypeName()
+		 * @param AnalyzerName - The name of the analyzer to use
+		 * @param OutputName - The name of the output in the analyzer
+		 */
+		static void RegisterPassthroughAnalyzerForType(FName TypeName, FName AnalyzerName, FName OutputName);
+
+		/**
+		* Enable the profiling of the MetaSound render for this playing instance. You
+		* must call this before calling "GetRuntimeCPUCoreUtilization" (below) or you will just 
+		* get 0.0 back for core utilization.
+		*/
+		void EnableRuntimeRenderTiming(bool Enable);
+
+		/**
+		* Get the CPU usage as "fraction of real time" used to render this metasound. 
+		* NOTE: The MetasoundSource asset MUST have had its EnableRenderTiming function called
+		* before the metasound is started!
+		*/
+		double GetCPUCoreUtilization() const;
+
+	private:
+		void SetGenerator(TWeakPtr<FMetasoundGenerator>&& InGenerator);
+
+		void RegisterGeneratorEvents();
+
+		void UnregisterGeneratorEvents() const;
+
+		TWeakObjectPtr<UMetaSoundSource> GetMetaSoundSource() const;
+
+		void SendParametersToGenerator() const;
+		
+		void FixUpOutputWatchers();
+
+		void CreateOutputWatcher(
+			const Frontend::FAnalyzerAddress& AnalyzerAddress,
+			const FOnMetasoundOutputValueChanged& OnOutputValueChanged);
+		
+		TWeakObjectPtr<UAudioComponent> AudioComponent;
+		const uint64 AudioComponentId;
+
+		/**
+		 * Delegate management
+		 */
+		void HandleGeneratorCreated(uint64 InAudioComponentId, TSharedPtr<FMetasoundGenerator> Generator);
+		FDelegateHandle GeneratorCreatedDelegateHandle;
+
+		void HandleGeneratorDestroyed(uint64 InAudioComponentId, TSharedPtr<FMetasoundGenerator> Generator);
+		FDelegateHandle GeneratorDestroyedDelegateHandle;
+
+		void HandleGeneratorGraphSet();
+		FDelegateHandle GeneratorGraphSetDelegateHandle;
+
+		void HandleGeneratorVertexInterfaceChanged(FVertexInterfaceData VertexInterfaceData);
+		FDelegateHandle GeneratorVertexInterfaceChangedDelegateHandle;
+
+		void HandleOutputChanged(
+			FName AnalyzerName,
+			FName OutputName,
+			FName AnalyzerOutputName,
+			TSharedPtr<IOutputStorage> OutputData);
+		FDelegateHandle GeneratorOutputChangedDelegateHandle;
+		
+		TWeakPtr<FMetasoundGenerator> Generator;
+	
+		FSharedMetasoundParameterStoragePtr LatestParameterState;
+
+		struct FPassthroughAnalyzerInfo
+		{
+			FName AnalyzerName;
+			FName OutputName;
+		};
+		static TMap<FName, FPassthroughAnalyzerInfo> PassthroughAnalyzers;
+
+		/**
+		 * Info about an output being watched by one or more listeners
+		 */
+		struct FOutputWatcher
+		{
+			Frontend::FAnalyzerAddress AnalyzerAddress;
+			FOnMetasoundOutputValueChangedMulticast OnOutputValueChanged;
+
+			FOutputWatcher(
+				const Frontend::FAnalyzerAddress& InAnalyzerAddress,
+				const FOnMetasoundOutputValueChanged& InOnOutputValueChanged)
+					: AnalyzerAddress(InAnalyzerAddress)
+			{
+				OnOutputValueChanged.AddUnique(InOnOutputValueChanged);
+			}
+		};
+
+		TArray<FOutputWatcher> OutputWatchers;
+
+		struct FOutputPayload
+		{
+			FName AnalyzerName;
+			FName OutputName;
+			FMetaSoundOutput OutputValue;
+
+			FOutputPayload(
+				const FName InAnalyzerName,
+				const FName InOutputName,
+				const FName AnalyzerOutputName,
+				TSharedPtr<IOutputStorage> OutputData)
+					: AnalyzerName(InAnalyzerName)
+					, OutputName(InOutputName)
+					, OutputValue(AnalyzerOutputName, OutputData)
+			{}
+		};
+	
+		TSpscQueue<FOutputPayload> ChangedOutputs;
+
+		bool bRuntimeRenderTimingShouldBeEnabled{ false };
+	};
+}
+
+/**
+ * Blueprint-facing interface to a FMetasoundGenerator on a UAudioComponent.
+ */
 UCLASS(BlueprintType,Category="MetaSound")
 class METASOUNDENGINE_API UMetasoundGeneratorHandle : public UObject
 {
 	GENERATED_BODY()
 
 public:
-
 	UFUNCTION(BlueprintCallable, Category="MetaSound")
 	static UMetasoundGeneratorHandle* CreateMetaSoundGeneratorHandle(UAudioComponent* OnComponent);
 
@@ -42,17 +253,6 @@ public:
 	 */
 	uint64 GetAudioComponentId() const;
 
-	// UMetasoundGeneratorHandle shields its "clients" from "cross thread" issues
-	// related to callbacks coming in the audio control or rendering threads that 
-	// game thread clients (e.g. blueprints) want to know about. That is why these 
-	// next delegate definitions do not need to be the "TS" variants. Assignments 
-	// to members of this type, and the broadcasts there to will all happen on the
-	// game thread. EVEN IF the instigator of those callbacks is on the audio
-	// render thread. 
-	DECLARE_MULTICAST_DELEGATE(FOnAttached);
-	DECLARE_MULTICAST_DELEGATE(FOnDetached);
-	DECLARE_MULTICAST_DELEGATE(FOnSetGraph);
-
 	/**
 	 * Makes a copy of the supplied parameter pack and passes it to the MetaSoundGenerator
 	 * for asynchronous processing. IT ALSO caches this copy so that if the AudioComponent
@@ -62,13 +262,29 @@ public:
 	UFUNCTION(BlueprintCallable, Category="MetaSoundParameterPack")
 	bool ApplyParameterPack(UMetasoundParameterPack* Pack);
 
-	TSharedPtr<Metasound::FMetasoundGenerator> GetGenerator();
+	/**
+	 * Get a shared pointer to the generator, if available.
+	 * NOTE: while the shared pointer is safe to get from threads other than the audio render thread,
+	 * not all methods on the generator are safe to use from other threads. Ensure you know what you're doing.
+	 * 
+	 * @return A shared pointer to the generator, if available.
+	 */
+	TSharedPtr<Metasound::FMetasoundGenerator> GetGenerator() const;
+
+	// UMetasoundGeneratorHandle shields its "clients" from "cross thread" issues
+	// related to callbacks coming in the audio control or rendering threads that 
+	// game thread clients (e.g. blueprints) want to know about. That is why these 
+	// next delegate definitions do not need to be the "TS" variants. Assignments 
+	// to members of this type, and the broadcasts there to will all happen on the
+	// game thread. EVEN IF the instigator of those callbacks is on the audio
+	// render thread. 
+	DECLARE_MULTICAST_DELEGATE(FOnAttached);
+	FOnAttached OnGeneratorHandleAttached;
 
-	UMetasoundGeneratorHandle::FOnAttached OnGeneratorHandleAttached;
-	UMetasoundGeneratorHandle::FOnDetached OnGeneratorHandleDetached;
-	// Note: We don't allow direct assignment to the OnGeneratorsGraphChanged delegate
-	// because we need to know that someone actually wants this message so we can
-	// start actively listening for the corresponding audio render thread callback...
+	DECLARE_MULTICAST_DELEGATE(FOnDetached);
+	FOnDetached OnGeneratorHandleDetached;
+	
+	DECLARE_MULTICAST_DELEGATE(FOnSetGraph);
 	FDelegateHandle AddGraphSetCallback(UMetasoundGeneratorHandle::FOnSetGraph::FDelegate&& Delegate);
 	bool RemoveGraphSetCallback(const FDelegateHandle& Handle);
 
@@ -88,22 +304,6 @@ public:
 		FName AnalyzerName = NAME_None,
 		FName AnalyzerOutputName = NAME_None);
 
-	/**
-	* Enable the profiling of the MetaSound render for this playing instance. You
-	* must call this before calling "GetRuntimeCPUCoreUtilization" (below) or you will just 
-	* get 0.0 back for core utilization.
-	*/
-	UFUNCTION(BlueprintCallable, Category="MetaSoundPerf", meta=(AdvancedDisplay = "2"))
-	void EnableRuntimeRenderTiming(bool Enable);
-
-	/**
-	* Get the CPU usage as "fraction of real time" used to render this metasound. 
-	* NOTE: The MetasoundSource asset MUST have had its EnableRenderTiming function called
-	* before the metasound is started!
-	*/
-	UFUNCTION(BlueprintCallable, Category="MetaSoundPerf", meta=(AdvancedDisplay = "2"))
-	double GetCPUCoreUtilization();
-
 	/**
 	 * Map a type name to a passthrough analyzer name to use as a default for UMetasoundOutputSubsystem::WatchOutput()
 	 *
@@ -116,107 +316,30 @@ public:
 	/**
 	 * Update any watched outputs
 	 */
-	void UpdateWatchers();
-	
-private:
-	void SetAudioComponent(UAudioComponent* InAudioComponent);
-	void CacheMetasoundSource();
-	void ClearCachedData();
-
-	/**
-	 * Attempts to pin the weak generator pointer. If the first attempt fails it checks to see
-	 * if it can "recapture" a pointer to a generator for the current AudioComponent/MetaSoundSource
-	 * combination. 
-	 */
-	TSharedPtr<Metasound::FMetasoundGenerator> PinGenerator();
+	UFUNCTION(BlueprintCallable, Category="MetaSoundOutput")
+	void UpdateWatchers() const;
 
 	/**
-	 * Functions for adding and removing our MetaSoundGenerator lifecycle delegates
-	 */
-	void AttachGeneratorDelegates();
-	void AttachGraphChangedDelegate();
-	void DetachGeneratorDelegates();
+	* Enable the profiling of the MetaSound render for this playing instance. You
+	* must call this before calling "GetRuntimeCPUCoreUtilization" (below) or you will just 
+	* get 0.0 back for core utilization.
+	*/
+	UFUNCTION(BlueprintCallable, Category="MetaSoundPerf")
+	void EnableRuntimeRenderTiming(bool Enable) const;
 
 	/**
-	 * Generator creation and destruction delegates we register with the UMetaSoundSource
-	 */
-	void OnSourceCreatedAGenerator(uint64 InAudioComponentId, TSharedPtr<Metasound::FMetasoundGenerator> InGenerator);
-	void OnSourceDestroyedAGenerator(uint64 InAudioComponentId);
+	* Get the CPU usage as "fraction of real time" used to render this metasound. 
+	* NOTE: The MetasoundSource asset MUST have had its EnableRenderTiming function called
+	* before the metasound is started!
+	*/
+	UFUNCTION(BlueprintCallable, Category="MetaSoundPerf")
+	double GetCPUCoreUtilization() const;
+	
+private:
+	bool InitGeneratorHandle(TWeakObjectPtr<UAudioComponent>&& AudioComponent);
+	
+	TSharedPtr<Metasound::FMetasoundGeneratorHandle> GeneratorHandle;
 	
-	TWeakObjectPtr<UAudioComponent> AudioComponent;
-	TWeakObjectPtr<UMetaSoundSource> CachedMetasoundSource;
-	TWeakPtr<Metasound::FMetasoundGenerator> CachedGeneratorPtr;
-	FSharedMetasoundParameterStoragePtr CachedParameterPack;
-
-	// Note: We don't allow direct assignment to the OnGeneratorsGraphChanged delegate
-	// because we need to know that someone actually wants this message so we can
-	// start actively listening for the corresponding audio render thread callback. 
-	// So these next members are private and a "client" that wants to be notified of
-	// the graph change have to call public functions declared above to add themselves.
 	DECLARE_MULTICAST_DELEGATE(FOnSetGraphMulticast);
 	FOnSetGraphMulticast OnGeneratorsGraphChanged;
-
-	FDelegateHandle GeneratorCreatedDelegateHandle;
-	FDelegateHandle GeneratorDestroyedDelegateHandle;
-	FDelegateHandle GeneratorGraphChangedDelegateHandle;
-
-	void CreateAnalyzer(
-		const Metasound::Frontend::FAnalyzerAddress& AnalyzerAddress,
-		const TSharedPtr<Metasound::FMetasoundGenerator> Generator);
-	void CreateListener(
-		const Metasound::Frontend::FAnalyzerAddress& AnalyzerAddress,
-		const FOnMetasoundOutputValueChanged& OnOutputValueChanged);
-
-	struct FPassthroughAnalyzerInfo
-	{
-		FName AnalyzerName;
-		FName OutputName;
-	};
-	static TMap<FName, FPassthroughAnalyzerInfo> PassthroughAnalyzers;
-
-	/**
-	 * Multicast delegate to broadcast to users calling WatchOutput
-	 */
-	DECLARE_DYNAMIC_MULTICAST_DELEGATE_TwoParams(FOnOutputValueChangedMulticast, FName, Name, const FMetaSoundOutput&, Output);
-
-	/**
-	 * Info about an output being watched by one or more listeners
-	 */
-	struct FOutputListener
-	{
-		Metasound::Frontend::FAnalyzerAddress AnalyzerAddress;
-		FOnOutputValueChangedMulticast OnOutputValueChanged;
-
-		FOutputListener(
-			const Metasound::Frontend::FAnalyzerAddress& InAnalyzerAddress,
-			const FOnMetasoundOutputValueChanged& InOnOutputValueChanged)
-				: AnalyzerAddress(InAnalyzerAddress)
-		{
-			OnOutputValueChanged.AddUnique(InOnOutputValueChanged);
-		}
-	};
-
-	TArray<FOutputListener> OutputListeners;
-
-	struct FOutputPayload
-	{
-		FName AnalyzerName;
-		FName OutputName;
-		FMetaSoundOutput OutputValue;
-
-		FOutputPayload(
-			const FName InAnalyzerName,
-			const FName InOutputName,
-			const FName AnalyzerOutputName,
-			TSharedPtr<Metasound::IOutputStorage> OutputData)
-				: AnalyzerName(InAnalyzerName)
-				, OutputName(InOutputName)
-				, OutputValue(AnalyzerOutputName, OutputData)
-		{}
-	};
-	
-	TSharedRef<TSpscQueue<FOutputPayload>> ChangedOutputs = MakeShared<TSpscQueue<FOutputPayload>>();
-	FDelegateHandle OutputChangedDelegateHandle;
-
-	bool bRuntimeRenderTimingShouldBeEnabled{ false };
 };
