Update caching of input data to use TSortedMap

diff --git a/Engine/Plugins/Runtime/Metasound/Source/MetasoundEngine/Private/MetasoundSource.cpp b/Engine/Plugins/Runtime/Metasound/Source/MetasoundEngine/Private/MetasoundSource.cpp
--- a/Engine/Plugins/Runtime/Metasound/Source/MetasoundEngine/Private/MetasoundSource.cpp
+++ b/Engine/Plugins/Runtime/Metasound/Source/MetasoundEngine/Private/MetasoundSource.cpp
@@ -346,6 +346,8 @@ void UMetaSoundSource::PostLoad()
 void UMetaSoundSource::InitParameters(TArray<FAudioParameter>& ParametersToInit, FName InFeatureName)
 {
 	using namespace Metasound::SourcePrivate;
+	using namespace Metasound::Frontend;
+	using FRuntimeInput = FMetasoundAssetBase::FRuntimeInput;
 
 	METASOUND_LLM_SCOPE;
 	METASOUND_TRACE_CPUPROFILER_EVENT_SCOPE(UMetaSoundSource::InitParameters);
@@ -355,31 +357,12 @@ void UMetaSoundSource::InitParameters(TArray<FAudioParameter>& ParametersToInit,
 	RegisterGraphWithFrontend(GetInitRegistrationOptions());
 	const FRuntimeData& RuntimeData = GetRuntimeData();
 
-	TMap<FName, const FMetasoundFrontendVertex*> PublicInputMap;
-	{
-		const TArray<FMetasoundFrontendClassInput>& PublicInputs = RuntimeData.PublicInputs;
-		Algo::Transform(PublicInputs, PublicInputMap, [](const FMetasoundFrontendClassInput& Input)
-		{
-			return TPair<FName, const FMetasoundFrontendVertex*>(Input.Name, &Input);
-		});
-	}
+	IDataTypeRegistry& DataTypeRegistry = IDataTypeRegistry::Get();
+	const Metasound::TSortedVertexNameMap<FRuntimeInput>& PublicInputMap = RuntimeData.PublicInputMap;
 
-	// Removes values that are not explicitly defined by the ParamType and returns
-	// whether or not the parameter is a valid input and should be included.
-	auto Sanitize = [&PublicInputMap](FAudioParameter& Parameter) -> bool
+	// Removes values that are not explicitly defined by the ParamType
+	auto Sanitize = [&PublicInputMap](FAudioParameter& Parameter)
 	{
-		const FMetasoundFrontendVertex* Input = PublicInputMap.FindRef(Parameter.ParamName);
-		if (!Input)
-		{
-			return false;
-		}
-
-		const bool bIsMatchingType = Parameter.TypeName.IsNone() || (Parameter.TypeName == Input->TypeName);
-		if (!bIsMatchingType)
-		{
-			return false;
-		}
-
 		switch (Parameter.ParamType)
 		{
 			case EAudioParameterType::Trigger:
@@ -458,20 +441,17 @@ void UMetaSoundSource::InitParameters(TArray<FAudioParameter>& ParametersToInit,
 			default:
 			break;
 		}
-
-		return true;
 	};
 
-	auto ConstructProxies = [this](FAudioParameter& OutParamToInit, FName VertexTypeName)
+	auto ConstructProxies = [&DataTypeRegistry](FAudioParameter& OutParamToInit, FName VertexTypeName)
 	{
 		using namespace Metasound;
-		using namespace Metasound::Frontend;
 
 		switch (OutParamToInit.ParamType)
 		{
 			case EAudioParameterType::Object:
 			{
-				TSharedPtr<Audio::IProxyData> ProxyPtr = IDataTypeRegistry::Get().CreateProxyFromUObject(VertexTypeName, OutParamToInit.ObjectParam);
+				TSharedPtr<Audio::IProxyData> ProxyPtr = DataTypeRegistry.CreateProxyFromUObject(VertexTypeName, OutParamToInit.ObjectParam);
 				OutParamToInit.ObjectProxies.Emplace(MoveTemp(ProxyPtr));
 
 				// Null out param as it is no longer needed (nor desired to be accessed once passed to the Audio Thread)
@@ -484,7 +464,7 @@ void UMetaSoundSource::InitParameters(TArray<FAudioParameter>& ParametersToInit,
 				const FName ElementTypeName = CreateElementTypeNameFromArrayTypeName(VertexTypeName);
 				for (TObjectPtr<UObject>& Object : OutParamToInit.ArrayObjectParam)
 				{
-					TSharedPtr<Audio::IProxyData> ProxyPtr = IDataTypeRegistry::Get().CreateProxyFromUObject(ElementTypeName, Object);
+					TSharedPtr<Audio::IProxyData> ProxyPtr = DataTypeRegistry.CreateProxyFromUObject(ElementTypeName, Object);
 					OutParamToInit.ObjectProxies.Emplace(MoveTemp(ProxyPtr));
 				}
 				// Reset param array as it is no longer needed (nor desired to be accessed once passed to the Audio Thread).
@@ -498,35 +478,23 @@ void UMetaSoundSource::InitParameters(TArray<FAudioParameter>& ParametersToInit,
 		}
 	};
 
+
 	for (int32 i = ParametersToInit.Num() - 1; i >= 0; --i)
 	{
+		bool bIsParameterValid = false;
+
 		FAudioParameter& Parameter = ParametersToInit[i];
-		
-#if !NO_LOGGING
-		// For logging in case of failure
-		const FString AssetName = GetName();
-#endif // !NO_LOGGING
-		
-		if (Sanitize(Parameter))
+		if (const FRuntimeInput* Input = PublicInputMap.Find(Parameter.ParamName))
 		{
-			const FMetasoundFrontendVertex* Vertex = PublicInputMap.FindRef(Parameter.ParamName);
-			if (IsParameterValid(Parameter, Vertex))
+			if (IsParameterValid(Parameter, Input->TypeName, DataTypeRegistry))
 			{
-				ConstructProxies(Parameter, Vertex->TypeName);
-			}
-			else
-			{
-#if !NO_LOGGING
-				if (::Metasound::MetaSoundParameterEnableWarningOnIgnoredParameterCVar)
-				{
-					UE_LOG(LogMetaSound, Warning, TEXT("Failed to set invalid parameter '%s' in asset '%s': Either does not exist or is unsupported type"), *Parameter.ParamName.ToString(), *AssetName);
-				}
-#endif // !NO_LOGGING
-				constexpr bool bAllowShrinking = false;
-				ParametersToInit.RemoveAtSwap(i, 1, bAllowShrinking);
+				Sanitize(Parameter);
+				ConstructProxies(Parameter, Input->TypeName);
+				bIsParameterValid = true;
 			}
 		}
-		else
+
+		if (!bIsParameterValid)
 		{
 			constexpr bool bAllowShrinking = false;
 			ParametersToInit.RemoveAtSwap(i, 1, bAllowShrinking);
@@ -534,6 +502,7 @@ void UMetaSoundSource::InitParameters(TArray<FAudioParameter>& ParametersToInit,
 #if !NO_LOGGING
 			if (::Metasound::MetaSoundParameterEnableWarningOnIgnoredParameterCVar)
 			{
+				const FString AssetName = GetName();
 				UE_LOG(LogMetaSound, Warning, TEXT("Failed to set parameter '%s' in asset '%s': No name specified, no transmittable input found, or type mismatch."), *Parameter.ParamName.ToString(), *AssetName);
 			}
 #endif // !NO_LOGGING
@@ -685,10 +654,9 @@ bool UMetaSoundSource::GetAllDefaultParameters(TArray<FAudioParameter>& OutParam
 	using namespace Metasound::Frontend;
 	using namespace Metasound::Engine;
 
-	// TODO: Make this use the cached runtime data's input copy. This call can become expensive if called repeatedly.
-	TArray<FMetasoundFrontendClassInput> PublicInputs = GetPublicClassInputs();
-	for(const FMetasoundFrontendClassInput& Input : PublicInputs)
+	for(const TPair<FVertexName, FMetasoundAssetBase::FRuntimeInput>& Pair : GetRuntimeData().PublicInputMap)
 	{
+		const FMetasoundAssetBase::FRuntimeInput& Input = Pair.Value;
 		FAudioParameter Params;
 		Params.ParamName = Input.Name;
 		Params.TypeName = Input.TypeName;
@@ -795,138 +763,134 @@ bool UMetaSoundSource::IsParameterValid(const FAudioParameter& InParameter) cons
 	{
 		return Input.Name == InParameter.ParamName;
 	});
-	return IsParameterValid(InParameter, Vertex);
+	
+	if (Vertex)
+	{
+		return IsParameterValid(InParameter, Vertex->TypeName, Metasound::Frontend::IDataTypeRegistry::Get());
+	}
+	else
+	{
+		return false;
+	}
 }
 
-bool UMetaSoundSource::IsParameterValid(const FAudioParameter& InParameter, const FMetasoundFrontendVertex* InVertex) const
+bool UMetaSoundSource::IsParameterValid(const FAudioParameter& InParameter, const FName& InTypeName, Metasound::Frontend::IDataTypeRegistry& InDataTypeRegistry) const
 {
 	using namespace Metasound;
 	using namespace Metasound::Frontend;
 
 	if (InParameter.ParamName.IsNone())
 	{
+		// Invalid parameter name
 		return false;
 	}
 
-	if (!InVertex)
+	if (!InParameter.TypeName.IsNone() && InParameter.TypeName != InTypeName)
 	{
+		// Mismatched parameter type and vertex data type
 		return false;
 	}
 
-	const FName& TypeName = InVertex->TypeName;
-	if (!InParameter.TypeName.IsNone() && InParameter.TypeName != TypeName)
+	// Special handling for UObject proxies
+	if (InParameter.ParamType == EAudioParameterType::Object)
+	{
+		return InDataTypeRegistry.IsValidUObjectForDataType(InTypeName, InParameter.ObjectParam);
+	}
+	else if (InParameter.ParamType == EAudioParameterType::ObjectArray)
+	{
+		bool bIsValid = true;
+
+		const FName ElementTypeName = CreateElementTypeNameFromArrayTypeName(InTypeName);
+		for (const UObject* Object : InParameter.ArrayObjectParam)
+		{
+			bIsValid = InDataTypeRegistry.IsValidUObjectForDataType(ElementTypeName, Object);
+			if (!bIsValid)
+			{
+				break;
+			}
+		}
+		return bIsValid;
+	}
+
+	const IDataTypeRegistryEntry* RegistryEntry = InDataTypeRegistry.FindDataTypeRegistryEntry(InTypeName);
+	if (!RegistryEntry)
 	{
+		// Unregistered MetaSound data type
 		return false;
 	}
 
-	bool bIsValid = false;
 	switch (InParameter.ParamType)
 	{
 		case EAudioParameterType::Trigger:
 		case EAudioParameterType::Boolean:
 		{
-			FDataTypeRegistryInfo DataTypeInfo;
-			bIsValid = IDataTypeRegistry::Get().GetDataTypeInfo(TypeName, DataTypeInfo);
-			bIsValid &= DataTypeInfo.bIsBoolParsable;
+			return RegistryEntry->GetDataTypeInfo().bIsBoolParsable;
 		}
 		break;
 
 		case EAudioParameterType::BooleanArray:
 		{
-			FDataTypeRegistryInfo DataTypeInfo;
-			bIsValid = IDataTypeRegistry::Get().GetDataTypeInfo(TypeName, DataTypeInfo);
-			bIsValid &= DataTypeInfo.bIsBoolArrayParsable;
+			return RegistryEntry->GetDataTypeInfo().bIsBoolArrayParsable;
 		}
 		break;
 
 		case EAudioParameterType::Float:
 		{
-			FDataTypeRegistryInfo DataTypeInfo;
-			bIsValid = IDataTypeRegistry::Get().GetDataTypeInfo(TypeName, DataTypeInfo);
-			bIsValid &= DataTypeInfo.bIsFloatParsable;
+			return RegistryEntry->GetDataTypeInfo().bIsFloatParsable;
 		}
 		break;
 
 		case EAudioParameterType::FloatArray:
 		{
-			FDataTypeRegistryInfo DataTypeInfo;
-			bIsValid = IDataTypeRegistry::Get().GetDataTypeInfo(TypeName, DataTypeInfo);
-			bIsValid &= DataTypeInfo.bIsFloatArrayParsable;
+			return RegistryEntry->GetDataTypeInfo().bIsFloatArrayParsable;
 		}
 		break;
 
 		case EAudioParameterType::Integer:
 		{
-			FDataTypeRegistryInfo DataTypeInfo;
-			bIsValid = IDataTypeRegistry::Get().GetDataTypeInfo(TypeName, DataTypeInfo);
-			bIsValid &= DataTypeInfo.bIsIntParsable;
+			return RegistryEntry->GetDataTypeInfo().bIsIntParsable;
 		}
 		break;
 
 		case EAudioParameterType::IntegerArray:
 		{
-			FDataTypeRegistryInfo DataTypeInfo;
-			bIsValid = IDataTypeRegistry::Get().GetDataTypeInfo(TypeName, DataTypeInfo);
-			bIsValid &= DataTypeInfo.bIsIntArrayParsable;
-
+			return RegistryEntry->GetDataTypeInfo().bIsIntArrayParsable;
 		}
 		break;
 
-		case EAudioParameterType::Object:
+		case EAudioParameterType::String:
 		{
-			bIsValid = IDataTypeRegistry::Get().IsValidUObjectForDataType(TypeName, InParameter.ObjectParam);
+			return RegistryEntry->GetDataTypeInfo().bIsStringParsable;
 		}
 		break;
 
-		case EAudioParameterType::ObjectArray:
+		case EAudioParameterType::StringArray:
 		{
-			bIsValid = true;
-
-			const FName ElementTypeName = CreateElementTypeNameFromArrayTypeName(TypeName);
-			for (const UObject* Object : InParameter.ArrayObjectParam)
-			{
-				bIsValid = IDataTypeRegistry::Get().IsValidUObjectForDataType(ElementTypeName, Object);
-				if (!bIsValid)
-				{
-					break;
-				}
-			}
+			return RegistryEntry->GetDataTypeInfo().bIsStringArrayParsable;
 		}
 		break;
 
-		case EAudioParameterType::String:
+		case EAudioParameterType::NoneArray:
 		{
-			FDataTypeRegistryInfo DataTypeInfo;
-			bIsValid = IDataTypeRegistry::Get().GetDataTypeInfo(TypeName, DataTypeInfo);
-			bIsValid &= DataTypeInfo.bIsStringParsable;
+			return RegistryEntry->GetDataTypeInfo().bIsDefaultArrayParsable;
 		}
-		break;
 
-		case EAudioParameterType::StringArray:
+		case EAudioParameterType::None:
 		{
-			FDataTypeRegistryInfo DataTypeInfo;
-			bIsValid = IDataTypeRegistry::Get().GetDataTypeInfo(TypeName, DataTypeInfo);
-			bIsValid &= DataTypeInfo.bIsStringArrayParsable;
+			return RegistryEntry->GetDataTypeInfo().bIsDefaultParsable;
 		}
 		break;
 
-		case EAudioParameterType::NoneArray:
-		{
-			FDataTypeRegistryInfo DataTypeInfo;
-			bIsValid = IDataTypeRegistry::Get().GetDataTypeInfo(TypeName, DataTypeInfo);
-			bIsValid &= DataTypeInfo.bIsDefaultArrayParsable;
-		}
-		case EAudioParameterType::None:
 		default:
 		{
-			FDataTypeRegistryInfo DataTypeInfo;
-			bIsValid = IDataTypeRegistry::Get().GetDataTypeInfo(TypeName, DataTypeInfo);
-			bIsValid &= DataTypeInfo.bIsDefaultParsable;
+			// All parameter types should be covered.
+			static_assert(static_cast<uint8>(EAudioParameterType::COUNT) == 13, "Possible unhandled EAudioParameterType");
+			checkNoEntry();
+			// Unhandled parameter type
+			return false;
 		}
-		break;
 	}
 
-	return bIsValid;
 }
 
 bool UMetaSoundSource::IsLooping() const
@@ -951,8 +915,16 @@ TSharedPtr<Audio::IParameterTransmitter> UMetaSoundSource::CreateParameterTransm
 
 	const FRuntimeData& RuntimeData = GetRuntimeData();
 
+	// Build list of parameters that can be set at runtime.
 	TArray<FName> ValidParameters;
-	Algo::Transform(RuntimeData.TransmittableInputs, ValidParameters, [](const FMetasoundFrontendClassInput& InputVertex) { return InputVertex.Name; });
+	for (const TPair<FVertexName, FMetasoundAssetBase::FRuntimeInput>& Pair : RuntimeData.PublicInputMap)
+	{
+		if (Pair.Value.bIsTransmittable && (Pair.Value.AccessType == EMetasoundFrontendVertexAccessType::Reference))
+		{
+			ValidParameters.Add(Pair.Value.Name);
+		}
+	}
+
 	FParameterRouter& Router = GetParameterRouter();
 	TSharedPtr<TSpscQueue<FMetaSoundParameterTransmitter::FParameter>> DataChannel = Router.FindOrCreateDataChannelForWriter(InParams.AudioDeviceID, InParams.InstanceID);
 
diff --git a/Engine/Plugins/Runtime/Metasound/Source/MetasoundEngine/Public/MetasoundSource.h b/Engine/Plugins/Runtime/Metasound/Source/MetasoundEngine/Public/MetasoundSource.h
--- a/Engine/Plugins/Runtime/Metasound/Source/MetasoundEngine/Public/MetasoundSource.h
+++ b/Engine/Plugins/Runtime/Metasound/Source/MetasoundEngine/Public/MetasoundSource.h
@@ -234,7 +234,7 @@ private:
 	}
 
 
-	bool IsParameterValid(const FAudioParameter& InParameter, const FMetasoundFrontendVertex* InVertex) const;
+	bool IsParameterValid(const FAudioParameter& InParameter, const FName& InTypeName, Metasound::Frontend::IDataTypeRegistry& InDataTypeRegistry) const;
 
 	static Metasound::SourcePrivate::FParameterRouter& GetParameterRouter();
 
diff --git a/Engine/Plugins/Runtime/Metasound/Source/MetasoundFrontend/Private/MetasoundAssetBase.cpp b/Engine/Plugins/Runtime/Metasound/Source/MetasoundFrontend/Private/MetasoundAssetBase.cpp
--- a/Engine/Plugins/Runtime/Metasound/Source/MetasoundFrontend/Private/MetasoundAssetBase.cpp
+++ b/Engine/Plugins/Runtime/Metasound/Source/MetasoundFrontend/Private/MetasoundAssetBase.cpp
@@ -121,38 +121,60 @@ namespace Metasound
 				}
 			}
 
-			// Remove all inputs from the provided array of public inputs which are of non-transmittable data types.
-			TArray<FMetasoundFrontendClassInput> GetTransmittableInputsFromPublicInputs(const TArray<FMetasoundFrontendClassInput>& InPublicInputs)
+			TSortedVertexNameMap<FMetasoundAssetBase::FRuntimeInput> GetPublicClassInputs(const FMetasoundFrontendDocument& InDoc)
 			{
+				using namespace Metasound;
 				using namespace Metasound::Frontend;
 
-				TArray<FMetasoundFrontendClassInput> TransmittableInputs;
+				METASOUND_TRACE_CPUPROFILER_EVENT_SCOPE(AssetBasePrivate::GetPublicClassInputs);
+
+				auto GetInputName = [](const FMetasoundFrontendClassInput& InInput) { return InInput.Name; };
 
 				IDataTypeRegistry& Registry = IDataTypeRegistry::Get();
-				Algo::TransformIf(InPublicInputs, TransmittableInputs,
-					[&Registry](const FMetasoundFrontendClassInput& Input)
+
+				TArray<const IInterfaceRegistryEntry*> Interfaces;
+				FMetaSoundFrontendDocumentBuilder::FindDeclaredInterfaces(InDoc, Interfaces);
+
+				// Inputs which are controlled by an interface are private unless
+				// their router name is `Audio::IParameterTransmitter::RouterName`
+				TSet<FVertexName> PrivateInputs;
+				for (const IInterfaceRegistryEntry* InterfaceEntry : Interfaces)
+				{
+					if (InterfaceEntry)
 					{
-						if (Input.AccessType != EMetasoundFrontendVertexAccessType::Reference)
+						if (InterfaceEntry->GetRouterName() != Audio::IParameterTransmitter::RouterName)
 						{
-							return false;
+							const FMetasoundFrontendInterface& Interface = InterfaceEntry->GetInterface();
+							Algo::Transform(Interface.Inputs, PrivateInputs, GetInputName);
 						}
+					}
+				}
 
-						FDataTypeRegistryInfo Info;
-						if (!Registry.GetDataTypeInfo(Input.TypeName, Info))
+				// Cache all inputs which are not private inputs.
+				TSortedVertexNameMap<FMetasoundAssetBase::FRuntimeInput> PublicInputs;
+				for (const FMetasoundFrontendClassInput& Input : InDoc.RootGraph.Interface.Inputs)
+				{
+					if (!PrivateInputs.Contains(Input.Name))
+					{
+						bool bIsTransmittable = false;
+						if (const IDataTypeRegistryEntry* RegistryEntry = Registry.FindDataTypeRegistryEntry(Input.TypeName))
 						{
-							return false;
+							bIsTransmittable = RegistryEntry->GetDataTypeInfo().bIsTransmittable;	
 						}
-
-						if (!Info.bIsTransmittable)
+						else
 						{
-							return false;
-						}
+							UE_LOG(LogMetaSound, Warning, TEXT("Failed to find data type '%s' in registry. Assuming data type is not transmittable"), *Input.TypeName.ToString());
 
-						return true;
-					}, [](const FMetasoundFrontendClassInput& Input) { return Input; }
-					);
+						}
+						PublicInputs.Add(Input.Name, FMetasoundAssetBase::FRuntimeInput{Input.Name, Input.TypeName, Input.AccessType, Input.DefaultLiteral, bIsTransmittable});
+					}
+				}
 
-				return TransmittableInputs;
+				// Add the parameter pack input that ALL Metasounds have
+				FMetasoundFrontendClassInput ParameterPackInput = UMetasoundParameterPack::GetClassInput();
+				PublicInputs.Add(ParameterPackInput.Name, FMetasoundAssetBase::FRuntimeInput{ParameterPackInput.Name, ParameterPackInput.TypeName, ParameterPackInput.AccessType, ParameterPackInput.DefaultLiteral, true /* bIsTransmittable */});
+				
+				return PublicInputs;
 			}
 
 			// Registers node by copying document. Updates to document require re-registration.
@@ -773,19 +795,19 @@ TArray<FMetasoundAssetBase::FSendInfoAndVertexName> FMetasoundAssetBase::GetSend
 
 	TArray<FSendInfoAndVertexName> SendInfos;
 
-	for (const FMetasoundFrontendClassInput& Vertex : RuntimeData.TransmittableInputs)
+	for (const TTuple<FVertexName, FRuntimeInput>& Entry: RuntimeData.PublicInputMap)
 	{
+		const FRuntimeInput& RuntimeInput = Entry.Get<1>();
 		FSendInfoAndVertexName Info;
 
 PRAGMA_DISABLE_DEPRECATION_WARNINGS
-		Info.SendInfo.Address = FMetaSoundParameterTransmitter::CreateSendAddressFromInstanceID(InInstanceID, Vertex.Name, Vertex.TypeName);
+		Info.SendInfo.Address = FMetaSoundParameterTransmitter::CreateSendAddressFromInstanceID(InInstanceID, RuntimeInput.Name, RuntimeInput.TypeName);
 PRAGMA_ENABLE_DEPRECATION_WARNINGS
-		Info.SendInfo.ParameterName = Vertex.Name;
-		Info.SendInfo.TypeName = Vertex.TypeName;
-		Info.VertexName = Vertex.Name;
+		Info.SendInfo.ParameterName = RuntimeInput.Name;
+		Info.SendInfo.TypeName = RuntimeInput.TypeName;
+		Info.VertexName = RuntimeInput.Name;
 
 		SendInfos.Add(Info);
-		
 	}
 
 	return SendInfos;
@@ -904,48 +926,6 @@ FString FMetasoundAssetBase::GetOwningAssetName() const
 	return FString();
 }
 
-TArray<FMetasoundFrontendClassInput> FMetasoundAssetBase::GetPublicClassInputs() const
-{
-	using namespace Metasound;
-	using namespace Metasound::Frontend;
-
-	METASOUND_TRACE_CPUPROFILER_EVENT_SCOPE(FMetasoundAssetBase::GetPublicClassInputs);
-
-	auto GetInputName = [](const FMetasoundFrontendClassInput& InInput) { return InInput.Name; };
-
-	// Inputs which are controlled by an interface are private. 
-	TArray<const IInterfaceRegistryEntry*> Interfaces;
-	TSet<FVertexName> PrivateInputs;
-	FMetaSoundFrontendDocumentBuilder::FindDeclaredInterfaces(GetDocumentChecked(), Interfaces);
-	for (const IInterfaceRegistryEntry* InterfaceEntry : Interfaces)
-	{
-		if (InterfaceEntry)
-		{
-			if (InterfaceEntry->GetRouterName() != Audio::IParameterTransmitter::RouterName)
-			{
-				const FMetasoundFrontendInterface& Interface = InterfaceEntry->GetInterface();
-				Algo::Transform(Interface.Inputs, PrivateInputs, GetInputName);
-			}
-		}
-	}
-
-	auto IsPublic = [&PrivateInputs](const FMetasoundFrontendClassVertex& InVertex)
-	{
-		return !PrivateInputs.Contains(InVertex.Name);
-	};
-
-	const FMetasoundFrontendDocument& Doc = GetDocumentChecked();
-	TArray<FMetasoundFrontendClassInput> PublicInputs;
-
-	Algo::CopyIf(Doc.RootGraph.Interface.Inputs, PublicInputs, IsPublic);
-
-	// Add the parameter pack input that ALL Metasounds have
-	PublicInputs.Add(UMetasoundParameterPack::GetClassInput());
-	
-	return PublicInputs;
-}
-
-
 #if WITH_EDITOR
 void FMetasoundAssetBase::RebuildReferencedAssetClasses()
 {
@@ -1045,20 +1025,15 @@ const FMetasoundAssetBase::FRuntimeData& FMetasoundAssetBase::CacheRuntimeData(c
 	CurrentCachedRuntimeDataChangeID = FGuid::NewGuid();
 	CachedRuntimeData.ChangeID = CurrentCachedRuntimeDataChangeID;
 
-	TArray<FMetasoundFrontendClassInput> PublicInputs = GetPublicClassInputs();
-	TArray<FMetasoundFrontendClassInput> TransmittableInputs = AssetBasePrivate::GetTransmittableInputsFromPublicInputs(PublicInputs);
+	Metasound::TSortedVertexNameMap<FMetasoundAssetBase::FRuntimeInput> PublicInputMap = AssetBasePrivate::GetPublicClassInputs(InDocument);
 
 	Metasound::Frontend::FProxyDataCache ProxyDataCache;
 	ProxyDataCache.CreateAndCacheProxies(InDocument);
 	TSharedPtr<Metasound::FGraph, ESPMode::ThreadSafe> Graph = BuildMetasoundDocument(InDocument, ProxyDataCache);
 
-	CachedRuntimeData =
-	{
-		CurrentCachedRuntimeDataChangeID,
-		MoveTemp(PublicInputs),
-		MoveTemp(TransmittableInputs),
-		MoveTemp(Graph)
-	};
+	CachedRuntimeData.ChangeID = CurrentCachedRuntimeDataChangeID;
+	CachedRuntimeData.PublicInputMap = MoveTemp(PublicInputMap);
+	CachedRuntimeData.Graph = MoveTemp(Graph);
 
 	return CachedRuntimeData;
 }
diff --git a/Engine/Plugins/Runtime/Metasound/Source/MetasoundFrontend/Private/MetasoundFrontendDocument.cpp b/Engine/Plugins/Runtime/Metasound/Source/MetasoundFrontend/Private/MetasoundFrontendDocument.cpp
--- a/Engine/Plugins/Runtime/Metasound/Source/MetasoundFrontend/Private/MetasoundFrontendDocument.cpp
+++ b/Engine/Plugins/Runtime/Metasound/Source/MetasoundFrontend/Private/MetasoundFrontendDocument.cpp
@@ -868,56 +868,58 @@ const TCHAR* LexToString(EMetasoundFrontendVertexAccessType InVertexAccess)
 
 namespace Metasound
 {
-	void ForEachLiteral(const FMetasoundFrontendDocument& InDoc, FForEachLiteralFunctionRef OnLiteral)
+	namespace Frontend
 	{
-		ForEachLiteral(InDoc.RootGraph, OnLiteral);
-		for (const FMetasoundFrontendGraphClass& GraphClass : InDoc.Subgraphs)
-		{
-			ForEachLiteral(GraphClass, OnLiteral);
-		}
-		for (const FMetasoundFrontendClass& Dependency : InDoc.Dependencies)
+		void ForEachLiteral(const FMetasoundFrontendDocument& InDoc, FForEachLiteralFunctionRef OnLiteral)
 		{
-			ForEachLiteral(Dependency, OnLiteral);
+			ForEachLiteral(InDoc.RootGraph, OnLiteral);
+			for (const FMetasoundFrontendGraphClass& GraphClass : InDoc.Subgraphs)
+			{
+				ForEachLiteral(GraphClass, OnLiteral);
+			}
+			for (const FMetasoundFrontendClass& Dependency : InDoc.Dependencies)
+			{
+				ForEachLiteral(Dependency, OnLiteral);
+			}
 		}
-	}
 
-	void ForEachLiteral(const FMetasoundFrontendGraphClass& InGraphClass, FForEachLiteralFunctionRef OnLiteral)
-	{
-		ForEachLiteral(static_cast<const FMetasoundFrontendClass&>(InGraphClass), OnLiteral);
-
-		for (const FMetasoundFrontendNode& Node : InGraphClass.Graph.Nodes)
+		void ForEachLiteral(const FMetasoundFrontendGraphClass& InGraphClass, FForEachLiteralFunctionRef OnLiteral)
 		{
-			ForEachLiteral(Node, OnLiteral);
-		}
+			ForEachLiteral(static_cast<const FMetasoundFrontendClass&>(InGraphClass), OnLiteral);
 
-		for (const FMetasoundFrontendVariable& Variable : InGraphClass.Graph.Variables)
-		{
-			OnLiteral(Variable.TypeName, Variable.Literal);
+			for (const FMetasoundFrontendNode& Node : InGraphClass.Graph.Nodes)
+			{
+				ForEachLiteral(Node, OnLiteral);
+			}
+
+			for (const FMetasoundFrontendVariable& Variable : InGraphClass.Graph.Variables)
+			{
+				OnLiteral(Variable.TypeName, Variable.Literal);
+			}
 		}
-	}
 
-	void ForEachLiteral(const FMetasoundFrontendClass& InClass, FForEachLiteralFunctionRef OnLiteral)
-	{
-		for (const FMetasoundFrontendClassInput& ClassInput : InClass.Interface.Inputs)
+		void ForEachLiteral(const FMetasoundFrontendClass& InClass, FForEachLiteralFunctionRef OnLiteral)
 		{
-			OnLiteral(ClassInput.TypeName, ClassInput.DefaultLiteral);
+			for (const FMetasoundFrontendClassInput& ClassInput : InClass.Interface.Inputs)
+			{
+				OnLiteral(ClassInput.TypeName, ClassInput.DefaultLiteral);
+			}
 		}
-	}
 
-	void ForEachLiteral(const FMetasoundFrontendNode& InNode, FForEachLiteralFunctionRef OnLiteral)
-	{
-		for (const FMetasoundFrontendVertexLiteral& VertexLiteral : InNode.InputLiterals)
+		void ForEachLiteral(const FMetasoundFrontendNode& InNode, FForEachLiteralFunctionRef OnLiteral)
 		{
-			auto HasEqualVertexID = [&VertexLiteral](const FMetasoundFrontendVertex& InVertex) -> bool
-			{ 
-				return InVertex.VertexID == VertexLiteral.VertexID; 
-			};
-
-			if (const FMetasoundFrontendVertex* InputVertex = InNode.Interface.Inputs.FindByPredicate(HasEqualVertexID))
+			for (const FMetasoundFrontendVertexLiteral& VertexLiteral : InNode.InputLiterals)
 			{
-				OnLiteral(InputVertex->TypeName, VertexLiteral.Value);
+				auto HasEqualVertexID = [&VertexLiteral](const FMetasoundFrontendVertex& InVertex) -> bool
+				{ 
+					return InVertex.VertexID == VertexLiteral.VertexID; 
+				};
+
+				if (const FMetasoundFrontendVertex* InputVertex = InNode.Interface.Inputs.FindByPredicate(HasEqualVertexID))
+				{
+					OnLiteral(InputVertex->TypeName, VertexLiteral.Value);
+				}
 			}
 		}
 	}
 }
-
diff --git a/Engine/Plugins/Runtime/Metasound/Source/MetasoundFrontend/Private/MetasoundFrontendProxyDataCache.h b/Engine/Plugins/Runtime/Metasound/Source/MetasoundFrontend/Private/MetasoundFrontendProxyDataCache.h
--- a/Engine/Plugins/Runtime/Metasound/Source/MetasoundFrontend/Private/MetasoundFrontendProxyDataCache.h
+++ b/Engine/Plugins/Runtime/Metasound/Source/MetasoundFrontend/Private/MetasoundFrontendProxyDataCache.h
@@ -68,8 +68,9 @@ namespace Metasound
 			const TSharedPtr<Audio::IProxyData>* FindProxy(const UObject* InObject) const;
 
 		private:
+			using FUObjectMemoryAddress = const void*;
 
-			TMap<const void*, TSharedPtr<Audio::IProxyData>> ProxyCache;
+			TMap<FUObjectMemoryAddress, TSharedPtr<Audio::IProxyData>> ProxyCache;
 		};
 	}
 }
diff --git a/Engine/Plugins/Runtime/Metasound/Source/MetasoundFrontend/Public/MetasoundAssetBase.h b/Engine/Plugins/Runtime/Metasound/Source/MetasoundFrontend/Public/MetasoundAssetBase.h
--- a/Engine/Plugins/Runtime/Metasound/Source/MetasoundFrontend/Public/MetasoundAssetBase.h
+++ b/Engine/Plugins/Runtime/Metasound/Source/MetasoundFrontend/Public/MetasoundAssetBase.h
@@ -31,6 +31,7 @@ namespace Metasound
 		METASOUNDFRONTEND_API float GetDefaultBlockRate();
 
 		class FProxyDataCache;
+
 	} // namespace Frontend
 } // namespace Metasound
 
@@ -45,6 +46,21 @@ class METASOUNDFRONTEND_API FMetasoundAssetBase
 public:
 	static const FString FileExtension;
 
+	// FRuntimeInput represents an input to a MetaSound which can be manipulated.
+	struct FRuntimeInput
+	{
+		// Name of input vertex
+		FName Name;
+		// Data type name of input vertex.
+		FName TypeName;
+		// Access type of input vertex.
+		EMetasoundFrontendVertexAccessType AccessType;
+		// Default literal of input vertex.
+		FMetasoundFrontendLiteral DefaultLiteral;
+		// True if the data type is transmittable. False otherwise.
+		bool bIsTransmittable;
+	};
+
 	FMetasoundAssetBase() = default;
 	virtual ~FMetasoundAssetBase() = default;
 
@@ -204,6 +220,7 @@ protected:
 	virtual Metasound::Frontend::FConstDocumentAccessPtr GetDocumentConstAccessPtr() const = 0;
 
 protected:
+	
 	// Container for runtime data of MetaSound graph.
 	struct FRuntimeData
 	{
@@ -211,9 +228,14 @@ protected:
 		FGuid ChangeID;
 
 		// Array of inputs which can be set for construction. 
+		UE_DEPRECATED(5.4, "Use PublicInputMap instead.")
 		TArray<FMetasoundFrontendClassInput> PublicInputs;
 
+		// Map of runtime inputs keyed by vertex names. 
+		Metasound::TSortedVertexNameMap<FRuntimeInput> PublicInputMap;
+
 		// Array of inputs which can be transmitted to.
+		UE_DEPRECATED(5.4, "Use PublicInputMap and DataTypeRegistry information instead.")
 		TArray<FMetasoundFrontendClassInput> TransmittableInputs;
 
 		// Core graph.
diff --git a/Engine/Plugins/Runtime/Metasound/Source/MetasoundFrontend/Public/MetasoundFrontendDocument.h b/Engine/Plugins/Runtime/Metasound/Source/MetasoundFrontend/Public/MetasoundFrontendDocument.h
--- a/Engine/Plugins/Runtime/Metasound/Source/MetasoundFrontend/Public/MetasoundFrontendDocument.h
+++ b/Engine/Plugins/Runtime/Metasound/Source/MetasoundFrontend/Public/MetasoundFrontendDocument.h
@@ -1655,19 +1655,22 @@ METASOUNDFRONTEND_API const TCHAR* LexToString(EMetasoundFrontendVertexAccessTyp
 
 namespace Metasound
 {
-	/** Signature of function called for each found literal. */
-	using FForEachLiteralFunctionRef = TFunctionRef<void(const FName& InDataTypeName, const FMetasoundFrontendLiteral&)>; 
+	namespace Frontend
+	{
+		/** Signature of function called for each found literal. */
+		using FForEachLiteralFunctionRef = TFunctionRef<void(const FName& InDataTypeName, const FMetasoundFrontendLiteral&)>; 
 
-	/** Execute the provided function for each literal on a FMetasoundFrontendDocument.*/
-	METASOUNDFRONTEND_API void ForEachLiteral(const FMetasoundFrontendDocument& InDoc, FForEachLiteralFunctionRef OnLiteral);
+		/** Execute the provided function for each literal on a FMetasoundFrontendDocument.*/
+		METASOUNDFRONTEND_API void ForEachLiteral(const FMetasoundFrontendDocument& InDoc, FForEachLiteralFunctionRef OnLiteral);
 
-	/** Execute the provided function for each literal on a FMetasoundFrontendGraphClass.*/
-	METASOUNDFRONTEND_API void ForEachLiteral(const FMetasoundFrontendGraphClass& InGraphClass, FForEachLiteralFunctionRef OnLiteral);
+		/** Execute the provided function for each literal on a FMetasoundFrontendGraphClass.*/
+		METASOUNDFRONTEND_API void ForEachLiteral(const FMetasoundFrontendGraphClass& InGraphClass, FForEachLiteralFunctionRef OnLiteral);
 
-	/** Execute the provided function for each literal on a FMetasoundFrontendClass.*/
-	METASOUNDFRONTEND_API void ForEachLiteral(const FMetasoundFrontendClass& InClass, FForEachLiteralFunctionRef OnLiteral);
+		/** Execute the provided function for each literal on a FMetasoundFrontendClass.*/
+		METASOUNDFRONTEND_API void ForEachLiteral(const FMetasoundFrontendClass& InClass, FForEachLiteralFunctionRef OnLiteral);
 
-	/** Execute the provided function for each literal on a FMetasoundFrontendNode.*/
-	METASOUNDFRONTEND_API void ForEachLiteral(const FMetasoundFrontendNode& InNode, FForEachLiteralFunctionRef OnLiteral);
+		/** Execute the provided function for each literal on a FMetasoundFrontendNode.*/
+		METASOUNDFRONTEND_API void ForEachLiteral(const FMetasoundFrontendNode& InNode, FForEachLiteralFunctionRef OnLiteral);
+	}
 }
 
diff --git a/Engine/Plugins/Runtime/Metasound/Source/MetasoundFrontend/Public/MetasoundFrontendLiteral.h b/Engine/Plugins/Runtime/Metasound/Source/MetasoundFrontend/Public/MetasoundFrontendLiteral.h
--- a/Engine/Plugins/Runtime/Metasound/Source/MetasoundFrontend/Public/MetasoundFrontendLiteral.h
+++ b/Engine/Plugins/Runtime/Metasound/Source/MetasoundFrontend/Public/MetasoundFrontendLiteral.h
@@ -146,10 +146,10 @@ public:
 	// @param InDataTypeRegistry - A pointer to an existing data type registry. If null, the data type registry will be retrieved within this function.
 	// @param InProxyDataCache - A pointer to an existing proxy data cache. If not null, UObject proxies will be retrieved from the cache. If null, UObject proxies will be created in this function.
 	//
-	// @returns An Fliteral. If the data type couldn't be found, or if the literal type was incompatible with the data type, then an invalid FLiteral is returned.
+	// @returns An FLiteral. If the data type couldn't be found, or if the literal type was incompatible with the data type, then an invalid FLiteral is returned.
 	Metasound::FLiteral ToLiteral(const FName& InMetaSoundDataType, const Metasound::Frontend::IDataTypeRegistry* InDataTypeRegistry=nullptr, const Metasound::Frontend::FProxyDataCache* InProxyDataCache=nullptr) const;
 
-	// Return a Metasound::Fliteral representation of this object, excluding UObject proxies.
+	// Return a Metasound::FLiteral representation of this object, excluding UObject proxies.
 	Metasound::FLiteral ToLiteralNoProxy() const;
 
 	// Convert the value to a string for printing. 
