Move MetasoundDocumentBuilders to cache-on-demand
- Invalidate subsystem registerd asset builder cache upon soft deprecated document access id creation
#rnx

diff --git a/Engine/Plugins/Runtime/Metasound/Source/MetasoundEngine/Private/Metasound.cpp b/Engine/Plugins/Runtime/Metasound/Source/MetasoundEngine/Private/Metasound.cpp
--- a/Engine/Plugins/Runtime/Metasound/Source/MetasoundEngine/Private/Metasound.cpp
+++ b/Engine/Plugins/Runtime/Metasound/Source/MetasoundEngine/Private/Metasound.cpp
@@ -57,6 +57,29 @@ UMetaSoundPatch::UMetaSoundPatch(const FObjectInitializer& ObjectInitializer)
 {
 }
 
+Metasound::Frontend::FDocumentAccessPtr UMetaSoundPatch::GetDocumentAccessPtr()
+{
+	using namespace Metasound::Frontend;
+
+	// Mutation of a document via the soft deprecated access ptr/controller system is not tracked by
+	// the builder registry, so the document cache is invalidated here. It is discouraged to mutate
+	// documents using both systems at the same time as it can corrupt a builder document's cache.
+	const FMetasoundFrontendClassName& Name = RootMetaSoundDocument.RootGraph.Metadata.GetClassName();
+	IDocumentBuilderRegistry::GetChecked().InvalidateDocumentCache(Name);
+
+	// Return document using FAccessPoint to inform the TAccessPtr when the 
+	// object is no longer valid.
+	return MakeAccessPtr<FDocumentAccessPtr>(RootMetaSoundDocument.AccessPoint, RootMetaSoundDocument);
+}
+
+Metasound::Frontend::FConstDocumentAccessPtr UMetaSoundPatch::GetDocumentConstAccessPtr() const
+{
+	using namespace Metasound::Frontend;
+	// Return document using FAccessPoint to inform the TAccessPtr when the 
+	// object is no longer valid.
+	return MakeAccessPtr<FConstDocumentAccessPtr>(RootMetaSoundDocument.AccessPoint, RootMetaSoundDocument);
+}
+
 const UClass& UMetaSoundPatch::GetBaseMetaSoundUClass() const
 {
 	return *UMetaSoundPatch::StaticClass();
diff --git a/Engine/Plugins/Runtime/Metasound/Source/MetasoundEngine/Private/MetasoundBuilderSubsystem.cpp b/Engine/Plugins/Runtime/Metasound/Source/MetasoundEngine/Private/MetasoundBuilderSubsystem.cpp
--- a/Engine/Plugins/Runtime/Metasound/Source/MetasoundEngine/Private/MetasoundBuilderSubsystem.cpp
+++ b/Engine/Plugins/Runtime/Metasound/Source/MetasoundEngine/Private/MetasoundBuilderSubsystem.cpp
@@ -636,6 +636,11 @@ UMetaSoundBuilderDocument* UMetaSoundBuilderBase::CreateTransientDocumentObject(
 	return &UMetaSoundBuilderDocument::Create(GetBuilderUClass());
 }
 
+const FMetaSoundFrontendDocumentBuilder& UMetaSoundBuilderBase::GetConstBuilder() const
+{
+	return Builder;
+}
+
 UObject* UMetaSoundBuilderBase::GetReferencedPresetAsset() const
 {
 	using namespace Metasound::Frontend;
@@ -682,6 +687,11 @@ bool UMetaSoundBuilderBase::InterfaceIsDeclared(FName InterfaceName) const
 	return Builder.IsInterfaceDeclared(InterfaceName);
 }
 
+void UMetaSoundBuilderBase::InvalidateCache()
+{
+	Builder.InvalidateCache();
+}
+
 bool UMetaSoundBuilderBase::IsPreset() const
 {
 	return Builder.IsPreset();
@@ -703,6 +713,11 @@ bool UMetaSoundBuilderBase::NodeOutputIsConnected(const FMetaSoundBuilderNodeOut
 	return Builder.IsNodeOutputConnected(OutputHandle.NodeID, OutputHandle.VertexID);
 }
 
+void UMetaSoundBuilderBase::ReloadCache(bool bPrimeCache)
+{
+	Builder.ReloadCache();
+}
+
 void UMetaSoundBuilderBase::RemoveGraphInput(FName Name, EMetaSoundBuilderResult& OutResult)
 {
 	const bool bRemoved = Builder.RemoveGraphInput(Name);
@@ -739,11 +754,6 @@ void UMetaSoundBuilderBase::RenameRootGraphClass(const FMetasoundFrontendClassNa
 }
 
 #if WITH_EDITOR
-void UMetaSoundBuilderBase::ReloadCache()
-{
-	Builder.ReloadCache();
-}
-
 void UMetaSoundBuilderBase::SetAuthor(const FString& InAuthor)
 {
 	Builder.SetAuthor(InAuthor);
@@ -1661,17 +1671,15 @@ bool UMetaSoundBuilderSubsystem::DetachBuilderFromAsset(const FMetasoundFrontend
 	return AssetBuilders.Remove(InClassName.GetFullName()) > 0;
 }
 
-const Metasound::Frontend::FDocumentModifyDelegates* UMetaSoundBuilderSubsystem::FindModifyDelegates(const FMetasoundFrontendClassName& InClassName) const
+void UMetaSoundBuilderSubsystem::InvalidateDocumentCache(const FMetasoundFrontendClassName& InClassName) const
 {
 	using namespace Metasound::Frontend;
 
 	TWeakObjectPtr<UMetaSoundBuilderBase> BuilderPtr = AssetBuilders.FindRef(InClassName.GetFullName());
 	if (BuilderPtr.IsValid())
 	{
-		return &BuilderPtr->GetConstBuilder().GetDocumentDelegates();
+		BuilderPtr->InvalidateCache();
 	}
-
-	return nullptr;
 }
 
 UMetaSoundBuilderBase* UMetaSoundBuilderSubsystem::FindBuilder(FName BuilderName)
diff --git a/Engine/Plugins/Runtime/Metasound/Source/MetasoundEngine/Private/MetasoundSource.cpp b/Engine/Plugins/Runtime/Metasound/Source/MetasoundEngine/Private/MetasoundSource.cpp
--- a/Engine/Plugins/Runtime/Metasound/Source/MetasoundEngine/Private/MetasoundSource.cpp
+++ b/Engine/Plugins/Runtime/Metasound/Source/MetasoundEngine/Private/MetasoundSource.cpp
@@ -542,6 +542,30 @@ float UMetaSoundSource::GetDuration() const
 	return IsOneShot() ? INDEFINITELY_LOOPING_DURATION - 1.0f : INDEFINITELY_LOOPING_DURATION;
 }
 
+Metasound::Frontend::FDocumentAccessPtr UMetaSoundSource::GetDocumentAccessPtr()
+{
+	using namespace Metasound::Frontend;
+
+	// Mutation of a document via the soft deprecated access ptr/controller system is not tracked by
+	// the builder registry, so the document cache is invalidated here. It is discouraged to mutate
+	// documents using both systems at the same time as it can corrupt a builder document's cache.
+	const FMetasoundFrontendClassName& Name = RootMetasoundDocument.RootGraph.Metadata.GetClassName();
+	IDocumentBuilderRegistry::GetChecked().InvalidateDocumentCache(Name);
+
+	// Return document using FAccessPoint to inform the TAccessPtr when the 
+	// object is no longer valid.
+	return MakeAccessPtr<FDocumentAccessPtr>(RootMetasoundDocument.AccessPoint, RootMetasoundDocument);
+}
+
+Metasound::Frontend::FConstDocumentAccessPtr UMetaSoundSource::GetDocumentConstAccessPtr() const
+{
+	using namespace Metasound::Frontend;
+
+	// Return document using FAccessPoint to inform the TAccessPtr when the 
+	// object is no longer valid.
+	return MakeAccessPtr<FConstDocumentAccessPtr>(RootMetasoundDocument.AccessPoint, RootMetasoundDocument);
+}
+
 bool UMetaSoundSource::ImplementsParameterInterface(Audio::FParameterInterfacePtr InInterface) const
 {
 	const FMetasoundFrontendVersion Version { InInterface->GetName(), { InInterface->GetVersion().Major, InInterface->GetVersion().Minor } };
diff --git a/Engine/Plugins/Runtime/Metasound/Source/MetasoundEngine/Public/Metasound.h b/Engine/Plugins/Runtime/Metasound/Source/MetasoundEngine/Public/Metasound.h
--- a/Engine/Plugins/Runtime/Metasound/Source/MetasoundEngine/Public/Metasound.h
+++ b/Engine/Plugins/Runtime/Metasound/Source/MetasoundEngine/Public/Metasound.h
@@ -168,21 +168,8 @@ protected:
 	virtual void SetReferencedAssetClasses(TSet<Metasound::Frontend::IMetaSoundAssetManager::FAssetInfo>&& InAssetClasses) override;
 #endif // #if WITH_EDITOR
 
-	Metasound::Frontend::FDocumentAccessPtr GetDocumentAccessPtr() override
-	{
-		using namespace Metasound::Frontend;
-		// Return document using FAccessPoint to inform the TAccessPtr when the 
-		// object is no longer valid.
-		return MakeAccessPtr<FDocumentAccessPtr>(RootMetaSoundDocument.AccessPoint, RootMetaSoundDocument);
-	}
-
-	Metasound::Frontend::FConstDocumentAccessPtr GetDocumentConstAccessPtr() const override
-	{
-		using namespace Metasound::Frontend;
-		// Return document using FAccessPoint to inform the TAccessPtr when the 
-		// object is no longer valid.
-		return MakeAccessPtr<FConstDocumentAccessPtr>(RootMetaSoundDocument.AccessPoint, RootMetaSoundDocument);
-	}
+	Metasound::Frontend::FDocumentAccessPtr GetDocumentAccessPtr() override;
+	Metasound::Frontend::FConstDocumentAccessPtr GetDocumentConstAccessPtr() const override;
 
 private:
 	virtual FMetasoundFrontendDocument& GetDocument() override
diff --git a/Engine/Plugins/Runtime/Metasound/Source/MetasoundEngine/Public/MetasoundBuilderSubsystem.h b/Engine/Plugins/Runtime/Metasound/Source/MetasoundEngine/Public/MetasoundBuilderSubsystem.h
--- a/Engine/Plugins/Runtime/Metasound/Source/MetasoundEngine/Public/MetasoundBuilderSubsystem.h
+++ b/Engine/Plugins/Runtime/Metasound/Source/MetasoundEngine/Public/MetasoundBuilderSubsystem.h
@@ -350,11 +350,11 @@ public:
 	// Rename the document's root graph class with a guid and optional namespace and variant
 	void RenameRootGraphClass(const FMetasoundFrontendClassName& InName);
 
-#if WITH_EDITOR
 	// Primarily used by editor transaction stack to avoid corruption when object is
 	// changed outside of the builder API. Generally discouraged for direct use otherwise.
-	void ReloadCache();
+	void ReloadCache(bool bPrimeCache = true);
 
+#if WITH_EDITOR
 	// Sets the author of the MetaSound.
 	void SetAuthor(const FString& InAuthor);
 #endif // WITH_EDITOR
@@ -390,10 +390,9 @@ protected:
 	// Should only be used when builders is not being attached to an existing, serialized MetaSound asset.
 	virtual void CreateTransientBuilder() PURE_VIRTUAL(UMetaSoundBuilderBase::CreateTransientBuilder, );
 
-	const FMetaSoundFrontendDocumentBuilder& GetConstBuilder() const
-	{
-		return Builder;
-	}
+	const FMetaSoundFrontendDocumentBuilder& GetConstBuilder() const;
+
+	void InvalidateCache();
 
 	// Runs build, conforming the document and corresponding object data on a MetaSound UObject to that managed by this builder.
 	template <typename UClassType>
@@ -570,8 +569,7 @@ private:
 
 public:
 	virtual void Initialize(FSubsystemCollectionBase& Collection) override;
-
-	virtual const Metasound::Frontend::FDocumentModifyDelegates* FindModifyDelegates(const FMetasoundFrontendClassName& InClassName) const override;
+	virtual void InvalidateDocumentCache(const FMetasoundFrontendClassName& InClassName) const override;
 
 	static UMetaSoundBuilderSubsystem& GetChecked();
 	static const UMetaSoundBuilderSubsystem& GetConstChecked();
diff --git a/Engine/Plugins/Runtime/Metasound/Source/MetasoundEngine/Public/MetasoundSource.h b/Engine/Plugins/Runtime/Metasound/Source/MetasoundEngine/Public/MetasoundSource.h
--- a/Engine/Plugins/Runtime/Metasound/Source/MetasoundEngine/Public/MetasoundSource.h
+++ b/Engine/Plugins/Runtime/Metasound/Source/MetasoundEngine/Public/MetasoundSource.h
@@ -201,21 +201,8 @@ public:
 	FOnGeneratorInstanceDestroyed OnGeneratorInstanceDestroyed;
 
 protected:
-	Metasound::Frontend::FDocumentAccessPtr GetDocumentAccessPtr() override
-	{
-		using namespace Metasound::Frontend;
-		// Return document using FAccessPoint to inform the TAccessPtr when the 
-		// object is no longer valid.
-		return MakeAccessPtr<FDocumentAccessPtr>(RootMetasoundDocument.AccessPoint, RootMetasoundDocument);
-	}
-
-	Metasound::Frontend::FConstDocumentAccessPtr GetDocumentConstAccessPtr() const override
-	{
-		using namespace Metasound::Frontend;
-		// Return document using FAccessPoint to inform the TAccessPtr when the 
-		// object is no longer valid.
-		return MakeAccessPtr<FConstDocumentAccessPtr>(RootMetasoundDocument.AccessPoint, RootMetasoundDocument);
-	}
+	Metasound::Frontend::FDocumentAccessPtr GetDocumentAccessPtr() override;
+	Metasound::Frontend::FConstDocumentAccessPtr GetDocumentConstAccessPtr() const override;
 
 	virtual const UClass& GetBaseMetaSoundUClass() const final override;
 	virtual const FMetasoundFrontendDocument& GetDocument() const override;
diff --git a/Engine/Plugins/Runtime/Metasound/Source/MetasoundFrontend/Private/MetasoundDocumentInterface.cpp b/Engine/Plugins/Runtime/Metasound/Source/MetasoundFrontend/Private/MetasoundDocumentInterface.cpp
--- a/Engine/Plugins/Runtime/Metasound/Source/MetasoundFrontend/Private/MetasoundDocumentInterface.cpp
+++ b/Engine/Plugins/Runtime/Metasound/Source/MetasoundFrontend/Private/MetasoundDocumentInterface.cpp
@@ -25,4 +25,9 @@ namespace Metasound::Frontend
 		checkf(!GetInstance, TEXT("Failed to initialize MetaSoundDocumentBuilderRegistry getter: Cannot reinitialize once initialized."))
 		GetInstance = MoveTemp(InGetInstance);
 	}
+
+	IMetaSoundDocumentBuilderRegistry& IMetaSoundDocumentBuilderRegistry::GetChecked()
+	{
+		return static_cast<IMetaSoundDocumentBuilderRegistry&>(IDocumentBuilderRegistry::GetChecked());
+	}
 } // namespace Metasound::Frontend
\ No newline at end of file
diff --git a/Engine/Plugins/Runtime/Metasound/Source/MetasoundFrontend/Private/MetasoundFrontendDocumentBuilder.cpp b/Engine/Plugins/Runtime/Metasound/Source/MetasoundFrontend/Private/MetasoundFrontendDocumentBuilder.cpp
--- a/Engine/Plugins/Runtime/Metasound/Source/MetasoundFrontend/Private/MetasoundFrontendDocumentBuilder.cpp
+++ b/Engine/Plugins/Runtime/Metasound/Source/MetasoundFrontend/Private/MetasoundFrontendDocumentBuilder.cpp
@@ -582,7 +582,7 @@ FMetaSoundFrontendDocumentBuilder::FMetaSoundFrontendDocumentBuilder(TScriptInte
 {
 	if (DocumentInterface)
 	{
-		ReloadCacheInternal();
+		InitCacheInternal();
 	}
 }
 
@@ -592,7 +592,7 @@ FMetaSoundFrontendDocumentBuilder::FMetaSoundFrontendDocumentBuilder(TScriptInte
 {
 	if (DocumentInterface)
 	{
-		ReloadCacheInternal();
+		InitCacheInternal();
 	}
 }
 
@@ -1207,7 +1207,7 @@ void FMetaSoundFrontendDocumentBuilder::ClearGraph()
 	GraphClass.PresetOptions.InputsInheritingDefault.Reset();
 	GetDocument().Interfaces.Reset();
 	RemoveUnusedDependencies();
-	ReloadCacheInternal();
+	InitCacheInternal();
 }
 
 bool FMetaSoundFrontendDocumentBuilder::ContainsDependencyOfType(EMetasoundFrontendClassType ClassType) const
@@ -1279,7 +1279,7 @@ bool FMetaSoundFrontendDocumentBuilder::ConvertToPreset(const FMetasoundFrontend
 	FRebuildPresetRootGraph RebuildPresetRootGraph(InReferencedDocument);
 	if (RebuildPresetRootGraph.Transform(GetDocument()))
 	{
-		ReloadCacheInternal();
+		InitCacheInternal();
 		return true;
 	}
 	return false;
@@ -1890,6 +1890,12 @@ void FMetaSoundFrontendDocumentBuilder::InitNodeLocations()
 #endif // WITH_EDITORONLY_DATA
 }
 
+void FMetaSoundFrontendDocumentBuilder::InvalidateCache()
+{
+	constexpr bool bPrimeCache = false;
+	InitCacheInternal(bPrimeCache);
+}
+
 bool FMetaSoundFrontendDocumentBuilder::IsDependencyReferenced(const FGuid& InClassID) const
 {
 	return DocumentCache->GetNodeCache().ContainsNodesOfClassID(InClassID);
@@ -2052,19 +2058,10 @@ bool FMetaSoundFrontendDocumentBuilder::TransformTemplateNodes()
 	return bModified;
 }
 
-void FMetaSoundFrontendDocumentBuilder::ReloadCacheInternal()
+void FMetaSoundFrontendDocumentBuilder::InitCacheInternal(bool bPrimeCache)
 {
 	using namespace Metasound::Frontend;
-	if (!DocumentCache.IsValid())
-	{
-		DocumentCache = MakeShared<FDocumentCache>(GetDocument());
-	}
-
-	// Must be called after the constructor to allow for passing shared pointer to
-	// sub-caches (i.e. can't be RAII because initialization requires creating
-	// additional shared pointers in the ctor which is explicitly forbidden)
-	TSharedPtr<FDocumentCache> CacheConcrete = StaticCastSharedPtr<FDocumentCache>(DocumentCache);
-	CacheConcrete->Init(*DocumentDelegates);
+	DocumentCache = FDocumentCache::Create(GetDocument(), DocumentDelegates, bPrimeCache);
 }
 
 bool FMetaSoundFrontendDocumentBuilder::RemoveDependency(const FGuid& InClassID)
@@ -2485,12 +2482,13 @@ bool FMetaSoundFrontendDocumentBuilder::RenameRootGraphClass(const FMetasoundFro
 	}
 }
 
-#if WITH_EDITOR
 void FMetaSoundFrontendDocumentBuilder::ReloadCache()
 {
-	ReloadCacheInternal();
+	constexpr bool bPrimeCache = true;
+	InitCacheInternal(bPrimeCache);
 }
 
+#if WITH_EDITOR
 void FMetaSoundFrontendDocumentBuilder::SetAuthor(const FString& InAuthor)
 {
 	FMetasoundFrontendClassMetadata& ClassMetadata = GetDocument().RootGraph.Metadata;
diff --git a/Engine/Plugins/Runtime/Metasound/Source/MetasoundFrontend/Private/MetasoundFrontendDocumentCache.cpp b/Engine/Plugins/Runtime/Metasound/Source/MetasoundFrontend/Private/MetasoundFrontendDocumentCache.cpp
--- a/Engine/Plugins/Runtime/Metasound/Source/MetasoundFrontend/Private/MetasoundFrontendDocumentCache.cpp
+++ b/Engine/Plugins/Runtime/Metasound/Source/MetasoundFrontend/Private/MetasoundFrontendDocumentCache.cpp
@@ -70,54 +70,27 @@ namespace Metasound::Frontend
 		}
 	}
 
-	FDocumentCache::FDocumentCache(const FMetasoundFrontendDocument& InDocument)
+	FDocumentCache::FDocumentCache(const FMetasoundFrontendDocument& InDocument, TSharedRef<FDocumentModifyDelegates> Delegates)
 		: Document(&InDocument)
+		, ModifyDelegates(Delegates)
 	{
 	}
 
-	void FDocumentCache::Init(FDocumentModifyDelegates& OutDelegates)
+	void FDocumentCache::Init(bool bPrimeCache)
 	{
-		METASOUND_TRACE_CPUPROFILER_EVENT_SCOPE(FDocumentCache::Init);
-
 		check(Document);
 
-		IDToIndex.Reset();
-		KeyToIndex.Reset();
-
-		for (int32 Index = 0; Index < Document->Dependencies.Num(); ++Index)
+		if (bPrimeCache)
 		{
-			const FMetasoundFrontendClass& Class = Document->Dependencies[Index];
-			const FMetasoundFrontendClassMetadata& Metadata = Class.Metadata;
-			FNodeRegistryKey Key = NodeRegistryKey::CreateKey(Metadata);
-			KeyToIndex.Add(MoveTemp(Key), Index);
-			IDToIndex.Add(Class.ID, Index);
+			DependencyCache = MakeShared<FDocumentDependencyCache>(GetDocument());
+			EdgeCache = FDocumentGraphEdgeCache::Create(AsShared(), ModifyDelegates->EdgeDelegates);
+			NodeCache = FDocumentGraphNodeCache::Create(AsShared(), ModifyDelegates->NodeDelegates);
+			InterfaceCache = FDocumentGraphInterfaceCache::Create(AsShared(), ModifyDelegates->InterfaceDelegates);
 		}
 
-		OutDelegates.OnDependencyAdded.Remove(OnAddedHandle);
-		OutDelegates.OnRemoveSwappingDependency.Remove(OnRemoveSwappingHandle);
-		OutDelegates.OnRenamingDependencyClass.Remove(OnRenamingDependencyClassHandle);
-
-		OnAddedHandle = OutDelegates.OnDependencyAdded.AddSP(this, &FDocumentCache::OnDependencyAdded);
-		OnRemoveSwappingHandle = OutDelegates.OnRemoveSwappingDependency.AddSP(this, &FDocumentCache::OnRemoveSwappingDependency);
-		OnRenamingDependencyClassHandle = OutDelegates.OnRenamingDependencyClass.AddSP(this, &FDocumentCache::OnRenamingDependencyClass);
-
-		if (!EdgeCache.IsValid())
-		{
-			EdgeCache = MakeShared<FDocumentGraphEdgeCache>(AsShared());
-		}
-		EdgeCache->Init(OutDelegates.EdgeDelegates);
-
-		if (!NodeCache.IsValid())
-		{
-			NodeCache = MakeShared<FDocumentGraphNodeCache>(AsShared());
-		}
-		NodeCache->Init(OutDelegates.NodeDelegates);
-
-		if (!InterfaceCache.IsValid())
-		{
-			InterfaceCache = MakeShared<FDocumentGraphInterfaceCache>(AsShared());
-		}
-		InterfaceCache->Init(OutDelegates.InterfaceDelegates);
+		ModifyDelegates->OnDependencyAdded.AddSP(this, &FDocumentCache::OnDependencyAdded);
+		ModifyDelegates->OnRemoveSwappingDependency.AddSP(this, &FDocumentCache::OnRemoveSwappingDependency);
+		ModifyDelegates->OnRenamingDependencyClass.AddSP(this, &FDocumentCache::OnRenamingDependencyClass);
 	}
 
 	const FMetasoundFrontendDocument& FDocumentCache::GetDocument() const
@@ -126,9 +99,20 @@ namespace Metasound::Frontend
 		return *Document;
 	}
 
+	TSharedRef<FDocumentCache> FDocumentCache::Create(const FMetasoundFrontendDocument& InDocument, TSharedRef<FDocumentModifyDelegates> Delegates, bool bPrimeCache = false)
+	{
+		METASOUND_TRACE_CPUPROFILER_EVENT_SCOPE(Metasound::Frontend::FDocumentCache::Create);
+
+		// Factory style constructor as restrictions on construction of shared pointers disallows passing of this document cache's pointer to sub-caches
+		TSharedRef<FDocumentCache> Cache = MakeShared<FDocumentCache>(InDocument, Delegates);
+		Cache->Init(bPrimeCache);
+		return Cache;
+	}
+
 	bool FDocumentCache::ContainsDependency(const FNodeRegistryKey& InClassKey) const
 	{
-		return KeyToIndex.Contains(InClassKey);
+		const FDocumentDependencyCache& Cache = GetDependencyCache();
+		return Cache.KeyToIndex.Contains(InClassKey);
 	}
 
 	bool FDocumentCache::ContainsDependencyOfType(EMetasoundFrontendClassType ClassType) const
@@ -140,7 +124,8 @@ namespace Metasound::Frontend
 
 	const FMetasoundFrontendClass* FDocumentCache::FindDependency(const FNodeRegistryKey& InClassKey) const
 	{
-		if (const int32* Index = KeyToIndex.Find(InClassKey))
+		const FDocumentDependencyCache& Cache = GetDependencyCache();
+		if (const int32* Index = Cache.KeyToIndex.Find(InClassKey))
 		{
 			return &GetDocument().Dependencies[*Index];
 		}
@@ -150,7 +135,8 @@ namespace Metasound::Frontend
 
 	const FMetasoundFrontendClass* FDocumentCache::FindDependency(const FGuid& InClassID) const
 	{
-		if (const int32* Index = IDToIndex.Find(InClassID))
+		const FDocumentDependencyCache& Cache = GetDependencyCache();
+		if (const int32* Index = Cache.IDToIndex.Find(InClassID))
 		{
 			return &GetDocument().Dependencies[*Index];
 		}
@@ -160,26 +146,58 @@ namespace Metasound::Frontend
 
 	const int32* FDocumentCache::FindDependencyIndex(const Metasound::Frontend::FNodeRegistryKey& InClassKey) const
 	{
-		return KeyToIndex.Find(InClassKey);
+		const FDocumentDependencyCache& Cache = GetDependencyCache();
+		return Cache.KeyToIndex.Find(InClassKey);
 	}
 
 	const int32* FDocumentCache::FindDependencyIndex(const FGuid& InClassID) const
 	{
-		return IDToIndex.Find(InClassID);
+		const FDocumentDependencyCache& Cache = GetDependencyCache();
+		return Cache.IDToIndex.Find(InClassID);
+	}
+
+	FDocumentCache::FDocumentDependencyCache& FDocumentCache::GetDependencyCache()
+	{
+		if (!DependencyCache.IsValid())
+		{
+			DependencyCache = MakeShared<FDocumentDependencyCache>(GetDocument());
+		}
+		return *DependencyCache;
+	}
+
+	const FDocumentCache::FDocumentDependencyCache& FDocumentCache::GetDependencyCache() const
+	{
+		if (!DependencyCache.IsValid())
+		{
+			DependencyCache = MakeShared<FDocumentDependencyCache>(GetDocument());
+		}
+		return *DependencyCache;
 	}
 
 	const IDocumentGraphEdgeCache& FDocumentCache::GetEdgeCache() const
 	{
+		if (!EdgeCache.IsValid())
+		{
+			EdgeCache = FDocumentGraphEdgeCache::Create(AsShared(), ModifyDelegates->EdgeDelegates);
+		}
 		return *EdgeCache;
 	}
 
 	const IDocumentGraphNodeCache& FDocumentCache::GetNodeCache() const
 	{
+		if (!NodeCache.IsValid())
+		{
+			NodeCache = FDocumentGraphNodeCache::Create(AsShared(), ModifyDelegates->NodeDelegates);
+		}
 		return *NodeCache;
 	}
 
 	const IDocumentGraphInterfaceCache& FDocumentCache::GetInterfaceCache() const
 	{
+		if (!InterfaceCache.IsValid())
+		{
+			InterfaceCache = FDocumentGraphInterfaceCache::Create(AsShared(), ModifyDelegates->InterfaceDelegates);
+		}
 		return *InterfaceCache;
 	}
 
@@ -189,8 +207,10 @@ namespace Metasound::Frontend
 
 		const FMetasoundFrontendClass& Dependency = GetDocument().Dependencies[NewIndex];
 		const FNodeRegistryKey Key = NodeRegistryKey::CreateKey(Dependency.Metadata);
-		KeyToIndex.Add(Key, NewIndex);
-		IDToIndex.Add(Dependency.ID, NewIndex);
+
+		FDocumentDependencyCache& Cache = GetDependencyCache();
+		Cache.KeyToIndex.Add(Key, NewIndex);
+		Cache.IDToIndex.Add(Dependency.ID, NewIndex);
 	}
 
 	void FDocumentCache::OnRemoveSwappingDependency(int32 SwapIndex, int32 LastIndex)
@@ -203,8 +223,10 @@ namespace Metasound::Frontend
 
 		const FNodeRegistryKey SwapKey = NodeRegistryKey::CreateKey(SwapDependency.Metadata);
 		const FNodeRegistryKey LastKey = NodeRegistryKey::CreateKey(LastDependency.Metadata);
-		DocumentCachePrivate::RemoveSwapMapIndexChecked(SwapKey, LastKey, KeyToIndex);
-		DocumentCachePrivate::RemoveSwapMapIndexChecked(SwapDependency.ID, LastDependency.ID, IDToIndex);
+
+		FDocumentDependencyCache& Cache = GetDependencyCache();
+		DocumentCachePrivate::RemoveSwapMapIndexChecked(SwapKey, LastKey, Cache.KeyToIndex);
+		DocumentCachePrivate::RemoveSwapMapIndexChecked(SwapDependency.ID, LastDependency.ID, Cache.IDToIndex);
 	}
 
 	void FDocumentCache::OnRenamingDependencyClass(const int32 IndexBeingRenamed, const FMetasoundFrontendClassName& NewName)
@@ -217,13 +239,36 @@ namespace Metasound::Frontend
 		NewMetadata.SetClassName(NewName);
 		const FNodeRegistryKey NewKey = NodeRegistryKey::CreateKey(NewMetadata);
 
-		KeyToIndex.Add(NewKey, IndexBeingRenamed);
-		KeyToIndex.Remove(OldKey);
+		FDocumentDependencyCache& Cache = GetDependencyCache();
+		Cache.KeyToIndex.Add(NewKey, IndexBeingRenamed);
+		Cache.KeyToIndex.Remove(OldKey);
 	}
 
-	FDocumentGraphInterfaceCache::FDocumentGraphInterfaceCache(TSharedRef<IDocumentCache> ParentCache)
+	FDocumentGraphInterfaceCache::FDocumentGraphInterfaceCache(TSharedRef<const IDocumentCache> ParentCache)
 		: Parent(ParentCache)
 	{
+		const FMetasoundFrontendGraphClass& GraphClass = Parent->GetDocument().RootGraph;
+		const TArray<FMetasoundFrontendClassInput>& Inputs = GraphClass.Interface.Inputs;
+		for (int32 Index = 0; Index < Inputs.Num(); ++Index)
+		{
+			InputNameToIndex.Add(Inputs[Index].Name, Index);
+		};
+
+		const TArray<FMetasoundFrontendClassOutput>& Outputs = GraphClass.Interface.Outputs;
+		for (int32 Index = 0; Index < Outputs.Num(); ++Index)
+		{
+			OutputNameToIndex.Add(Outputs[Index].Name, Index);
+		};
+	}
+
+	TSharedRef<FDocumentGraphInterfaceCache> FDocumentGraphInterfaceCache::Create(TSharedRef<const IDocumentCache> ParentCache, FInterfaceModifyDelegates& OutDelegates)
+	{
+		METASOUND_TRACE_CPUPROFILER_EVENT_SCOPE(Metasound::Frontend::FDocumentGraphInterfaceCache::Create);
+
+		// Factory style constructor as restrictions on construction of shared pointers disallows passing of this document cache's pointer to sub-caches
+		TSharedRef<FDocumentGraphInterfaceCache> Cache = MakeShared<FDocumentGraphInterfaceCache>(ParentCache);
+		Cache->Init(OutDelegates);
+		return Cache;
 	}
 
 	const FMetasoundFrontendClassInput* FDocumentGraphInterfaceCache::FindInput(FName InputName) const
@@ -250,31 +295,11 @@ namespace Metasound::Frontend
 
 	void FDocumentGraphInterfaceCache::Init(FInterfaceModifyDelegates& OutDelegates)
 	{
-		InputNameToIndex.Reset();
-		OutputNameToIndex.Reset();
 
-		const FMetasoundFrontendGraphClass& GraphClass = Parent->GetDocument().RootGraph;
-		const TArray<FMetasoundFrontendClassInput>& Inputs = GraphClass.Interface.Inputs;
-		for (int32 Index = 0; Index < Inputs.Num(); ++Index)
-		{
-			InputNameToIndex.Add(Inputs[Index].Name, Index);
-		};
-
-		const TArray<FMetasoundFrontendClassOutput>& Outputs = GraphClass.Interface.Outputs;
-		for (int32 Index = 0; Index < Outputs.Num(); ++Index)
-		{
-			OutputNameToIndex.Add(Outputs[Index].Name, Index);
-		};
-
-		OutDelegates.OnInputAdded.Remove(OnInputAddedHandle);
-		OutDelegates.OnOutputAdded.Remove(OnOutputAddedHandle);
-		OutDelegates.OnRemovingInput.Remove(OnRemovingInputHandle);
-		OutDelegates.OnRemovingOutput.Remove(OnRemovingOutputHandle);
-
-		OnInputAddedHandle = OutDelegates.OnInputAdded.AddSP(this, &FDocumentGraphInterfaceCache::OnInputAdded);
-		OnOutputAddedHandle = OutDelegates.OnOutputAdded.AddSP(this, &FDocumentGraphInterfaceCache::OnOutputAdded);
-		OnRemovingInputHandle = OutDelegates.OnRemovingInput.AddSP(this, &FDocumentGraphInterfaceCache::OnRemovingInput);
-		OnRemovingOutputHandle = OutDelegates.OnRemovingOutput.AddSP(this, &FDocumentGraphInterfaceCache::OnRemovingOutput);
+		OutDelegates.OnInputAdded.AddSP(this, &FDocumentGraphInterfaceCache::OnInputAdded);
+		OutDelegates.OnOutputAdded.AddSP(this, &FDocumentGraphInterfaceCache::OnOutputAdded);
+		OutDelegates.OnRemovingInput.AddSP(this, &FDocumentGraphInterfaceCache::OnRemovingInput);
+		OutDelegates.OnRemovingOutput.AddSP(this, &FDocumentGraphInterfaceCache::OnRemovingOutput);
 	}
 
 	void FDocumentGraphInterfaceCache::OnInputAdded(int32 NewIndex)
@@ -309,16 +334,9 @@ namespace Metasound::Frontend
 		OutputNameToIndex.Remove(Output.Name);
 	}
 
-	FDocumentGraphNodeCache::FDocumentGraphNodeCache(TSharedRef<IDocumentCache> ParentCache)
+	FDocumentGraphNodeCache::FDocumentGraphNodeCache(TSharedRef<const IDocumentCache> ParentCache)
 		: Parent(ParentCache)
 	{
-	}
-
-	void FDocumentGraphNodeCache::Init(FNodeModifyDelegates& OutDelegates)
-	{
-		IDToIndex.Reset();
-		ClassIDToNodeIndices.Reset();
-
 		const FMetasoundFrontendGraph& Graph = Parent->GetDocument().RootGraph.Graph;
 		const TArray<FMetasoundFrontendNode>& Nodes = Graph.Nodes;
 		for (int32 Index = 0; Index < Nodes.Num(); ++Index)
@@ -331,12 +349,6 @@ namespace Metasound::Frontend
 			IDToIndex.Add(NodeID, Index);
 			ClassIDToNodeIndices.FindOrAdd(Node.ClassID).Add(Index);
 		}
-
-		OutDelegates.OnNodeAdded.Remove(OnAddedHandle);
-		OnAddedHandle = OutDelegates.OnNodeAdded.AddSP(this, &FDocumentGraphNodeCache::OnNodeAdded);
-
-		OutDelegates.OnRemoveSwappingNode.Remove(OnRemoveSwappingHandle);
-		OnRemoveSwappingHandle = OutDelegates.OnRemoveSwappingNode.AddSP(this, &FDocumentGraphNodeCache::OnRemoveSwappingNode);
 	}
 
 	bool FDocumentGraphNodeCache::ContainsNode(const FGuid& InNodeID) const
@@ -344,6 +356,16 @@ namespace Metasound::Frontend
 		return IDToIndex.Contains(InNodeID);
 	}
 
+	TSharedRef<FDocumentGraphNodeCache> FDocumentGraphNodeCache::Create(TSharedRef<const IDocumentCache> ParentCache, FNodeModifyDelegates& OutDelegates)
+	{
+		METASOUND_TRACE_CPUPROFILER_EVENT_SCOPE(Metasound::Frontend::FDocumentGraphNodeCache::Create);
+
+		// Factory style constructor as restrictions on construction of shared pointers disallows passing of this document cache's pointer to sub-caches
+		TSharedRef<FDocumentGraphNodeCache> Cache = MakeShared<FDocumentGraphNodeCache>(ParentCache);
+		Cache->Init(OutDelegates);
+		return Cache;
+	}
+
 	const int32* FDocumentGraphNodeCache::FindNodeIndex(const FGuid& InNodeID) const
 	{
 		return IDToIndex.Find(InNodeID);
@@ -654,6 +676,12 @@ namespace Metasound::Frontend
 		return nullptr;
 	}
 
+	void FDocumentGraphNodeCache::Init(FNodeModifyDelegates& OutDelegates)
+	{
+		OutDelegates.OnNodeAdded.AddSP(this, &FDocumentGraphNodeCache::OnNodeAdded);
+		OutDelegates.OnRemoveSwappingNode.AddSP(this, &FDocumentGraphNodeCache::OnRemoveSwappingNode);
+	}
+
 	void FDocumentGraphNodeCache::OnNodeAdded(int32 InNewIndex)
 	{
 		const FMetasoundFrontendDocument& Document = Parent->GetDocument();
@@ -676,16 +704,9 @@ namespace Metasound::Frontend
 		DocumentCachePrivate::RemoveSwapMapIndexChecked(SwapNode.GetID(), LastNode.GetID(), IDToIndex);
 	}
 
-	FDocumentGraphEdgeCache::FDocumentGraphEdgeCache(TSharedRef<IDocumentCache> ParentCache)
+	FDocumentGraphEdgeCache::FDocumentGraphEdgeCache(TSharedRef<const IDocumentCache> ParentCache)
 		: Parent(ParentCache)
 	{
-	}
-
-	void FDocumentGraphEdgeCache::Init(FEdgeModifyDelegates& OutDelegates)
-	{
-		OutputToEdgeIndices.Reset();
-		InputToEdgeIndex.Reset();
-
 		const TArray<FMetasoundFrontendEdge>& Edges = Parent->GetDocument().RootGraph.Graph.Edges;
 		for (int32 Index = 0; Index < Edges.Num(); ++Index)
 		{
@@ -693,12 +714,6 @@ namespace Metasound::Frontend
 			OutputToEdgeIndices.FindOrAdd(Edge.GetFromVertexHandle()).Add(Index);
 			InputToEdgeIndex.Add(Edge.GetToVertexHandle()) = Index;
 		}
-
-		OutDelegates.OnEdgeAdded.Remove(OnAddedHandle);
-		OnAddedHandle = OutDelegates.OnEdgeAdded.AddSP(this, &FDocumentGraphEdgeCache::OnEdgeAdded);
-
-		OutDelegates.OnRemoveSwappingEdge.Remove(OnRemoveSwappingHandle);
-		OnRemoveSwappingHandle = OutDelegates.OnRemoveSwappingEdge.AddSP(this, &FDocumentGraphEdgeCache::OnRemoveSwappingEdge);
 	}
 
 	bool FDocumentGraphEdgeCache::ContainsEdge(const FMetasoundFrontendEdge& InEdge) const
@@ -714,6 +729,16 @@ namespace Metasound::Frontend
 		return false;
 	}
 
+	TSharedRef<FDocumentGraphEdgeCache> FDocumentGraphEdgeCache::Create(TSharedRef<const IDocumentCache> ParentCache, FEdgeModifyDelegates& OutDelegates)
+	{
+		METASOUND_TRACE_CPUPROFILER_EVENT_SCOPE(Metasound::Frontend::FDocumentGraphEdgeCache::Create);
+
+		// Factory style constructor as restrictions on construction of shared pointers disallows passing of this document cache's pointer to sub-caches
+		TSharedRef<FDocumentGraphEdgeCache> Cache = MakeShared<FDocumentGraphEdgeCache>(ParentCache);
+		Cache->Init(OutDelegates);
+		return Cache;
+	}
+
 	TArray<const FMetasoundFrontendEdge*> FDocumentGraphEdgeCache::FindEdges(const FGuid& InNodeID, const FGuid& InVertexID) const
 	{
 		TArray<const FMetasoundFrontendEdge*> Edges;
@@ -750,6 +775,12 @@ namespace Metasound::Frontend
 		return { };
 	}
 
+	void FDocumentGraphEdgeCache::Init(FEdgeModifyDelegates& OutDelegates)
+	{
+		OutDelegates.OnEdgeAdded.AddSP(this, &FDocumentGraphEdgeCache::OnEdgeAdded);
+		OutDelegates.OnRemoveSwappingEdge.AddSP(this, &FDocumentGraphEdgeCache::OnRemoveSwappingEdge);
+	}
+
 	bool FDocumentGraphEdgeCache::IsNodeInputConnected(const FGuid& InNodeID, const FGuid& InVertexID) const
 	{
 		return InputToEdgeIndex.Contains(FMetasoundFrontendVertexHandle { InNodeID, InVertexID });
@@ -779,4 +810,18 @@ namespace Metasound::Frontend
 		DocumentCachePrivate::RemoveSwapMapIndexChecked(SwapEdge.GetToVertexHandle(), LastEdge.GetToVertexHandle(), InputToEdgeIndex);
 		DocumentCachePrivate::RemoveSwapMapArrayIndexChecked(SwapEdge.GetFromVertexHandle(), SwapIndex, LastEdge.GetFromVertexHandle(), LastIndex, OutputToEdgeIndices);
 	}
+
+	FDocumentCache::FDocumentDependencyCache::FDocumentDependencyCache(const FMetasoundFrontendDocument& InDocument)
+	{
+		METASOUND_TRACE_CPUPROFILER_EVENT_SCOPE(Metasound::Frontend::FDocumentDependencyCache::Create);
+
+		for (int32 Index = 0; Index < InDocument.Dependencies.Num(); ++Index)
+		{
+			const FMetasoundFrontendClass& Class = InDocument.Dependencies[Index];
+			const FMetasoundFrontendClassMetadata& Metadata = Class.Metadata;
+			FNodeRegistryKey Key = NodeRegistryKey::CreateKey(Metadata);
+			KeyToIndex.Add(MoveTemp(Key), Index);
+			IDToIndex.Add(Class.ID, Index);
+		}
+	}
 } // namespace Metasound::Frontend
diff --git a/Engine/Plugins/Runtime/Metasound/Source/MetasoundFrontend/Private/MetasoundFrontendDocumentCache.h b/Engine/Plugins/Runtime/Metasound/Source/MetasoundFrontend/Private/MetasoundFrontendDocumentCache.h
--- a/Engine/Plugins/Runtime/Metasound/Source/MetasoundFrontend/Private/MetasoundFrontendDocumentCache.h
+++ b/Engine/Plugins/Runtime/Metasound/Source/MetasoundFrontend/Private/MetasoundFrontendDocumentCache.h
@@ -18,9 +18,13 @@ namespace Metasound::Frontend
 
 	class FDocumentGraphEdgeCache : public IDocumentGraphEdgeCache
 	{
+		template <typename ObjectType, ESPMode Mode>
+		friend class SharedPointerInternals::TIntrusiveReferenceController;
+
+		FDocumentGraphEdgeCache(TSharedRef<const IDocumentCache> ParentCache);
+
 	public:
-		FDocumentGraphEdgeCache() = default;
-		FDocumentGraphEdgeCache(TSharedRef<IDocumentCache> ParentCache);
+		static TSharedRef<FDocumentGraphEdgeCache> Create(TSharedRef<const IDocumentCache> ParentCache, FEdgeModifyDelegates& OutDelegates);
 		virtual ~FDocumentGraphEdgeCache() = default;
 
 		// IDocumentGraphEdgeCache implementation
@@ -33,9 +37,9 @@ namespace Metasound::Frontend
 		virtual const int32* FindEdgeIndexToNodeInput(const FGuid& InNodeID, const FGuid& InVertexID) const override;
 		virtual const TArrayView<const int32> FindEdgeIndicesFromNodeOutput(const FGuid& InNodeID, const FGuid& InVertexID) const override;
 
-		void Init(FEdgeModifyDelegates& OutDelegates);
 
 	private:
+		void Init(FEdgeModifyDelegates& OutDelegates);
 		void OnEdgeAdded(int32 InNewIndex);
 		void OnRemoveSwappingEdge(int32 SwapIndex, int32 LastIndex);
 
@@ -45,18 +49,19 @@ namespace Metasound::Frontend
 		// Cache of Input NodeId/VertexId pairs to associated edge indices
 		TMap<FMetasoundFrontendVertexHandle, int32> InputToEdgeIndex;
 
-		TSharedPtr<IDocumentCache> Parent;
-
-		FDelegateHandle OnAddedHandle;
-		FDelegateHandle OnRemoveSwappingHandle;
+		TSharedPtr<const IDocumentCache> Parent;
 	};
 
 
 	class FDocumentGraphNodeCache : public IDocumentGraphNodeCache
 	{
+		template <typename ObjectType, ESPMode Mode>
+		friend class SharedPointerInternals::TIntrusiveReferenceController;
+
+		FDocumentGraphNodeCache(TSharedRef<const IDocumentCache> ParentCache);
+
 	public:
-		FDocumentGraphNodeCache() = default;
-		FDocumentGraphNodeCache(TSharedRef<IDocumentCache> ParentCache);
+		static TSharedRef<FDocumentGraphNodeCache> Create(TSharedRef<const IDocumentCache> ParentCache, FNodeModifyDelegates& OutDelegates);
 		virtual ~FDocumentGraphNodeCache() = default;
 
 		// IDocumentGraphNodeCache implementation
@@ -80,9 +85,8 @@ namespace Metasound::Frontend
 		virtual const FMetasoundFrontendVertex* FindReroutedOutputVertex(const FGuid& InNodeID, const FGuid& InVertexID, const FMetasoundFrontendNode** ConnectedNodes = nullptr, bool* bOutIsRerouted = nullptr) const override;
 		virtual const FMetasoundFrontendVertex* FindReroutedOutputVertex(const FGuid& InNodeID, FName InVertexName, const FMetasoundFrontendNode** ConnectedNodes = nullptr, bool* bOutIsRerouted = nullptr) const override;
 
-		void Init(FNodeModifyDelegates& OutDelegates);
-
 	private:
+		void Init(FNodeModifyDelegates& OutDelegates);
 		void OnNodeAdded(int32 NewIndex);
 		void OnRemoveSwappingNode(int32 IndexBeingRemoved, int32 LastIndex);
 
@@ -92,27 +96,27 @@ namespace Metasound::Frontend
 		// Cache of ClassID to referencing node indices
 		TSortedMap<FGuid, TArray<int32>> ClassIDToNodeIndices;
 
-		TSharedPtr<IDocumentCache> Parent;
-
-		FDelegateHandle OnAddedHandle;
-		FDelegateHandle OnRemoveSwappingHandle;
+		TSharedPtr<const IDocumentCache> Parent;
 	};
 
 
 	class FDocumentGraphInterfaceCache : public IDocumentGraphInterfaceCache
 	{
+		template <typename ObjectType, ESPMode Mode>
+		friend class SharedPointerInternals::TIntrusiveReferenceController;
+
+		FDocumentGraphInterfaceCache(TSharedRef<const IDocumentCache> ParentCache);
+
 	public:
-		FDocumentGraphInterfaceCache() = default;
-		FDocumentGraphInterfaceCache(TSharedRef<IDocumentCache> ParentCache);
+		static TSharedRef<FDocumentGraphInterfaceCache> Create(TSharedRef<const IDocumentCache> ParentCache, FInterfaceModifyDelegates& OutDelegates);
 		virtual ~FDocumentGraphInterfaceCache() = default;
 
 		// IDocumentGraphInterfaceCache implementation
 		virtual const FMetasoundFrontendClassInput* FindInput(FName InputName) const override;
 		virtual const FMetasoundFrontendClassOutput* FindOutput(FName OutputName) const override;
 
-		void Init(FInterfaceModifyDelegates& OutDelegates);
-
 	private:
+		void Init(FInterfaceModifyDelegates& OutDelegates);
 		void OnInputAdded(int32 NewIndex);
 		void OnOutputAdded(int32 NewIndex);
 		void OnRemovingInput(int32 IndexBeingRemoved);
@@ -124,20 +128,19 @@ namespace Metasound::Frontend
 		// Cache of Output name to array index of output
 		TMap<FName, int32> OutputNameToIndex;
 
-		TSharedPtr<IDocumentCache> Parent;
-
-		FDelegateHandle OnInputAddedHandle;
-		FDelegateHandle OnOutputAddedHandle;
-		FDelegateHandle OnRemovingInputHandle;
-		FDelegateHandle OnRemovingOutputHandle;
+		TSharedPtr<const IDocumentCache> Parent;
 	};
 
 
 	class FDocumentCache : public IDocumentCache
 	{
+		template <typename ObjectType, ESPMode Mode>
+		friend class SharedPointerInternals::TIntrusiveReferenceController;
+
+		FDocumentCache(const FMetasoundFrontendDocument& InDocument, TSharedRef<FDocumentModifyDelegates> Delegates);
+
 	public:
-		FDocumentCache() = default;
-		FDocumentCache(const FMetasoundFrontendDocument& InDocument);
+		static TSharedRef<FDocumentCache> Create(const FMetasoundFrontendDocument& InDocument, TSharedRef<FDocumentModifyDelegates> Delegates, bool bPrimeCache);
 		virtual ~FDocumentCache() = default;
 
 		virtual bool ContainsDependency(const FNodeRegistryKey& InClassKey) const override;
@@ -153,27 +156,38 @@ namespace Metasound::Frontend
 		virtual const IDocumentGraphNodeCache& GetNodeCache() const override;
 		virtual const IDocumentGraphInterfaceCache& GetInterfaceCache() const override;
 
-		void Init(FDocumentModifyDelegates& OutDelegates);
-
 	private:
+		void Init(bool bPrimeCache);
+
 		void OnDependencyAdded(int32 InNewIndex);
 		void OnRemoveSwappingDependency(int32 SwapIndex, int32 LastIndex);
 		void OnRenamingDependencyClass(const int32 IndexBeingRenamed, const FMetasoundFrontendClassName& NewName);
 
-		// Cache of dependency (Class) ID to corresponding class dependency index
-		TSortedMap<FGuid, int32> IDToIndex;
+		struct FDocumentDependencyCache : public TSharedFromThis<FDocumentDependencyCache>
+		{
+			FDocumentDependencyCache(const FMetasoundFrontendDocument& InDocument);
+			~FDocumentDependencyCache() = default;
+
+			// Cache of dependency (Class) ID to corresponding class dependency index
+			TSortedMap<FGuid, int32> IDToIndex;
 
-		// Cache of version data to corresponding class dependency index
-		TSortedMap<FNodeRegistryKey, int32> KeyToIndex;
+			// Cache of version data to corresponding class dependency index
+			TSortedMap<FNodeRegistryKey, int32> KeyToIndex;
+		};
 
-		TSharedPtr<FDocumentGraphEdgeCache> EdgeCache;
-		TSharedPtr<FDocumentGraphNodeCache> NodeCache;
-		TSharedPtr<FDocumentGraphInterfaceCache> InterfaceCache;
+		// Private as cached queried data for fast access at document layer is minimal (i.e. dependencies by class id),
+		// so API is just passed through document cache interface. Opaque create getter API kept similar to calls above
+		// to separately implemented caches for parity.
+		FDocumentDependencyCache& GetDependencyCache();
+		const FDocumentDependencyCache& GetDependencyCache() const;
 
-		FDelegateHandle OnAddedHandle;
-		FDelegateHandle OnRemoveSwappingHandle;
-		FDelegateHandle OnRenamingDependencyClassHandle;
+		// Mutable to be able to be loaded on demand if cache does not exist at the time of an initial query
+		mutable TSharedPtr<FDocumentDependencyCache> DependencyCache;
+		mutable TSharedPtr<FDocumentGraphEdgeCache> EdgeCache;
+		mutable TSharedPtr<FDocumentGraphNodeCache> NodeCache;
+		mutable TSharedPtr<FDocumentGraphInterfaceCache> InterfaceCache;
 
 		const FMetasoundFrontendDocument* Document = nullptr;
+		TSharedRef<FDocumentModifyDelegates> ModifyDelegates;
 	};
 } // namespace Metasound::Frontend
diff --git a/Engine/Plugins/Runtime/Metasound/Source/MetasoundFrontend/Public/MetasoundDocumentInterface.h b/Engine/Plugins/Runtime/Metasound/Source/MetasoundFrontend/Public/MetasoundDocumentInterface.h
--- a/Engine/Plugins/Runtime/Metasound/Source/MetasoundFrontend/Public/MetasoundDocumentInterface.h
+++ b/Engine/Plugins/Runtime/Metasound/Source/MetasoundFrontend/Public/MetasoundDocumentInterface.h
@@ -45,11 +45,10 @@ namespace Metasound::Frontend
 	class METASOUNDFRONTEND_API IDocumentBuilderRegistry
 	{
 	public:
-		// Returns delegates used to mutate any internal builder cached state or notify listeners of external system that
-		// has mutated a given document class.  Exists primarily for backward compatibility with the DocumentController system,
-		// and is not recommended for use outside of MetaSound plugin as it may be deprecated in the future (best practice is to
-		// mutate all documents using MetasoundDocumentBuilders).
-		virtual const FDocumentModifyDelegates* FindModifyDelegates(const FMetasoundFrontendClassName& InClassName) const = 0;
+		// Invalidates the cache of a given document's builder should one be registered, causing it to be rebuilt.  Not recommended
+		// for general use, and is only available in case the document is modified by a system other than an active, registered builder
+		// (ex. via the soft deprecated controller/handle API).
+		virtual void InvalidateDocumentCache(const FMetasoundFrontendClassName& InClassName) const = 0;
 
 		static IDocumentBuilderRegistry& GetChecked();
 
@@ -57,6 +56,13 @@ namespace Metasound::Frontend
 		static void Set(TUniqueFunction<IDocumentBuilderRegistry&()>&& InGetInstance);
 	};
 
-	// Deprecated: use shorter name
-	using IMetaSoundDocumentBuilderRegistry = IDocumentBuilderRegistry;
+	class METASOUNDFRONTEND_API IMetaSoundDocumentBuilderRegistry : public IDocumentBuilderRegistry
+	{
+	public:
+		UE_DEPRECATED(5.4, "Public exposition of modify delegates no longer available to discourage unsafe manipulation of builder document cache")
+		virtual const FDocumentModifyDelegates* FindModifyDelegates(const FMetasoundFrontendClassName& InClassName) const = 0;
+
+		UE_DEPRECATED(5.4, "Use 'IDocumentBuilderRegistry' instead")
+		static IMetaSoundDocumentBuilderRegistry& GetChecked();
+	};
 } // namespace Metasound::Frontend
diff --git a/Engine/Plugins/Runtime/Metasound/Source/MetasoundFrontend/Public/MetasoundFrontendDocumentBuilder.h b/Engine/Plugins/Runtime/Metasound/Source/MetasoundFrontend/Public/MetasoundFrontendDocumentBuilder.h
--- a/Engine/Plugins/Runtime/Metasound/Source/MetasoundFrontend/Public/MetasoundFrontendDocumentBuilder.h
+++ b/Engine/Plugins/Runtime/Metasound/Source/MetasoundFrontend/Public/MetasoundFrontendDocumentBuilder.h
@@ -217,6 +217,9 @@ public:
 	void InitDocument();
 	void InitNodeLocations();
 
+	// Clears the cached data enabling fast document queries. Generally discouraged unless document has been unavoidably mutated outside of the builder.
+	void InvalidateCache();
+
 	bool IsDependencyReferenced(const FGuid& InClassID) const;
 	bool IsNodeInputConnected(const FGuid& InNodeID, const FGuid& InVertexID) const;
 	bool IsNodeOutputConnected(const FGuid& InNodeID, const FGuid& InVertexID) const;
@@ -231,6 +234,10 @@ public:
 
 	bool ModifyInterfaces(Metasound::Frontend::FModifyInterfaceOptions&& InOptions);
 
+	// Reloads the cached data enabling fast document queries. Can be expensive and is generally discouraged in favor of using InvalidateCache
+	// only when document has been known to be mutated outside of the builder (also, invalidate doesn't prime the cache and only reloads piecemeal upon request).
+	void ReloadCache();
+
 	bool RemoveDependency(const FGuid& InClassID);
 	bool RemoveDependency(EMetasoundFrontendClassType ClassType, const FMetasoundFrontendClassName& InClassName, const FMetasoundFrontendVersionNumber& InClassVersionNumber);
 	bool RemoveEdge(const FMetasoundFrontendEdge& EdgeToRemove);
@@ -247,10 +254,6 @@ public:
 	bool RenameRootGraphClass(const FMetasoundFrontendClassName& InName);
 
 #if WITH_EDITOR
-	// Primarily used by editor transaction stack to avoid corruption when object is
-	// changed outside of the builder API. Generally discouraged for direct use otherwise
-	// as it can be expensive and builder API should manage internal cache state directly.
-	void ReloadCache();
 
 	void SetAuthor(const FString& InAuthor);
 
@@ -285,7 +288,7 @@ private:
 
 	const TSet<FMetasoundFrontendVersion>* FindNodeClassInterfaces(const FGuid& InNodeID) const;
 
-	void ReloadCacheInternal();
+	void InitCacheInternal(bool bPrimeCache = false);
 
 	bool SetGraphInputInheritsDefault(FName InName, bool bInputInheritsDefault);
 
@@ -293,10 +296,5 @@ private:
 	TScriptInterface<IMetaSoundDocumentInterface> DocumentInterface;
 
 	TSharedPtr<Metasound::Frontend::IDocumentCache> DocumentCache;
-
-	// SharedRef to struct of delegates fired when mutating document.  These are shared to safely support copying
-	// FrontendDocumentBuilders.  Copying is typically dissuaded as you can still have disparate builders that are
-	// not operating on the same cache; using a shared builder registered within the MetaSoundBuilderSubsystem
-	// in MetaSoundEngine is recommended.
 	TSharedRef<Metasound::Frontend::FDocumentModifyDelegates> DocumentDelegates;
 };
