Fix for uninitialized builder registry getter causing crash in tests
#rnx

diff --git a/Engine/Plugins/Runtime/Metasound/Source/MetasoundEngine/Private/MetasoundSource.cpp b/Engine/Plugins/Runtime/Metasound/Source/MetasoundEngine/Private/MetasoundSource.cpp
--- a/Engine/Plugins/Runtime/Metasound/Source/MetasoundEngine/Private/MetasoundSource.cpp
+++ b/Engine/Plugins/Runtime/Metasound/Source/MetasoundEngine/Private/MetasoundSource.cpp
@@ -550,7 +550,10 @@ Metasound::Frontend::FDocumentAccessPtr UMetaSoundSource::GetDocumentAccessPtr()
 	// the builder registry, so the document cache is invalidated here. It is discouraged to mutate
 	// documents using both systems at the same time as it can corrupt a builder document's cache.
 	const FMetasoundFrontendClassName& Name = RootMetasoundDocument.RootGraph.Metadata.GetClassName();
-	IDocumentBuilderRegistry::GetChecked().InvalidateDocumentCache(Name);
+	if (IDocumentBuilderRegistry* DocRegistry = IDocumentBuilderRegistry::Get())
+	{
+		DocRegistry->InvalidateDocumentCache(Name);
+	}
 
 	// Return document using FAccessPoint to inform the TAccessPtr when the 
 	// object is no longer valid.
diff --git a/Engine/Plugins/Runtime/Metasound/Source/MetasoundFrontend/Private/MetasoundDocumentInterface.cpp b/Engine/Plugins/Runtime/Metasound/Source/MetasoundFrontend/Private/MetasoundDocumentInterface.cpp
--- a/Engine/Plugins/Runtime/Metasound/Source/MetasoundFrontend/Private/MetasoundDocumentInterface.cpp
+++ b/Engine/Plugins/Runtime/Metasound/Source/MetasoundFrontend/Private/MetasoundDocumentInterface.cpp
@@ -7,9 +7,22 @@ namespace Metasound::Frontend
 {
 	namespace DocumentBuilderRegistryPrivate
 	{
+		static bool bInitialized = false;
 		TUniqueFunction<IDocumentBuilderRegistry&()> GetInstance;
 	} // namespace DocumentBuilderRegistryPrivate
 
+	IDocumentBuilderRegistry* IDocumentBuilderRegistry::Get()
+	{
+		using namespace DocumentBuilderRegistryPrivate;
+
+		if (!DocumentBuilderRegistryPrivate::bInitialized)
+		{
+			return nullptr;
+		}
+		
+		return &GetInstance();
+	}
+
 	IDocumentBuilderRegistry& IDocumentBuilderRegistry::GetChecked()
 	{
 		using namespace DocumentBuilderRegistryPrivate;
@@ -24,6 +37,7 @@ namespace Metasound::Frontend
 
 		checkf(!GetInstance, TEXT("Failed to initialize MetaSoundDocumentBuilderRegistry getter: Cannot reinitialize once initialized."))
 		GetInstance = MoveTemp(InGetInstance);
+		DocumentBuilderRegistryPrivate::bInitialized = true;
 	}
 
 	IMetaSoundDocumentBuilderRegistry& IMetaSoundDocumentBuilderRegistry::GetChecked()
diff --git a/Engine/Plugins/Runtime/Metasound/Source/MetasoundFrontend/Public/MetasoundDocumentInterface.h b/Engine/Plugins/Runtime/Metasound/Source/MetasoundFrontend/Public/MetasoundDocumentInterface.h
--- a/Engine/Plugins/Runtime/Metasound/Source/MetasoundFrontend/Public/MetasoundDocumentInterface.h
+++ b/Engine/Plugins/Runtime/Metasound/Source/MetasoundFrontend/Public/MetasoundDocumentInterface.h
@@ -50,6 +50,7 @@ namespace Metasound::Frontend
 		// (ex. via the soft deprecated controller/handle API).
 		virtual void InvalidateDocumentCache(const FMetasoundFrontendClassName& InClassName) const = 0;
 
+		static IDocumentBuilderRegistry* Get();
 		static IDocumentBuilderRegistry& GetChecked();
 
 	protected:
