Remove deprecation warnings on struct members for CIS build fix
#rnx

diff --git a/Engine/Plugins/Runtime/Metasound/Source/MetasoundFrontend/Public/MetasoundAssetBase.h b/Engine/Plugins/Runtime/Metasound/Source/MetasoundFrontend/Public/MetasoundAssetBase.h
--- a/Engine/Plugins/Runtime/Metasound/Source/MetasoundFrontend/Public/MetasoundAssetBase.h
+++ b/Engine/Plugins/Runtime/Metasound/Source/MetasoundFrontend/Public/MetasoundAssetBase.h
@@ -228,14 +228,12 @@ protected:
 		FGuid ChangeID;
 
 		// Array of inputs which can be set for construction. 
-		UE_DEPRECATED(5.4, "Use PublicInputMap instead.")
 		TArray<FMetasoundFrontendClassInput> PublicInputs;
 
 		// Map of runtime inputs keyed by vertex names. 
 		Metasound::TSortedVertexNameMap<FRuntimeInput> PublicInputMap;
 
 		// Array of inputs which can be transmitted to.
-		UE_DEPRECATED(5.4, "Use PublicInputMap and DataTypeRegistry information instead.")
 		TArray<FMetasoundFrontendClassInput> TransmittableInputs;
 
 		// Core graph.
