Fix for crash due to CacheRuntimeData causing proxy gen when cooking with deterministic guid gen on (which should be omitted from CookMetaSound altogether as its not required)
#rnx

diff --git a/Engine/Plugins/Runtime/Metasound/Source/MetasoundFrontend/Private/MetasoundAssetBase.cpp b/Engine/Plugins/Runtime/Metasound/Source/MetasoundFrontend/Private/MetasoundAssetBase.cpp
--- a/Engine/Plugins/Runtime/Metasound/Source/MetasoundFrontend/Private/MetasoundAssetBase.cpp
+++ b/Engine/Plugins/Runtime/Metasound/Source/MetasoundFrontend/Private/MetasoundAssetBase.cpp
@@ -381,7 +381,7 @@ void FMetasoundAssetBase::CookMetaSound()
 	const FString AssetName = Owner->GetName();
 
 	{
-	// Performs document transforms on local copy, which reduces document footprint & renders transforming unnecessary unless altered at runtime when registering
+		// Performs document transforms on local copy, which reduces document footprint & renders transforming unnecessary unless altered at runtime when registering
 		FMetaSoundFrontendDocumentBuilder DocBuilder(Owner);
 		const bool bContainsTemplateDependency = DocBuilder.ContainsDependencyOfType(EMetasoundFrontendClassType::Template);
 		if (bContainsTemplateDependency)
@@ -390,7 +390,6 @@ void FMetasoundAssetBase::CookMetaSound()
 		}
 
 		const IMetaSoundDocumentInterface& Interface = *TScriptInterface<IMetaSoundDocumentInterface>(Owner).GetInterface();
-		CacheRuntimeData(Interface.GetDocument());
 		TUniquePtr<INodeRegistryEntry> RegistryEntry = MakeUnique<AssetBasePrivate::FNodeRegistryEntry>(AssetName, Interface.GetDocument(), FSoftObjectPath(Owner));
 		UnregisterGraphWithFrontend();
 		RegistryKey = FMetasoundFrontendRegistryContainer::Get()->RegisterNode(MoveTemp(RegistryEntry));
