If at first you succeed, you potenitally forgot another place you failed the same way...
#rnx

diff --git a/Engine/Plugins/Runtime/Metasound/Source/MetasoundEngine/Private/Metasound.cpp b/Engine/Plugins/Runtime/Metasound/Source/MetasoundEngine/Private/Metasound.cpp
--- a/Engine/Plugins/Runtime/Metasound/Source/MetasoundEngine/Private/Metasound.cpp
+++ b/Engine/Plugins/Runtime/Metasound/Source/MetasoundEngine/Private/Metasound.cpp
@@ -65,7 +65,10 @@ Metasound::Frontend::FDocumentAccessPtr UMetaSoundPatch::GetDocumentAccessPtr()
 	// the builder registry, so the document cache is invalidated here. It is discouraged to mutate
 	// documents using both systems at the same time as it can corrupt a builder document's cache.
 	const FMetasoundFrontendClassName& Name = RootMetaSoundDocument.RootGraph.Metadata.GetClassName();
-	IDocumentBuilderRegistry::GetChecked().InvalidateDocumentCache(Name);
+	if (IDocumentBuilderRegistry* DocRegistry = IDocumentBuilderRegistry::Get())
+	{
+		DocRegistry->InvalidateDocumentCache(Name);
+	}
 
 	// Return document using FAccessPoint to inform the TAccessPtr when the 
 	// object is no longer valid.
