[Metasound] Add Blueprint API to create FMetasoundFrontendLiteral instances of supported types

diff --git a/Engine/Plugins/Runtime/Metasound/Source/MetasoundFrontend/MetasoundFrontend.Build.cs b/Engine/Plugins/Runtime/Metasound/Source/MetasoundFrontend/MetasoundFrontend.Build.cs
--- a/Engine/Plugins/Runtime/Metasound/Source/MetasoundFrontend/MetasoundFrontend.Build.cs
+++ b/Engine/Plugins/Runtime/Metasound/Source/MetasoundFrontend/MetasoundFrontend.Build.cs
@@ -18,6 +18,7 @@ namespace UnrealBuildTool.Rules
 					"AudioExtensions",
 					"Core",
 					"CoreUObject",
+					"Engine",
 					"Serialization",
 					"SignalProcessing",
 					"MetasoundGraphCore"
diff --git a/Engine/Plugins/Runtime/Metasound/Source/MetasoundFrontend/Private/MetasoundFrontendLiteral.cpp b/Engine/Plugins/Runtime/Metasound/Source/MetasoundFrontend/Private/MetasoundFrontendLiteral.cpp
--- a/Engine/Plugins/Runtime/Metasound/Source/MetasoundFrontend/Private/MetasoundFrontendLiteral.cpp
+++ b/Engine/Plugins/Runtime/Metasound/Source/MetasoundFrontend/Private/MetasoundFrontendLiteral.cpp
@@ -78,6 +78,14 @@ namespace MetasoundFrontendLiteralPrivate
 		// (c++ does not allow partial specialization of functions)
 		return TLiteralValueToStringHelper<Type>::Convert(InType);
 	}
+	
+	template <typename TLiteralType>
+	FMetasoundFrontendLiteral CreatePODMetaSoundLiteral(const TLiteralType& Value)
+	{
+		FMetasoundFrontendLiteral Literal;
+		Literal.Set(Value);
+		return Literal;
+	}
 }
 
 FMetasoundFrontendLiteral::FMetasoundFrontendLiteral(const FAudioParameter& InParameter)
@@ -870,4 +878,63 @@ namespace Metasound
 	}
 }
 
+FMetasoundFrontendLiteral UMetasoundFrontendLiteralBlueprintAccess::CreateBoolMetaSoundLiteral(bool Value)
+{
+	return MetasoundFrontendLiteralPrivate::CreatePODMetaSoundLiteral(Value);
+}
+
+FMetasoundFrontendLiteral UMetasoundFrontendLiteralBlueprintAccess::CreateBoolArrayMetaSoundLiteral(
+	const TArray<bool>& Value)
+{
+	return MetasoundFrontendLiteralPrivate::CreatePODMetaSoundLiteral(Value);
+}
+
+FMetasoundFrontendLiteral UMetasoundFrontendLiteralBlueprintAccess::CreateFloatMetaSoundLiteral(float Value)
+{
+	return MetasoundFrontendLiteralPrivate::CreatePODMetaSoundLiteral(Value);
+}
+
+FMetasoundFrontendLiteral UMetasoundFrontendLiteralBlueprintAccess::CreateFloatArrayMetaSoundLiteral(
+	const TArray<float>& Value)
+{
+	return MetasoundFrontendLiteralPrivate::CreatePODMetaSoundLiteral(Value);
+}
+
+FMetasoundFrontendLiteral UMetasoundFrontendLiteralBlueprintAccess::CreateIntMetaSoundLiteral(int32 Value)
+{
+	return MetasoundFrontendLiteralPrivate::CreatePODMetaSoundLiteral(Value);
+}
+
+FMetasoundFrontendLiteral UMetasoundFrontendLiteralBlueprintAccess::CreateIntArrayMetaSoundLiteral(
+	const TArray<int32>& Value)
+{
+	return MetasoundFrontendLiteralPrivate::CreatePODMetaSoundLiteral(Value);
+}
+
+FMetasoundFrontendLiteral UMetasoundFrontendLiteralBlueprintAccess::CreateObjectMetaSoundLiteral(UObject* Value)
+{
+	return MetasoundFrontendLiteralPrivate::CreatePODMetaSoundLiteral(Value);
+}
+
+FMetasoundFrontendLiteral UMetasoundFrontendLiteralBlueprintAccess::CreateObjectArrayMetaSoundLiteral(
+	const TArray<UObject*>& Value)
+{
+	return MetasoundFrontendLiteralPrivate::CreatePODMetaSoundLiteral(Value);
+}
 
+FMetasoundFrontendLiteral UMetasoundFrontendLiteralBlueprintAccess::CreateStringMetaSoundLiteral(const FString& Value)
+{
+	return MetasoundFrontendLiteralPrivate::CreatePODMetaSoundLiteral(Value);
+}
+
+FMetasoundFrontendLiteral UMetasoundFrontendLiteralBlueprintAccess::CreateStringArrayMetaSoundLiteral(
+	const TArray<FString>& Value)
+{
+	return MetasoundFrontendLiteralPrivate::CreatePODMetaSoundLiteral(Value);
+}
+
+FMetasoundFrontendLiteral UMetasoundFrontendLiteralBlueprintAccess::CreateMetaSoundLiteralFromParam(
+	const FAudioParameter& Param)
+{
+	return FMetasoundFrontendLiteral{ Param };
+}
diff --git a/Engine/Plugins/Runtime/Metasound/Source/MetasoundFrontend/Public/MetasoundFrontendLiteral.h b/Engine/Plugins/Runtime/Metasound/Source/MetasoundFrontend/Public/MetasoundFrontendLiteral.h
--- a/Engine/Plugins/Runtime/Metasound/Source/MetasoundFrontend/Public/MetasoundFrontendLiteral.h
+++ b/Engine/Plugins/Runtime/Metasound/Source/MetasoundFrontend/Public/MetasoundFrontendLiteral.h
@@ -6,6 +6,8 @@
 #include "CoreMinimal.h"
 #include "MetasoundLiteral.h"
 
+#include "Kismet/BlueprintFunctionLibrary.h"
+
 #include <type_traits>
 
 #include "MetasoundFrontendLiteral.generated.h"
@@ -175,3 +177,44 @@ namespace Metasound
 	}
 }
 
+/**
+ * Blueprint support for FMetasoundFrontendLiteral
+ */
+UCLASS()
+class UMetasoundFrontendLiteralBlueprintAccess : public UBlueprintFunctionLibrary
+{
+	GENERATED_BODY()
+
+	UFUNCTION(BlueprintCallable, Category = "Audio|MetaSound", meta = (DisplayName = "Create MetaSound Bool Literal"))
+	static UPARAM(DisplayName = "Bool Literal") FMetasoundFrontendLiteral CreateBoolMetaSoundLiteral(bool Value);
+
+	UFUNCTION(BlueprintCallable, Category = "Audio|MetaSound", meta = (DisplayName = "Create MetaSound Bool Array Literal"))
+	static UPARAM(DisplayName = "Bool Array Literal") FMetasoundFrontendLiteral CreateBoolArrayMetaSoundLiteral(const TArray<bool>& Value);
+
+	UFUNCTION(BlueprintCallable, Category = "Audio|MetaSound", meta = (DisplayName = "Create MetaSound Float Literal"))
+	static UPARAM(DisplayName = "Float Literal") FMetasoundFrontendLiteral CreateFloatMetaSoundLiteral(float Value);
+
+	UFUNCTION(BlueprintCallable, Category = "Audio|MetaSound", meta = (DisplayName = "Create MetaSound Float Array Literal"))
+	static UPARAM(DisplayName = "Float Array Literal") FMetasoundFrontendLiteral CreateFloatArrayMetaSoundLiteral(const TArray<float>& Value);
+
+	UFUNCTION(BlueprintCallable, Category = "Audio|MetaSound", meta = (DisplayName = "Create MetaSound Int Literal"))
+	static UPARAM(DisplayName = "Int32 Literal") FMetasoundFrontendLiteral CreateIntMetaSoundLiteral(int32 Value);
+
+	UFUNCTION(BlueprintCallable, Category = "Audio|MetaSound", meta = (DisplayName = "Create MetaSound Int Array Literal"))
+	static UPARAM(DisplayName = "Int32 Array Literal") FMetasoundFrontendLiteral CreateIntArrayMetaSoundLiteral(const TArray<int32>& Value);
+
+	UFUNCTION(BlueprintCallable, Category = "Audio|MetaSound", meta = (DisplayName = "Create MetaSound Object Literal"))
+	static UPARAM(DisplayName = "Object Literal") FMetasoundFrontendLiteral CreateObjectMetaSoundLiteral(UObject* Value);
+
+	UFUNCTION(BlueprintCallable, Category = "Audio|MetaSound", meta = (DisplayName = "Create MetaSound Object Array Literal"))
+	static UPARAM(DisplayName = "Object Array Literal") FMetasoundFrontendLiteral CreateObjectArrayMetaSoundLiteral(const TArray<UObject*>& Value);
+
+	UFUNCTION(BlueprintCallable, Category = "Audio|MetaSound", meta = (DisplayName = "Create MetaSound String Literal"))
+	static UPARAM(DisplayName = "String Literal") FMetasoundFrontendLiteral CreateStringMetaSoundLiteral(const FString& Value);
+
+	UFUNCTION(BlueprintCallable, Category = "Audio|MetaSound", meta = (DisplayName = "Create MetaSound String Array Literal"))
+	static UPARAM(DisplayName = "String Array Literal") FMetasoundFrontendLiteral CreateStringArrayMetaSoundLiteral(const TArray<FString>& Value);
+
+	UFUNCTION(BlueprintCallable, Category = "Audio|MetaSound", meta = (DisplayName = "Create MetaSound Literal From AudioParameter"))
+	static UPARAM(DisplayName = "Param Literal") FMetasoundFrontendLiteral CreateMetaSoundLiteralFromParam(const FAudioParameter& Param);
+};
