[Metasound] Guard against unbounded growth of the output queue

diff --git a/Engine/Plugins/Runtime/Metasound/Source/MetasoundEngine/Private/MetasoundGeneratorHandle.cpp b/Engine/Plugins/Runtime/Metasound/Source/MetasoundEngine/Private/MetasoundGeneratorHandle.cpp
--- a/Engine/Plugins/Runtime/Metasound/Source/MetasoundEngine/Private/MetasoundGeneratorHandle.cpp
+++ b/Engine/Plugins/Runtime/Metasound/Source/MetasoundEngine/Private/MetasoundGeneratorHandle.cpp
@@ -222,6 +222,8 @@ namespace Metasound
 
 		check(IsInGameThread());
 
+		int32 NumDequeued = 0;
+
 		while (TOptional<FOutputPayload> ChangedOutput = ChangedOutputs.Dequeue())
 		{
 			if (const FOutputWatcher* OutputListener = OutputWatchers.FindByPredicate(
@@ -234,7 +236,11 @@ namespace Metasound
 			{
 				OutputListener->OnOutputValueChanged.Broadcast(ChangedOutput->OutputName, ChangedOutput->OutputValue);
 			}
+
+			++NumDequeued;
 		}
+
+		ChangedOutputsQueueCount.store(FMath::Max(0, ChangedOutputsQueueCount.load() - NumDequeued));
 	}
 
 	void FMetasoundGeneratorHandle::RegisterPassthroughAnalyzerForType(
@@ -531,8 +537,21 @@ namespace Metasound
 	{
 		METASOUND_LLM_SCOPE;
 		METASOUND_TRACE_CPUPROFILER_EVENT_SCOPE(FMetasoundGeneratorHandle::HandleOutputChanged);
+
+		if (ChangedOutputsQueueCount >= ChangedOutputsQueueMax)
+		{
+			// Log only once per handle
+			if (ChangedOutputsQueueShouldLogIfFull.load())
+			{
+				UE_LOG(LogMetaSound, Warning, TEXT("UMetasoundGeneratorHandle output queue is full."));
+				ChangedOutputsQueueShouldLogIfFull.store(false);
+			}
+			
+			return;
+		}
 		
 		ChangedOutputs.Enqueue(AnalyzerName, OutputName, AnalyzerOutputName, OutputData);
+		ChangedOutputsQueueCount.fetch_add(1);
 	}
 }
 
diff --git a/Engine/Plugins/Runtime/Metasound/Source/MetasoundEngine/Public/MetasoundGeneratorHandle.h b/Engine/Plugins/Runtime/Metasound/Source/MetasoundEngine/Public/MetasoundGeneratorHandle.h
--- a/Engine/Plugins/Runtime/Metasound/Source/MetasoundEngine/Public/MetasoundGeneratorHandle.h
+++ b/Engine/Plugins/Runtime/Metasound/Source/MetasoundEngine/Public/MetasoundGeneratorHandle.h
@@ -2,6 +2,8 @@
 
 #pragma once
 
+#include <atomic>
+
 #include "MetasoundOutput.h"
 #include "MetasoundParameterPack.h"
 #include "Analysis/MetasoundFrontendAnalyzerAddress.h"
@@ -225,6 +227,11 @@ namespace Metasound
 	
 		TSpscQueue<FOutputPayload> ChangedOutputs;
 
+		// Keep ChangedOutputs from growing infinitely
+		static constexpr int32 ChangedOutputsQueueMax = 1024;
+		std::atomic<int32> ChangedOutputsQueueCount{ 0 };
+		std::atomic<bool> ChangedOutputsQueueShouldLogIfFull{ true };
+
 		bool bRuntimeRenderTimingShouldBeEnabled{ false };
 	};
 }
