[Metasound] Ensure input node defaults are propagated

diff --git a/Engine/Plugins/Runtime/Metasound/Source/MetasoundEngine/Private/MetasoundBuilderSubsystem.cpp b/Engine/Plugins/Runtime/Metasound/Source/MetasoundEngine/Private/MetasoundBuilderSubsystem.cpp
--- a/Engine/Plugins/Runtime/Metasound/Source/MetasoundEngine/Private/MetasoundBuilderSubsystem.cpp
+++ b/Engine/Plugins/Runtime/Metasound/Source/MetasoundEngine/Private/MetasoundBuilderSubsystem.cpp
@@ -1172,13 +1172,35 @@ void UMetaSoundSourceBuilder::OnNodeAdded(int32 NodeIndex) const
 			{
 				const FName& DataTypeName = NodeClass->Metadata.GetClassName().Name;
 				const FMetasoundFrontendVertex& InputVertex = AddedNode.Interface.Inputs.Last();
-				const FMetasoundFrontendLiteral& DefaultLiteral = NodeClass->Interface.Inputs.Last().DefaultLiteral;
+
+				const FMetasoundFrontendLiteral* DefaultLiteral = nullptr;
+				auto HasEqualVertexID = [&VertexID = InputVertex.VertexID](const FMetasoundFrontendVertexLiteral& InVertexLiteral)
+				{
+					return InVertexLiteral.VertexID == VertexID;
+				};
+
+				// Check for default literal on Node
+				if (const FMetasoundFrontendVertexLiteral* VertexLiteralOnNode = AddedNode.InputLiterals.FindByPredicate(HasEqualVertexID))
+				{
+					DefaultLiteral = &(VertexLiteralOnNode->Value);
+				}
+				// Check for default literal on class input
+				else if (const FMetasoundFrontendClassInput* GraphInput = Builder.FindGraphInput(InputVertex.Name))
+				{
+					DefaultLiteral = &(GraphInput->DefaultLiteral);
+				}
+				else
+				{
+					// As a last resort, get default literal on node class 
+					DefaultLiteral = &(NodeClass->Interface.Inputs.Last().DefaultLiteral);
+				}
+
 				FInputNodeConstructorParams InitData
 				{
 					AddedNode.Name,
 					AddedNode.GetID(),
 					InputVertex.Name,
-					DefaultLiteral.ToLiteral(DataTypeName)
+					DefaultLiteral->ToLiteral(DataTypeName)
 				};
 
 				NewNode = DataTypeRegistry.CreateInputNode(DataTypeName, MoveTemp(InitData));
diff --git a/Engine/Plugins/Runtime/Metasound/Source/MetasoundFrontend/Private/MetasoundFrontendDocumentBuilder.cpp b/Engine/Plugins/Runtime/Metasound/Source/MetasoundFrontend/Private/MetasoundFrontendDocumentBuilder.cpp
--- a/Engine/Plugins/Runtime/Metasound/Source/MetasoundFrontend/Private/MetasoundFrontendDocumentBuilder.cpp
+++ b/Engine/Plugins/Runtime/Metasound/Source/MetasoundFrontend/Private/MetasoundFrontendDocumentBuilder.cpp
@@ -74,7 +74,7 @@ namespace Metasound::Frontend
 			return true;
 		}
 
-		void FinalizeGraphVertexNode(FMetasoundFrontendNode& InOutNode, const FMetasoundFrontendClassVertex& InVertex)
+		void SetNodeAndVertexNames(FMetasoundFrontendNode& InOutNode, const FMetasoundFrontendClassVertex& InVertex)
 		{
 			InOutNode.Name = InVertex.Name;
 			// Set name on related vertices of input node
@@ -98,6 +98,38 @@ namespace Metasound::Frontend
 			}
 		}
 
+		void SetDefaultLiteralOnInputNode(FMetasoundFrontendNode& InOutNode, const FMetasoundFrontendClassInput& InClassInput)
+		{
+			// Set the default literal on the nodes inputs so that it gets passed to the instantiated TInputNode on a live
+			// auditioned MetaSound
+			auto IsVertexWithName = [&Name = InClassInput.Name](const FMetasoundFrontendVertex& InVertex)
+			{
+				return InVertex.Name == Name;
+			};
+
+			if (const FMetasoundFrontendVertex* InputVertex = InOutNode.Interface.Inputs.FindByPredicate(IsVertexWithName))
+			{
+				auto IsVertexLiteralWithVertexID = [&VertexID = InputVertex->VertexID](const FMetasoundFrontendVertexLiteral& VertexLiteral)
+				{
+					return VertexLiteral.VertexID == VertexID;
+				};
+				if (FMetasoundFrontendVertexLiteral* VertexLiteral = InOutNode.InputLiterals.FindByPredicate(IsVertexLiteralWithVertexID))
+				{
+					// Update existing literal default value with value from class input.
+					VertexLiteral->Value = InClassInput.DefaultLiteral;
+				}
+				else
+				{
+					// Add literal default value with value from class input.
+					InOutNode.InputLiterals.Add(FMetasoundFrontendVertexLiteral{InputVertex->VertexID, InClassInput.DefaultLiteral});
+				}
+			}
+			else
+			{
+				UE_LOG(LogMetaSound, Error, TEXT("Input node associated with graph input vertex of name '%s' does not contain input vertex with matching name."), *InClassInput.Name.ToString());
+			}
+		}
+
 		const FString GetDebugName(const IMetaSoundDocumentInterface& DocumentInterface)
 		{
 			const FMetasoundFrontendClassMetadata& Metadata = DocumentInterface.GetDocument().RootGraph.Metadata;
@@ -923,10 +955,19 @@ const FMetasoundFrontendNode* FMetaSoundFrontendDocumentBuilder::AddGraphInput(c
 
 		auto FinalizeNode = [&InClassInput](FMetasoundFrontendNode& InOutNode, const Metasound::Frontend::FNodeRegistryKey&)
 		{
-			DocumentBuilderPrivate::FinalizeGraphVertexNode(InOutNode, InClassInput);
+			// Sets the name of the node an vertices on the node to match the class vertex name
+			DocumentBuilderPrivate::SetNodeAndVertexNames(InOutNode, InClassInput);
+
+			// Set the default literal on the nodes inputs so that it gets passed to the instantiated TInputNode on a live
+			// auditioned MetaSound.
+			DocumentBuilderPrivate::SetDefaultLiteralOnInputNode(InOutNode, InClassInput);
 		};
 		if (FMetasoundFrontendNode* NewNode = AddNodeInternal(Class.Metadata, FinalizeNode, InClassInput.NodeID))
 		{
+			// Remove the default literal on the node added during the "FinalizeNode" call. This matches how 
+			// nodes are serialized in editor. The default literals are only stored on the FMetasoundFrontendClassInputs.
+			NewNode->InputLiterals.Reset();
+
 			const int32 NewIndex = RootGraph.Interface.Inputs.Num();
 			FMetasoundFrontendClassInput& NewInput = RootGraph.Interface.Inputs.Add_GetRef(InClassInput);
 			if (!NewInput.VertexID.IsValid())
@@ -1009,7 +1050,7 @@ const FMetasoundFrontendNode* FMetaSoundFrontendDocumentBuilder::AddGraphOutput(
 
 		auto FinalizeNode = [&InClassOutput](FMetasoundFrontendNode& InOutNode, const Metasound::Frontend::FNodeRegistryKey&)
 		{
-			DocumentBuilderPrivate::FinalizeGraphVertexNode(InOutNode, InClassOutput);
+			DocumentBuilderPrivate::SetNodeAndVertexNames(InOutNode, InClassOutput);
 		};
 		if (FMetasoundFrontendNode* NewNode = AddNodeInternal(Class.Metadata, FinalizeNode, InClassOutput.NodeID))
 		{
