Fix for crash in tests when MetaSoundBuilderSubsystem is not yet loaded
#rnx

diff --git a/Engine/Plugins/Runtime/Metasound/Source/MetasoundEngine/Private/Metasound.cpp b/Engine/Plugins/Runtime/Metasound/Source/MetasoundEngine/Private/Metasound.cpp
--- a/Engine/Plugins/Runtime/Metasound/Source/MetasoundEngine/Private/Metasound.cpp
+++ b/Engine/Plugins/Runtime/Metasound/Source/MetasoundEngine/Private/Metasound.cpp
@@ -64,10 +64,10 @@ Metasound::Frontend::FDocumentAccessPtr UMetaSoundPatch::GetDocumentAccessPtr()
 	// Mutation of a document via the soft deprecated access ptr/controller system is not tracked by
 	// the builder registry, so the document cache is invalidated here. It is discouraged to mutate
 	// documents using both systems at the same time as it can corrupt a builder document's cache.
-	const FMetasoundFrontendClassName& Name = RootMetaSoundDocument.RootGraph.Metadata.GetClassName();
-	if (IDocumentBuilderRegistry* DocRegistry = IDocumentBuilderRegistry::Get())
+	if (UMetaSoundBuilderSubsystem* BuilderSubsystem = UMetaSoundBuilderSubsystem::Get())
 	{
-		DocRegistry->InvalidateDocumentCache(Name);
+		const FMetasoundFrontendClassName& Name = RootMetaSoundDocument.RootGraph.Metadata.GetClassName();
+		BuilderSubsystem->InvalidateDocumentCache(Name);
 	}
 
 	// Return document using FAccessPoint to inform the TAccessPtr when the 
diff --git a/Engine/Plugins/Runtime/Metasound/Source/MetasoundEngine/Private/MetasoundBuilderSubsystem.cpp b/Engine/Plugins/Runtime/Metasound/Source/MetasoundEngine/Private/MetasoundBuilderSubsystem.cpp
--- a/Engine/Plugins/Runtime/Metasound/Source/MetasoundEngine/Private/MetasoundBuilderSubsystem.cpp
+++ b/Engine/Plugins/Runtime/Metasound/Source/MetasoundEngine/Private/MetasoundBuilderSubsystem.cpp
@@ -1613,6 +1613,19 @@ UMetaSoundSourceBuilder* UMetaSoundBuilderSubsystem::CreateSourcePresetBuilder(F
 	return nullptr;
 }
 
+UMetaSoundBuilderSubsystem* UMetaSoundBuilderSubsystem::Get()
+{
+	if (GEngine)
+	{
+		if (UMetaSoundBuilderSubsystem* BuilderSubsystem = GEngine->GetEngineSubsystem<UMetaSoundBuilderSubsystem>())
+		{
+			return BuilderSubsystem;
+		}
+	}
+
+	return nullptr;
+}
+
 UMetaSoundBuilderSubsystem& UMetaSoundBuilderSubsystem::GetChecked()
 {
 	checkf(GEngine, TEXT("Cannot access UMetaSoundBuilderSubsystem without engine loaded"));
@@ -1621,6 +1634,19 @@ UMetaSoundBuilderSubsystem& UMetaSoundBuilderSubsystem::GetChecked()
 	return *BuilderSubsystem;
 }
 
+const UMetaSoundBuilderSubsystem* UMetaSoundBuilderSubsystem::GetConst()
+{
+	if (GEngine)
+	{
+		if (const UMetaSoundBuilderSubsystem* BuilderSubsystem = GEngine->GetEngineSubsystem<const UMetaSoundBuilderSubsystem>())
+		{
+			return BuilderSubsystem;
+		}
+	}
+
+	return nullptr;
+}
+
 const UMetaSoundBuilderSubsystem& UMetaSoundBuilderSubsystem::GetConstChecked()
 {
 	checkf(GEngine, TEXT("Cannot access UMetaSoundBuilderSubsystem without engine loaded"));
diff --git a/Engine/Plugins/Runtime/Metasound/Source/MetasoundEngine/Private/MetasoundSource.cpp b/Engine/Plugins/Runtime/Metasound/Source/MetasoundEngine/Private/MetasoundSource.cpp
--- a/Engine/Plugins/Runtime/Metasound/Source/MetasoundEngine/Private/MetasoundSource.cpp
+++ b/Engine/Plugins/Runtime/Metasound/Source/MetasoundEngine/Private/MetasoundSource.cpp
@@ -549,10 +549,10 @@ Metasound::Frontend::FDocumentAccessPtr UMetaSoundSource::GetDocumentAccessPtr()
 	// Mutation of a document via the soft deprecated access ptr/controller system is not tracked by
 	// the builder registry, so the document cache is invalidated here. It is discouraged to mutate
 	// documents using both systems at the same time as it can corrupt a builder document's cache.
-	const FMetasoundFrontendClassName& Name = RootMetasoundDocument.RootGraph.Metadata.GetClassName();
-	if (IDocumentBuilderRegistry* DocRegistry = IDocumentBuilderRegistry::Get())
+	if (UMetaSoundBuilderSubsystem* BuilderSubsystem = UMetaSoundBuilderSubsystem::Get())
 	{
-		DocRegistry->InvalidateDocumentCache(Name);
+		const FMetasoundFrontendClassName& Name = RootMetasoundDocument.RootGraph.Metadata.GetClassName();
+		BuilderSubsystem->InvalidateDocumentCache(Name);
 	}
 
 	// Return document using FAccessPoint to inform the TAccessPtr when the 
diff --git a/Engine/Plugins/Runtime/Metasound/Source/MetasoundEngine/Public/MetasoundBuilderSubsystem.h b/Engine/Plugins/Runtime/Metasound/Source/MetasoundEngine/Public/MetasoundBuilderSubsystem.h
--- a/Engine/Plugins/Runtime/Metasound/Source/MetasoundEngine/Public/MetasoundBuilderSubsystem.h
+++ b/Engine/Plugins/Runtime/Metasound/Source/MetasoundEngine/Public/MetasoundBuilderSubsystem.h
@@ -571,7 +571,9 @@ public:
 	virtual void Initialize(FSubsystemCollectionBase& Collection) override;
 	virtual void InvalidateDocumentCache(const FMetasoundFrontendClassName& InClassName) const override;
 
+	static UMetaSoundBuilderSubsystem* Get();
 	static UMetaSoundBuilderSubsystem& GetChecked();
+	static const UMetaSoundBuilderSubsystem* GetConst();
 	static const UMetaSoundBuilderSubsystem& GetConstChecked();
 
 	UMetaSoundBuilderBase& AttachBuilderToAssetChecked(UObject& InMetaSound) const;
