[Metasound] Enable pin on envelope follower node

diff --git a/Engine/Plugins/Runtime/Metasound/Source/MetasoundEngineTest/MetasoundEngineTest.Build.cs b/Engine/Plugins/Runtime/Metasound/Source/MetasoundEngineTest/MetasoundEngineTest.Build.cs
--- a/Engine/Plugins/Runtime/Metasound/Source/MetasoundEngineTest/MetasoundEngineTest.Build.cs
+++ b/Engine/Plugins/Runtime/Metasound/Source/MetasoundEngineTest/MetasoundEngineTest.Build.cs
@@ -30,7 +30,8 @@ namespace UnrealBuildTool.Rules
 					"MetasoundGraphCore",
 					"MetasoundGenerator",
 					"MetasoundFrontend",
-					"MetasoundEngine"
+					"MetasoundEngine",
+					"MetasoundStandardNodes"
 				}
 			);
 		}
diff --git a/Engine/Plugins/Runtime/Metasound/Source/MetasoundEngineTest/Private/EnvelopeFollowerNodeTests.cpp b/Engine/Plugins/Runtime/Metasound/Source/MetasoundEngineTest/Private/EnvelopeFollowerNodeTests.cpp
new file mode 100644
--- /dev/null
+++ b/Engine/Plugins/Runtime/Metasound/Source/MetasoundEngineTest/Private/EnvelopeFollowerNodeTests.cpp
@@ -0,0 +1,118 @@
+// Copyright Epic Games, Inc. All Rights Reserved.
+
+#include "MetasoundGenerator.h"
+#include "MetasoundStandardNodesNames.h"
+#include "NodeTestGraphBuilder.h"
+
+#include "Misc/AutomationTest.h"
+
+#if WITH_DEV_AUTOMATION_TESTS
+
+namespace Metasound::Test::Nodes::EnvelopeFollower
+{
+	IMPLEMENT_SIMPLE_AUTOMATION_TEST(
+		FMetasoundEnvelopeFollowerNodeEnableTest,
+		"Audio.Metasound.Nodes.EnvelopeFollower.Enable",
+		EAutomationTestFlags::EditorContext | EAutomationTestFlags::EngineFilter)
+	bool FMetasoundEnvelopeFollowerNodeEnableTest::RunTest(const FString&)
+	{
+		// Make a graph with a sine oscillator going into an envelope follower
+		TUniquePtr<FMetasoundGenerator> Generator;
+		{
+			FNodeTestGraphBuilder Builder;
+
+			const Frontend::FNodeHandle Oscillator = Builder.AddNode({ StandardNodes::Namespace, TEXT("Sine"), StandardNodes::AudioVariant }, 1);
+			const Frontend::FNodeHandle EnvelopeFollower = Builder.AddNode({StandardNodes::Namespace, TEXT("Envelope Follower"), TEXT("") }, 1);
+			FNodeTestGraphBuilder::ConnectNodes(Oscillator, "Audio", EnvelopeFollower, "In");
+
+			// Expose the enable pin on the envelope follower
+			Builder.AddAndConnectDataReferenceInput(EnvelopeFollower, "Enable", GetMetasoundDataTypeName<bool>());
+
+			// Expose the audio and float outputs of the envelope follower
+			Builder.AddAndConnectDataReferenceOutput(EnvelopeFollower, "Envelope", GetMetasoundDataTypeName<float>());
+			Builder.AddAndConnectDataReferenceOutput(EnvelopeFollower, "Audio Envelope", GetMetasoundDataTypeName<FAudioBuffer>());
+
+			Generator = Builder.BuildGenerator();
+		}
+		UTEST_NOT_NULL("Created generator", Generator.Get());
+
+		// Get the graph output refs
+		TOptional<TDataReadReference<float>> FloatEnvelopeOutput = Generator->GetOutputReadReference<float>("Envelope");
+		UTEST_TRUE("Got float envelope output", FloatEnvelopeOutput.IsSet());
+		TOptional<TDataReadReference<FAudioBuffer>> AudioEnvelopeOutput = Generator->GetOutputReadReference<FAudioBuffer>("Audio Envelope");
+		UTEST_TRUE("Got audio envelope output", AudioEnvelopeOutput.IsSet());
+		
+		// Enable, render a few blocks and expect some output
+		const int32 BlockSize = Generator->OperatorSettings.GetNumFramesPerBlock();
+		FAudioBuffer DummyBuffer{ BlockSize };
+
+		Generator->SetInputValue("Enable", true);
+		
+		constexpr int32 NumBlocksToRender = 4;
+		for (int32 i = 0; i < NumBlocksToRender; ++i)
+		{
+			Generator->OnGenerateAudio(DummyBuffer.GetData(), DummyBuffer.Num());
+		}
+
+		{
+			UTEST_TRUE("Float envelope was non-zero", **FloatEnvelopeOutput > 0.0f);
+			bool HadNonZeroEnvelope = false;
+			for (int32 i = 0; i < BlockSize; ++i)
+			{
+				const float Value = (*AudioEnvelopeOutput)->GetData()[i];
+				if (Value > 0.0f)
+				{
+					HadNonZeroEnvelope = true;
+					break;
+				}
+			}
+			UTEST_TRUE("Audio envelope had non-zero output", HadNonZeroEnvelope);
+		}
+
+		// Disable, render a block and expect zero output
+		Generator->SetInputValue("Enable", false);
+		Generator->OnGenerateAudio(DummyBuffer.GetData(), DummyBuffer.Num());
+
+		{
+			UTEST_EQUAL("Float envelope was zero", **FloatEnvelopeOutput, 0.0f);
+			bool HadNonZeroEnvelope = false;
+			for (int32 i = 0; i < BlockSize; ++i)
+			{
+				const float Value = (*AudioEnvelopeOutput)->GetData()[i];
+				if (Value > 0.0f)
+				{
+					HadNonZeroEnvelope = true;
+					break;
+				}
+			}
+			UTEST_FALSE("Audio envelope had zero output", HadNonZeroEnvelope);
+		}
+
+		// Enable, render a few blocks and expect some output
+		Generator->SetInputValue("Enable", true);
+		for (int32 i = 0; i < NumBlocksToRender; ++i)
+		{
+			Generator->OnGenerateAudio(DummyBuffer.GetData(), DummyBuffer.Num());
+		}
+
+		{
+			UTEST_TRUE("Float envelope was non-zero", **FloatEnvelopeOutput > 0.0f);
+			bool HadNonZeroEnvelope = false;
+			for (int32 i = 0; i < BlockSize; ++i)
+			{
+				const float Value = (*AudioEnvelopeOutput)->GetData()[i];
+				if (Value > 0.0f)
+				{
+					HadNonZeroEnvelope = true;
+					break;
+				}
+			}
+			UTEST_TRUE("Audio envelope had non-zero output", HadNonZeroEnvelope);
+		}
+		
+		return true;
+	}
+
+}
+
+#endif
diff --git a/Engine/Plugins/Runtime/Metasound/Source/MetasoundStandardNodes/Private/MetasoundEnvelopeFollowerNode.cpp b/Engine/Plugins/Runtime/Metasound/Source/MetasoundStandardNodes/Private/MetasoundEnvelopeFollowerNode.cpp
--- a/Engine/Plugins/Runtime/Metasound/Source/MetasoundStandardNodes/Private/MetasoundEnvelopeFollowerNode.cpp
+++ b/Engine/Plugins/Runtime/Metasound/Source/MetasoundStandardNodes/Private/MetasoundEnvelopeFollowerNode.cpp
@@ -28,6 +28,7 @@ namespace Metasound
 {
 	namespace EnvelopeFollowerVertexNames
 	{
+		METASOUND_PARAM(InParamEnable, "Enable", "Enable the envelope follower.")
 		METASOUND_PARAM(InParamAudioInput, "In", "Audio input.")
 		METASOUND_PARAM(InParamAttackTime, "Attack Time", "The attack time of the envelope follower.")
 		METASOUND_PARAM(InParamReleaseTime,"Release Time" , "The release time of the envelope follower.")
@@ -45,6 +46,7 @@ namespace Metasound
 		static TUniquePtr<IOperator> CreateOperator(const FCreateOperatorParams& InParams, FBuildErrorArray& OutErrors);
 
 		FEnvelopeFollowerOperator(const FCreateOperatorParams& InOperatorSettings,
+			const FBoolReadRef& InEnable,
 			const FAudioBufferReadRef& InAudioInput,
 			const FTimeReadRef& InAttackTime,
 			const FTimeReadRef& InReleaseTime,
@@ -58,6 +60,10 @@ namespace Metasound
 		void Execute();
 
 	private:
+		// Whether the enveloper follower is enabled
+		FBoolReadRef EnableInput;
+		bool OutputNeedsResetOnDisable = true;
+		
 		// The input audio buffer
 		FAudioBufferReadRef AudioInput;
 
@@ -83,11 +89,13 @@ namespace Metasound
 	};
 
 	FEnvelopeFollowerOperator::FEnvelopeFollowerOperator(const FCreateOperatorParams& InParams,
+		const FBoolReadRef& InEnable,
 		const FAudioBufferReadRef& InAudioInput,
 		const FTimeReadRef& InAttackTime,
 		const FTimeReadRef& InReleaseTime,
 		const FEnvelopePeakModeReadRef& InEnvelopeMode)
-		: AudioInput(InAudioInput)
+		: EnableInput(InEnable)
+		, AudioInput(InAudioInput)
 		, AttackTimeInput(InAttackTime)
 		, ReleaseTimeInput(InReleaseTime)
 		, FollowModeInput(InEnvelopeMode)
@@ -101,6 +109,7 @@ namespace Metasound
 	{
 		using namespace EnvelopeFollowerVertexNames;
 
+		InOutVertexData.BindReadVertex(METASOUND_GET_PARAM_NAME(InParamEnable), EnableInput);
 		InOutVertexData.BindReadVertex(METASOUND_GET_PARAM_NAME(InParamAudioInput), AudioInput);
 		InOutVertexData.BindReadVertex(METASOUND_GET_PARAM_NAME(InParamAttackTime), AttackTimeInput);
 		InOutVertexData.BindReadVertex(METASOUND_GET_PARAM_NAME(InParamReleaseTime), ReleaseTimeInput);
@@ -133,6 +142,8 @@ namespace Metasound
 
 	void FEnvelopeFollowerOperator::Reset(const IOperator::FResetParams& InParams)
 	{
+		OutputNeedsResetOnDisable = true;
+		
 		PrevAttackTime = FMath::Max(FTime::ToMilliseconds(*AttackTimeInput), 0.0);
 		PrevReleaseTime = FMath::Max(FTime::ToMilliseconds(*ReleaseTimeInput), 0.0);
 
@@ -150,6 +161,21 @@ namespace Metasound
 
 	void FEnvelopeFollowerOperator::Execute()
 	{
+		// Skip rendering if disabled
+		if (!*EnableInput)
+		{
+			if (OutputNeedsResetOnDisable)
+			{
+				EnvelopeAudioOutput->Zero();
+				*EnvelopeFloatOutput = 0;
+				OutputNeedsResetOnDisable = false;
+			}
+			
+			return;
+		}
+
+		OutputNeedsResetOnDisable = true;
+		
 		// Check for any input changes
 		double CurrentAttackTime = FMath::Max(FTime::ToMilliseconds(*AttackTimeInput), 0.0);
 		if (!FMath::IsNearlyEqual(CurrentAttackTime, PrevAttackTime))
@@ -201,6 +227,7 @@ namespace Metasound
 
 		static const FVertexInterface Interface(
 			FInputVertexInterface(
+				TInputDataVertex<bool>(METASOUND_GET_PARAM_NAME_AND_METADATA(InParamEnable), true),
 				TInputDataVertex<FAudioBuffer>(METASOUND_GET_PARAM_NAME_AND_METADATA(InParamAudioInput)),
 				TInputDataVertex<FTime>(METASOUND_GET_PARAM_NAME_AND_METADATA(InParamAttackTime), 0.01f),
 				TInputDataVertex<FTime>(METASOUND_GET_PARAM_NAME_AND_METADATA(InParamReleaseTime), 0.1f),
@@ -242,16 +269,16 @@ namespace Metasound
 	{
 		using namespace EnvelopeFollowerVertexNames;
 
-		const FEnvelopeFollowerNode& EnvelopeFollowerNode = static_cast<const FEnvelopeFollowerNode&>(InParams.Node);
 		const FDataReferenceCollection& InputCollection = InParams.InputDataReferences;
 		const FInputVertexInterface& InputInterface = GetVertexInterface().GetInputInterface();
 
+		FBoolReadRef EnableIn = InputCollection.GetDataReadReferenceOrConstructWithVertexDefault<bool>(InputInterface, METASOUND_GET_PARAM_NAME(InParamEnable), InParams.OperatorSettings);
 		FAudioBufferReadRef AudioIn = InputCollection.GetDataReadReferenceOrConstruct<FAudioBuffer>(METASOUND_GET_PARAM_NAME(InParamAudioInput), InParams.OperatorSettings);
 		FTimeReadRef AttackTime = InputCollection.GetDataReadReferenceOrConstructWithVertexDefault<FTime>(InputInterface, METASOUND_GET_PARAM_NAME(InParamAttackTime), InParams.OperatorSettings);
 		FTimeReadRef ReleaseTime = InputCollection.GetDataReadReferenceOrConstructWithVertexDefault<FTime>(InputInterface, METASOUND_GET_PARAM_NAME(InParamReleaseTime), InParams.OperatorSettings);
 		FEnvelopePeakModeReadRef EnvelopeModeIn = InputCollection.GetDataReadReferenceOrConstructWithVertexDefault<FEnumEnvelopePeakMode>(InputInterface, METASOUND_GET_PARAM_NAME(InParamFollowMode), InParams.OperatorSettings);
 
-		return MakeUnique<FEnvelopeFollowerOperator>(InParams, AudioIn, AttackTime, ReleaseTime, EnvelopeModeIn);
+		return MakeUnique<FEnvelopeFollowerOperator>(InParams, EnableIn, AudioIn, AttackTime, ReleaseTime, EnvelopeModeIn);
 	}
 
 	FEnvelopeFollowerNode::FEnvelopeFollowerNode(const FNodeInitData& InitData)
diff --git a/Engine/Plugins/Runtime/Metasound/Source/MetasoundStandardNodes/Public/MetasoundStandardNodesNames.h b/Engine/Plugins/Runtime/Metasound/Source/MetasoundStandardNodes/Public/MetasoundStandardNodesNames.h
--- a/Engine/Plugins/Runtime/Metasound/Source/MetasoundStandardNodes/Public/MetasoundStandardNodesNames.h
+++ b/Engine/Plugins/Runtime/Metasound/Source/MetasoundStandardNodes/Public/MetasoundStandardNodesNames.h
@@ -8,7 +8,7 @@ namespace Metasound
 {
 	namespace StandardNodes
 	{
-		extern const FName Namespace;
-		extern const FName AudioVariant;
+		METASOUNDSTANDARDNODES_API extern const FName Namespace;
+		METASOUNDSTANDARDNODES_API extern const FName AudioVariant;
 	}
 }
