global metasound sample rate override

diff --git a/Engine/Plugins/Runtime/Metasound/Source/MetasoundEngine/Private/MetasoundSource.cpp b/Engine/Plugins/Runtime/Metasound/Source/MetasoundEngine/Private/MetasoundSource.cpp
--- a/Engine/Plugins/Runtime/Metasound/Source/MetasoundEngine/Private/MetasoundSource.cpp
+++ b/Engine/Plugins/Runtime/Metasound/Source/MetasoundEngine/Private/MetasoundSource.cpp
@@ -152,7 +152,15 @@ UMetaSoundSource::UMetaSoundSource(const FObjectInitializer& ObjectInitializer)
 	NumChannels = 1;
 
 	// todo: ensure that we have a method so that the audio engine can be authoritative over the sample rate the UMetaSoundSource runs at.
-	SampleRate = 48000.f;
+	const int32 SampleRateOverride = Metasound::Frontend::GetDefaultSampleRate();
+	if (SampleRateOverride != INDEX_NONE)
+	{
+		SampleRate = SampleRateOverride;
+	}
+	else
+	{
+		SampleRate = 48000.f;
+	}
 }
 
 const UClass& UMetaSoundSource::GetBaseMetaSoundUClass() const
@@ -585,8 +593,10 @@ ISoundGeneratorPtr UMetaSoundSource::CreateSoundGenerator(const FSoundGeneratorI
 	METASOUND_LLM_SCOPE;
 	METASOUND_TRACE_CPUPROFILER_EVENT_SCOPE(UMetaSoundSource::CreateSoundGenerator);
 
-	SampleRate = InParams.SampleRate;
-	FOperatorSettings InSettings = GetOperatorSettings(static_cast<FSampleRate>(SampleRate));
+	FOperatorSettings InSettings = GetOperatorSettings(static_cast<FSampleRate>(InParams.SampleRate));
+
+	SampleRate = InSettings.GetSampleRate();
+
 	FMetasoundEnvironment Environment = CreateEnvironment(InParams);
 	FParameterRouter& Router = GetParameterRouter();
 	TSharedPtr<TSpscQueue<FMetaSoundParameterTransmitter::FParameter>> DataChannel = Router.FindOrCreateDataChannelForReader(InParams.AudioDeviceID, InParams.InstanceID);
@@ -964,6 +974,11 @@ TSharedPtr<Audio::IParameterTransmitter> UMetaSoundSource::CreateParameterTransm
 Metasound::FOperatorSettings UMetaSoundSource::GetOperatorSettings(Metasound::FSampleRate InSampleRate) const
 {
 	const float BlockRate = Metasound::Frontend::GetDefaultBlockRate();
+	const int32 SampleRateOverride = Metasound::Frontend::GetDefaultSampleRate();
+	if (SampleRateOverride != INDEX_NONE)
+	{
+		return Metasound::FOperatorSettings(SampleRateOverride, BlockRate);
+	}
 	return Metasound::FOperatorSettings(InSampleRate, BlockRate);
 }
 
diff --git a/Engine/Plugins/Runtime/Metasound/Source/MetasoundFrontend/Private/MetasoundAssetBase.cpp b/Engine/Plugins/Runtime/Metasound/Source/MetasoundFrontend/Private/MetasoundAssetBase.cpp
--- a/Engine/Plugins/Runtime/Metasound/Source/MetasoundFrontend/Private/MetasoundAssetBase.cpp
+++ b/Engine/Plugins/Runtime/Metasound/Source/MetasoundFrontend/Private/MetasoundAssetBase.cpp
@@ -48,6 +48,9 @@ namespace Metasound
 		{
 			static float BlockRate = 100.f;
 
+			// Negative value means we use the audio mixer device sample rate
+			static int32 SampleRate = INDEX_NONE;
+
 			TScriptInterface<IMetaSoundDocumentInterface> BuildRegistryDocument(TScriptInterface<IMetaSoundDocumentInterface> DocumentInterface)
 			{
 				using namespace Metasound::Frontend;
@@ -255,10 +258,27 @@ namespace Metasound
 			TEXT("Default: 100.0f, Min: 1.0f, Max: 1000.0f"),
 			ECVF_Default);
 
+		FAutoConsoleVariableRef CVarMetaSoundSampleRate(
+			TEXT("au.MetaSound.SampleRate"),
+			AssetBasePrivate::SampleRate,
+			TEXT("Overrides the sample rate of metasounds. Negative values default to audio mixer sample rate.\n")
+			TEXT("Default: -1, Min: 8000, Max: 48000"),
+			ECVF_Default);
+
+
 		float GetDefaultBlockRate()
 		{
 			return FMath::Clamp(AssetBasePrivate::BlockRate, 1.0f, 1000.0f);
 		}
+
+		int32 GetDefaultSampleRate()
+		{
+			if (AssetBasePrivate::SampleRate > 0)
+			{
+				return FMath::Clamp(AssetBasePrivate::SampleRate, 8000, 48000);
+			}
+			return AssetBasePrivate::SampleRate;
+		}
 	} // namespace Frontend
 } // namespace Metasound
 
diff --git a/Engine/Plugins/Runtime/Metasound/Source/MetasoundFrontend/Public/MetasoundAssetBase.h b/Engine/Plugins/Runtime/Metasound/Source/MetasoundFrontend/Public/MetasoundAssetBase.h
--- a/Engine/Plugins/Runtime/Metasound/Source/MetasoundFrontend/Public/MetasoundAssetBase.h
+++ b/Engine/Plugins/Runtime/Metasound/Source/MetasoundFrontend/Public/MetasoundAssetBase.h
@@ -29,6 +29,7 @@ namespace Metasound
 		class IInterfaceRegistryEntry;
 
 		METASOUNDFRONTEND_API float GetDefaultBlockRate();
+		METASOUNDFRONTEND_API int32 GetDefaultSampleRate();
 
 		class FProxyDataCache;
 
