Fix for failed assert in lowlevel metasound tests

diff --git a/Engine/Plugins/Runtime/Metasound/Source/MetasoundFrontend/Private/MetasoundFrontendDocumentAccessPtr.cpp b/Engine/Plugins/Runtime/Metasound/Source/MetasoundFrontend/Private/MetasoundFrontendDocumentAccessPtr.cpp
--- a/Engine/Plugins/Runtime/Metasound/Source/MetasoundFrontend/Private/MetasoundFrontendDocumentAccessPtr.cpp
+++ b/Engine/Plugins/Runtime/Metasound/Source/MetasoundFrontend/Private/MetasoundFrontendDocumentAccessPtr.cpp
@@ -409,6 +409,11 @@ namespace Metasound
 				{
 					return NodeRegistryKey::IsEqual(RegistryKey, NodeRegistryKey::CreateKey(InClass.Metadata));
 				};
+
+				bool operator()(const FMetasoundFrontendGraphClass& InGraphClass) const
+				{
+					return NodeRegistryKey::IsEqual(RegistryKey, NodeRegistryKey::CreateKey(InGraphClass));
+				};
 			private:
 				FNodeRegistryKey RegistryKey;
 			};
diff --git a/Engine/Plugins/Runtime/Metasound/Source/MetasoundFrontend/Private/MetasoundFrontendGraphController.cpp b/Engine/Plugins/Runtime/Metasound/Source/MetasoundFrontend/Private/MetasoundFrontendGraphController.cpp
--- a/Engine/Plugins/Runtime/Metasound/Source/MetasoundFrontend/Private/MetasoundFrontendGraphController.cpp
+++ b/Engine/Plugins/Runtime/Metasound/Source/MetasoundFrontend/Private/MetasoundFrontendGraphController.cpp
@@ -1652,17 +1652,20 @@ namespace Metasound
 					FMetasoundFrontendNode& Node = GraphClass->Graph.Nodes.Emplace_GetRef(*NodeClass);
 
 					// Cache the asset name on the node if it node is reference to asset-defined graph.
-					const FNodeRegistryKey RegistryKey = NodeRegistryKey::CreateKey(NodeClass->Metadata);
-					if (IMetaSoundAssetManager* AssetManager = IMetaSoundAssetManager::Get())
+					if (NodeClass->Metadata.GetType() == EMetasoundFrontendClassType::External)
 					{
-						if (const FSoftObjectPath* Path = AssetManager->FindObjectPathFromKey(RegistryKey))
+						if (IMetaSoundAssetManager* AssetManager = IMetaSoundAssetManager::Get())
 						{
-							const FString& AssetName = Path->GetAssetName();
-							Node.Name = *AssetName;
+							const FNodeRegistryKey RegistryKey = NodeRegistryKey::CreateKey(NodeClass->Metadata);
+							if (const FSoftObjectPath* Path = AssetManager->FindObjectPathFromKey(RegistryKey))
+							{
+								const FString& AssetName = Path->GetAssetName();
+								Node.Name = *AssetName;
+							}
 						}
 					}
-					Node.UpdateID(InNodeGuid);
 
+					Node.UpdateID(InNodeGuid);
 #if WITH_EDITOR
 					if (FMetasoundFrontendDocumentMetadata* DocMetadata = OwningDocument->GetMetadata())
 					{
