Register metasound graphs asynchronously

diff --git a/Engine/Plugins/Runtime/Metasound/Source/MetasoundEditor/Private/MetasoundEditorSubsystem.cpp b/Engine/Plugins/Runtime/Metasound/Source/MetasoundEditor/Private/MetasoundEditorSubsystem.cpp
--- a/Engine/Plugins/Runtime/Metasound/Source/MetasoundEditor/Private/MetasoundEditorSubsystem.cpp
+++ b/Engine/Plugins/Runtime/Metasound/Source/MetasoundEditor/Private/MetasoundEditorSubsystem.cpp
@@ -158,9 +158,7 @@ void UMetaSoundEditorSubsystem::InitAsset(UObject& InNewMetaSound, UObject* InRe
 
 		// Initialize preset with referenced Metasound 
 		TScriptInterface<IMetaSoundDocumentInterface> ReferencedDocInterface = InReferencedMetaSound;
-		const IMetaSoundDocumentInterface* ReferencedInterface = ReferencedDocInterface.GetInterface();
-		check(ReferencedInterface);
-		Builder.ConvertToPreset(ReferencedInterface->GetDocument());
+		Builder.ConvertToPreset(ReferencedDocInterface->GetConstDocument());
 
 		// Update asset object data from interfaces 
 		FMetasoundAssetBase* PresetAsset = IMetasoundUObjectRegistry::Get().GetObjectAsAssetBase(&InNewMetaSound);
@@ -168,7 +166,7 @@ void UMetaSoundEditorSubsystem::InitAsset(UObject& InNewMetaSound, UObject* InRe
 		PresetAsset->ConformObjectDataToInterfaces();
 
 		// Copy sound wave settings to preset for sources
-		if (&ReferencedInterface->GetBaseMetaSoundUClass() == UMetaSoundSource::StaticClass())
+		if (&ReferencedDocInterface->GetBaseMetaSoundUClass() == UMetaSoundSource::StaticClass())
 		{
 			SetSoundWaveSettingsFromTemplate(*CastChecked<USoundWave>(&InNewMetaSound), *CastChecked<USoundWave>(InReferencedMetaSound));
 		}
diff --git a/Engine/Plugins/Runtime/Metasound/Source/MetasoundEditor/Private/MetasoundFactory.cpp b/Engine/Plugins/Runtime/Metasound/Source/MetasoundEditor/Private/MetasoundFactory.cpp
--- a/Engine/Plugins/Runtime/Metasound/Source/MetasoundEditor/Private/MetasoundFactory.cpp
+++ b/Engine/Plugins/Runtime/Metasound/Source/MetasoundEditor/Private/MetasoundFactory.cpp
@@ -61,9 +61,7 @@ void UMetaSoundBaseFactory::InitAsset(UObject& InNewMetaSound, UObject* InRefere
 		
 		// Initialize preset with referenced Metasound 
 		TScriptInterface<IMetaSoundDocumentInterface> ReferencedDocInterface = InReferencedMetaSound;
-		const IMetaSoundDocumentInterface* ReferencedInterface = ReferencedDocInterface.GetInterface();
-		check(ReferencedInterface);
-		Builder.ConvertToPreset(ReferencedInterface->GetDocument());
+		Builder.ConvertToPreset(ReferencedDocInterface->GetConstDocument());
 		
 		// Update asset object data from interfaces 
 		FMetasoundAssetBase* PresetAsset = IMetasoundUObjectRegistry::Get().GetObjectAsAssetBase(&InNewMetaSound);
diff --git a/Engine/Plugins/Runtime/Metasound/Source/MetasoundEngine/Private/Metasound.cpp b/Engine/Plugins/Runtime/Metasound/Source/MetasoundEngine/Private/Metasound.cpp
--- a/Engine/Plugins/Runtime/Metasound/Source/MetasoundEngine/Private/Metasound.cpp
+++ b/Engine/Plugins/Runtime/Metasound/Source/MetasoundEngine/Private/Metasound.cpp
@@ -88,7 +88,7 @@ const UClass& UMetaSoundPatch::GetBaseMetaSoundUClass() const
 	return *UMetaSoundPatch::StaticClass();
 }
 
-const FMetasoundFrontendDocument& UMetaSoundPatch::GetDocument() const
+const FMetasoundFrontendDocument& UMetaSoundPatch::GetConstDocument() const
 {
 	return RootMetaSoundDocument;
 }
@@ -196,5 +196,32 @@ void UMetaSoundPatch::OnAsyncReferencedAssetsLoaded(const TArray<FMetasoundAsset
 	Metasound::FMetaSoundEngineAssetHelper::OnAsyncReferencedAssetsLoaded(*this, InAsyncReferences);
 }
 
+bool UMetaSoundPatch::IsBuilderActive() const
+{
+	return bIsBuilderActive;
+}
+
+void UMetaSoundPatch::OnBeginActiveBuilder()
+{
+	if (bIsBuilderActive)
+	{
+		UE_LOG(LogMetaSound, Error, TEXT("OnBeginActiveBuilder() call while prior builder is still active. This may indicate that multiple builders are attempting to modify the MetaSound %s concurrently."), *GetOwningAssetName())
+	}
+
+	// If a builder is activating, make sure any in-flight registration
+	// tasks have completed. Async registration tasks use the FMetasoundFrontendDocument
+	// that lives on this object. We need to make sure that registration task
+	// completes so that the FMetasoundFrontendDocument does not get modified
+	// by a builder while it is also being read by async registration.
+	WaitForAsyncGraphRegistration();
+
+	bIsBuilderActive = true;
+}
+
+void UMetaSoundPatch::OnFinishActiveBuilder()
+{
+	bIsBuilderActive = false;
+}
+
 #undef LOCTEXT_NAMESPACE // MetaSound
 
diff --git a/Engine/Plugins/Runtime/Metasound/Source/MetasoundEngine/Private/MetasoundBuilderSubsystem.cpp b/Engine/Plugins/Runtime/Metasound/Source/MetasoundEngine/Private/MetasoundBuilderSubsystem.cpp
--- a/Engine/Plugins/Runtime/Metasound/Source/MetasoundEngine/Private/MetasoundBuilderSubsystem.cpp
+++ b/Engine/Plugins/Runtime/Metasound/Source/MetasoundEngine/Private/MetasoundBuilderSubsystem.cpp
@@ -55,6 +55,15 @@ namespace Metasound::Engine
 	} // namespace BuilderSubsystemPrivate
 } // namespace Metasound::Engine
 
+void UMetaSoundBuilderBase::FinishDestroy()
+{
+	// Need to detach before destroying UPROPERTYs as the Builder 
+	// often holds a TScriptInterface<IMetaSoundDocumentInterface> of
+	// a UPROPERTY that lives on this or derived objects. 
+	Builder.FinishBuilding();
+
+	Super::FinishDestroy();
+}
 
 FMetaSoundBuilderNodeOutputHandle UMetaSoundBuilderBase::AddGraphInputNode(FName Name, FName DataType, FMetasoundFrontendLiteral DefaultValue, EMetaSoundBuilderResult& OutResult, bool bIsConstructorInput)
 {
@@ -169,9 +178,7 @@ FMetaSoundNodeHandle UMetaSoundBuilderBase::AddNode(const TScriptInterface<IMeta
 		{
 			MetaSoundAsset->RegisterGraphWithFrontend();
 
-			const IMetaSoundDocumentInterface* Interface = NodeClass.GetInterface();
-			check(Interface);
-			const FMetasoundFrontendDocument& NodeClassDoc = Interface->GetDocument();
+			const FMetasoundFrontendDocument& NodeClassDoc = NodeClass->GetConstDocument();
 			const FMetasoundFrontendGraphClass& NodeClassGraph = NodeClassDoc.RootGraph;
 
 			if (const FMetasoundFrontendNode* NewNode = Builder.AddGraphNode(NodeClassGraph))
@@ -406,7 +413,7 @@ void UMetaSoundBuilderBase::ConvertToPreset(const TScriptInterface<IMetaSoundDoc
 		ReferencedMetaSoundAsset->RegisterGraphWithFrontend();
 	}
 
-	const FMetasoundFrontendDocument& ReferencedDocument = ReferencedInterface->GetDocument();
+	const FMetasoundFrontendDocument& ReferencedDocument = ReferencedInterface->GetConstDocument();
 	if (Builder.ConvertToPreset(ReferencedDocument))
 	{
 		OutResult = EMetaSoundBuilderResult::Succeeded;
diff --git a/Engine/Plugins/Runtime/Metasound/Source/MetasoundEngine/Private/MetasoundSource.cpp b/Engine/Plugins/Runtime/Metasound/Source/MetasoundEngine/Private/MetasoundSource.cpp
--- a/Engine/Plugins/Runtime/Metasound/Source/MetasoundEngine/Private/MetasoundSource.cpp
+++ b/Engine/Plugins/Runtime/Metasound/Source/MetasoundEngine/Private/MetasoundSource.cpp
@@ -168,7 +168,7 @@ const UClass& UMetaSoundSource::GetBaseMetaSoundUClass() const
 	return *UMetaSoundSource::StaticClass();
 }
 
-const FMetasoundFrontendDocument& UMetaSoundSource::GetDocument() const
+const FMetasoundFrontendDocument& UMetaSoundSource::GetConstDocument() const
 {
 	return RootMetasoundDocument;
 }
@@ -218,7 +218,7 @@ void UMetaSoundSource::PostEditChangeOutputFormat()
 		// TODO: Once builders are notified of controller changes and can be safely persistent, this
 		// can be removed so builders can be shared and not have to be created for each change output
 		// format mutation transaction.
-		BuilderSubsystem.DetachBuilderFromAsset(GetDocument().RootGraph.Metadata.GetClassName());
+		BuilderSubsystem.DetachBuilderFromAsset(GetConstDocument().RootGraph.Metadata.GetClassName());
 	}
 
 	if (Result == EMetaSoundBuilderResult::Succeeded)
@@ -360,11 +360,20 @@ void UMetaSoundSource::InitParameters(TArray<FAudioParameter>& ParametersToInit,
 	METASOUND_LLM_SCOPE;
 	METASOUND_TRACE_CPUPROFILER_EVENT_SCOPE(UMetaSoundSource::InitParameters);
 
-	// Have to call register vs a simple get as the source may have yet to start playing/has not been registered
-	// via InitResources. If it has, this call is fast and returns the already cached RuntimeData.
-	RegisterGraphWithFrontend(GetInitRegistrationOptions());
 	const FRuntimeData& RuntimeData = GetRuntimeData();
 
+	// To initialize parameters, we need the PublicInputMap which lives on FRuntimeData. 
+	// If the runtime data has not valid, it can be created via a call to RegisterGraphWithFrontend(...) 
+	// which subsequently pupulates the runtime data.
+	if (!RuntimeData.IsValid())
+	{
+		// If a InitParameters is called before InitResources, the graph will not 
+		// yet be registered. RegisterGraphWithFrontend is called here to cover that 
+		// scenario. 
+		RegisterGraphWithFrontend(GetInitRegistrationOptions());
+		check(RuntimeData.IsValid());
+	}
+
 	IDataTypeRegistry& DataTypeRegistry = IDataTypeRegistry::Get();
 	const Metasound::TSortedVertexNameMap<FRuntimeInput>& PublicInputMap = RuntimeData.PublicInputMap;
 
@@ -649,7 +658,7 @@ ISoundGeneratorPtr UMetaSoundSource::CreateSoundGenerator(const FSoundGeneratorI
 	}
 	else
 	{
-		TSharedPtr<const IGraph, ESPMode::ThreadSafe> MetasoundGraph = GetRuntimeData().Graph;
+		TSharedPtr<const IGraph> MetasoundGraph = FMetasoundFrontendRegistryContainer::Get()->GetGraph(GetRegistryKey());
 		if (!MetasoundGraph.IsValid())
 		{
 			return ISoundGeneratorPtr(nullptr);
@@ -691,6 +700,12 @@ bool UMetaSoundSource::GetAllDefaultParameters(TArray<FAudioParameter>& OutParam
 	using namespace Metasound::Frontend;
 	using namespace Metasound::Engine;
 
+	const FRuntimeData& RuntimeData = GetRuntimeData();
+	if(!RuntimeData.IsValid())
+	{
+		UE_LOG(LogMetaSound, Warning, TEXT("Default parameters may be incorrect. Accessing invalid runtime data on MetaSound %s"), *GetOwningAssetName());
+	}
+
 	for(const TPair<FVertexName, FMetasoundAssetBase::FRuntimeInput>& Pair : GetRuntimeData().PublicInputMap)
 	{
 		const FMetasoundAssetBase::FRuntimeInput& Input = Pair.Value;
@@ -951,6 +966,10 @@ TSharedPtr<Audio::IParameterTransmitter> UMetaSoundSource::CreateParameterTransm
 	METASOUND_LLM_SCOPE;
 
 	const FRuntimeData& RuntimeData = GetRuntimeData();
+	if(!RuntimeData.IsValid())
+	{
+		UE_LOG(LogMetaSound, Warning, TEXT("Parameter Transmitter may not work. Accessing invalid runtime data on MetaSound %s"), *GetOwningAssetName());
+	}
 
 	// Build list of parameters that can be set at runtime.
 	TArray<FName> ValidParameters;
@@ -1085,6 +1104,33 @@ Metasound::SourcePrivate::FParameterRouter& UMetaSoundSource::GetParameterRouter
 	return Router;
 }
 
+bool UMetaSoundSource::IsBuilderActive() const
+{
+	return bIsBuilderActive;
+}
+
+void UMetaSoundSource::OnBeginActiveBuilder()
+{
+	if (bIsBuilderActive)
+	{
+		UE_LOG(LogMetaSound, Error, TEXT("OnBeginActiveBuilder() call while prior builder is still active. This may indicate that multiple builders are attempting to modify the MetaSound %s concurrently."), *GetOwningAssetName())
+	}
+
+	// If a builder is activating, make sure any in-flight registration
+	// tasks have completed. Async registration tasks use the FMetasoundFrontendDocument
+	// that lives on this object. We need to make sure that registration task
+	// completes so that the FMetasoundFrontendDocument does not get modified
+	// by a builder while it is also being read by async registration.
+	WaitForAsyncGraphRegistration();
+
+	bIsBuilderActive = true;
+}
+
+void UMetaSoundSource::OnFinishActiveBuilder()
+{
+	bIsBuilderActive = false;
+}
+
 TSharedPtr<Metasound::DynamicGraph::FDynamicOperatorTransactor> UMetaSoundSource::SetDynamicGeneratorEnabled(bool bInIsEnabled)
 {
 	using namespace Metasound;
@@ -1094,10 +1140,25 @@ TSharedPtr<Metasound::DynamicGraph::FDynamicOperatorTransactor> UMetaSoundSource
 	{
 		if (!DynamicTransactor.IsValid())
 		{
-			TSharedPtr<FGraph> CurrentGraph = GetRuntimeData().Graph;
-			if (CurrentGraph.IsValid())
+			// If a FGraph exists for this UMetaSoundSource, then we need to initialize
+			// the DynamicTransactor with the existing FGraph so it has the correct
+			// initial state. 
+			//
+			// Currently, any existing FGraph will be stored in the node registry,
+			// hence we check if the graph is registered and retrieve the current 
+			// graph to see if any FGraph already exists. 
+			if (IsRegistered())
 			{
-				DynamicTransactor = MakeShared<FDynamicOperatorTransactor>(*CurrentGraph);
+				TSharedPtr<const FGraph> CurrentGraph = FMetasoundFrontendRegistryContainer::Get()->GetGraph(GetRegistryKey());
+				if (CurrentGraph.IsValid())
+				{
+					DynamicTransactor = MakeShared<FDynamicOperatorTransactor>(*CurrentGraph);
+				}
+				else
+				{
+					UE_LOG(LogMetaSound, Warning, TEXT("Failed to get existing graph for dynamic metasound %s. Initializing to empty graph."), *GetOwningAssetName());
+					DynamicTransactor = MakeShared<FDynamicOperatorTransactor>();
+				}
 			}
 			else
 			{
diff --git a/Engine/Plugins/Runtime/Metasound/Source/MetasoundEngine/Public/Metasound.h b/Engine/Plugins/Runtime/Metasound/Source/MetasoundEngine/Public/Metasound.h
--- a/Engine/Plugins/Runtime/Metasound/Source/MetasoundEngine/Public/Metasound.h
+++ b/Engine/Plugins/Runtime/Metasound/Source/MetasoundEngine/Public/Metasound.h
@@ -131,7 +131,7 @@ public:
 #endif // #if WITH_EDITORONLY_DATA
 
 	virtual const UClass& GetBaseMetaSoundUClass() const final override;
-	virtual const FMetasoundFrontendDocument& GetDocument() const override;
+	virtual const FMetasoundFrontendDocument& GetConstDocument() const override;
 
 #if WITH_EDITOR
 	virtual void PostDuplicate(EDuplicateMode::Type DuplicateMode) override;
@@ -176,4 +176,10 @@ private:
 	{
 		return RootMetaSoundDocument;
 	}
+
+	virtual bool IsBuilderActive() const override;
+	virtual void OnBeginActiveBuilder() override;
+	virtual void OnFinishActiveBuilder() override;
+
+	bool bIsBuilderActive = false;
 };
diff --git a/Engine/Plugins/Runtime/Metasound/Source/MetasoundEngine/Public/MetasoundBuilderSubsystem.h b/Engine/Plugins/Runtime/Metasound/Source/MetasoundEngine/Public/MetasoundBuilderSubsystem.h
--- a/Engine/Plugins/Runtime/Metasound/Source/MetasoundEngine/Public/MetasoundBuilderSubsystem.h
+++ b/Engine/Plugins/Runtime/Metasound/Source/MetasoundEngine/Public/MetasoundBuilderSubsystem.h
@@ -132,6 +132,10 @@ class METASOUNDENGINE_API UMetaSoundBuilderBase : public UObject
 	GENERATED_BODY()
 
 public:
+	// Begin UObject interface
+	virtual void FinishDestroy() override;
+	// End UObject interface
+
 	// Adds a graph input node with the given name, DataType, and sets the graph input to default value.
 	// Returns the new input node's output handle if it was successfully created, or an invalid handle if it failed.
 	UFUNCTION(BlueprintCallable, Category = "Audio|MetaSound|Builder", meta = (ExpandEnumAsExecs = "OutResult", AdvancedDisplay = "3"))
@@ -698,7 +702,7 @@ private:
 		check(InMetaSoundObject->IsAsset());
 
 		TScriptInterface<IMetaSoundDocumentInterface> DocInterface = InMetaSoundObject;
-		const FMetasoundFrontendDocument& Document = static_cast<const IMetaSoundDocumentInterface*>(DocInterface.GetInterface())->GetDocument();
+		const FMetasoundFrontendDocument& Document = DocInterface->GetConstDocument();
 		const FName FullClassName = Document.RootGraph.Metadata.GetClassName().GetFullName();
 		TWeakObjectPtr<UMetaSoundBuilderBase> Builder = AssetBuilders.FindRef(FullClassName);
 		if (Builder.IsValid())
diff --git a/Engine/Plugins/Runtime/Metasound/Source/MetasoundEngine/Public/MetasoundSource.h b/Engine/Plugins/Runtime/Metasound/Source/MetasoundEngine/Public/MetasoundSource.h
--- a/Engine/Plugins/Runtime/Metasound/Source/MetasoundEngine/Public/MetasoundSource.h
+++ b/Engine/Plugins/Runtime/Metasound/Source/MetasoundEngine/Public/MetasoundSource.h
@@ -205,7 +205,7 @@ protected:
 	Metasound::Frontend::FConstDocumentAccessPtr GetDocumentConstAccessPtr() const override;
 
 	virtual const UClass& GetBaseMetaSoundUClass() const final override;
-	virtual const FMetasoundFrontendDocument& GetDocument() const override;
+	virtual const FMetasoundFrontendDocument& GetConstDocument() const override;
 
 	/** Gets all the default parameters for this Asset.  */
 	virtual bool GetAllDefaultParameters(TArray<FAudioParameter>& OutParameters) const override;
@@ -220,6 +220,9 @@ private:
 		return RootMetasoundDocument;
 	}
 
+	virtual bool IsBuilderActive() const override;
+	virtual void OnBeginActiveBuilder() override;
+	virtual void OnFinishActiveBuilder() override;
 
 	bool IsParameterValid(const FAudioParameter& InParameter, const FName& InTypeName, Metasound::Frontend::IDataTypeRegistry& InDataTypeRegistry) const;
 
@@ -254,4 +257,6 @@ private:
 	TSharedPtr<Metasound::DynamicGraph::FDynamicOperatorTransactor> GetDynamicGeneratorTransactor() const;
 
 	TSharedPtr<Metasound::DynamicGraph::FDynamicOperatorTransactor> DynamicTransactor;
+
+	bool bIsBuilderActive = false;
 };
diff --git a/Engine/Plugins/Runtime/Metasound/Source/MetasoundFrontend/Private/MetasoundAssetBase.cpp b/Engine/Plugins/Runtime/Metasound/Source/MetasoundFrontend/Private/MetasoundAssetBase.cpp
--- a/Engine/Plugins/Runtime/Metasound/Source/MetasoundFrontend/Private/MetasoundAssetBase.cpp
+++ b/Engine/Plugins/Runtime/Metasound/Source/MetasoundFrontend/Private/MetasoundAssetBase.cpp
@@ -57,9 +57,6 @@ namespace Metasound
 
 				METASOUND_TRACE_CPUPROFILER_EVENT_SCOPE(BuildRegistryDocument);
 
-				const IMetaSoundDocumentInterface* ConstInterface = DocumentInterface.GetInterface();
-				check(ConstInterface != nullptr);
-
 #if WITH_EDITOR
 				// Node template transform is performed on copy of local document to avoid overwriting editable data
 				constexpr bool bTransformDocumentBeforeRegistering = true;
@@ -69,7 +66,7 @@ namespace Metasound
 				const bool bTransformDocumentBeforeRegistering = MetaSoundEnableCookDeterministicIDGeneration == 0;
 #endif // WITH_EDITOR
 
-				const FMetasoundFrontendDocument& Document = ConstInterface->GetDocument();
+				const FMetasoundFrontendDocument& Document = DocumentInterface->GetConstDocument();
 				const TArray<FMetasoundFrontendClass>& Dependencies = Document.Dependencies;
 
 				if (bTransformDocumentBeforeRegistering)
@@ -181,44 +178,30 @@ namespace Metasound
 			}
 
 			// Registers node by copying document. Updates to document require re-registration.
-			class FNodeRegistryEntry : public INodeRegistryEntry
+			// This registry entry does not support node creation as it is only intended to be
+			// used when cooking MetaSounds. 
+			class FDocumentNodeRegistryEntryForCook : public INodeRegistryEntry
 			{
 			public:
-				FNodeRegistryEntry(const FString& InName, const FMetasoundFrontendDocument& InDocument, const FSoftObjectPath& InAssetPath)
-					: Name(InName)
-					, Document(InDocument)
+				FDocumentNodeRegistryEntryForCook(const FMetasoundFrontendDocument& InDocument, const FSoftObjectPath& InAssetPath)
+					: Interfaces(InDocument.Interfaces)
+					, FrontendClass(InDocument.RootGraph)
+					, ClassInfo(InDocument.RootGraph, InAssetPath)
 				{
 					// Copy FrontendClass to preserve original document.
-					FrontendClass = Document.RootGraph;
 					FrontendClass.Metadata.SetType(EMetasoundFrontendClassType::External);
-					ClassInfo = FNodeClassInfo(Document.RootGraph, InAssetPath);
 				}
 
-				FNodeRegistryEntry(const FString& InName, FMetasoundFrontendDocument&& InDocument, const FSoftObjectPath& InAssetPath)
-					: Name(InName)
-					, Document(MoveTemp(InDocument))
-				{
-					// Copy FrontendClass to preserve original document.
-					FrontendClass = Document.RootGraph;
-					FrontendClass.Metadata.SetType(EMetasoundFrontendClassType::External);
-					ClassInfo = FNodeClassInfo(Document.RootGraph, InAssetPath);
-				}
+				FDocumentNodeRegistryEntryForCook(const FDocumentNodeRegistryEntryForCook& InOther) = default;
 
-				virtual ~FNodeRegistryEntry() = default;
+				virtual ~FDocumentNodeRegistryEntryForCook() = default;
 
 				virtual const FNodeClassInfo& GetClassInfo() const override
 				{
 					return ClassInfo;
 				}
 
-				virtual TUniquePtr<INode> CreateNode(const FNodeInitData&) const override
-				{
-					FProxyDataCache ProxyDataCache;
-					ProxyDataCache.CreateAndCacheProxies(Document);
-
-					return FFrontendGraphBuilder().CreateGraph(Document, ProxyDataCache, Name);
-				}
-
+				virtual TUniquePtr<INode> CreateNode(const FNodeInitData&) const override { return nullptr; }
 				virtual TUniquePtr<INode> CreateNode(FDefaultLiteralNodeConstructorParams&&) const override { return nullptr; }
 				virtual TUniquePtr<INode> CreateNode(FDefaultNamedVertexNodeConstructorParams&&) const override { return nullptr; }
 				virtual TUniquePtr<INode> CreateNode(FDefaultNamedVertexWithLiteralNodeConstructorParams&&) const override { return nullptr; }
@@ -230,12 +213,12 @@ namespace Metasound
 
 				virtual TUniquePtr<INodeRegistryEntry> Clone() const override
 				{
-					return MakeUnique<FNodeRegistryEntry>(Name, Document, ClassInfo.AssetPath);
+					return MakeUnique<FDocumentNodeRegistryEntryForCook>(*this);
 				}
 
 				virtual const TSet<FMetasoundFrontendVersion>* GetImplementedInterfaces() const override
 				{
-					return &Document.Interfaces;
+					return &Interfaces;
 				}
 
 				virtual bool IsNative() const override
@@ -244,8 +227,7 @@ namespace Metasound
 				}
 
 			private:
-				FString Name;
-				FMetasoundFrontendDocument Document;
+				TSet<FMetasoundFrontendVersion> Interfaces;
 				FMetasoundFrontendClass FrontendClass;
 				FNodeClassInfo ClassInfo;
 			};
@@ -345,14 +327,11 @@ void FMetasoundAssetBase::RegisterGraphWithFrontend(Metasound::Frontend::FMetaSo
 
 	{
 		TScriptInterface<IMetaSoundDocumentInterface> RegistryDocInterface = AssetBasePrivate::BuildRegistryDocument(Owner);
-		const FMetasoundFrontendDocument& RegistryDocument = static_cast<const IMetaSoundDocumentInterface*>(RegistryDocInterface.GetInterface())->GetDocument();
-		CacheRuntimeData(RegistryDocument);
-		TUniquePtr<INodeRegistryEntry> RegistryEntry = MakeUnique<AssetBasePrivate::FNodeRegistryEntry>(AssetName, RegistryDocument, FSoftObjectPath(Owner));
 		UnregisterGraphWithFrontend();
-		RegistryKey = FMetasoundFrontendRegistryContainer::Get()->RegisterNode(MoveTemp(RegistryEntry));
+		RegistryKey = CacheRuntimeData(RegistryDocInterface);
 	}
 
-	if (NodeRegistryKey::IsValid(RegistryKey) && FMetasoundFrontendRegistryContainer::Get()->IsNodeRegistered(RegistryKey))
+	if (NodeRegistryKey::IsValid(RegistryKey))
 	{
 #if WITH_EDITORONLY_DATA
 		UpdateAssetRegistry();
@@ -398,7 +377,6 @@ void FMetasoundAssetBase::CookMetaSound()
 
 	UObject* Owner = GetOwningAsset();
 	check(Owner);
-	const FString AssetName = Owner->GetName();
 
 	{
 		// Performs document transforms on local copy, which reduces document footprint & renders transforming unnecessary unless altered at runtime when registering
@@ -409,13 +387,16 @@ void FMetasoundAssetBase::CookMetaSound()
 			DocBuilder.TransformTemplateNodes();
 		}
 
-		const IMetaSoundDocumentInterface& Interface = *TScriptInterface<IMetaSoundDocumentInterface>(Owner).GetInterface();
-		TUniquePtr<INodeRegistryEntry> RegistryEntry = MakeUnique<AssetBasePrivate::FNodeRegistryEntry>(AssetName, Interface.GetDocument(), FSoftObjectPath(Owner));
+		// During cook, we need to register the node so that it is available for other graphs, but we need to avoid
+		// creating proxies. To do so, we use a special node registration object which reflects the necessary information
+		// for the node registry, but does not create INodes.
 		UnregisterGraphWithFrontend();
+		const FMetasoundFrontendDocument& Document = TScriptInterface<IMetaSoundDocumentInterface>(Owner)->GetConstDocument();
+		TUniquePtr<INodeRegistryEntry> RegistryEntry = MakeUnique<AssetBasePrivate::FDocumentNodeRegistryEntryForCook>(Document, FSoftObjectPath(Owner));
 		RegistryKey = FMetasoundFrontendRegistryContainer::Get()->RegisterNode(MoveTemp(RegistryEntry));
 	}
 
-	if (NodeRegistryKey::IsValid(RegistryKey) && FMetasoundFrontendRegistryContainer::Get()->IsNodeRegistered(RegistryKey))
+	if (NodeRegistryKey::IsValid(RegistryKey))
 	{
 #if WITH_EDITORONLY_DATA
 		UpdateAssetRegistry();
@@ -426,7 +407,7 @@ void FMetasoundAssetBase::CookMetaSound()
 		const UClass* Class = Owner->GetClass();
 		check(Class);
 		const FString ClassName = Class->GetName();
-		UE_LOG(LogMetaSound, Error, TEXT("Registration failed during cook for MetaSound node class '%s' of UObject class '%s'"), *AssetName, *ClassName);
+		UE_LOG(LogMetaSound, Error, TEXT("Registration failed during cook for MetaSound node class '%s' of UObject class '%s'"), *GetOwningAssetName(), *ClassName);
 	}
 }
 
@@ -700,35 +681,13 @@ const FMetasoundFrontendDocumentModifyContext& FMetasoundAssetBase::GetModifyCon
 }
 #endif // WITH_EDITOR
 
-TSharedPtr<Metasound::FGraph, ESPMode::ThreadSafe> FMetasoundAssetBase::BuildMetasoundDocument(const FMetasoundFrontendDocument& InDocument, const Metasound::Frontend::FProxyDataCache& InProxies) const
-{
-	using namespace Metasound;
-	using namespace Metasound::Frontend;
-
-	METASOUND_TRACE_CPUPROFILER_EVENT_SCOPE(MetaSoundAssetBase::BuildMetasoundDocument);
-
-	// Create graph which can spawn instances. 
-	TUniquePtr<FFrontendGraph> FrontendGraph = FFrontendGraphBuilder::CreateGraph(InDocument, InProxies, GetOwningAssetName());
-	if (!FrontendGraph.IsValid())
-	{
-		UE_LOG(LogMetaSound, Error, TEXT("Failed to build MetaSound graph in asset '%s'"), *GetOwningAssetName());
-	}
 
-	TSharedPtr<Metasound::FGraph, ESPMode::ThreadSafe> SharedGraph(FrontendGraph.Release());
-
-	return SharedGraph;
-}
 
 bool FMetasoundAssetBase::IsRegistered() const
 {
 	using namespace Metasound::Frontend;
 
-	if (!NodeRegistryKey::IsValid(RegistryKey))
-	{
-		return false;
-	}
-
-	return FMetasoundFrontendRegistryContainer::Get()->IsNodeRegistered(RegistryKey);
+	return NodeRegistryKey::IsValid(RegistryKey);
 }
 
 bool FMetasoundAssetBase::IsReferencedAsset(const FMetasoundAssetBase& InAsset) const
@@ -810,7 +769,10 @@ TArray<FMetasoundAssetBase::FSendInfoAndVertexName> FMetasoundAssetBase::GetSend
 	check(IsInGameThread() || IsInAudioThread());
 
 	const FRuntimeData& RuntimeData = GetRuntimeData();
-	checkf(CurrentCachedRuntimeDataChangeID == RuntimeData.ChangeID, TEXT("Asset must have up-to-date cached RuntimeData prior to calling GetSendInfos"));
+	if(!RuntimeData.IsValid())
+	{
+		UE_LOG(LogMetaSound, Warning, TEXT("Send infos may be incorrect. Accessing invalid runtime data on MetaSound %s"), *GetOwningAssetName());
+	}
 
 	TArray<FSendInfoAndVertexName> SendInfos;
 
@@ -1035,32 +997,41 @@ TSharedPtr<FMetasoundFrontendDocument> FMetasoundAssetBase::PreprocessDocument()
 	return nullptr;
 }
 
-const FMetasoundAssetBase::FRuntimeData& FMetasoundAssetBase::CacheRuntimeData(const FMetasoundFrontendDocument& InDocument)
+Metasound::Frontend::FNodeRegistryKey FMetasoundAssetBase::CacheRuntimeData(const TScriptInterface<IMetaSoundDocumentInterface>& InDocumentInterface)
 {
 	using namespace Metasound::Frontend;
 
 	METASOUND_TRACE_CPUPROFILER_EVENT_SCOPE(MetaSoundAssetBase::CacheRuntimeData);
 
-	CurrentCachedRuntimeDataChangeID = FGuid::NewGuid();
-	CachedRuntimeData.ChangeID = CurrentCachedRuntimeDataChangeID;
+	const FMetasoundFrontendDocument& Document = InDocumentInterface->GetConstDocument();
+	Metasound::TSortedVertexNameMap<FMetasoundAssetBase::FRuntimeInput> PublicInputMap = AssetBasePrivate::GetPublicClassInputs(Document);
 
-	Metasound::TSortedVertexNameMap<FMetasoundAssetBase::FRuntimeInput> PublicInputMap = AssetBasePrivate::GetPublicClassInputs(InDocument);
+	// If the IMetaSoundDocumentInterface actively modified by a builder, then 
+	// we must build synchronously to avoid a race condition on reading/writing 
+	// the IMetaSoundDocumentInterface
+	const bool bAsync = !IsBuilderActive();
 
-	Metasound::Frontend::FProxyDataCache ProxyDataCache;
-	ProxyDataCache.CreateAndCacheProxies(InDocument);
-	TSharedPtr<Metasound::FGraph, ESPMode::ThreadSafe> Graph = BuildMetasoundDocument(InDocument, ProxyDataCache);
+	FNodeClassInfo NodeClassInfo { Document.RootGraph, FSoftObjectPath(GetOwningAsset()) };
+	FNodeRegistryKey NewRegistryKey = FMetasoundFrontendRegistryContainer::Get()->RegisterGraph(NodeClassInfo, InDocumentInterface, bAsync);
 
-	CachedRuntimeData.ChangeID = CurrentCachedRuntimeDataChangeID;
+	CachedRuntimeData.ChangeID = FGuid::NewGuid();
 	CachedRuntimeData.PublicInputMap = MoveTemp(PublicInputMap);
-	CachedRuntimeData.Graph = MoveTemp(Graph);
+	CachedRuntimeData.Graph = nullptr; // Graph will be retrieved later. Retrieving the graph forces async graph registration to complete. 
 
-	return CachedRuntimeData;
+	return NewRegistryKey;
 }
 
-const FMetasoundAssetBase::FRuntimeData& FMetasoundAssetBase::GetRuntimeData() const
+void FMetasoundAssetBase::WaitForAsyncGraphRegistration()
 {
-	ensureMsgf(CurrentCachedRuntimeDataChangeID == CachedRuntimeData.ChangeID, TEXT("Accessing out-of-date runtime data: MetaSound asset '%s'."), *GetOwningAssetName());
+	using namespace Metasound::Frontend;
+	if (NodeRegistryKey::IsValid(RegistryKey))
+	{
+		FMetasoundFrontendRegistryContainer::Get()->WaitForAsyncGraphRegistration(RegistryKey);
+	}
+}
 
+const FMetasoundAssetBase::FRuntimeData& FMetasoundAssetBase::GetRuntimeData() const
+{
 	return CachedRuntimeData;
 }
 
diff --git a/Engine/Plugins/Runtime/Metasound/Source/MetasoundFrontend/Private/MetasoundFrontendDocumentBuilder.cpp b/Engine/Plugins/Runtime/Metasound/Source/MetasoundFrontend/Private/MetasoundFrontendDocumentBuilder.cpp
--- a/Engine/Plugins/Runtime/Metasound/Source/MetasoundFrontend/Private/MetasoundFrontendDocumentBuilder.cpp
+++ b/Engine/Plugins/Runtime/Metasound/Source/MetasoundFrontend/Private/MetasoundFrontendDocumentBuilder.cpp
@@ -132,7 +132,7 @@ namespace Metasound::Frontend
 
 		const FString GetDebugName(const IMetaSoundDocumentInterface& DocumentInterface)
 		{
-			const FMetasoundFrontendClassMetadata& Metadata = DocumentInterface.GetDocument().RootGraph.Metadata;
+			const FMetasoundFrontendClassMetadata& Metadata = DocumentInterface.GetConstDocument().RootGraph.Metadata;
 			const FMetasoundFrontendVersionNumber& Version = Metadata.GetVersion();
 			const FNodeRegistryKey RegKey = NodeRegistryKey::CreateKey(EMetasoundFrontendClassType::External, Metadata.GetClassName().ToString(), Version.Major, Version.Minor);
 			if (const FMetasoundAssetBase* Asset = IMetaSoundAssetManager::GetChecked().TryLoadAssetFromKey(RegKey))
@@ -582,12 +582,12 @@ UMetaSoundBuilderDocument& UMetaSoundBuilderDocument::Create(const UClass& InBui
 UMetaSoundBuilderDocument& UMetaSoundBuilderDocument::Create(const IMetaSoundDocumentInterface& InDocToCopy)
 {
 	UMetaSoundBuilderDocument* DocObject = NewObject<UMetaSoundBuilderDocument>();
-	DocObject->Document = InDocToCopy.GetDocument();
+	DocObject->Document = InDocToCopy.GetConstDocument();
 	DocObject->MetaSoundUClass = InDocToCopy.GetBaseMetaSoundUClass();
 	return *DocObject;
 }
 
-const FMetasoundFrontendDocument& UMetaSoundBuilderDocument::GetDocument() const
+const FMetasoundFrontendDocument& UMetaSoundBuilderDocument::GetConstDocument() const
 {
 	return Document;
 }
@@ -603,6 +603,16 @@ FMetasoundFrontendDocument& UMetaSoundBuilderDocument::GetDocument()
 	return Document;
 }
 
+void UMetaSoundBuilderDocument::OnBeginActiveBuilder()
+{
+	// Nothing to do here. UMetaSoundBuilderDocuments are always being used by builders
+}
+
+void UMetaSoundBuilderDocument::OnFinishActiveBuilder()
+{
+	// Nothing to do here. UMetaSoundBuilderDocuments are always being used by builders
+}
+
 FMetaSoundFrontendDocumentBuilder::FMetaSoundFrontendDocumentBuilder()
 	: DocumentDelegates(MakeShared<Metasound::Frontend::FDocumentModifyDelegates>())
 {
@@ -614,6 +624,7 @@ FMetaSoundFrontendDocumentBuilder::FMetaSoundFrontendDocumentBuilder(TScriptInte
 {
 	if (DocumentInterface)
 	{
+		BeginBuilding();
 		InitCacheInternal();
 	}
 }
@@ -624,9 +635,14 @@ FMetaSoundFrontendDocumentBuilder::FMetaSoundFrontendDocumentBuilder(TScriptInte
 {
 	if (DocumentInterface)
 	{
+		BeginBuilding();
 		InitCacheInternal();
 	}
 }
+FMetaSoundFrontendDocumentBuilder::~FMetaSoundFrontendDocumentBuilder()
+{
+	FinishBuilding();
+}
 
 const FMetasoundFrontendClass* FMetaSoundFrontendDocumentBuilder::AddDependency(const FMetasoundFrontendClass& InClass)
 {
@@ -2105,6 +2121,23 @@ void FMetaSoundFrontendDocumentBuilder::InitCacheInternal(bool bPrimeCache)
 	DocumentCache = FDocumentCache::Create(GetDocument(), DocumentDelegates, bPrimeCache);
 }
 
+void FMetaSoundFrontendDocumentBuilder::BeginBuilding()
+{
+	if (DocumentInterface)
+	{
+		DocumentInterface->OnBeginActiveBuilder();
+	}
+}
+
+void FMetaSoundFrontendDocumentBuilder::FinishBuilding()
+{
+	if (DocumentInterface)
+	{
+		DocumentInterface->OnFinishActiveBuilder();
+	}
+	DocumentInterface = TScriptInterface<IMetaSoundDocumentInterface>();
+}
+
 bool FMetaSoundFrontendDocumentBuilder::RemoveDependency(const FGuid& InClassID)
 {
 	using namespace Metasound::Frontend;
diff --git a/Engine/Plugins/Runtime/Metasound/Source/MetasoundFrontend/Private/MetasoundFrontendDocumentCache.cpp b/Engine/Plugins/Runtime/Metasound/Source/MetasoundFrontend/Private/MetasoundFrontendDocumentCache.cpp
--- a/Engine/Plugins/Runtime/Metasound/Source/MetasoundFrontend/Private/MetasoundFrontendDocumentCache.cpp
+++ b/Engine/Plugins/Runtime/Metasound/Source/MetasoundFrontend/Private/MetasoundFrontendDocumentCache.cpp
@@ -78,10 +78,13 @@ namespace Metasound::Frontend
 
 	void FDocumentCache::Init(bool bPrimeCache)
 	{
+		METASOUND_TRACE_CPUPROFILER_EVENT_SCOPE(Metasound::Frontend::FDocumentCache::Init);
 		check(Document);
 
 		if (bPrimeCache)
 		{
+			METASOUND_TRACE_CPUPROFILER_EVENT_SCOPE(Metasound::Frontend::FDocumentCache::Init_Prime);
+
 			DependencyCache = MakeShared<FDocumentDependencyCache>(GetDocument());
 			EdgeCache = FDocumentGraphEdgeCache::Create(AsShared(), ModifyDelegates->EdgeDelegates);
 			NodeCache = FDocumentGraphNodeCache::Create(AsShared(), ModifyDelegates->NodeDelegates);
diff --git a/Engine/Plugins/Runtime/Metasound/Source/MetasoundFrontend/Private/MetasoundFrontendNodeRegistryPrivate.h b/Engine/Plugins/Runtime/Metasound/Source/MetasoundFrontend/Private/MetasoundFrontendNodeRegistryPrivate.h
--- a/Engine/Plugins/Runtime/Metasound/Source/MetasoundFrontend/Private/MetasoundFrontendNodeRegistryPrivate.h
+++ b/Engine/Plugins/Runtime/Metasound/Source/MetasoundFrontend/Private/MetasoundFrontendNodeRegistryPrivate.h
@@ -1,18 +1,32 @@
 // Copyright Epic Games, Inc. All Rights Reserved.
 #pragma once
 
-#include "MetasoundFrontendDataTypeRegistry.h"
-#include "MetasoundFrontendRegistries.h"
-
 #include "Containers/Array.h"
 #include "Containers/Map.h"
 #include "HAL/CriticalSection.h"
+#include "Tasks/Pipe.h"
+#include "UObject/GCObject.h"
+
+#include "MetasoundFrontendDataTypeRegistry.h"
+#include "MetasoundFrontendRegistries.h"
 #include "MetasoundFrontendRegistryTransaction.h"
 
+struct FMetasoundFrontendDocument; 
+
 namespace Metasound
 {
+	class FProxyDataCache;
+	class FGraph;
+
 	namespace Frontend
 	{
+		struct FNodeClassInfo;
+
+		namespace MetasoundFrontendRegistryPrivate
+		{
+			void BuildAndRegisterGraphFromDocument(const FMetasoundFrontendDocument& InDocument, const FProxyDataCache& InProxyDataCache, const FNodeClassInfo& InNodeClassInfo);
+		}
+
 		using FNodeRegistryTransactionBuffer = TTransactionBuffer<FNodeRegistryTransaction>;
 		using FNodeRegistryTransactionStream = TTransactionStream<FNodeRegistryTransaction>; 
 
@@ -37,7 +51,11 @@ namespace Metasound
 		};
 
 		// Registry container private implementation.
-		class FRegistryContainerImpl : public FMetasoundFrontendRegistryContainer
+		class FRegistryContainerImpl
+			: public FMetasoundFrontendRegistryContainer
+#if WITH_ENGINE
+			, public FGCObject
+#endif // #if WITH_ENGINE
 		{
 
 		public:
@@ -67,6 +85,24 @@ namespace Metasound
 			// This is called on module startup. This invokes any registration commands enqueued by our registration macros.
 			virtual void RegisterPendingNodes() override;
 
+			// Register a graph from an IMetaSoundDocumentInterface
+			virtual FNodeRegistryKey RegisterGraph(const FNodeClassInfo& InAssetPath, const TScriptInterface<IMetaSoundDocumentInterface>& InDocument, bool bAsync) override;
+			
+			// Wait for async graph registration to complete for a specific graph
+			virtual void WaitForAsyncGraphRegistration(const FNodeRegistryKey& InRegistryKey) const override;
+			
+			// Retrieve a registered graph. 
+			//
+			// If the graph is registered asynchronously, this will wait until the registration task has completed.
+			virtual TSharedPtr<const FGraph> GetGraph(const FNodeRegistryKey& InRegistryKey) const override;
+			
+#if WITH_ENGINE
+			// Begin FGCObject interface
+			virtual void AddReferencedObjects( FReferenceCollector& Collector ) override;
+			virtual FString GetReferencerName() const override;
+			// End FGCObject interface
+#endif // #if WITH_ENGINE
+
 			/** Register external node with the frontend.
 			 *
 			 * @param InCreateNode - Function for creating node from FNodeInitData.
@@ -76,7 +112,6 @@ namespace Metasound
 			 */
 			virtual FNodeRegistryKey RegisterNode(TUniquePtr<Metasound::Frontend::INodeRegistryEntry>&& InEntry) override;
 
-			virtual void ForEachNodeRegistryTransactionSince(Metasound::Frontend::FRegistryTransactionID InSince, Metasound::Frontend::FRegistryTransactionID* OutCurrentRegistryTransactionID, TFunctionRef<void(const Metasound::Frontend::FNodeRegistryTransaction&)> InFunc) const override;
 			virtual bool UnregisterNode(const FNodeRegistryKey& InKey) override;
 			virtual bool IsNodeRegistered(const FNodeRegistryKey& InKey) const override;
 			virtual bool IsNodeNative(const FNodeRegistryKey& InKey) const override;
@@ -89,12 +124,8 @@ namespace Metasound
 			virtual bool FindFrontendClassFromRegistered(const FNodeRegistryKey& InKey, FMetasoundFrontendClass& OutClass) override;
 			virtual const TSet<FMetasoundFrontendVersion>* FindImplementedInterfacesFromRegistered(const Metasound::Frontend::FNodeRegistryKey& InKey) const override;
 			virtual bool FindNodeClassInfoFromRegistered(const Metasound::Frontend::FNodeRegistryKey& InKey, FNodeClassInfo& OutInfo) override;
-			UE_DEPRECATED(5.1, "Use FindInputNodeRegistryKeyForDataType with EMetasoundFrontendVertexAccessType instead.")
-			virtual bool FindInputNodeRegistryKeyForDataType(const FName& InDataTypeName, FNodeRegistryKey& OutKey) override;
 			virtual bool FindInputNodeRegistryKeyForDataType(const FName& InDataTypeName, const EMetasoundFrontendVertexAccessType InAccessType, FNodeRegistryKey& OutKey) override;
 			virtual bool FindVariableNodeRegistryKeyForDataType(const FName& InDataTypeName, FNodeRegistryKey& OutKey) override;
-			UE_DEPRECATED(5.1, "Use FindOutputNodeRegistryKeyForDataType with EMetasoundFrontendVertexAccessType instead.")
-			virtual bool FindOutputNodeRegistryKeyForDataType(const FName& InDataTypeName, FNodeRegistryKey& OutKey) override;
 			virtual bool FindOutputNodeRegistryKeyForDataType(const FName& InDataTypeName, const EMetasoundFrontendVertexAccessType InAccessType, FNodeRegistryKey& OutKey) override;
 
 			// Create a new instance of a C++ implemented node from the registry.
@@ -115,12 +146,16 @@ namespace Metasound
 			TUniquePtr<FNodeRegistryTransactionStream> CreateTransactionStream();
 
 		private:
+			friend void MetasoundFrontendRegistryPrivate::BuildAndRegisterGraphFromDocument(const FMetasoundFrontendDocument& InDocument, const FProxyDataCache& InProxyDataCache, const FNodeClassInfo& InNodeClassInfo);
+
 			static FRegistryContainerImpl* LazySingleton;
 
+			void RegisterGraph(const FNodeRegistryKey& InKey, TSharedRef<const FGraph> InGraph);
 			const INodeRegistryEntry* FindNodeEntry(const FNodeRegistryKey& InKey) const;
 
 			const INodeRegistryTemplateEntry* FindNodeTemplateEntry(const FNodeRegistryKey& InKey) const;
 
+
 			// This buffer is used to enqueue nodes and datatypes to register when the module has been initialized,
 			// in order to avoid bad behavior with ensures, logs, etc. on static initialization.
 			// The bad news is that TInlineAllocator is the safest allocator to use on static init.
@@ -136,10 +171,23 @@ namespace Metasound
 			// Registry in which we keep all information about dynamically-generated templated nodes via in C++.
 			TMap<FNodeRegistryKey, TSharedRef<INodeRegistryTemplateEntry, ESPMode::ThreadSafe>> RegisteredNodeTemplates;
 
+			TMap<FNodeRegistryKey, TSharedRef<const FGraph>> RegisteredGraphs;
+
 			// Registry in which we keep lists of possible nodes to use to convert between two datatypes
 			TMap<FConverterNodeRegistryKey, FConverterNodeRegistryValue> ConverterNodeRegistry;
 
 			TSharedRef<FNodeRegistryTransactionBuffer> TransactionBuffer;
+
+			struct FActiveRegistrationTaskInfo
+			{
+				UE::Tasks::FTask Task;
+				TObjectPtr<UObject> OwningObject;
+			};
+
+			mutable FCriticalSection RegistryMapsCriticalSection;
+			mutable FCriticalSection ActiveRegistrationTasksCriticalSection;
+			UE::Tasks::FPipe AsyncRegistrationPipe;
+			TMap<FNodeRegistryKey, FActiveRegistrationTaskInfo> ActiveRegistrationTasks;
 		};
 	}
 }
diff --git a/Engine/Plugins/Runtime/Metasound/Source/MetasoundFrontend/Private/MetasoundFrontendRegistries.cpp b/Engine/Plugins/Runtime/Metasound/Source/MetasoundFrontend/Private/MetasoundFrontendRegistries.cpp
--- a/Engine/Plugins/Runtime/Metasound/Source/MetasoundFrontend/Private/MetasoundFrontendRegistries.cpp
+++ b/Engine/Plugins/Runtime/Metasound/Source/MetasoundFrontend/Private/MetasoundFrontendRegistries.cpp
@@ -3,17 +3,29 @@
 #include "MetasoundFrontendRegistries.h"
 
 #include "Algo/ForEach.h"
-#include "CoreMinimal.h"
+#include "Containers/UnrealString.h"
+#include "Containers/Array.h"
+#include "Containers/Set.h"
 #include "HAL/PlatformTime.h"
+#include "Misc/Guid.h"
+#include "Misc/ScopeLock.h"
+#include "Tasks/Pipe.h"
+#include "Tasks/Task.h"
+#include "Templates/SharedPointer.h"
+#include "Templates/UniquePtr.h"
+#include "UObject/NameTypes.h"
+#include "UObject/ScriptInterface.h"
+
+#include "MetasoundDocumentInterface.h"
 #include "MetasoundFrontendDataTypeRegistry.h"
+#include "MetasoundFrontendGraph.h"
 #include "MetasoundFrontendNodeRegistryPrivate.h"
+#include "MetasoundFrontendProxyDataCache.h"
 #include "MetasoundFrontendRegistryTransaction.h"
 #include "MetasoundFrontendSearchEngine.h"
 #include "MetasoundLog.h"
 #include "MetasoundRouter.h"
 #include "MetasoundTrace.h"
-#include "Misc/ScopeLock.h"
-
 
 namespace Metasound
 {
@@ -72,6 +84,185 @@ namespace Metasound
 						return InvalidType;
 				}
 			}
+
+			// FGraphNode is used to create unique INodes based off of a IGraph. 
+			//
+			// Individual nodes need to reflect their InstanceName and InstanceID, but otherwise
+			// they simply encapsulate a shared set of behavior. To minimize memory usage, a single
+			// shared IGraph is used for all nodes referring to the same IGraph.
+			class FGraphNode : public INode
+			{
+				// This adapter class forwards the correct FBuilderOperatorParams
+				// to the graph's operator creation method. Many operator creation
+				// methods downcast the supplied INode in `FBuilderOperatorParams`
+				// and so it is required that it point to the correct runtime instance
+				// when calling CreateOperator(...)
+				class FGraphOperatorFactoryAdapter : public IOperatorFactory 
+				{
+				public:
+					FGraphOperatorFactoryAdapter(const IGraph& InGraph)
+					: Graph(&InGraph)
+					, GraphFactory(InGraph.GetDefaultOperatorFactory())
+					{
+					}
+
+					virtual ~FGraphOperatorFactoryAdapter() = default;
+
+					virtual TUniquePtr<IOperator> CreateOperator(const FBuildOperatorParams& InParams, FBuildResults& OutResults) override
+					{
+						FBuildOperatorParams ForwardParams
+						{
+							*Graph,  // Point to correct INode instance
+							InParams.OperatorSettings,
+							InParams.InputData,
+							InParams.Environment,
+							InParams.Builder
+						};
+						return GraphFactory->CreateOperator(ForwardParams, OutResults);
+					}
+				private:
+					const IGraph* Graph; // Only store pointer because owning node keeps wrapped IGraph alive.  
+					FOperatorFactorySharedRef GraphFactory;
+				};
+
+			public:
+				FGraphNode(FNodeInitData InNodeInitData, TSharedRef<const IGraph> InGraphToWrap)
+				: NodeInitData(MoveTemp(InNodeInitData))
+				, Factory(MakeShared<FGraphOperatorFactoryAdapter>(*InGraphToWrap))
+				, Graph(MoveTemp(InGraphToWrap))
+				{
+				}
+
+				virtual const FName& GetInstanceName() const override
+				{
+					return NodeInitData.InstanceName;
+				}
+
+				virtual const FGuid& GetInstanceID() const override
+				{
+					return NodeInitData.InstanceID;
+				}
+
+				virtual const FNodeClassMetadata& GetMetadata() const override
+				{
+					return Graph->GetMetadata();
+				}
+
+				virtual const FVertexInterface& GetVertexInterface() const override
+				{
+					return Graph->GetVertexInterface();
+				}
+
+				virtual bool SetVertexInterface(const FVertexInterface& InInterface) override
+				{
+					// Cannot set vertex interface because of const reference to IGraph
+					// Return true if the supplied interface is the same as the existing interface.
+					return Graph->GetVertexInterface() == InInterface;
+				}
+
+				virtual bool IsVertexInterfaceSupported(const FVertexInterface& InInterface) const override
+				{
+					return Graph->IsVertexInterfaceSupported(InInterface);
+				}
+
+				virtual FOperatorFactorySharedRef GetDefaultOperatorFactory() const override
+				{
+					return Factory;
+				}
+
+			private:
+
+				FNodeInitData NodeInitData;
+				TSharedRef<FGraphOperatorFactoryAdapter> Factory;
+				TSharedRef<const IGraph> Graph;
+			};
+
+
+			// FDocumentNodeRegistryEntry encapsulates a node registry entry for a FGraph
+			class FDocumentNodeRegistryEntry : public INodeRegistryEntry
+			{
+			public:
+				FDocumentNodeRegistryEntry(const FMetasoundFrontendGraphClass& InGraphClass, const TSet<FMetasoundFrontendVersion>& InInterfaces, const FNodeClassInfo& InNodeClassInfo, TSharedRef<const IGraph> InGraph)
+				: FrontendClass(InGraphClass)
+				, Interfaces(InInterfaces)
+				, ClassInfo(InNodeClassInfo)
+				, Graph(InGraph)
+				{
+					FrontendClass.Metadata.SetType(EMetasoundFrontendClassType::External);
+				}
+
+				FDocumentNodeRegistryEntry(const FDocumentNodeRegistryEntry&) = default;
+
+				virtual ~FDocumentNodeRegistryEntry() 
+				{
+				};
+
+				virtual const FNodeClassInfo& GetClassInfo() const override
+				{
+					return ClassInfo;
+				}
+
+				virtual TUniquePtr<INode> CreateNode(const FNodeInitData& InNodeInitData) const override
+				{
+					return MakeUnique<FGraphNode>(InNodeInitData, Graph);
+				}
+
+				virtual TUniquePtr<INode> CreateNode(FDefaultLiteralNodeConstructorParams&&) const override { return nullptr; }
+				virtual TUniquePtr<INode> CreateNode(FDefaultNamedVertexNodeConstructorParams&&) const override { return nullptr; }
+				virtual TUniquePtr<INode> CreateNode(FDefaultNamedVertexWithLiteralNodeConstructorParams&&) const override { return nullptr; }
+
+				virtual const FMetasoundFrontendClass& GetFrontendClass() const override
+				{
+					return FrontendClass;
+				}
+
+				virtual TUniquePtr<INodeRegistryEntry> Clone() const override
+				{
+					return MakeUnique<FDocumentNodeRegistryEntry>(*this);
+				}
+
+				virtual const TSet<FMetasoundFrontendVersion>* GetImplementedInterfaces() const override
+				{
+					return &Interfaces;
+				}
+
+				virtual bool IsNative() const override
+				{
+					return false;
+				}
+
+			private:
+
+				FMetasoundFrontendClass FrontendClass;
+				TSet<FMetasoundFrontendVersion> Interfaces;
+				FNodeClassInfo ClassInfo;
+				TSharedRef<const IGraph> Graph;
+			};
+
+			// Builds an IGraph from a FMetasoundFrontendDocument and registers it in the node registry.
+			void BuildAndRegisterGraphFromDocument(const FMetasoundFrontendDocument& InDocument, const FProxyDataCache& InProxyDataCache, const FNodeClassInfo& InNodeClassInfo)
+			{
+				METASOUND_TRACE_CPUPROFILER_EVENT_SCOPE(Metasound::FRegistryContainerImpl::BuildAndRegisterGraphFromDocument);
+
+				TUniquePtr<FFrontendGraph> FrontendGraph = FFrontendGraphBuilder::CreateGraph(InDocument, InProxyDataCache, InNodeClassInfo.AssetPath.ToString());
+				if (!FrontendGraph.IsValid())
+				{
+					UE_LOG(LogMetaSound, Error, TEXT("Failed to build MetaSound graph in asset '%s'"), *InNodeClassInfo.AssetPath.ToString());
+					return;
+				}
+
+				TSharedRef<const FGraph> GraphToRegister = MakeShareable<const FGraph>(FrontendGraph.Release());
+
+				TUniquePtr<INodeRegistryEntry> RegistryEntry = MakeUnique<FDocumentNodeRegistryEntry>(
+						InDocument.RootGraph, 
+						InDocument.Interfaces, 
+						InNodeClassInfo, 
+						GraphToRegister);
+				FRegistryContainerImpl& Registry = FRegistryContainerImpl::Get();
+
+				FNodeRegistryKey RegistryKey = Registry.RegisterNode(MoveTemp(RegistryEntry));
+				Registry.RegisterGraph(RegistryKey, MoveTemp(GraphToRegister));
+			}
 		}
 
 		FRegistryContainerImpl* FRegistryContainerImpl::LazySingleton = nullptr;
@@ -97,6 +288,7 @@ namespace Metasound
 
 		FRegistryContainerImpl::FRegistryContainerImpl()
 		: TransactionBuffer(MakeShared<FNodeRegistryTransactionBuffer>())
+		, AsyncRegistrationPipe( UE_SOURCE_LOCATION )
 		{
 		}
 
@@ -132,11 +324,6 @@ namespace Metasound
 			return true;
 		}
 
-		void FRegistryContainerImpl::ForEachNodeRegistryTransactionSince(FRegistryTransactionID InSince, FRegistryTransactionID* OutCurrentRegistryTransactionID, TFunctionRef<void(const FNodeRegistryTransaction&)> InFunc) const 
-		{
-			UE_LOG(LogMetaSound, Error, TEXT("FMetasoundFrontendRegistryContainer::ForEachNodeRegistryTransactionSince is no longer supported and will not function as expected."));
-		}
-
 		TUniquePtr<Metasound::INode> FRegistryContainerImpl::CreateNode(const FNodeRegistryKey& InKey, const Metasound::FNodeInitData& InInitData) const
 		{
 			if (const INodeRegistryEntry* Entry = FindNodeEntry(InKey))
@@ -205,6 +392,117 @@ namespace Metasound
 			return MakeUnique<FNodeRegistryTransactionStream>(TransactionBuffer);
 		}
 
+		FNodeRegistryKey FRegistryContainerImpl::RegisterGraph(const FNodeClassInfo& InNodeClassInfo, const TScriptInterface<IMetaSoundDocumentInterface>& InDocumentInterface, bool bAsync)
+		{
+			using namespace UE;
+
+			check(InDocumentInterface);
+
+			const FMetasoundFrontendDocument& Document = InDocumentInterface->GetConstDocument();
+			FNodeRegistryKey RegistryKey = NodeRegistryKey::CreateKey(Document.RootGraph);
+
+			// Proxies are created synchronously to avoid creating proxies in async tasks. Proxies
+			// are created from UObjects which need to be protected from GC and non-GT access.
+			FProxyDataCache ProxyDataCache;
+			ProxyDataCache.CreateAndCacheProxies(Document);
+
+
+			if (bAsync)
+			{	
+				// Wait for any async tasks that are inflight which correspond to the same graph
+				FScopeLock LockActiveReg(&ActiveRegistrationTasksCriticalSection);
+				if (const FActiveRegistrationTaskInfo* ActiveTaskInfo = ActiveRegistrationTasks.Find(RegistryKey))
+				{
+					UE_LOG(LogMetaSound, Warning, TEXT("Waiting for async registration task to finish before beginning new registration task for same registration key (%s) with asset (%s))"), *RegistryKey, *InNodeClassInfo.AssetPath.ToString());
+					ActiveTaskInfo->Task.Wait(FTimespan::FromSeconds(1.));
+				}
+			
+				Tasks::FTask BuildAndRegisterTask = AsyncRegistrationPipe.Launch(
+					UE_SOURCE_LOCATION,
+					[RegistryKey, NodeClassInfo=InNodeClassInfo, DocumentInterface=InDocumentInterface, ProxyDataCache=MoveTemp(ProxyDataCache)]()
+					{
+						// Build and add the graph to the node registry
+						MetasoundFrontendRegistryPrivate::BuildAndRegisterGraphFromDocument(DocumentInterface->GetConstDocument(), ProxyDataCache, NodeClassInfo);
+
+						// cleanup async task tracking
+						{
+							FRegistryContainerImpl& Registry = FRegistryContainerImpl::Get();
+
+							FScopeLock LockActiveReg(&Registry.ActiveRegistrationTasksCriticalSection);
+							int32 NumRemoved = Registry.ActiveRegistrationTasks.Remove(RegistryKey);
+							if (NumRemoved != 1)
+							{
+								UE_LOG(LogMetaSound, Warning, TEXT("Multiple active registration tasks for the same asset may result in incorrect MetaSound graphs being instantiated during runtime. Registry Key: (%s) Asset Path: (%s)"), *RegistryKey, *NodeClassInfo.AssetPath.ToString());
+							}
+						}
+					}
+				);
+
+				ActiveRegistrationTasks.Add(RegistryKey, FActiveRegistrationTaskInfo{BuildAndRegisterTask, InDocumentInterface.GetObject()});
+			}
+			else
+			{
+				// Build and register graph synchronously
+				MetasoundFrontendRegistryPrivate::BuildAndRegisterGraphFromDocument(InDocumentInterface->GetConstDocument(), ProxyDataCache, InNodeClassInfo);
+			}
+
+			return RegistryKey;
+		}
+
+#if WITH_ENGINE
+		void FRegistryContainerImpl::AddReferencedObjects(FReferenceCollector& Collector)
+		{
+			// Inflight IMetaSoundDOucmentInterface objects are protected from GC while their
+			// async tasks or being processed. 
+			FScopeLock LockActiveReg(&ActiveRegistrationTasksCriticalSection);
+			for (TPair<FNodeRegistryKey, FActiveRegistrationTaskInfo>& Pair : ActiveRegistrationTasks)
+			{
+				if (Pair.Value.OwningObject)
+				{
+					Collector.AddReferencedObject(Pair.Value.OwningObject);
+				}
+			}
+		}
+
+		FString FRegistryContainerImpl::GetReferencerName() const
+		{
+			return TEXT("Metasound::Frontend::FRegistryContainerImpl");
+		}
+#endif // #if WITH_ENGINE
+
+		void FRegistryContainerImpl::RegisterGraph(const FNodeRegistryKey& InKey, TSharedRef<const FGraph> InGraph)
+		{
+			FScopeLock Lock(&RegistryMapsCriticalSection);
+
+			if (const TSharedRef<const FGraph>* ExistingGraph = RegisteredGraphs.Find(InKey))
+			{
+				UE_LOG(LogMetaSound, Error, TEXT("Multiple graphs are registered with the same registry key (%s). The existing registered graph (%s) will be replaced with the new graph (%s)."),  *InKey, *((*ExistingGraph)->GetInstanceName().ToString()), *(InGraph->GetInstanceName().ToString()));
+			}
+
+			RegisteredGraphs.Add(InKey, MoveTemp(InGraph));
+		}
+
+		TSharedPtr<const Metasound::FGraph> FRegistryContainerImpl::GetGraph(const Metasound::Frontend::FNodeRegistryKey& InRegistryKey) const
+		{
+			WaitForAsyncGraphRegistration(InRegistryKey);
+
+			TSharedPtr<const FGraph> Graph;
+			{
+				FScopeLock Lock(&RegistryMapsCriticalSection);
+				if (const TSharedRef<const FGraph>* RegisteredGraph = RegisteredGraphs.Find(InRegistryKey))
+				{
+					Graph = *RegisteredGraph;
+				}
+			}
+
+			if (!Graph)
+			{
+				UE_LOG(LogMetaSound, Error, TEXT("Could not find graph with registry key (%s)."),  *InRegistryKey);
+			}
+
+			return Graph;
+		}
+
 		FNodeRegistryKey FRegistryContainerImpl::RegisterNode(TUniquePtr<INodeRegistryEntry>&& InEntry)
 		{
 			METASOUND_LLM_SCOPE;
@@ -219,23 +517,27 @@ namespace Metasound
 
 				Key = NodeRegistryKey::CreateKey(Entry->GetClassInfo());
 
-				// check to see if an identical node was already registered, and log
-				if (RegisteredNodes.Contains(Key))
 				{
-					TSharedRef<INodeRegistryEntry, ESPMode::ThreadSafe> Node = RegisteredNodes[Key];
-					const FNodeClassInfo& ClassInfo = Node->GetClassInfo();
-					UE_LOG(LogMetaSound, Error, TEXT("Node with registry key '%s' already registered by asset '%s'. The previously registered node will be overwritten. " 
-						"This can happen if two classes share the same name or if METASOUND_REGISTER_NODE is defined in a public header."
-					), *Key, *ClassInfo.AssetPath.ToString());
+					FScopeLock Lock(&RegistryMapsCriticalSection);
+
+					// check to see if an identical node was already registered, and log
+					if (RegisteredNodes.Contains(Key))
+					{
+						TSharedRef<INodeRegistryEntry, ESPMode::ThreadSafe> Node = RegisteredNodes[Key];
+						const FNodeClassInfo& ClassInfo = Node->GetClassInfo();
+						UE_LOG(LogMetaSound, Error, TEXT("Node with registry key '%s' already registered by asset '%s'. The previously registered node will be overwritten. " 
+							"This can happen if two classes share the same name or if METASOUND_REGISTER_NODE is defined in a public header."
+						), *Key, *ClassInfo.AssetPath.ToString());
+					}
+
+					// Store registry elements in map so nodes can be queried using registry key.
+					RegisteredNodes.Add(Key, Entry);
 				}
 
 				// Store update to newly registered node in history so nodes
 				// can be queried by transaction ID
 				
 				TransactionBuffer->AddTransaction(FNodeRegistryTransaction(FNodeRegistryTransaction::ETransactionType::NodeRegistration, Entry->GetClassInfo(), Timestamp));
-
-				// Store registry elements in map so nodes can be queried using registry key.
-				RegisteredNodes.Add(Key, Entry);
 			}
 
 			return Key;
@@ -255,19 +557,23 @@ namespace Metasound
 
 				Key = NodeRegistryKey::CreateKey(Entry->GetClassInfo());
 
-				// check to see if an identical node was already registered, and log
-				ensureAlwaysMsgf(
-					!RegisteredNodes.Contains(Key),
-					TEXT("Node template with registry key '%s' already registered. The previously registered node will be overwritten."),
-					*Key);
+				{
+					FScopeLock Lock(&RegistryMapsCriticalSection);
+					// check to see if an identical node was already registered, and log
+					ensureAlwaysMsgf(
+						!RegisteredNodeTemplates.Contains(Key),
+						TEXT("Node template with registry key '%s' already registered. The previously registered node will be overwritten."),
+						*Key);
+
+
+					// Store registry elements in map so nodes can be queried using registry key.
+					RegisteredNodeTemplates.Add(Key, Entry);
+				}
 
 				// Store update to newly registered node in history so nodes
 				// can be queried by transaction ID
 
 				TransactionBuffer->AddTransaction(FNodeRegistryTransaction(FNodeRegistryTransaction::ETransactionType::NodeRegistration, Entry->GetClassInfo(), Timestamp));
-
-				// Store registry elements in map so nodes can be queried using registry key.
-				RegisteredNodeTemplates.Add(Key, Entry);
 			}
 
 			return Key;
@@ -285,7 +591,11 @@ namespace Metasound
 
 					TransactionBuffer->AddTransaction(FNodeRegistryTransaction(FNodeRegistryTransaction::ETransactionType::NodeUnregistration, Entry->GetClassInfo(), Timestamp));
 
-					RegisteredNodes.Remove(InKey);
+					{
+						FScopeLock Lock(&RegistryMapsCriticalSection);
+						RegisteredNodes.Remove(InKey);
+						RegisteredGraphs.Remove(InKey);
+					}
 					return true;
 				}
 			}
@@ -305,7 +615,10 @@ namespace Metasound
 
 					TransactionBuffer->AddTransaction(FNodeRegistryTransaction(FNodeRegistryTransaction::ETransactionType::NodeUnregistration, Entry->GetClassInfo(), Timestamp));
 
-					RegisteredNodeTemplates.Remove(InKey);
+					{
+						FScopeLock Lock(&RegistryMapsCriticalSection);
+						RegisteredNodeTemplates.Remove(InKey);
+					}
 					return true;
 				}
 			}
@@ -336,7 +649,12 @@ namespace Metasound
 
 		bool FRegistryContainerImpl::IsNodeRegistered(const FNodeRegistryKey& InKey) const
 		{
-			return RegisteredNodes.Contains(InKey) || RegisteredNodeTemplates.Contains(InKey);
+			WaitForAsyncGraphRegistration(InKey);
+
+			{
+				FScopeLock Lock(&RegistryMapsCriticalSection);
+				return RegisteredNodes.Contains(InKey) || RegisteredNodeTemplates.Contains(InKey);
+			}
 		}
 
 		bool FRegistryContainerImpl::IsNodeNative(const FNodeRegistryKey& InKey) const
@@ -398,12 +716,6 @@ namespace Metasound
 			return false;
 		}
 
-		bool FRegistryContainerImpl::FindInputNodeRegistryKeyForDataType(const FName& InDataTypeName, FNodeRegistryKey& OutKey)
-		{
-			UE_LOG(LogMetaSound, Error, TEXT("FRegistryContainerImpl::FindInputNodeRegistryKeyForDataType with these parameters is no longer supported and should not be called. Use the one with EMetasoundFrontendVertexAccessType instead."));
-			return false;
-		}
-
 		bool FRegistryContainerImpl::FindInputNodeRegistryKeyForDataType(const FName& InDataTypeName, const EMetasoundFrontendVertexAccessType InAccessType, FNodeRegistryKey& OutKey)
 		{
 			FMetasoundFrontendClass Class;
@@ -451,12 +763,6 @@ namespace Metasound
 			return false;
 		}
 
-		bool FRegistryContainerImpl::FindOutputNodeRegistryKeyForDataType(const FName& InDataTypeName, FNodeRegistryKey& OutKey)
-		{
-			UE_LOG(LogMetaSound, Error, TEXT("FRegistryContainerImpl::FindOutputNodeRegistryKeyForDataType with these parameters is no longer supported and should not be called. Use the one with EMetasoundFrontendVertexAccessType instead."));
-			return false;
-		}
-
 		bool FRegistryContainerImpl::FindOutputNodeRegistryKeyForDataType(const FName& InDataTypeName, const EMetasoundFrontendVertexAccessType InAccessType, FNodeRegistryKey& OutKey)
 		{
 			FMetasoundFrontendClass Class;
@@ -488,6 +794,7 @@ namespace Metasound
 
 		void FRegistryContainerImpl::IterateRegistry(Metasound::FIterateMetasoundFrontendClassFunction InIterFunc, EMetasoundFrontendClassType InClassType) const
 		{
+			UE_LOG(LogMetaSound, Warning, TEXT("Calling FMetasoundRegistryContainer::IterateRegistry(...) is not threadsafe. Please use Metasound::Frontend::ISearchEngine instead"));
 			auto WrappedFunc = [&](const TPair<FNodeRegistryKey, TSharedPtr<INodeRegistryEntry, ESPMode::ThreadSafe>>& Pair)
 			{
 				InIterFunc(Pair.Value->GetFrontendClass());
@@ -511,16 +818,27 @@ namespace Metasound
 
 		const INodeRegistryEntry* FRegistryContainerImpl::FindNodeEntry(const FNodeRegistryKey& InKey) const
 		{
-			if (const TSharedRef<INodeRegistryEntry, ESPMode::ThreadSafe>* Entry = RegisteredNodes.Find(InKey))
+			WaitForAsyncGraphRegistration(InKey);
+
 			{
-				return &Entry->Get();
+				// This scope lock proctects against race conditions manipulating the `RegisteredNodes` map, but it does not
+				// protect against the individual INodeRegistryEntry from being removed after this function returns. Generally
+				// this is not an issue as the acess to the INodeRegistryEntry pointer happens on the same thread as calls to 
+				// remove node registry entries. 
+				FScopeLock Lock(&RegistryMapsCriticalSection);
+				if (const TSharedRef<INodeRegistryEntry, ESPMode::ThreadSafe>* Entry = RegisteredNodes.Find(InKey))
+				{
+					return &Entry->Get();
+				}
 			}
 
 			return nullptr;
 		}
 
+
 		const INodeRegistryTemplateEntry* FRegistryContainerImpl::FindNodeTemplateEntry(const FNodeRegistryKey& InKey) const
 		{
+			FScopeLock Lock(&RegistryMapsCriticalSection);
 			if (const TSharedRef<INodeRegistryTemplateEntry, ESPMode::ThreadSafe>* Entry = RegisteredNodeTemplates.Find(InKey))
 			{
 				return &Entry->Get();
@@ -529,6 +847,24 @@ namespace Metasound
 			return nullptr;
 		}
 
+		void FRegistryContainerImpl::WaitForAsyncGraphRegistration(const FNodeRegistryKey& InRegistryKey) const
+		{
+			UE::Tasks::FTask ActiveRegistrationTask;
+			{
+				FScopeLock Lock(&ActiveRegistrationTasksCriticalSection);
+				if (const FActiveRegistrationTaskInfo* FoundTask = ActiveRegistrationTasks.Find(InRegistryKey))
+				{
+					ActiveRegistrationTask = FoundTask->Task;
+				}
+			}
+
+			if (ActiveRegistrationTask.IsValid())
+			{
+				METASOUND_TRACE_CPUPROFILER_EVENT_SCOPE(Metasound::FRegistryContainerImpl::WaitForRegistrationTaskToComplete);
+				ActiveRegistrationTask.Wait();
+			}
+		}
+
 		FNodeRegistryTransaction::FNodeRegistryTransaction(ETransactionType InType, const FNodeClassInfo& InNodeClassInfo, FNodeRegistryTransaction::FTimeType InTimestamp)
 		: Type(InType)
 		, NodeClassInfo(InNodeClassInfo)
@@ -536,11 +872,6 @@ namespace Metasound
 		{
 		}
 
-		FNodeRegistryTransaction::FNodeRegistryTransaction(ETransactionType InType, const FNodeRegistryKey& InKey, const FNodeClassInfo& InNodeClassInfo, FNodeRegistryTransaction::FTimeType InTimestamp)
-		: FNodeRegistryTransaction(InType, InNodeClassInfo, InTimestamp)
-		{
-		}
-
 		FNodeRegistryTransaction::ETransactionType FNodeRegistryTransaction::GetTransactionType() const
 		{
 			return Type;
@@ -567,6 +898,11 @@ namespace Metasound
 			FNodeRegistryKey CreateKey(EMetasoundFrontendClassType InType, const FString& InFullClassName, int32 InMajorVersion, int32 InMinorVersion)
 			{
 				using namespace MetasoundFrontendRegistryPrivate;
+				if (InType == EMetasoundFrontendClassType::Graph)
+				{
+					// No graphs are registered. Any registered graph should be registered as an external node.
+					InType = EMetasoundFrontendClassType::External;
+				}
 				const FString RegistryKey = FString::Format(TEXT("{0}_{1}_{2}.{3}"), {*GetClassTypeString(InType), *InFullClassName, InMajorVersion, InMinorVersion});
 				return RegistryKey;
 			}
@@ -724,12 +1060,6 @@ bool FMetasoundFrontendRegistryContainer::GetNodeClassInfoFromRegistered(const F
 	return false;
 }
 
-bool FMetasoundFrontendRegistryContainer::GetInputNodeRegistryKeyForDataType(const FName& InDataTypeName, FNodeRegistryKey& OutKey)
-{
-	UE_LOG(LogMetaSound, Error, TEXT("FMetasoundFrontendRegistryContainer::GetInputNodeRegistryKeyForDataType with these parameters is no longer supported and should not be called. Use the one with EMetasoundFrontendVertexAccessType instead."));
-	return false;
-}
-
 bool FMetasoundFrontendRegistryContainer::GetInputNodeRegistryKeyForDataType(const FName& InDataTypeName, const EMetasoundFrontendVertexAccessType InAccessType, FNodeRegistryKey& OutKey)
 {
 	if (FMetasoundFrontendRegistryContainer* Registry = FMetasoundFrontendRegistryContainer::Get())
@@ -748,12 +1078,6 @@ bool FMetasoundFrontendRegistryContainer::GetVariableNodeRegistryKeyForDataType(
 	return false;
 }
 
-bool FMetasoundFrontendRegistryContainer::GetOutputNodeRegistryKeyForDataType(const FName& InDataTypeName, FNodeRegistryKey& OutKey)
-{
-	UE_LOG(LogMetaSound, Error, TEXT("FMetasoundFrontendRegistryContainer::GetOutputNodeRegistryKeyForDataType with these parameters is no longer supported and should not be called. Use the one with EMetasoundFrontendVertexAccessType instead."));
-	return false;
-}
-
 bool FMetasoundFrontendRegistryContainer::GetOutputNodeRegistryKeyForDataType(const FName& InDataTypeName, const EMetasoundFrontendVertexAccessType InVertexAccessType, FNodeRegistryKey& OutKey)
 {
 	if (FMetasoundFrontendRegistryContainer* Registry = FMetasoundFrontendRegistryContainer::Get())
diff --git a/Engine/Plugins/Runtime/Metasound/Source/MetasoundFrontend/Public/MetasoundAssetBase.h b/Engine/Plugins/Runtime/Metasound/Source/MetasoundFrontend/Public/MetasoundAssetBase.h
--- a/Engine/Plugins/Runtime/Metasound/Source/MetasoundFrontend/Public/MetasoundAssetBase.h
+++ b/Engine/Plugins/Runtime/Metasound/Source/MetasoundFrontend/Public/MetasoundAssetBase.h
@@ -220,6 +220,13 @@ protected:
 	// Returns an access pointer to the document.
 	virtual Metasound::Frontend::FConstDocumentAccessPtr GetDocumentConstAccessPtr() const = 0;
 
+	// Waits for a graph to be registered in the scenario where a graph is registered on an async task.
+	//
+	// When a graph is registered, the underlying IMetaSoundDocumentInterface may be accessed on an
+	// async tasks. If modifications need to be made to the IMetaSoundDocumentInterface, callers should
+	// wait for the inflight graph registration to complete by calling this method.
+	void WaitForAsyncGraphRegistration();
+
 protected:
 	
 	// Container for runtime data of MetaSound graph.
@@ -239,6 +246,11 @@ protected:
 
 		// Core graph.
 		TSharedPtr<Metasound::FGraph, ESPMode::ThreadSafe> Graph;
+
+		bool IsValid() const 
+		{
+			return ChangeID.IsValid();
+		}
 	};
 
 	// Returns the cached runtime data.
@@ -262,13 +274,18 @@ private:
 #if WITH_EDITORONLY_DATA
 	void UpdateAssetRegistry();
 #endif
-	// Returns the cached runtime data. Call updates cached data if out-of-date.
-	const FRuntimeData& CacheRuntimeData(const FMetasoundFrontendDocument& InPreprocessedDoc);
+
+	// Returns true if the IMetaSoundDocumentInterface is currently has an active builder.
+	//
+	// If true, calls to register the graph will be performed synchronously in order to avoid
+	// race conditions with an active builder.
+	virtual bool IsBuilderActive() const = 0;
+
+	// Returns the cached registry key.
+	Metasound::Frontend::FNodeRegistryKey CacheRuntimeData(const TScriptInterface<IMetaSoundDocumentInterface>& InDoc);
 
 	Metasound::Frontend::FNodeRegistryKey RegistryKey;
 
-	// Cache ID is used to determine whether CachedRuntimeData is out-of-date.
-	FGuid CurrentCachedRuntimeDataChangeID;
 	FRuntimeData CachedRuntimeData;
 
 	TSharedPtr<Metasound::FGraph, ESPMode::ThreadSafe> BuildMetasoundDocument(const FMetasoundFrontendDocument& InPreprocessDoc, const Metasound::Frontend::FProxyDataCache& InProxies) const;
diff --git a/Engine/Plugins/Runtime/Metasound/Source/MetasoundFrontend/Public/MetasoundDocumentInterface.h b/Engine/Plugins/Runtime/Metasound/Source/MetasoundFrontend/Public/MetasoundDocumentInterface.h
--- a/Engine/Plugins/Runtime/Metasound/Source/MetasoundFrontend/Public/MetasoundDocumentInterface.h
+++ b/Engine/Plugins/Runtime/Metasound/Source/MetasoundFrontend/Public/MetasoundDocumentInterface.h
@@ -29,14 +29,29 @@ class METASOUNDFRONTEND_API IMetaSoundDocumentInterface : public IInterface
 public:
 	// Returns read-only reference to the the MetaSoundFrontendDocument
 	// containing all MetaSound runtime & editor data.
-	virtual const FMetasoundFrontendDocument& GetDocument() const = 0;
+	UE_DEPRECATED(5.4, "Use GetConstDocument instead")
+	virtual const FMetasoundFrontendDocument& GetDocument() const 
+	{
+		return GetConstDocument();
+	}
+
+	// Returns read-only reference to the the MetaSoundFrontendDocument
+	// containing all MetaSound runtime & editor data.
+	virtual const FMetasoundFrontendDocument& GetConstDocument() const = 0;
 
 	// Returns the parent class registered with the MetaSound UObject registry.
 	virtual const UClass& GetBaseMetaSoundUClass() const = 0;
 
+
 private:
 	virtual FMetasoundFrontendDocument& GetDocument() = 0;
 
+	// Derived classes can implement these methods to react to a builder beginning
+	// or finishing. Begin and Finish are tied to the lifetime of the active 
+	// FMetaSoundFrontendDocumentBuilder.
+	virtual void OnBeginActiveBuilder() = 0;
+	virtual void OnFinishActiveBuilder() = 0;
+
 	friend struct FMetaSoundFrontendDocumentBuilder;
 };
 
diff --git a/Engine/Plugins/Runtime/Metasound/Source/MetasoundFrontend/Public/MetasoundFrontendDocumentBuilder.h b/Engine/Plugins/Runtime/Metasound/Source/MetasoundFrontend/Public/MetasoundFrontendDocumentBuilder.h
--- a/Engine/Plugins/Runtime/Metasound/Source/MetasoundFrontend/Public/MetasoundFrontendDocumentBuilder.h
+++ b/Engine/Plugins/Runtime/Metasound/Source/MetasoundFrontend/Public/MetasoundFrontendDocumentBuilder.h
@@ -102,7 +102,7 @@ public:
 	static UMetaSoundBuilderDocument& Create(const IMetaSoundDocumentInterface& InDocToCopy);
 
 	// Returns the document
-	virtual const FMetasoundFrontendDocument& GetDocument() const override;
+	virtual const FMetasoundFrontendDocument& GetConstDocument() const override;
 
 	// Base MetaSoundClass that document is published to.
 	virtual const UClass& GetBaseMetaSoundUClass() const final override;
@@ -110,13 +110,14 @@ public:
 private:
 	virtual FMetasoundFrontendDocument& GetDocument() override;
 
+	virtual void OnBeginActiveBuilder() override;
+	virtual void OnFinishActiveBuilder() override;
+
 	UPROPERTY(Transient)
 	FMetasoundFrontendDocument Document;
 
 	UPROPERTY(Transient)
 	TObjectPtr<const UClass> MetaSoundUClass = nullptr;
-
-	friend class FMetasoundAssetBase;
 };
 
 // Builder used to support dynamically generating MetaSound documents at runtime. Builder contains caches that speed up
@@ -134,6 +135,11 @@ public:
 	FMetaSoundFrontendDocumentBuilder();
 	FMetaSoundFrontendDocumentBuilder(TScriptInterface<IMetaSoundDocumentInterface> InDocumentInterface);
 	FMetaSoundFrontendDocumentBuilder(TScriptInterface<IMetaSoundDocumentInterface> InDocumentInterface, TSharedRef<Metasound::Frontend::FDocumentModifyDelegates> InDocumentDelegates);
+	virtual ~FMetaSoundFrontendDocumentBuilder();
+		
+	// Call when the builder will no longer modify the IMetaSoundDocumentInterface
+	void FinishBuilding();
+
 
 	const FMetasoundFrontendClass* AddDependency(const FMetasoundFrontendClass& InClass);
 	void AddEdge(FMetasoundFrontendEdge&& InNewEdge);
@@ -292,6 +298,9 @@ private:
 
 	bool SetGraphInputInheritsDefault(FName InName, bool bInputInheritsDefault);
 
+	// Called to mark that the builder is beginning modifications to the IMetaSoundDocumentInterface
+	void BeginBuilding();
+
 	UPROPERTY(Transient)
 	TScriptInterface<IMetaSoundDocumentInterface> DocumentInterface;
 
diff --git a/Engine/Plugins/Runtime/Metasound/Source/MetasoundFrontend/Public/MetasoundFrontendRegistries.h b/Engine/Plugins/Runtime/Metasound/Source/MetasoundFrontend/Public/MetasoundFrontendRegistries.h
--- a/Engine/Plugins/Runtime/Metasound/Source/MetasoundFrontend/Public/MetasoundFrontendRegistries.h
+++ b/Engine/Plugins/Runtime/Metasound/Source/MetasoundFrontend/Public/MetasoundFrontendRegistries.h
@@ -13,9 +13,15 @@
 #include "MetasoundVertex.h"
 #include "Templates/Function.h"
 
+class IMetaSoundDocumentInterface;
+
+template <typename InInterfaceType>
+class TScriptInterface;
+
 namespace Metasound
 {
 	using FIterateMetasoundFrontendClassFunction = TFunctionRef<void(const FMetasoundFrontendClass&)>;
+	class FGraph;
 
 	namespace Frontend
 	{
@@ -213,8 +219,6 @@ namespace Metasound
 				return FString(TEXT(""));
 			}
 
-			UE_DEPRECATED("5.1", "This constructor is deprecated. Use a different constructor for FNodeREegistryTransaction")
-			FNodeRegistryTransaction(ETransactionType InType, const FNodeRegistryKey& InKey, const FNodeClassInfo& InNodeClassInfo, FTimeType InTimestamp);
 			FNodeRegistryTransaction(ETransactionType InType, const FNodeClassInfo& InNodeClassInfo, FTimeType InTimestamp);
 
 			ETransactionType GetTransactionType() const;
@@ -281,12 +285,8 @@ public:
 
 	static bool GetFrontendClassFromRegistered(const FNodeRegistryKey& InKey, FMetasoundFrontendClass& OutClass);
 	static bool GetNodeClassInfoFromRegistered(const FNodeRegistryKey& InKey, FNodeClassInfo& OutInfo);
-	UE_DEPRECATED(5.1, "Use GetInputNodeRegistryKeyForDataType with EMetasoundFrontendVertexAccessType instead.")
-	static bool GetInputNodeRegistryKeyForDataType(const FName& InDataTypeName, FNodeRegistryKey& OutKey);
 	static bool GetInputNodeRegistryKeyForDataType(const FName& InDataTypeName, const EMetasoundFrontendVertexAccessType InAccessType, FNodeRegistryKey& OutKey);
 	static bool GetVariableNodeRegistryKeyForDataType(const FName& InDataTypeName, FNodeRegistryKey& OutKey);
-	UE_DEPRECATED(5.1, "Use GetOutputNodeRegistryKeyForDataType with EMetasoundFrontendVertexAccessType instead.")
-	static bool GetOutputNodeRegistryKeyForDataType(const FName& InDataTypeName, FNodeRegistryKey& OutKey);
 	static bool GetOutputNodeRegistryKeyForDataType(const FName& InDataTypeName, const EMetasoundFrontendVertexAccessType InAccessType, FNodeRegistryKey& OutKey);
 
 
@@ -303,9 +303,16 @@ public:
 	// This is called on module startup. This invokes any registration commands enqueued by our registration macros.
 	virtual void RegisterPendingNodes() = 0;
 
-	/** Perform function for each registry transaction since a given transaction ID. */
-	UE_DEPRECATED(5.1, "ForEachNodeRegistryTransactionSince is no longer be supported")
-	virtual void ForEachNodeRegistryTransactionSince(Metasound::Frontend::FRegistryTransactionID InSince, Metasound::Frontend::FRegistryTransactionID* OutCurrentRegistryTransactionID, TFunctionRef<void(const Metasound::Frontend::FNodeRegistryTransaction&)> InFunc) const = 0;
+	// Register a graph from an IMetaSoundDocumentInterface
+	virtual FNodeRegistryKey RegisterGraph(const FNodeClassInfo& InAssetPath, const TScriptInterface<IMetaSoundDocumentInterface>& InDocument, bool bAsync=true) = 0;
+
+	// Wait for async graph registration to complete for a specific graph
+	virtual void WaitForAsyncGraphRegistration(const FNodeRegistryKey& InRegistryKey) const = 0;
+
+	// Retrieve a registered graph. 
+	//
+	// If the graph is registered asynchronously, this will wait until the registration task has completed.
+	virtual TSharedPtr<const Metasound::FGraph> GetGraph(const Metasound::Frontend::FNodeRegistryKey& InRegistryKey) const = 0;
 
 	/** Register an external node with the frontend.
 	 *
@@ -338,12 +345,8 @@ public:
 	virtual bool FindFrontendClassFromRegistered(const Metasound::Frontend::FNodeRegistryKey& InKey, FMetasoundFrontendClass& OutClass) = 0;
 	virtual const TSet<FMetasoundFrontendVersion>* FindImplementedInterfacesFromRegistered(const Metasound::Frontend::FNodeRegistryKey& InKey) const = 0;
 	virtual bool FindNodeClassInfoFromRegistered(const Metasound::Frontend::FNodeRegistryKey& InKey, FNodeClassInfo& OutInfo) = 0;
-	UE_DEPRECATED(5.1, "Use FindInputNodeRegistryKeyForDataType with EMetasoundFrontendVertexAccessType instead.")
-	virtual bool FindInputNodeRegistryKeyForDataType(const FName& InDataTypeName, FNodeRegistryKey& OutKey) = 0;
 	virtual bool FindInputNodeRegistryKeyForDataType(const FName& InDataTypeName, const EMetasoundFrontendVertexAccessType InAccessType, FNodeRegistryKey& OutKey) = 0;
 	virtual bool FindVariableNodeRegistryKeyForDataType(const FName& InDataTypeName, FNodeRegistryKey& OutKey) = 0;
-	UE_DEPRECATED(5.1, "Use FindOutputNodeRegistryKeyForDataType with EMetasoundFrontendVertexAccessType instead.")
-	virtual bool FindOutputNodeRegistryKeyForDataType(const FName& InDataTypeName, FNodeRegistryKey& OutKey) = 0;
 	virtual bool FindOutputNodeRegistryKeyForDataType(const FName& InDataTypeName, const EMetasoundFrontendVertexAccessType InAccessType, FNodeRegistryKey& OutKey) = 0;
 
 	virtual TUniquePtr<Metasound::INode> CreateNode(const FNodeRegistryKey& InKey, const Metasound::FNodeInitData&) const = 0;
