Metasound per platform blockrate and samplerate

diff --git a/Engine/Plugins/Runtime/Metasound/Config/DefaultMetaSound.ini b/Engine/Plugins/Runtime/Metasound/Config/DefaultMetaSound.ini
new file mode 100644
--- /dev/null
+++ b/Engine/Plugins/Runtime/Metasound/Config/DefaultMetaSound.ini
@@ -0,0 +1,5 @@
+[/Script/MetasoundEngine.MetaSoundSettings]
++QualitySettings=(UniqueId=A6909AF846FBBE271647E0A2E54855AC,Name="High",SampleRate=48000,BlockRate=100.000000)
+
+[/Script/MetasoundEngine.MetaSoundSource]
+QualitySettingGuid=A6909AF846FBBE271647E0A2E54855AC
diff --git a/Engine/Plugins/Runtime/Metasound/Source/MetasoundEngine/Private/MetasoundSource.cpp b/Engine/Plugins/Runtime/Metasound/Source/MetasoundEngine/Private/MetasoundSource.cpp
--- a/Engine/Plugins/Runtime/Metasound/Source/MetasoundEngine/Private/MetasoundSource.cpp
+++ b/Engine/Plugins/Runtime/Metasound/Source/MetasoundEngine/Private/MetasoundSource.cpp
@@ -150,17 +150,6 @@ UMetaSoundSource::UMetaSoundSource(const FObjectInitializer& ObjectInitializer)
 	bProcedural = true;
 	bRequiresStopFade = true;
 	NumChannels = 1;
-
-	// todo: ensure that we have a method so that the audio engine can be authoritative over the sample rate the UMetaSoundSource runs at.
-	const int32 SampleRateOverride = Metasound::Frontend::GetDefaultSampleRate();
-	if (SampleRateOverride != INDEX_NONE)
-	{
-		SampleRate = SampleRateOverride;
-	}
-	else
-	{
-		SampleRate = 48000.f;
-	}
 }
 
 const UClass& UMetaSoundSource::GetBaseMetaSoundUClass() const
@@ -201,6 +190,31 @@ void UMetaSoundSource::PostEditChangeProperty(FPropertyChangedEvent& InEvent)
 	{
 		PostEditChangeOutputFormat();
 	}
+	if (InEvent.GetPropertyName() == GET_MEMBER_NAME_CHECKED(UMetaSoundSource, SampleRateOverride) ||
+		InEvent.GetPropertyName() == GET_MEMBER_NAME_CHECKED(UMetaSoundSource, BlockRateOverride) ||
+		InEvent.GetPropertyName() == GET_MEMBER_NAME_CHECKED(UMetaSoundSource, QualitySetting) )
+	{
+		PostEditChangeQualitySettings();
+	}
+}
+
+bool UMetaSoundSource::CanEditChange(const FProperty* InProperty) const
+{
+	if (!Super::CanEditChange(InProperty))
+	{
+		return false;
+	}
+
+	// Allow changes to quality if we don't have any overrides.
+	if (InProperty->GetFName() == GET_MEMBER_NAME_CHECKED(UMetaSoundSource, QualitySetting))
+	{
+		const bool bBlockRateIsZero = BlockRateOverride.GetValue() == 0;
+		const bool bSampleRateIsZero = SampleRateOverride.GetValue() == 0;
+
+		return bBlockRateIsZero && bSampleRateIsZero;
+	}
+		
+	return true;
 }
 
 void UMetaSoundSource::PostEditChangeOutputFormat()
@@ -235,6 +249,25 @@ void UMetaSoundSource::PostEditChangeOutputFormat()
 		MarkMetasoundDocumentDirty();
 	}
 }
+
+void UMetaSoundSource::PostEditChangeQualitySettings()
+{
+	// Re-cache Operator settings by clearing the Optional.
+	OperatorSettings.Reset();
+
+	// Refresh the SampleRate (which is what the engine sees from the operator settings).
+	SampleRate = GetOperatorSettings(SampleRate).GetSampleRate();
+
+	// Always refresh the GUID with the selection.
+	if (const UMetaSoundSettings* Settings = GetDefault<UMetaSoundSettings>())	
+	{
+		auto FindByName = [&Name = QualitySetting](const FMetaSoundQualitySettings& Q) -> bool { return Q.Name == Name; };
+		if (const FMetaSoundQualitySettings* Found = Settings->QualitySettings.FindByPredicate(FindByName))
+		{
+			QualitySettingGuid = Found->UniqueId;
+		}
+	}
+}
 #endif // WITH_EDITOR
 
 bool UMetaSoundSource::ConformObjectDataToInterfaces()
@@ -349,6 +382,78 @@ void UMetaSoundSource::PostLoad()
 
 	Duration = GetDuration();
 	bLooping = IsLooping();
+
+	PostLoadQualitySettings();
+}
+
+void UMetaSoundSource::PostLoadQualitySettings()
+{
+#if WITH_EDITORONLY_DATA
+	
+	// Ensure that our Quality settings resolve. 
+	if (UMetaSoundSettings* Settings = GetMutableDefault<UMetaSoundSettings>())
+	{
+		ResolveQualitySettings(Settings);
+		
+		// Register for any changes to the settings while we're open in the editor.
+		Settings->OnSettingChanged().AddWeakLambda(this, [WeakSource = MakeWeakObjectPtr(this)](UObject* InObj, struct FPropertyChangedEvent& InEvent)
+		{
+			if (
+				WeakSource.IsValid() &&
+				InEvent.GetMemberPropertyName() == GET_MEMBER_NAME_CHECKED(UMetaSoundSettings, QualitySettings)
+			)
+			{
+				WeakSource->ResolveQualitySettings(CastChecked<UMetaSoundSettings>(InObj));
+			}
+		});
+	}
+#endif //WITH_EDITORONLY_DATA
+
+	// Override SampleRate with the Operator settings version which uses our Quality settings.
+	SampleRate = GetOperatorSettings(SampleRate).GetSampleRate();
+}
+
+void UMetaSoundSource::ResolveQualitySettings(const UMetaSoundSettings* Settings)
+{
+	const FMetaSoundQualitySettings* Resolved = nullptr;
+
+	// 1. Try and resolve by name. (most should resolve unless its been renamed, deleted).
+	auto FindByName = [&Name = QualitySetting](const FMetaSoundQualitySettings& Q) -> bool { return Q.Name == Name; };
+	Resolved = Settings->QualitySettings.FindByPredicate(FindByName);
+
+#if WITH_EDITORONLY_DATA
+
+	// 2. If that failed, try by guid (if its been renamed in the settings, we can still find it).
+	if (!Resolved && QualitySettingGuid.IsValid())
+	{
+		auto FindByGuid = [&Guid = QualitySettingGuid](const FMetaSoundQualitySettings& Q) -> bool { return Q.UniqueId == Guid; };
+		Resolved = Settings->QualitySettings.FindByPredicate(FindByName);
+	}
+
+	// 3. If still failed to resolve, use defaults and warn.
+	if (!Resolved)
+	{
+		UE_LOG(LogMetaSound, Warning, TEXT("Failed to resolve Quality '%s', resetting to the default."), *QualitySetting.ToString());
+
+		// Reset to defaults. (and make sure they are sane)
+		QualitySetting = GetDefault<UMetaSoundSource>()->QualitySetting;
+		QualitySettingGuid = GetDefault<UMetaSoundSource>()->QualitySettingGuid;
+		if (!Settings->QualitySettings.FindByPredicate(FindByName) && !Settings->QualitySettings.IsEmpty())
+		{
+			// Default doesn't point to anything, use first one in the list.
+			QualitySetting = Settings->QualitySettings[0].Name;
+			QualitySettingGuid = Settings->QualitySettings[0].UniqueId;
+		}				
+	}
+
+	// Refresh the guid/name now we've resolved to correctly reflect.
+	if (Resolved)
+	{
+		QualitySetting = Resolved->Name;
+		QualitySettingGuid = Resolved->UniqueId;
+	}
+
+#endif //WITH_EDITORONLY_DATA
 }
 
 void UMetaSoundSource::InitParameters(TArray<FAudioParameter>& ParametersToInit, FName InFeatureName)
@@ -991,14 +1096,67 @@ TSharedPtr<Audio::IParameterTransmitter> UMetaSoundSource::CreateParameterTransm
 }
 
 Metasound::FOperatorSettings UMetaSoundSource::GetOperatorSettings(Metasound::FSampleRate InSampleRate) const
-{
-	const float BlockRate = Metasound::Frontend::GetDefaultBlockRate();
-	const int32 SampleRateOverride = Metasound::Frontend::GetDefaultSampleRate();
-	if (SampleRateOverride != INDEX_NONE)
+{	
+	if (!OperatorSettings)
 	{
-		return Metasound::FOperatorSettings(SampleRateOverride, BlockRate);
+		using namespace Metasound;
+		using namespace Metasound::SourcePrivate;
+
+		// Lazy Query and cache on the optional.
+		auto QueryQualitySettings = [&](Metasound::FSampleRate InSampleRate) -> Metasound::FOperatorSettings
+		{
+			static const int32 DefaultSampleRateConstant = 48000;
+			static const float DefaultBlockRateConstant = 100.f;
+			
+			// 1. Sensible defaults.
+			FSampleRate SampleRate = DefaultSampleRateConstant;
+			float BlockRate = DefaultBlockRateConstant;
+
+			// 2. Query CVars. (Override with CVars if they are > 0)
+			const float BlockRateCVar = Metasound::Frontend::GetDefaultBlockRate();
+			const int32 SampleRateCvar = Metasound::Frontend::GetDefaultSampleRate();
+
+			if (SampleRateCvar != INDEX_NONE)
+			{
+				SampleRate = SampleRateCvar;
+			}
+			if (BlockRateCVar > 0)
+			{
+				BlockRate = BlockRateCVar;
+			}
+
+			// 3. Query our quality settings.
+			if (const UMetaSoundSettings* Settings = GetDefault<UMetaSoundSettings>())
+			{
+				if (const FMetaSoundQualitySettings* Found = Settings->QualitySettings.FindByPredicate([&QT = QualitySetting](const FMetaSoundQualitySettings& Q) -> bool { return Q.Name == QT; }))
+				{
+					// Allow partial applications of settings, if some are non-zero.
+					if (Found->BlockRate > 0.f)
+					{
+						BlockRate = Found->BlockRate;
+					}
+					if (Found->SampleRate > 0.f)
+					{
+						SampleRate = Found->SampleRate;
+					}
+				}
+			}
+
+			// 4. Do per asset overrides.
+			if (const float SerializedBlockRate = BlockRateOverride.GetValue(); SerializedBlockRate > 0.0f)
+			{
+				BlockRate = SerializedBlockRate;
+			}
+			if (const int32 SerializedSampleRate = SampleRateOverride.GetValue(); SerializedSampleRate > 0)
+			{
+				SampleRate = SerializedSampleRate;
+			}
+
+			return Metasound::FOperatorSettings(SampleRate, BlockRate);
+		};
+		OperatorSettings = QueryQualitySettings(InSampleRate);
 	}
-	return Metasound::FOperatorSettings(InSampleRate, BlockRate);
+	return *OperatorSettings;
 }
 
 Metasound::FMetasoundEnvironment UMetaSoundSource::CreateEnvironment() const
diff --git a/Engine/Plugins/Runtime/Metasound/Source/MetasoundEngine/Public/MetasoundSettings.h b/Engine/Plugins/Runtime/Metasound/Source/MetasoundEngine/Public/MetasoundSettings.h
--- a/Engine/Plugins/Runtime/Metasound/Source/MetasoundEngine/Public/MetasoundSettings.h
+++ b/Engine/Plugins/Runtime/Metasound/Source/MetasoundEngine/Public/MetasoundSettings.h
@@ -33,6 +33,32 @@ struct METASOUNDENGINE_API FDefaultMetaSoundAssetAutoUpdateSettings
 	FSoftObjectPath MetaSound;
 };
 
+USTRUCT()
+struct METASOUNDENGINE_API FMetaSoundQualitySettings
+{
+	GENERATED_BODY()
+
+#if WITH_EDITORONLY_DATA
+	/** A hidden GUID that will be generated once when adding a new entry. This prevents orphaning of renamed entries. **/
+	UPROPERTY()
+	FGuid UniqueId = FGuid::NewGuid();
+#endif //WITH_EDITORONLY_DATA
+
+	/** Name of this quality setting. This will appear in the quality dropdown list.
+		The names should be unique and adequately describe the Entry. "High", "Low" etc. **/
+	UPROPERTY(EditAnywhere, Category = "Quality")
+	FName Name;
+
+	/** Sample Rate (in Hz). NOTE: A Zero value will force the default. **/
+	UPROPERTY(EditAnywhere, Category = "Quality", meta = (ClampMin = "0", ClampMax="96000"))
+	int32 SampleRate = 0;
+
+	/** Block Rate (in Hz). NOTE: A Zero value will force the default.  **/
+	UPROPERTY(EditAnywhere, Category = "Quality", meta = (ClampMin = "0", ClampMax="100"))
+	float BlockRate = 0.f;
+};
+
+
 UCLASS(config = MetaSound, defaultconfig, meta = (DisplayName = "MetaSounds"))
 class METASOUNDENGINE_API UMetaSoundSettings : public UDeveloperSettings
 {
@@ -63,14 +89,46 @@ public:
 	  */
 	UPROPERTY(EditAnywhere, config, Category = Registration, meta = (RelativePath, LongPackageName))
 	TArray<FDirectoryPath> DirectoriesToRegister;
-
+	
+	/** Array of possible quality settings for Metasounds to chose from */
+	UPROPERTY(EditAnywhere, config, Category = Quality)
+	TArray<FMetaSoundQualitySettings> QualitySettings;
+		
 	UPROPERTY(Transient)
-	int32 DenyListCacheChangeID = 0;
+	int32 DenyListCacheChangeID = 0;	
 
 #if WITH_EDITOR
 	virtual void PostEditChangeProperty(FPropertyChangedEvent& PropertyChangedEvent) override
-	{
+	{		
+		Super::PostEditChangeProperty(PropertyChangedEvent);
+
 		DenyListCacheChangeID++;
 	}
 #endif // WITH_EDITOR
 };
+
+UCLASS()
+class UMetaSoundQuality : public UObject
+{
+	GENERATED_BODY()
+public:
+
+	/**
+	* Returns a list of quality settings to present to a combobox
+	* */
+	UFUNCTION()
+	static TArray<FName> GetQualityList()
+	{
+		TArray<FName> Names;
+
+		if (const UMetaSoundSettings* Settings = GetDefault<UMetaSoundSettings>())
+		{
+			Algo::Transform(Settings->QualitySettings, Names, [](const FMetaSoundQualitySettings& Quality) -> FName 
+			{
+				return Quality.Name;
+			});
+		}
+		return Names;
+	}
+
+};
diff --git a/Engine/Plugins/Runtime/Metasound/Source/MetasoundEngine/Public/MetasoundSource.h b/Engine/Plugins/Runtime/Metasound/Source/MetasoundEngine/Public/MetasoundSource.h
--- a/Engine/Plugins/Runtime/Metasound/Source/MetasoundEngine/Public/MetasoundSource.h
+++ b/Engine/Plugins/Runtime/Metasound/Source/MetasoundEngine/Public/MetasoundSource.h
@@ -40,13 +40,15 @@ namespace Audio
 	using DeviceID = uint32;
 }
 
+class UMetaSoundSettings;
+
 DECLARE_TS_MULTICAST_DELEGATE_TwoParams(FOnGeneratorInstanceCreated, uint64, TSharedPtr<Metasound::FMetasoundGenerator>);
 DECLARE_TS_MULTICAST_DELEGATE_TwoParams(FOnGeneratorInstanceDestroyed, uint64, TSharedPtr<Metasound::FMetasoundGenerator>);
 
 /**
  * This Metasound type can be played as an audio source.
  */
-UCLASS(hidecategories = object, BlueprintType)
+UCLASS(hidecategories = object, BlueprintType, config = Metasound, defaultconfig)
 class METASOUNDENGINE_API UMetaSoundSource : public USoundWaveProcedural, public FMetasoundAssetBase, public IMetaSoundDocumentInterface
 {
 	GENERATED_BODY()
@@ -79,6 +81,24 @@ public:
 	UPROPERTY(EditAnywhere, BlueprintReadOnly, Category = Metasound)
 	EMetaSoundOutputAudioFormat OutputFormat;
 
+	// The Quality this Metasound will use. These are defined in the MetaSounds project settings.
+	UPROPERTY(config, EditAnywhere, BlueprintReadWrite, meta = (GetOptions="MetasoundEngine.MetaSoundQuality.GetQualityList"), Category = "Metasound")
+	FName QualitySetting;
+
+#if WITH_EDITORONLY_DATA
+	// This a editor only look up for the Quality Setting above. Preventing orphaning of the original name.
+	UPROPERTY()
+	FGuid QualitySettingGuid;
+#endif //WITH_EDITOR_DATA
+
+	// Override the BlockRate for this Sound (overrides the Platforms Quality Settings for this asset)
+	UPROPERTY(EditAnywhere, BlueprintReadOnly, AdvancedDisplay, Category = Metasound, meta = (UIMin = 1, UIMax = 1000, DisplayAfter="OutputFormat", DisplayName = "Override Block Rate (in Hz)"))
+	FPerPlatformFloat BlockRateOverride = 0.f;
+
+	// Override the SampleRate for this Sound (overrides the 
+	UPROPERTY(EditAnywhere, BlueprintReadOnly, AdvancedDisplay, Category = Metasound, meta = (UIMin = 100, UIMax = 96000, DisplayName = "Override Sample Rate (in Hz)"))
+	FPerPlatformInt SampleRateOverride = 0;
+
 	UPROPERTY(AssetRegistrySearchable)
 	FGuid AssetClassID;
 
@@ -150,8 +170,11 @@ public:
 
 	virtual void PostEditChangeProperty(FPropertyChangedEvent& InEvent) override;
 
+	virtual bool CanEditChange(const FProperty* InProperty) const override;
+
 private:
 	void PostEditChangeOutputFormat();
+	void PostEditChangeQualitySettings();
 public:
 
 #endif // WITH_EDITOR
@@ -169,6 +192,9 @@ public:
 	virtual void Serialize(FArchive& Ar) override;
 	virtual void PostLoad() override;
 
+	void PostLoadQualitySettings();
+	void ResolveQualitySettings(const UMetaSoundSettings* Settings);
+
 	virtual bool ConformObjectDataToInterfaces() override;
 
 	UObject* GetOwningAsset() override
@@ -259,4 +285,7 @@ private:
 	TSharedPtr<Metasound::DynamicGraph::FDynamicOperatorTransactor> DynamicTransactor;
 
 	bool bIsBuilderActive = false;
+	 * Lazy (Cached) Operator Settings. Built in GetOperatorSettings
+	 */
+	mutable TOptional<Metasound::FOperatorSettings> OperatorSettings;
 };
