Mute warning on Quality Missing for CIS

diff --git a/Engine/Plugins/Runtime/Metasound/Source/MetasoundEngine/Private/MetasoundSource.cpp b/Engine/Plugins/Runtime/Metasound/Source/MetasoundEngine/Private/MetasoundSource.cpp
--- a/Engine/Plugins/Runtime/Metasound/Source/MetasoundEngine/Private/MetasoundSource.cpp
+++ b/Engine/Plugins/Runtime/Metasound/Source/MetasoundEngine/Private/MetasoundSource.cpp
@@ -433,7 +433,8 @@ void UMetaSoundSource::ResolveQualitySettings(const UMetaSoundSettings* Settings
 	// 3. If still failed to resolve, use defaults and warn.
 	if (!Resolved)
 	{
-		UE_LOG(LogMetaSound, Warning, TEXT("Failed to resolve Quality '%s', resetting to the default."), *QualitySetting.ToString());
+		// Disable the warning for now.
+		//UE_LOG(LogMetaSound, Warning, TEXT("Failed to resolve Quality '%s', resetting to the default."), *QualitySetting.ToString());
 
 		// Reset to defaults. (and make sure they are sane)
 		QualitySetting = GetDefault<UMetaSoundSource>()->QualitySetting;
