Unblock playtest. Fix crash in metasound source playing
#rnx

diff --git a/Engine/Plugins/Runtime/Metasound/Source/MetasoundFrontend/Private/MetasoundAssetBase.cpp b/Engine/Plugins/Runtime/Metasound/Source/MetasoundFrontend/Private/MetasoundAssetBase.cpp
--- a/Engine/Plugins/Runtime/Metasound/Source/MetasoundFrontend/Private/MetasoundAssetBase.cpp
+++ b/Engine/Plugins/Runtime/Metasound/Source/MetasoundFrontend/Private/MetasoundAssetBase.cpp
@@ -1009,7 +1009,8 @@ Metasound::Frontend::FNodeRegistryKey FMetasoundAssetBase::CacheRuntimeData(cons
 	// If the IMetaSoundDocumentInterface actively modified by a builder, then 
 	// we must build synchronously to avoid a race condition on reading/writing 
 	// the IMetaSoundDocumentInterface
-	const bool bAsync = !IsBuilderActive();
+	//const bool bAsync = !IsBuilderActive();
+	constexpr bool bAsync = false; // TODO: Active crash needs addressing where graph registration and async metasound building happen on same thread. AsyncBuilder deadlocks while waiting for task to complete. 
 
 	FNodeClassInfo NodeClassInfo { Document.RootGraph, FSoftObjectPath(GetOwningAsset()) };
 	FNodeRegistryKey NewRegistryKey = FMetasoundFrontendRegistryContainer::Get()->RegisterGraph(NodeClassInfo, InDocumentInterface, bAsync);
