Add Insights trace markup for MetaSound operator creation

diff --git a/Engine/Plugins/Runtime/Metasound/Source/MetasoundFrontend/Private/MetasoundFrontendDataTypeRegistry.cpp b/Engine/Plugins/Runtime/Metasound/Source/MetasoundFrontend/Private/MetasoundFrontendDataTypeRegistry.cpp
--- a/Engine/Plugins/Runtime/Metasound/Source/MetasoundFrontend/Private/MetasoundFrontendDataTypeRegistry.cpp
+++ b/Engine/Plugins/Runtime/Metasound/Source/MetasoundFrontend/Private/MetasoundFrontendDataTypeRegistry.cpp
@@ -855,6 +855,8 @@ namespace Metasound
 
 			TSharedPtr<Audio::IProxyData> FDataTypeRegistry::CreateProxyFromUObject(const FName& InDataType, UObject* InObject) const
 			{
+				METASOUND_TRACE_CPUPROFILER_EVENT_SCOPE_TEXT(*FString::Printf(TEXT("Metasound::FDataTypeRegistry::CreateProxyFromUObject Type:%s"), *InDataType.ToString()));
+
 				TSharedPtr<Audio::IProxyData> ProxyPtr;
 
 				if (const IDataTypeRegistryEntry* Entry = FindDataTypeEntry(InDataType))
@@ -930,6 +932,8 @@ namespace Metasound
 
 			FLiteral FDataTypeRegistry::CreateLiteralFromUObjectArray(const FName& InDataType, const TArray<UObject*>& InObjectArray) const
 			{
+				METASOUND_TRACE_CPUPROFILER_EVENT_SCOPE_TEXT(*FString::Printf(TEXT("Metasound::FDataTypeRegistry::CreateLiteralFromUObjectArray Type:%s"), *InDataType.ToString()));
+
 				TArray<TSharedPtr<Audio::IProxyData>> ProxyArray;
 				const IDataTypeRegistryEntry* DataTypeEntry = FindDataTypeEntry(InDataType);
 				if (!DataTypeEntry)
diff --git a/Engine/Plugins/Runtime/Metasound/Source/MetasoundGraphCore/Private/MetasoundOperatorBuilder.cpp b/Engine/Plugins/Runtime/Metasound/Source/MetasoundGraphCore/Private/MetasoundOperatorBuilder.cpp
--- a/Engine/Plugins/Runtime/Metasound/Source/MetasoundGraphCore/Private/MetasoundOperatorBuilder.cpp
+++ b/Engine/Plugins/Runtime/Metasound/Source/MetasoundGraphCore/Private/MetasoundOperatorBuilder.cpp
@@ -330,21 +330,26 @@ namespace Metasound
 		// Create FOperatorInfos from Nodes
 		TSortedMap<FOperatorID, FGraphOperatorData::FOperatorInfo>& OperatorMap = InOutGraphOperatorData.OperatorMap;
 		TArray<FOperatorID>& OperatorOrder = InOutGraphOperatorData.OperatorOrder;
-
-		for (const INode* Node : InSortedNodes)
 		{
-			FOperatorID OperatorID = GetOperatorID(Node);
-			OperatorOrder.Add(OperatorID);
-			OperatorMap.Add(OperatorID, FGraphOperatorData::FOperatorInfo{nullptr, Node->GetVertexInterface()});
+			METASOUND_TRACE_CPUPROFILER_EVENT_SCOPE(Metasound::FOperatorBuilder::InitializeOperatorInfo::Nodes);
+			for (const INode* Node : InSortedNodes)
+			{
+				FOperatorID OperatorID = GetOperatorID(Node);
+				OperatorOrder.Add(OperatorID);
+				OperatorMap.Add(OperatorID, FGraphOperatorData::FOperatorInfo{nullptr, Node->GetVertexInterface()});
+			}
 		}
 
-		// Set the output destinations on operator infos
-		for (const FDataEdge& Edge : InGraph.GetDataEdges())
 		{
-			const FOperatorID FromOperatorID = GetOperatorID(Edge.From.Node);
-			FGraphOperatorData::FOperatorInfo& OperatorInfo = OperatorMap.FindChecked(FromOperatorID);
+			METASOUND_TRACE_CPUPROFILER_EVENT_SCOPE(Metasound::FOperatorBuilder::InitializeOperatorInfo::OutputDestinations);
+			// Set the output destinations on operator infos
+			for (const FDataEdge& Edge : InGraph.GetDataEdges())
+			{
+				const FOperatorID FromOperatorID = GetOperatorID(Edge.From.Node);
+				FGraphOperatorData::FOperatorInfo& OperatorInfo = OperatorMap.FindChecked(FromOperatorID);
 
-			OperatorInfo.OutputConnections.FindOrAdd(Edge.From.Vertex.VertexName).Add(FGraphOperatorData::FVertexDestination{GetOperatorID(Edge.To.Node), Edge.To.Vertex.VertexName});
+				OperatorInfo.OutputConnections.FindOrAdd(Edge.From.Vertex.VertexName).Add(FGraphOperatorData::FVertexDestination{GetOperatorID(Edge.To.Node), Edge.To.Vertex.VertexName});
+			}
 		}
 	}
 
@@ -455,15 +460,19 @@ namespace Metasound
 			FGraphOperatorData::FOperatorInfo& OperatorInfo = OperatorMap.FindChecked(OperatorID);
 
 			{
-				METASOUND_TRACE_CPUPROFILER_EVENT_SCOPE(Metasound::FOperatorBuilder::CreateOperators::CreateAndBind);
+				METASOUND_TRACE_CPUPROFILER_EVENT_SCOPE_TEXT(*FString::Printf(TEXT("Metasound::FOperatorBuilder::CreateOperators::CreateAndBind %s"), *Node->GetMetadata().ClassName.GetFullName().ToString()));
+
 				FBuildOperatorParams CreateParams{*Node, InOutContext.Settings, OperatorInfo.VertexData.GetInputs(), InOutContext.Environment, this};
 				FOperatorFactorySharedRef Factory = Node->GetDefaultOperatorFactory();
 				if (ProfileOperators && Profiling::OperatorShouldBeProfiled(Node->GetMetadata()))
 				{
+					METASOUND_TRACE_CPUPROFILER_EVENT_SCOPE(Metasound::FOperatorBuilder::CreateOperators::CreateOperator);
+
 					OperatorInfo.Operator = MakeUnique<FProfilingOperator>(Factory->CreateOperator(CreateParams, InOutContext.Results), Node);
 				}
 				else
 				{
+					METASOUND_TRACE_CPUPROFILER_EVENT_SCOPE(Metasound::FOperatorBuilder::CreateOperators::CreateOperator);
 					OperatorInfo.Operator = Factory->CreateOperator(CreateParams, InOutContext.Results);
 				}
 
@@ -473,12 +482,17 @@ namespace Metasound
 				}
 
 				// Bind vertex to operator data
-				OperatorInfo.Operator->BindInputs(OperatorInfo.VertexData.GetInputs());
-				OperatorInfo.Operator->BindOutputs(OperatorInfo.VertexData.GetOutputs());
+				{
+					METASOUND_TRACE_CPUPROFILER_EVENT_SCOPE(Metasound::FOperatorBuilder::CreateOperators::BindInputsAndOutputs);
+					OperatorInfo.Operator->BindInputs(OperatorInfo.VertexData.GetInputs());
+					OperatorInfo.Operator->BindOutputs(OperatorInfo.VertexData.GetOutputs());
+				}
 				
 				// Check if outputs are bound correctly.
 				if (BuilderSettings.bValidateOperatorOutputsAreBound)
 				{
+					METASOUND_TRACE_CPUPROFILER_EVENT_SCOPE(Metasound::FOperatorBuilder::CreateOperators::ValidateOperatorOutputsAreBound);
+
 					BuildStatus |= ValidateOperatorOutputsAreBound(*Node, OperatorInfo.VertexData.GetOutputs());
 				}
 			}
