Use FLazyName instead of TChar in MetaSound parameter macros

diff --git a/Engine/Plugins/Runtime/Metasound/Source/MetasoundFrontend/Public/MetasoundArrayNodes.h b/Engine/Plugins/Runtime/Metasound/Source/MetasoundFrontend/Public/MetasoundArrayNodes.h
--- a/Engine/Plugins/Runtime/Metasound/Source/MetasoundFrontend/Public/MetasoundArrayNodes.h
+++ b/Engine/Plugins/Runtime/Metasound/Source/MetasoundFrontend/Public/MetasoundArrayNodes.h
@@ -45,7 +45,7 @@ namespace Metasound
 
 	namespace ArrayNodeVertexNames
 	{
-		static const TCHAR* InputInitialArrayName = TEXT("Array");
+		static const FLazyName InputInitialArrayName = TEXT("Array");
 #if WITH_EDITOR
 		static const FText InputInitialArrayTooltip = LOCTEXT("InitialArrayTooltip", "Initial Array");
 		static const FText InputInitialArrayDisplayName = LOCTEXT("InitialArrayDisplayName", "Init Array"); 
diff --git a/Engine/Plugins/Runtime/Metasound/Source/MetasoundGraphCore/Public/MetasoundParamHelper.h b/Engine/Plugins/Runtime/Metasound/Source/MetasoundGraphCore/Public/MetasoundParamHelper.h
--- a/Engine/Plugins/Runtime/Metasound/Source/MetasoundGraphCore/Public/MetasoundParamHelper.h
+++ b/Engine/Plugins/Runtime/Metasound/Source/MetasoundGraphCore/Public/MetasoundParamHelper.h
@@ -3,6 +3,7 @@
 #pragma once
 
 #include "MetasoundVertex.h"
+#include "UObject/NameTypes.h"
 
 /*
 
@@ -48,13 +49,13 @@
 #if WITH_EDITOR
 #define LOC_DEFINE_REGION
 #define METASOUND_PARAM(NAME, NAME_TEXT, TOOLTIP_TEXT) \
-	static const TCHAR* NAME##Name = TEXT(NAME_TEXT); \
+	static const FLazyName NAME##Name = TEXT(NAME_TEXT); \
 	static const FText NAME##Tooltip = LOCTEXT(#NAME "Tooltip", TOOLTIP_TEXT); \
 	static const FText NAME##DisplayName = LOCTEXT(#NAME "DisplayName", NAME_TEXT); 
 #undef LOC_DEFINE_REGION
 #else 
 #define METASOUND_PARAM(NAME, NAME_TEXT, TOOLTIP_TEXT) \
-	static const TCHAR* NAME##Name = TEXT(NAME_TEXT); \
+	static const FLazyName NAME##Name = TEXT(NAME_TEXT); \
 	static const FText NAME##Tooltip = FText::GetEmpty(); \
 	static const FText NAME##DisplayName = FText::GetEmpty();
 #endif // WITH_EDITOR
@@ -66,14 +67,14 @@
 #define METASOUND_GET_PARAM_DISPLAYNAME(NAME) NAME##DisplayName
 #define METASOUND_GET_PARAM_NAME_AND_METADATA(NAME) METASOUND_GET_PARAM_NAME(NAME), METASOUND_GET_PARAM_METADATA(NAME)
 
-#define METASOUND_GET_PARAM_NAME_WITH_INDEX(NAME, INDEX) *FString::Format(NAME##Name, {INDEX})
+#define METASOUND_GET_PARAM_NAME_WITH_INDEX(NAME, INDEX) *FString::Format(*static_cast<FName>(NAME##Name).ToString(), {INDEX})
 #if WITH_EDITOR
-#define METASOUND_GET_PARAM_NAME_WITH_INDEX_AND_TT(NAME, INDEX)  *FString::Format(NAME##Name, {INDEX}), FText::Format(NAME##Tooltip, INDEX)
+#define METASOUND_GET_PARAM_NAME_WITH_INDEX_AND_TT(NAME, INDEX)  METASOUND_GET_PARAM_NAME_WITH_INDEX(NAME, INDEX), FText::Format(NAME##Tooltip, INDEX)
 #define METASOUND_GET_PARAM_TT_WITH_INDEX(NAME, INDEX)  FText::Format(NAME##Tooltip, INDEX)
-#define METASOUND_GET_PARAM_NAME_WITH_INDEX_AND_METADATA(NAME, INDEX)  *FString::Format(NAME##Name, {INDEX}), FDataVertexMetadata {FText::Format(NAME##Tooltip, INDEX)}
+#define METASOUND_GET_PARAM_NAME_WITH_INDEX_AND_METADATA(NAME, INDEX)  METASOUND_GET_PARAM_NAME_WITH_INDEX(NAME, INDEX), FDataVertexMetadata {FText::Format(NAME##Tooltip, INDEX)}
 #else 
-#define METASOUND_GET_PARAM_NAME_WITH_INDEX_AND_TT(NAME, INDEX)  *FString::Format(NAME##Name, {INDEX}), FText::GetEmpty()
+#define METASOUND_GET_PARAM_NAME_WITH_INDEX_AND_TT(NAME, INDEX)  *FString::Format(*static_cast<FName>(NAME##Name).ToString(), {INDEX}), FText::GetEmpty()
 #define METASOUND_GET_PARAM_TT_WITH_INDEX(NAME, INDEX)  FText::GetEmpty();
-#define METASOUND_GET_PARAM_NAME_WITH_INDEX_AND_METADATA(NAME, INDEX)  *FString::Format(NAME##Name, {INDEX}), FDataVertexMetadata {FText::GetEmpty()}
+#define METASOUND_GET_PARAM_NAME_WITH_INDEX_AND_METADATA(NAME, INDEX)  *FString::Format(*static_cast<FName>(NAME##Name).ToString(), {INDEX}), FDataVertexMetadata {FText::GetEmpty()}
 #endif // WITH_EDITOR
 
