Fixing ADSR clicking on envelope release

diff --git a/Engine/Plugins/Runtime/Metasound/Source/MetasoundStandardNodes/Private/MetasoundADSREnvelopeNode.cpp b/Engine/Plugins/Runtime/Metasound/Source/MetasoundStandardNodes/Private/MetasoundADSREnvelopeNode.cpp
--- a/Engine/Plugins/Runtime/Metasound/Source/MetasoundStandardNodes/Private/MetasoundADSREnvelopeNode.cpp
+++ b/Engine/Plugins/Runtime/Metasound/Source/MetasoundStandardNodes/Private/MetasoundADSREnvelopeNode.cpp
@@ -224,7 +224,6 @@ namespace Metasound
 
 			static void GetNextEnvelopeOutput(FEnvState& InState, int32 StartFrame, int32 EndFrame, TArray<int32>& OutOnDecayFrames, TArray<int32>& OutOnSustainFrames, TArray<int32>& OutOnDoneFrames, FAudioBuffer& OutEnvelopeValue)
 			{
-				OutEnvelopeValue.Zero();
 				// If we are not active zero the buffer and early exit
 				if (InState.CurrentSampleIndex == INDEX_NONE)
 				{
@@ -344,6 +343,13 @@ namespace Metasound
 									// Envelope is done
 									if (InState.EnvEase.IsDone())
 									{
+										// Zero out the rest of the envelope
+										int32 NumSamplesLeft = EndFrame - i - 1;
+										if (NumSamplesLeft > 0)
+										{
+											FMemory::Memzero(&OutEnvPtr[i + 1], sizeof(float) * NumSamplesLeft);
+										}
+
 										InState.CurrentSampleIndex = INDEX_NONE;
 										OutOnDoneFrames.Add(i);
 										break;
