Fixed usage of TRACE_CPUPROFILER_EVENT_SCOPE macros in engine plugins.

* TRACE_CPUPROFILER_EVENT_SCOPE(ScopeName) --> to be used with a plain text as scope name (not a string!). If the parameter is a string ("abc" or TEXT("abc")), the quotes (and the TEXT prefix) will be included in the actual timer name.
* TRACE_CPUPROFILER_EVENT_SCOPE_STR("scope name") --> to be used with a static string (const ANSICHAR* or const TCHAR*); ex: when the scope name needs to include a space char
* TRACE_CPUPROFILER_EVENT_SCOPE_TEXT(*MyScopeName) --> to be used with a dynamic string (const ANSICHAR* or const TCHAR*)
See also the comments for these macros in ProfilingDebugging\CpuProfilerTrace.h.

diff --git a/Engine/Plugins/AI/MassCrowd/Source/MassCrowd/Private/MassCrowdServerRepresentationLODProcessor.cpp b/Engine/Plugins/AI/MassCrowd/Source/MassCrowd/Private/MassCrowdServerRepresentationLODProcessor.cpp
--- a/Engine/Plugins/AI/MassCrowd/Source/MassCrowd/Private/MassCrowdServerRepresentationLODProcessor.cpp
+++ b/Engine/Plugins/AI/MassCrowd/Source/MassCrowd/Private/MassCrowdServerRepresentationLODProcessor.cpp
@@ -52,14 +52,14 @@ void UMassCrowdServerRepresentationLODProcessor::Initialize(UObject& InOwner)
 
 void UMassCrowdServerRepresentationLODProcessor::Execute(FMassEntityManager& EntityManager, FMassExecutionContext& Context)
 {
-	TRACE_CPUPROFILER_EVENT_SCOPE(TEXT("CrowdServerRepresentationLOD"))
+	TRACE_CPUPROFILER_EVENT_SCOPE(CrowdServerRepresentationLOD)
 
 	const UMassLODSubsystem& LODSubsystem = Context.GetSubsystemChecked<UMassLODSubsystem>();
 	const TArray<FViewerInfo>& Viewers = LODSubsystem.GetViewers();
 	LODCalculator.PrepareExecution(Viewers);
 	
 	{
-		TRACE_CPUPROFILER_EVENT_SCOPE(TEXT("CalculateLOD"))
+		TRACE_CPUPROFILER_EVENT_SCOPE(CalculateLOD)
 		
 		EntityQuery.ForEachEntityChunk(EntityManager, Context, [this](FMassExecutionContext& Context)
 		{
@@ -70,7 +70,7 @@ void UMassCrowdServerRepresentationLODProcessor::Execute(FMassEntityManager& Ent
 	}
 
 	{
-		TRACE_CPUPROFILER_EVENT_SCOPE(TEXT("AdjustDistancesAndLODFromCount"))
+		TRACE_CPUPROFILER_EVENT_SCOPE(AdjustDistancesAndLODFromCount)
 		
 		if (LODCalculator.AdjustDistancesFromCount())
 		{
@@ -87,7 +87,7 @@ void UMassCrowdServerRepresentationLODProcessor::Execute(FMassEntityManager& Ent
 	// Optional debug display
 	if (UE::MassCrowd::bDebugCrowdServerRepresentationLOD)
 	{
-		TRACE_CPUPROFILER_EVENT_SCOPE(TEXT("DebugDisplayLOD"))
+		TRACE_CPUPROFILER_EVENT_SCOPE(DebugDisplayLOD)
 		UWorld* World = EntityManager.GetWorld();
 		EntityQuery.ForEachEntityChunk(EntityManager, Context, [this, World](FMassExecutionContext& Context)
 		{
diff --git a/Engine/Plugins/AI/MassCrowd/Source/MassCrowd/Private/MassCrowdVisualizationLODProcessor.cpp b/Engine/Plugins/AI/MassCrowd/Source/MassCrowd/Private/MassCrowdVisualizationLODProcessor.cpp
--- a/Engine/Plugins/AI/MassCrowd/Source/MassCrowd/Private/MassCrowdVisualizationLODProcessor.cpp
+++ b/Engine/Plugins/AI/MassCrowd/Source/MassCrowd/Private/MassCrowdVisualizationLODProcessor.cpp
@@ -56,14 +56,14 @@ void UMassCrowdVisualizationLODProcessor::Execute(FMassEntityManager& EntityMana
 
 	ForceOffLOD((bool)UE::MassCrowd::GCrowdTurnOffVisualization);
 
-	TRACE_CPUPROFILER_EVENT_SCOPE(TEXT("CrowdVisualizationLOD"))
+	TRACE_CPUPROFILER_EVENT_SCOPE(CrowdVisualizationLOD)
 
 	Super::Execute(EntityManager, Context);
 
 #if WITH_MASSGAMEPLAY_DEBUG
 	if (UE::MassCrowd::bDebugCrowdVisualizationLOD)
 	{
-		TRACE_CPUPROFILER_EVENT_SCOPE(TEXT("DebugDisplayLOD"))
+		TRACE_CPUPROFILER_EVENT_SCOPE(DebugDisplayLOD)
 
 		DebugEntityQuery.ForEachEntityChunk(EntityManager, Context, [World](FMassExecutionContext& Context)
 		{
@@ -77,7 +77,7 @@ void UMassCrowdVisualizationLODProcessor::Execute(FMassEntityManager& EntityMana
 
 	if (UE::MassCrowd::bDebugShowISMUnderSpecifiedRange > 0)
 	{
-		TRACE_CPUPROFILER_EVENT_SCOPE(TEXT("ShowISMUnderSpecifiedRange"))
+		TRACE_CPUPROFILER_EVENT_SCOPE(ShowISMUnderSpecifiedRange)
 
 		DebugEntityQuery.ForEachEntityChunk(EntityManager, Context, [World](const FMassExecutionContext& Context)
 		{
diff --git a/Engine/Plugins/AI/MassCrowd/Source/MassCrowd/Private/MassCrowdVisualizationProcessor.cpp b/Engine/Plugins/AI/MassCrowd/Source/MassCrowd/Private/MassCrowdVisualizationProcessor.cpp
--- a/Engine/Plugins/AI/MassCrowd/Source/MassCrowd/Private/MassCrowdVisualizationProcessor.cpp
+++ b/Engine/Plugins/AI/MassCrowd/Source/MassCrowd/Private/MassCrowdVisualizationProcessor.cpp
@@ -78,7 +78,7 @@ void UMassDebugCrowdVisualizationProcessor::Execute(FMassEntityManager& EntityMa
 	{
 		check(World);
 
-		TRACE_CPUPROFILER_EVENT_SCOPE(TEXT("DebugDisplayVisualType"))
+		TRACE_CPUPROFILER_EVENT_SCOPE(DebugDisplayVisualType)
 
 		EntityQuery.ForEachEntityChunk(EntityManager, Context, [this](const FMassExecutionContext& Context)
 		{
diff --git a/Engine/Plugins/Experimental/Text3D/Source/Text3D/Private/Text3DComponent.cpp b/Engine/Plugins/Experimental/Text3D/Source/Text3D/Private/Text3DComponent.cpp
--- a/Engine/Plugins/Experimental/Text3D/Source/Text3D/Private/Text3DComponent.cpp
+++ b/Engine/Plugins/Experimental/Text3D/Source/Text3D/Private/Text3DComponent.cpp
@@ -1050,7 +1050,7 @@ void UText3DComponent::BuildTextMesh(const bool& bCleanCache)
 
 void UText3DComponent::BuildTextMeshInternal(const bool& bCleanCache)
 {
-	TRACE_CPUPROFILER_EVENT_SCOPE(TEXT("UText3DComponent::Rebuild"));
+	TRACE_CPUPROFILER_EVENT_SCOPE(UText3DComponent::Rebuild);
 
 	ON_SCOPE_EXIT { bIsBuilding = false; };
 
diff --git a/Engine/Plugins/Runtime/ComputeFramework/Source/ComputeFramework/Private/ComputeGraphWorker.cpp b/Engine/Plugins/Runtime/ComputeFramework/Source/ComputeFramework/Private/ComputeGraphWorker.cpp
--- a/Engine/Plugins/Runtime/ComputeFramework/Source/ComputeFramework/Private/ComputeGraphWorker.cpp
+++ b/Engine/Plugins/Runtime/ComputeFramework/Source/ComputeFramework/Private/ComputeGraphWorker.cpp
@@ -69,7 +69,7 @@ bool FComputeGraphTaskWorker::HasWork(FName InExecutionGroupName) const
 
 void FComputeGraphTaskWorker::SubmitWork(FRDGBuilder& GraphBuilder, FName InExecutionGroupName, ERHIFeatureLevel::Type FeatureLevel)
 {
-	TRACE_CPUPROFILER_EVENT_SCOPE("ComputeFramework::ExecuteBatches");
+	TRACE_CPUPROFILER_EVENT_SCOPE(ComputeFramework::ExecuteBatches);
 	RDG_EVENT_SCOPE(GraphBuilder, "ComputeFramework::ExecuteBatches");
 	RDG_GPU_STAT_SCOPE(GraphBuilder, ComputeFramework_ExecuteBatches);
 
diff --git a/Engine/Plugins/Runtime/MassGameplay/Source/MassLOD/Private/MassLODCollectorProcessor.cpp b/Engine/Plugins/Runtime/MassGameplay/Source/MassLOD/Private/MassLODCollectorProcessor.cpp
--- a/Engine/Plugins/Runtime/MassGameplay/Source/MassLOD/Private/MassLODCollectorProcessor.cpp
+++ b/Engine/Plugins/Runtime/MassGameplay/Source/MassLOD/Private/MassLODCollectorProcessor.cpp
@@ -67,14 +67,14 @@ template <bool bLocalViewersOnly>
 void UMassLODCollectorProcessor::ExecuteInternal(FMassEntityManager& EntityManager, FMassExecutionContext& Context)
 {
 	{
-		TRACE_CPUPROFILER_EVENT_SCOPE(TEXT("Close"));
+		TRACE_CPUPROFILER_EVENT_SCOPE(Close);
 		EntityQuery_VisibleRangeAndOnLOD.ForEachEntityChunk(EntityManager, Context, [this](FMassExecutionContext& Context) { CollectLODForChunk<bLocalViewersOnly>(Context); });
 		EntityQuery_VisibleRangeOnly.ForEachEntityChunk(EntityManager, Context, [this](FMassExecutionContext& Context) { CollectLODForChunk<bLocalViewersOnly>(Context); });
 		EntityQuery_OnLODOnly.ForEachEntityChunk(EntityManager, Context, [this](FMassExecutionContext& Context) { CollectLODForChunk<bLocalViewersOnly>(Context); });
 	}
 
 	{
-		TRACE_CPUPROFILER_EVENT_SCOPE(TEXT("Far"));
+		TRACE_CPUPROFILER_EVENT_SCOPE(Far);
 		EntityQuery_NotVisibleRangeAndOffLOD.ForEachEntityChunk(EntityManager, Context, [this](FMassExecutionContext& Context) { CollectLODForChunk<bLocalViewersOnly>(Context); });
 	}
 }
diff --git a/Engine/Plugins/Runtime/MassGameplay/Source/MassLOD/Private/MassSimulationLOD.cpp b/Engine/Plugins/Runtime/MassGameplay/Source/MassLOD/Private/MassSimulationLOD.cpp
--- a/Engine/Plugins/Runtime/MassGameplay/Source/MassLOD/Private/MassSimulationLOD.cpp
+++ b/Engine/Plugins/Runtime/MassGameplay/Source/MassLOD/Private/MassSimulationLOD.cpp
@@ -115,10 +115,10 @@ void UMassSimulationLODProcessor::ConfigureQueries()
 
 void UMassSimulationLODProcessor::Execute(FMassEntityManager& EntityManager, FMassExecutionContext& Context)
 {
-	TRACE_CPUPROFILER_EVENT_SCOPE(TEXT("SimulationLOD"))
+	TRACE_CPUPROFILER_EVENT_SCOPE(SimulationLOD)
 
 	{
-		TRACE_CPUPROFILER_EVENT_SCOPE(TEXT("PrepareExecution"));
+		TRACE_CPUPROFILER_EVENT_SCOPE(PrepareExecution);
 
 		const UMassLODSubsystem& LODSubsystem = Context.GetSubsystemChecked<UMassLODSubsystem>();
 		const TArray<FViewerInfo>& Viewers = LODSubsystem.GetViewers();
@@ -130,7 +130,7 @@ void UMassSimulationLODProcessor::Execute(FMassEntityManager& EntityManager, FMa
 	}
 
 	{
-		TRACE_CPUPROFILER_EVENT_SCOPE(TEXT("CalculateLOD"));
+		TRACE_CPUPROFILER_EVENT_SCOPE(CalculateLOD);
 		EntityQueryCalculateLOD.ForEachEntityChunk(EntityManager, Context, [](FMassExecutionContext& Context)
 		{
 			FMassSimulationLODSharedFragment& LODSharedFragment = Context.GetMutableSharedFragment<FMassSimulationLODSharedFragment>();
@@ -142,7 +142,7 @@ void UMassSimulationLODProcessor::Execute(FMassEntityManager& EntityManager, FMa
 
 	if (bDoAdjustmentFromCount)
 	{
-		TRACE_CPUPROFILER_EVENT_SCOPE(TEXT("AdjustDistancesAndLODFromCount"));
+		TRACE_CPUPROFILER_EVENT_SCOPE(AdjustDistancesAndLODFromCount);
 		EntityManager.ForEachSharedFragment<FMassSimulationLODSharedFragment>([](FMassSimulationLODSharedFragment& LODSharedFragment)
 		{
 			LODSharedFragment.bHasAdjustedDistancesFromCount = LODSharedFragment.LODCalculator.AdjustDistancesFromCount();
@@ -159,7 +159,7 @@ void UMassSimulationLODProcessor::Execute(FMassEntityManager& EntityManager, FMa
 
 	UWorld* World = EntityManager.GetWorld();
 	{
-		TRACE_CPUPROFILER_EVENT_SCOPE(TEXT("VariableTickRates"))
+		TRACE_CPUPROFILER_EVENT_SCOPE(VariableTickRates)
 		check(World);
 		const double Time = World->GetTimeSeconds();
 		EntityQueryVariableTick.ForEachEntityChunk(EntityManager, Context, [Time](FMassExecutionContext& Context)
@@ -173,7 +173,7 @@ void UMassSimulationLODProcessor::Execute(FMassEntityManager& EntityManager, FMa
 	}
 
 	{
-		TRACE_CPUPROFILER_EVENT_SCOPE(TEXT("SetLODTags"))
+		TRACE_CPUPROFILER_EVENT_SCOPE(SetLODTags)
 		check(World);
 		EntityQuerySetLODTag.ForEachEntityChunk(EntityManager, Context, [](FMassExecutionContext& Context)
 		{
@@ -195,7 +195,7 @@ void UMassSimulationLODProcessor::Execute(FMassEntityManager& EntityManager, FMa
 	// Optional debug display
 	if (UE::MassLOD::bDebugSimulationLOD)
 	{
-		TRACE_CPUPROFILER_EVENT_SCOPE(TEXT("DebugDisplayLOD"));
+		TRACE_CPUPROFILER_EVENT_SCOPE(DebugDisplayLOD);
 		EntityQuery.ForEachEntityChunk(EntityManager, Context, [World](FMassExecutionContext& Context)
 		{
 			FMassSimulationLODSharedFragment& LODSharedFragment = Context.GetMutableSharedFragment<FMassSimulationLODSharedFragment>();
diff --git a/Engine/Plugins/Runtime/Metasound/Source/MetasoundGenerator/Private/MetasoundGeneratorBuilder.cpp b/Engine/Plugins/Runtime/Metasound/Source/MetasoundGenerator/Private/MetasoundGeneratorBuilder.cpp
--- a/Engine/Plugins/Runtime/Metasound/Source/MetasoundGenerator/Private/MetasoundGeneratorBuilder.cpp
+++ b/Engine/Plugins/Runtime/Metasound/Source/MetasoundGenerator/Private/MetasoundGeneratorBuilder.cpp
@@ -167,7 +167,7 @@ namespace Metasound
 
 		TUniquePtr<Frontend::FGraphAnalyzer> BuildGraphAnalyzer(TMap<FGuid, FDataReferenceCollection>&& InInternalDataReferences, const FMetasoundEnvironment& InEnvironment, const FOperatorSettings& InOperatorSettings)
 		{
-			METASOUND_TRACE_CPUPROFILER_EVENT_SCOPE_TEXT(TEXT("GeneratorBuilder::BuildGraphAnalyzer"));
+			METASOUND_TRACE_CPUPROFILER_EVENT_SCOPE(GeneratorBuilder::BuildGraphAnalyzer);
 			using namespace Frontend;
 
 			const uint64 InstanceID = InEnvironment.GetValue<uint64>(SourceInterface::Environment::TransmitterID);
@@ -176,7 +176,7 @@ namespace Metasound
 
 		FInputVertexInterfaceData BuildGraphOperatorInputs(const FOperatorSettings& InOperatorSettings, FMetasoundGeneratorInitParams& InInitParams)
 		{
-			METASOUND_TRACE_CPUPROFILER_EVENT_SCOPE_TEXT(TEXT("GeneratorBuilder::BuildGraphOperatorInputs"));
+			METASOUND_TRACE_CPUPROFILER_EVENT_SCOPE(GeneratorBuilder::BuildGraphOperatorInputs);
 
 			using namespace Frontend;
 
@@ -248,7 +248,7 @@ namespace Metasound
 
 		FOperatorAndInputs BuildGraphOperator(const FOperatorSettings& InOperatorSettings, FMetasoundGeneratorInitParams& InInitParams, FBuildResults& OutBuildResults)
 		{
-			METASOUND_TRACE_CPUPROFILER_EVENT_SCOPE_TEXT(TEXT("GeneratorBuilder::BuildGraphOperator"));
+			METASOUND_TRACE_CPUPROFILER_EVENT_SCOPE(GeneratorBuilder::BuildGraphOperator);
 
 			// Create inputs to graph operator
 			FOperatorAndInputs OpAndInputs;
@@ -264,7 +264,7 @@ namespace Metasound
 
 		FOperatorAndInputs BuildDynamicGraphOperator(const FOperatorSettings& InOperatorSettings, FMetasoundDynamicGraphGeneratorInitParams& InInitParams, const DynamicGraph::FDynamicOperatorUpdateCallbacks& InCallbacks, FBuildResults& OutBuildResults)
 		{
-			METASOUND_TRACE_CPUPROFILER_EVENT_SCOPE_TEXT(TEXT("GeneratorBuilder::BuildDynamicGraphOperator"));
+			METASOUND_TRACE_CPUPROFILER_EVENT_SCOPE(GeneratorBuilder::BuildDynamicGraphOperator);
 
 			FOperatorAndInputs OpAndInputs;
 			OpAndInputs.Inputs = BuildGraphOperatorInputs(InOperatorSettings, InInitParams);
@@ -294,7 +294,7 @@ namespace Metasound
 			using namespace Frontend;
 			using namespace MetasoundGeneratorPrivate;
 
-			METASOUND_TRACE_CPUPROFILER_EVENT_SCOPE_TEXT(TEXT("GeneratorBuilder::BuildGeneratorData"));
+			METASOUND_TRACE_CPUPROFILER_EVENT_SCOPE(GeneratorBuilder::BuildGeneratorData);
 
 			checkf(InGraphOperatorAndInputs.Operator.IsValid(), TEXT("Graph operator must be a valid object"))
 
@@ -341,7 +341,7 @@ namespace Metasound
 
 		void ApplyAudioParameters(const FOperatorSettings& InOperatorSettings, TArray<FAudioParameter>&& InParameters, FInputVertexInterfaceData& InInterface)
 		{
-			METASOUND_TRACE_CPUPROFILER_EVENT_SCOPE("GeneratorBuilder::ApplyAudioParameters");
+			METASOUND_TRACE_CPUPROFILER_EVENT_SCOPE(GeneratorBuilder::ApplyAudioParameters);
 			Frontend::IDataTypeRegistry& DataTypeRegistry = Frontend::IDataTypeRegistry::Get();
 			for (FAudioParameter& Parameter : InParameters)
 			{
diff --git a/Engine/Plugins/VirtualProduction/RemoteControl/Source/WebRemoteControl/Private/RemoteControlWebSocketServer.cpp b/Engine/Plugins/VirtualProduction/RemoteControl/Source/WebRemoteControl/Private/RemoteControlWebSocketServer.cpp
--- a/Engine/Plugins/VirtualProduction/RemoteControl/Source/WebRemoteControl/Private/RemoteControlWebSocketServer.cpp
+++ b/Engine/Plugins/VirtualProduction/RemoteControl/Source/WebRemoteControl/Private/RemoteControlWebSocketServer.cpp
@@ -261,7 +261,7 @@ void FRCWebSocketServer::ReceivedRawPacket(void* Data, int32 Size, FGuid ClientI
 		{
 			case ERCWebSocketCompressionMode::ZLIB:
 			{
-				TRACE_CPUPROFILER_EVENT_SCOPE(TEXT("DecompressWebSocketData"));
+				TRACE_CPUPROFILER_EVENT_SCOPE(DecompressWebSocketData);
 				
 				// Read the header containing the message's uncompressed size
 				const TArrayView<uint8> CompressedData = MakeArrayView(static_cast<uint8*>(Data), Size);
@@ -347,7 +347,7 @@ void FRCWebSocketServer::SendOnConnection(FWebSocketConnection& Connection, cons
 	{
 	case ERCWebSocketCompressionMode::ZLIB:
 		{
-			TRACE_CPUPROFILER_EVENT_SCOPE(TEXT("CompressWebSocketData"));
+			TRACE_CPUPROFILER_EVENT_SCOPE(CompressWebSocketData);
 
 			int32 CompressedSize = InUTF8Payload.Num();
 
