fix instance name for traces of metasound patches

diff --git a/Engine/Plugins/Runtime/Metasound/Source/MetasoundFrontend/Private/MetasoundFrontendRegistries.cpp b/Engine/Plugins/Runtime/Metasound/Source/MetasoundFrontend/Private/MetasoundFrontendRegistries.cpp
--- a/Engine/Plugins/Runtime/Metasound/Source/MetasoundFrontend/Private/MetasoundFrontendRegistries.cpp
+++ b/Engine/Plugins/Runtime/Metasound/Source/MetasoundFrontend/Private/MetasoundFrontendRegistries.cpp
@@ -126,8 +126,8 @@ namespace Metasound
 				};
 
 			public:
-				FGraphNode(FNodeInitData InNodeInitData, TSharedRef<const IGraph> InGraphToWrap)
-				: NodeInitData(MoveTemp(InNodeInitData))
+				FGraphNode(const FNodeInitData& InNodeInitData, TSharedRef<const IGraph> InGraphToWrap)
+				: InstanceID(InNodeInitData.InstanceID)
 				, Factory(MakeShared<FGraphOperatorFactoryAdapter>(*InGraphToWrap))
 				, Graph(MoveTemp(InGraphToWrap))
 				{
@@ -135,12 +135,14 @@ namespace Metasound
 
 				virtual const FName& GetInstanceName() const override
 				{
-					return NodeInitData.InstanceName;
+					// Use the instance name of underlying graph because it refers
+					// to the actual asset name.
+					return Graph->GetInstanceName();
 				}
 
 				virtual const FGuid& GetInstanceID() const override
 				{
-					return NodeInitData.InstanceID;
+					return InstanceID;
 				}
 
 				virtual const FNodeClassMetadata& GetMetadata() const override
@@ -172,7 +174,7 @@ namespace Metasound
 
 			private:
 
-				FNodeInitData NodeInitData;
+				FGuid InstanceID;
 				TSharedRef<FGraphOperatorFactoryAdapter> Factory;
 				TSharedRef<const IGraph> Graph;
 			};
