Limit the amount of log spam:
LogMetaSound: Warning: Underrun detected in audio bus writer node

diff --git a/Engine/Plugins/Runtime/Metasound/Source/MetasoundEngine/Private/MetasoundAudioBusWriterNode.cpp b/Engine/Plugins/Runtime/Metasound/Source/MetasoundEngine/Private/MetasoundAudioBusWriterNode.cpp
--- a/Engine/Plugins/Runtime/Metasound/Source/MetasoundEngine/Private/MetasoundAudioBusWriterNode.cpp
+++ b/Engine/Plugins/Runtime/Metasound/Source/MetasoundEngine/Private/MetasoundAudioBusWriterNode.cpp
@@ -289,10 +289,11 @@ namespace Metasound
 			}
 
 			// Pushes the interleaved data to the audio bus
-			int32 SamplesPushed = AudioBusPatchInput.PushAudio(InterleavedBuffer.GetData(), InterleavedBuffer.Num());
-			if (SamplesPushed < InterleavedBuffer.Num())
+			const int32 SamplesPushed = AudioBusPatchInput.PushAudio(InterleavedBuffer.GetData(), InterleavedBuffer.Num());
+			if (SamplesPushed < InterleavedBuffer.Num() && !bWasUnderrunReported)
 			{
 				UE_LOG(LogMetaSound, Warning, TEXT("Underrun detected in audio bus writer node."));
+				bWasUnderrunReported = true;
 			}
 		}
 
@@ -312,6 +313,7 @@ namespace Metasound
 		Audio::FPatchInput AudioBusPatchInput;
 		uint32 AudioBusChannels = INDEX_NONE;
 		uint32 AudioBusId = 0;
+		bool bWasUnderrunReported = false;
 		int32 BlockSizeFrames = 0;
 		enum class EConnectionState : uint8
 		{
