[Metasound] Use FindClassOutputWithName instead of GetOutputNodeWithName to get output metadata in FMetasoundGeneratorHandle::WatchOutputs

This results in a 98% improvement in CPU time.

diff --git a/Engine/Plugins/Runtime/Metasound/Source/MetasoundEngine/Private/MetasoundGeneratorHandle.cpp b/Engine/Plugins/Runtime/Metasound/Source/MetasoundEngine/Private/MetasoundGeneratorHandle.cpp
--- a/Engine/Plugins/Runtime/Metasound/Source/MetasoundEngine/Private/MetasoundGeneratorHandle.cpp
+++ b/Engine/Plugins/Runtime/Metasound/Source/MetasoundEngine/Private/MetasoundGeneratorHandle.cpp
@@ -163,23 +163,20 @@ namespace Metasound
 			}
 		
 			// Find the node id and type name
-			const Frontend::FNodeHandle Node =
-				Source->GetRootGraphHandle()->GetOutputNodeWithName(AnalyzerAddress.OutputName);
-
-			if (!Node->IsValid())
-			{
-				return false;
-			}
-
-			AnalyzerAddress.NodeID = Node->GetID();
-
-			// We expect output nodes to have only one output
-			if (!ensure(Node->GetNumOutputs() == 1))
+			const FMetasoundFrontendClassOutput* OutputPtr =
+				Source->GetConstDocument().RootGraph.Interface.Outputs.FindByPredicate(
+					[&AnalyzerAddress](const FMetasoundFrontendClassOutput& Output)
+					{
+						return Output.Name == AnalyzerAddress.OutputName;
+					});
+
+			if (nullptr == OutputPtr)
 			{
 				return false;
 			}
-
-			AnalyzerAddress.DataType = Node->GetOutputs()[0]->GetDataType();
+			
+			AnalyzerAddress.NodeID = OutputPtr->NodeID;
+			AnalyzerAddress.DataType = OutputPtr->TypeName;
 		}
 		
 		// If no analyzer name was provided, try to find a passthrough analyzer
