[Metasound] Use a map for faster generator handle output lookup.

This results in a ~70% CPU time improvement.

diff --git a/Engine/Plugins/Runtime/Metasound/Source/MetasoundEngine/Private/MetasoundGeneratorHandle.cpp b/Engine/Plugins/Runtime/Metasound/Source/MetasoundEngine/Private/MetasoundGeneratorHandle.cpp
--- a/Engine/Plugins/Runtime/Metasound/Source/MetasoundEngine/Private/MetasoundGeneratorHandle.cpp
+++ b/Engine/Plugins/Runtime/Metasound/Source/MetasoundEngine/Private/MetasoundGeneratorHandle.cpp
@@ -223,15 +223,16 @@ namespace Metasound
 
 		while (TOptional<FOutputPayload> ChangedOutput = ChangedOutputs.Dequeue())
 		{
-			if (const FOutputWatcher* OutputListener = OutputWatchers.FindByPredicate(
-				[&ChangedOutput](const FOutputWatcher& ExistingListener)
-				{
-					return ChangedOutput->AnalyzerName == ExistingListener.AnalyzerAddress.AnalyzerName
-					&& ChangedOutput->OutputName == ExistingListener.AnalyzerAddress.OutputName
-					&& ChangedOutput->OutputValue.Name == ExistingListener.AnalyzerAddress.AnalyzerMemberName;
-				}))
+			const FOutputWatcherKey WatcherKey
+			{
+				ChangedOutput->OutputName,
+				ChangedOutput->AnalyzerName,
+				ChangedOutput->OutputValue.Name
+			};
+
+			if (const FOutputWatcher* Watcher = OutputWatchers.Find(WatcherKey))
 			{
-				OutputListener->OnOutputValueChanged.Broadcast(ChangedOutput->OutputName, ChangedOutput->OutputValue);
+				Watcher->OnOutputValueChanged.Broadcast(ChangedOutput->OutputName, ChangedOutput->OutputValue);
 			}
 
 			++NumDequeued;
@@ -411,9 +412,9 @@ namespace Metasound
 		{
 			// For each watcher, make sure the generator has a corresponding analyzer
 			// (will fail gracefully on duplicates or non-existent outputs)
-			for (const FOutputWatcher& Watcher : OutputWatchers)
+			for (const auto& Watcher : OutputWatchers)
 			{
-				PinnedGenerator->AddOutputVertexAnalyzer(Watcher.AnalyzerAddress);
+				PinnedGenerator->AddOutputVertexAnalyzer(Watcher.Value.AnalyzerAddress);
 			}
 		}
 	}
@@ -428,20 +429,21 @@ namespace Metasound
 		check(IsInGameThread()); // modifying watchers isn't thread-safe
 		
 		// If we already have a watcher for this output, just add the delegate to that one
-		if (FOutputWatcher* Watcher = OutputWatchers.FindByPredicate(
-			[&AnalyzerAddress](const FOutputWatcher& ExistingListener)
-			{
-				return AnalyzerAddress.OutputName == ExistingListener.AnalyzerAddress.OutputName
-					&& AnalyzerAddress.AnalyzerName == ExistingListener.AnalyzerAddress.AnalyzerName
-					&& AnalyzerAddress.AnalyzerMemberName == ExistingListener.AnalyzerAddress.AnalyzerMemberName;
-			}))
+		const FOutputWatcherKey WatcherKey
+		{
+			AnalyzerAddress.OutputName,
+			AnalyzerAddress.AnalyzerName,
+			AnalyzerAddress.AnalyzerMemberName
+		};
+
+		if (FOutputWatcher* Watcher = OutputWatchers.Find(WatcherKey))
 		{
 			Watcher->OnOutputValueChanged.AddUnique(OnOutputValueChanged);
 		}
 		// Otherwise add a new watcher
 		else
 		{
-			OutputWatchers.Emplace(AnalyzerAddress, OnOutputValueChanged);
+			OutputWatchers.Emplace(WatcherKey, { AnalyzerAddress, OnOutputValueChanged });
 		}
 	}
 
diff --git a/Engine/Plugins/Runtime/Metasound/Source/MetasoundEngine/Public/MetasoundGeneratorHandle.h b/Engine/Plugins/Runtime/Metasound/Source/MetasoundEngine/Public/MetasoundGeneratorHandle.h
--- a/Engine/Plugins/Runtime/Metasound/Source/MetasoundEngine/Public/MetasoundGeneratorHandle.h
+++ b/Engine/Plugins/Runtime/Metasound/Source/MetasoundEngine/Public/MetasoundGeneratorHandle.h
@@ -189,6 +189,32 @@ namespace Metasound
 		};
 		static TMap<FName, FPassthroughAnalyzerInfo> PassthroughAnalyzers;
 
+		struct FOutputWatcherKey
+		{
+			FName OutputName;
+			FName AnalyzerName;
+			FName AnalyzerMemberName;
+
+			bool operator==(const FOutputWatcherKey& Other) const
+			{
+				return OutputName == Other.OutputName
+				&& AnalyzerName == Other.AnalyzerName
+				&& AnalyzerMemberName == Other.AnalyzerMemberName;
+			}
+
+			bool operator!=(const FOutputWatcherKey& Other) const
+			{
+				return !(*this == Other);
+			}
+		};
+
+		friend uint32 GetTypeHash(FOutputWatcherKey Key)
+		{
+			return HashCombineFast(
+				HashCombineFast(GetTypeHash(Key.OutputName), GetTypeHash(Key.AnalyzerName)),
+				GetTypeHash(Key.AnalyzerMemberName));
+		}
+
 		/**
 		 * Info about an output being watched by one or more listeners
 		 */
@@ -205,8 +231,8 @@ namespace Metasound
 				OnOutputValueChanged.AddUnique(InOnOutputValueChanged);
 			}
 		};
-
-		TArray<FOutputWatcher> OutputWatchers;
+		
+		TMap<FOutputWatcherKey, FOutputWatcher> OutputWatchers;
 
 		struct FOutputPayload
 		{
