Fix deadlock when async building metasounds
Ensure MetaSound graphs are only registered on the game thread.

diff --git a/Engine/Plugins/Runtime/Metasound/Source/MetasoundEngine/Private/MetasoundSource.cpp b/Engine/Plugins/Runtime/Metasound/Source/MetasoundEngine/Private/MetasoundSource.cpp
--- a/Engine/Plugins/Runtime/Metasound/Source/MetasoundEngine/Private/MetasoundSource.cpp
+++ b/Engine/Plugins/Runtime/Metasound/Source/MetasoundEngine/Private/MetasoundSource.cpp
@@ -5,6 +5,7 @@
 #include "Algo/Transform.h"
 #include "AssetRegistry/AssetRegistryModule.h"
 #include "AudioDeviceManager.h"
+#include "Containers/Ticker.h"
 #include "IAudioParameterInterfaceRegistry.h"
 #include "Interfaces/MetasoundOutputFormatInterfaces.h"
 #include "Interfaces/MetasoundFrontendSourceInterface.h"
@@ -459,202 +460,83 @@
 
 void UMetaSoundSource::InitParameters(TArray<FAudioParameter>& ParametersToInit, FName InFeatureName)
 {
-	using namespace Metasound::SourcePrivate;
-	using namespace Metasound::Frontend;
-	using FRuntimeInput = FMetasoundAssetBase::FRuntimeInput;
-
 	METASOUND_LLM_SCOPE;
 	METASOUND_TRACE_CPUPROFILER_EVENT_SCOPE(UMetaSoundSource::InitParameters);
 
-	const FRuntimeData& RuntimeData = GetRuntimeData();
+	if (bIsBuilderActive)
+	{
+		InitParametersInternal(CreateRuntimeInputMap(), ParametersToInit, InFeatureName);
+	}
+	else
+	{
+		const bool bIsRuntimeInputDataValid = RuntimeInputData.bIsValid.load();
+		if (bIsRuntimeInputDataValid)
+		{
+			InitParametersInternal(RuntimeInputData.InputMap, ParametersToInit, InFeatureName);
+		}
+		else
+		{
+			// The runtime input data should have been cached, but is not so we use
+			// a fallback method. If this is occurring, then callers need to ensure
+			// that InitResources has been called before this method executes or else
+			// suffer the consequences of incurring significant performance losses 
+			// each time a parameter is set on the MetaSound. 
+			UE_LOG(LogMetaSound, Warning, TEXT("Initializing parameters on uninitialized UMetaSoundSource %s will result in slower performance. UMetaSoundSource::InitResources should finish executing on the game thread before attempting to call UMetaSoundSource::InitParameters(...)"), *GetOwningAssetName());
+			InitParametersInternal(CreateRuntimeInputMap(), ParametersToInit, InFeatureName);
+		}
+	}
+}
 
-	// To initialize parameters, we need the PublicInputMap which lives on FRuntimeData. 
-	// If the runtime data has not valid, it can be created via a call to RegisterGraphWithFrontend(...) 
-	// which subsequently pupulates the runtime data.
-	if (!RuntimeData.IsValid())
+void UMetaSoundSource::InitResources()
+{
+	using namespace Metasound::Frontend;
+	using namespace Metasound::SourcePrivate;
+
+	METASOUND_LLM_SCOPE;
+	METASOUND_TRACE_CPUPROFILER_EVENT_SCOPE(UMetaSoundSource::InitResources); 
+
+	if (IsInGameThread())
 	{
-		// If a InitParameters is called before InitResources, the graph will not 
-		// yet be registered. RegisterGraphWithFrontend is called here to cover that 
-		// scenario. 
 		RegisterGraphWithFrontend(GetInitRegistrationOptions());
-		check(RuntimeData.IsValid());
 	}
+	else
+	{
+		const bool bIsInGCSafeThread = IsInAudioThread() || IsInAsyncLoadingThread(); // Audio Thread is safe from GC, so we can safely construct the TWeakObjectPtr<> to this.
+		if (!bIsInGCSafeThread)
+		{
+			UE_LOG(LogMetaSound, Warning, TEXT("Attempt to call UMetaSoundSource::InitResources() on %s from thread which may not provide garbage collection safety of the UMetaSoundSource"), *GetOwningAssetName());
+		}
 
-	IDataTypeRegistry& DataTypeRegistry = IDataTypeRegistry::Get();
-	const Metasound::TSortedVertexNameMap<FRuntimeInput>& PublicInputMap = RuntimeData.PublicInputMap;
-
-	// Removes values that are not explicitly defined by the ParamType
-	auto Sanitize = [&PublicInputMap](FAudioParameter& Parameter)
-	{
-		switch (Parameter.ParamType)
-		{
-			case EAudioParameterType::Trigger:
+		ExecuteOnGameThread(
+			UE_SOURCE_LOCATION, 
+			[MetaSoundSourcePtr=TWeakObjectPtr<UMetaSoundSource>(this)]()
 			{
-				Parameter = FAudioParameter(Parameter.ParamName, EAudioParameterType::Trigger);
-			}
-			break;
-			
-			case EAudioParameterType::Boolean:
-			{
-				Parameter = FAudioParameter(Parameter.ParamName, Parameter.BoolParam);
-			}
-			break;
-
-			case EAudioParameterType::BooleanArray:
-			{
-				TArray<bool> TempArray = Parameter.ArrayBoolParam;
-				Parameter = FAudioParameter(Parameter.ParamName, MoveTemp(TempArray));
-			}
-			break;
-
-			case EAudioParameterType::Float:
-			{
-				Parameter = FAudioParameter(Parameter.ParamName, Parameter.FloatParam);
-			}
-			break;
-
-			case EAudioParameterType::FloatArray:
-			{
-				TArray<float> TempArray = Parameter.ArrayFloatParam;
-				Parameter = FAudioParameter(Parameter.ParamName, MoveTemp(TempArray));
-			}
-			break;
-
-			case EAudioParameterType::Integer:
-			{
-				Parameter = FAudioParameter(Parameter.ParamName, Parameter.IntParam);
-			}
-			break;
-
-			case EAudioParameterType::IntegerArray:
-			{
-				TArray<int32> TempArray = Parameter.ArrayIntParam;
-				Parameter = FAudioParameter(Parameter.ParamName, MoveTemp(TempArray));
-			}
-			break;
-
-			case EAudioParameterType::Object:
-			{
-				Parameter = FAudioParameter(Parameter.ParamName, Parameter.ObjectParam);
-			}
-			break;
-
-			case EAudioParameterType::ObjectArray:
-			{
-				TArray<UObject*> TempArray = Parameter.ArrayObjectParam;
-				Parameter = FAudioParameter(Parameter.ParamName, MoveTemp(TempArray));
-			}
-			break;
-
-
-			case EAudioParameterType::String:
-			{
-				Parameter = FAudioParameter(Parameter.ParamName, Parameter.StringParam);
-			}
-			break;
-
-			case EAudioParameterType::StringArray:
-			{
-				TArray<FString> TempArray = Parameter.ArrayStringParam;
-				Parameter = FAudioParameter(Parameter.ParamName, MoveTemp(TempArray));
-			}
-			break;
-
-			case EAudioParameterType::None:
-			default:
-			break;
-		}
-	};
-
-	auto ConstructProxies = [&DataTypeRegistry](FAudioParameter& OutParamToInit, FName VertexTypeName)
-	{
-		using namespace Metasound;
-
-		switch (OutParamToInit.ParamType)
-		{
-			case EAudioParameterType::Object:
-			{
-				TSharedPtr<Audio::IProxyData> ProxyPtr = DataTypeRegistry.CreateProxyFromUObject(VertexTypeName, OutParamToInit.ObjectParam);
-				OutParamToInit.ObjectProxies.Emplace(MoveTemp(ProxyPtr));
-
-				// Null out param as it is no longer needed (nor desired to be accessed once passed to the Audio Thread)
-				OutParamToInit.ObjectParam = nullptr;
-			}
-			break;
-
-			case EAudioParameterType::ObjectArray:
-			{
-				const FName ElementTypeName = CreateElementTypeNameFromArrayTypeName(VertexTypeName);
-				for (TObjectPtr<UObject>& Object : OutParamToInit.ArrayObjectParam)
+				if (UMetaSoundSource* Source = MetaSoundSourcePtr.Get())
 				{
-					TSharedPtr<Audio::IProxyData> ProxyPtr = DataTypeRegistry.CreateProxyFromUObject(ElementTypeName, Object);
-					OutParamToInit.ObjectProxies.Emplace(MoveTemp(ProxyPtr));
+					Source->InitResources();
 				}
-				// Reset param array as it is no longer needed (nor desired to be accessed once passed to the Audio Thread).
-				// All object manipulation hereafter should be done via proxies
-				OutParamToInit.ArrayObjectParam.Reset();
 			}
-			break;
-
-			default:
-				break;
-		}
-	};
-
-
-	for (int32 i = ParametersToInit.Num() - 1; i >= 0; --i)
-	{
-		bool bIsParameterValid = false;
-
-		FAudioParameter& Parameter = ParametersToInit[i];
-		if (const FRuntimeInput* Input = PublicInputMap.Find(Parameter.ParamName))
-		{
-			if (IsParameterValid(Parameter, Input->TypeName, DataTypeRegistry))
-			{
-				Sanitize(Parameter);
-				ConstructProxies(Parameter, Input->TypeName);
-				bIsParameterValid = true;
-			}
-		}
-
-		if (!bIsParameterValid)
-		{
-			constexpr bool bAllowShrinking = false;
-			ParametersToInit.RemoveAtSwap(i, 1, bAllowShrinking);
-
-#if !NO_LOGGING
-			if (::Metasound::MetaSoundParameterEnableWarningOnIgnoredParameterCVar)
-			{
-				const FString AssetName = GetName();
-				UE_LOG(LogMetaSound, Warning, TEXT("Failed to set parameter '%s' in asset '%s': No name specified, no transmittable input found, or type mismatch."), *Parameter.ParamName.ToString(), *AssetName);
-			}
-#endif // !NO_LOGGING
-		}
+		);
 	}
 }
 
-void UMetaSoundSource::InitResources()
+void UMetaSoundSource::RegisterGraphWithFrontend(Metasound::Frontend::FMetaSoundAssetRegistrationOptions InRegistrationOptions)
 {
-	using namespace Metasound::Frontend;
-	using namespace Metasound::SourcePrivate;
+	check(IsInGameThread());
 
-	METASOUND_LLM_SCOPE;
-	METASOUND_TRACE_CPUPROFILER_EVENT_SCOPE(UMetaSoundSource::InitResources);
-
-	RegisterGraphWithFrontend(GetInitRegistrationOptions());
+	const bool bIsRuntimeInputDataValid = RuntimeInputData.bIsValid.load();
+	if (!bIsRuntimeInputDataValid)
+	{
+		CacheRuntimeInputData();
+	}
+	FMetasoundAssetBase::RegisterGraphWithFrontend(InRegistrationOptions);
 }
 
 bool UMetaSoundSource::IsPlayable() const
 {
-	// todo: cache off whether this metasound is buildable to an operator.
 	return true;
 }
 
-bool UMetaSoundSource::SupportsSubtitles() const
-{
-	return Super::SupportsSubtitles();
-}
-
 float UMetaSoundSource::GetDuration() const
 {
 	// This is an unfortunate function required by logic in determining what sounds can be potentially
@@ -713,7 +595,7 @@
 	SampleRate = InSettings.GetSampleRate();
 
 	FMetasoundEnvironment Environment = CreateEnvironment(InParams);
-	FParameterRouter& Router = GetParameterRouter();
+	FParameterRouter& Router = UMetaSoundSource::GetParameterRouter();
 	TSharedPtr<TSpscQueue<FMetaSoundParameterTransmitter::FParameter>> DataChannel = Router.FindOrCreateDataChannelForReader(InParams.AudioDeviceID, InParams.InstanceID);
 
 	FOperatorBuilderSettings BuilderSettings = FOperatorBuilderSettings::GetDefaultSettings();
@@ -806,15 +688,17 @@
 	using namespace Metasound::Frontend;
 	using namespace Metasound::Engine;
 
-	const FRuntimeData& RuntimeData = GetRuntimeData();
-	if(!RuntimeData.IsValid())
+	
+	if(!RuntimeInputData.bIsValid.load())
 	{
-		UE_LOG(LogMetaSound, Warning, TEXT("Default parameters may be incorrect. Accessing invalid runtime data on MetaSound %s"), *GetOwningAssetName());
+		UE_LOG(LogMetaSound, Warning, TEXT("Default parameters will be ommitted. Accessing invalid runtime data on MetaSound %s. Ensure that UMetaSoundSource::InitResources() is executed on the game thread before calling UMetaSoundSource::GetAllDefaultParameters(...)"), *GetOwningAssetName());
+		return false;
 	}
 
-	for(const TPair<FVertexName, FMetasoundAssetBase::FRuntimeInput>& Pair : GetRuntimeData().PublicInputMap)
+
+	for(const TPair<FVertexName, FRuntimeInput>& Pair : RuntimeInputData.InputMap)
 	{
-		const FMetasoundAssetBase::FRuntimeInput& Input = Pair.Value;
+		const FRuntimeInput& Input = Pair.Value;
 		FAudioParameter Params;
 		Params.ParamName = Input.Name;
 		Params.TypeName = Input.TypeName;
@@ -914,6 +798,168 @@
 	return true;
 }
 
+void UMetaSoundSource::InitParametersInternal(const Metasound::TSortedVertexNameMap<FRuntimeInput>& InInputMap, TArray<FAudioParameter>& ParametersToInit, FName InFeatureName) const
+{
+	using namespace Metasound;
+	using namespace Metasound::Frontend;
+
+	METASOUND_LLM_SCOPE;
+	METASOUND_TRACE_CPUPROFILER_EVENT_SCOPE(UMetaSoundSource::InitParametersInternal);
+
+	checkf(IsInGameThread() || IsInAudioThread(), TEXT("Parameter initialization must happen on the GameThread or AudioThread to allow for safe creation of UObject proxies"));
+
+	IDataTypeRegistry& DataTypeRegistry = IDataTypeRegistry::Get();
+
+	// Removes values that are not explicitly defined by the ParamType
+	auto Sanitize = [](FAudioParameter& Parameter)
+	{
+		switch (Parameter.ParamType)
+		{
+			case EAudioParameterType::Trigger:
+			{
+				Parameter = FAudioParameter(Parameter.ParamName, EAudioParameterType::Trigger);
+			}
+			break;
+			
+			case EAudioParameterType::Boolean:
+			{
+				Parameter = FAudioParameter(Parameter.ParamName, Parameter.BoolParam);
+			}
+			break;
+
+			case EAudioParameterType::BooleanArray:
+			{
+				TArray<bool> TempArray = Parameter.ArrayBoolParam;
+				Parameter = FAudioParameter(Parameter.ParamName, MoveTemp(TempArray));
+			}
+			break;
+
+			case EAudioParameterType::Float:
+			{
+				Parameter = FAudioParameter(Parameter.ParamName, Parameter.FloatParam);
+			}
+			break;
+
+			case EAudioParameterType::FloatArray:
+			{
+				TArray<float> TempArray = Parameter.ArrayFloatParam;
+				Parameter = FAudioParameter(Parameter.ParamName, MoveTemp(TempArray));
+			}
+			break;
+
+			case EAudioParameterType::Integer:
+			{
+				Parameter = FAudioParameter(Parameter.ParamName, Parameter.IntParam);
+			}
+			break;
+
+			case EAudioParameterType::IntegerArray:
+			{
+				TArray<int32> TempArray = Parameter.ArrayIntParam;
+				Parameter = FAudioParameter(Parameter.ParamName, MoveTemp(TempArray));
+			}
+			break;
+
+			case EAudioParameterType::Object:
+			{
+				Parameter = FAudioParameter(Parameter.ParamName, Parameter.ObjectParam);
+			}
+			break;
+
+			case EAudioParameterType::ObjectArray:
+			{
+				TArray<UObject*> TempArray = Parameter.ArrayObjectParam;
+				Parameter = FAudioParameter(Parameter.ParamName, MoveTemp(TempArray));
+			}
+			break;
+
+
+			case EAudioParameterType::String:
+			{
+				Parameter = FAudioParameter(Parameter.ParamName, Parameter.StringParam);
+			}
+			break;
+
+			case EAudioParameterType::StringArray:
+			{
+				TArray<FString> TempArray = Parameter.ArrayStringParam;
+				Parameter = FAudioParameter(Parameter.ParamName, MoveTemp(TempArray));
+			}
+			break;
+
+			case EAudioParameterType::None:
+			default:
+			break;
+		}
+	};
+
+	auto ConstructProxies = [&DataTypeRegistry](FAudioParameter& OutParamToInit, FName VertexTypeName)
+	{
+		using namespace Metasound;
+
+		switch (OutParamToInit.ParamType)
+		{
+			case EAudioParameterType::Object:
+			{
+				TSharedPtr<Audio::IProxyData> ProxyPtr = DataTypeRegistry.CreateProxyFromUObject(VertexTypeName, OutParamToInit.ObjectParam);
+				OutParamToInit.ObjectProxies.Emplace(MoveTemp(ProxyPtr));
+
+				// Null out param as it is no longer needed (nor desired to be accessed once passed to the Audio Thread)
+				OutParamToInit.ObjectParam = nullptr;
+			}
+			break;
+
+			case EAudioParameterType::ObjectArray:
+			{
+				const FName ElementTypeName = CreateElementTypeNameFromArrayTypeName(VertexTypeName);
+				for (TObjectPtr<UObject>& Object : OutParamToInit.ArrayObjectParam)
+				{
+					TSharedPtr<Audio::IProxyData> ProxyPtr = DataTypeRegistry.CreateProxyFromUObject(ElementTypeName, Object);
+					OutParamToInit.ObjectProxies.Emplace(MoveTemp(ProxyPtr));
+				}
+				// Reset param array as it is no longer needed (nor desired to be accessed once passed to the Audio Thread).
+				// All object manipulation hereafter should be done via proxies
+				OutParamToInit.ArrayObjectParam.Reset();
+			}
+			break;
+
+			default:
+				break;
+		}
+	};
+
+
+	for (int32 i = ParametersToInit.Num() - 1; i >= 0; --i)
+	{
+		bool bIsParameterValid = false;
+
+		FAudioParameter& Parameter = ParametersToInit[i];
+		if (const FRuntimeInput* Input = InInputMap.Find(Parameter.ParamName))
+		{
+			if (IsParameterValidInternal(Parameter, Input->TypeName, DataTypeRegistry))
+			{
+				Sanitize(Parameter);
+				ConstructProxies(Parameter, Input->TypeName);
+				bIsParameterValid = true;
+			}
+		}
+
+		if (!bIsParameterValid)
+		{
+			constexpr bool bAllowShrinking = false;
+			ParametersToInit.RemoveAtSwap(i, 1, bAllowShrinking);
+
+#if !NO_LOGGING
+			if (::Metasound::MetaSoundParameterEnableWarningOnIgnoredParameterCVar)
+			{
+				const FString AssetName = GetName();
+				UE_LOG(LogMetaSound, Warning, TEXT("Failed to set parameter '%s' in asset '%s': No name specified, no transmittable input found, or type mismatch."), *Parameter.ParamName.ToString(), *AssetName);
+			}
+#endif // !NO_LOGGING
+		}
+	}
+}
+
 bool UMetaSoundSource::IsParameterValid(const FAudioParameter& InParameter) const
 {
 	const TArray<FMetasoundFrontendClassInput>& Inputs = GetDocumentChecked().RootGraph.Interface.Inputs;
@@ -924,7 +970,7 @@
 	
 	if (Vertex)
 	{
-		return IsParameterValid(InParameter, Vertex->TypeName, Metasound::Frontend::IDataTypeRegistry::Get());
+		return IsParameterValidInternal(InParameter, Vertex->TypeName, Metasound::Frontend::IDataTypeRegistry::Get());
 	}
 	else
 	{
@@ -932,7 +978,7 @@
 	}
 }
 
-bool UMetaSoundSource::IsParameterValid(const FAudioParameter& InParameter, const FName& InTypeName, Metasound::Frontend::IDataTypeRegistry& InDataTypeRegistry) const
+bool UMetaSoundSource::IsParameterValidInternal(const FAudioParameter& InParameter, const FName& InTypeName, Metasound::Frontend::IDataTypeRegistry& InDataTypeRegistry) const
 {
 	using namespace Metasound;
 	using namespace Metasound::Frontend;
@@ -1071,29 +1117,52 @@
 
 	METASOUND_LLM_SCOPE;
 
-	const FRuntimeData& RuntimeData = GetRuntimeData();
-	if(!RuntimeData.IsValid())
+	auto CreateParameterTransmitterInternal = [this](const TSortedVertexNameMap<FRuntimeInput>& InInputMap, Audio::FParameterTransmitterInitParams& InParams)
 	{
-		UE_LOG(LogMetaSound, Warning, TEXT("Parameter Transmitter may not work. Accessing invalid runtime data on MetaSound %s"), *GetOwningAssetName());
+		// Build list of parameters that can be set at runtime.
+		TArray<FName> ValidParameters;
+		for (const TPair<FVertexName, FRuntimeInput>& Pair : InInputMap)
+		{
+			if (Pair.Value.bIsTransmittable && (Pair.Value.AccessType == EMetasoundFrontendVertexAccessType::Reference))
+			{
+				ValidParameters.Add(Pair.Value.Name);
+			}
+		}
+
+		FParameterRouter& Router = UMetaSoundSource::GetParameterRouter();
+		TSharedPtr<TSpscQueue<FMetaSoundParameterTransmitter::FParameter>> DataChannel = Router.FindOrCreateDataChannelForWriter(InParams.AudioDeviceID, InParams.InstanceID);
+
+		Metasound::FMetaSoundParameterTransmitter::FInitParams InitParams
+		(
+			GetOperatorSettings(InParams.SampleRate), 
+			InParams.InstanceID, 
+			MoveTemp(InParams.DefaultParams), 
+			MoveTemp(ValidParameters), 
+			DataChannel
+		);
+
+		InitParams.DebugMetaSoundName = this->GetFName();
+
+		return MakeShared<Metasound::FMetaSoundParameterTransmitter>(MoveTemp(InitParams));
+	};
+
+	const bool bIsRuntimeInputDataValid = RuntimeInputData.bIsValid.load();
+	const bool bCreateInputMapOnTheFly = bIsBuilderActive || !bIsRuntimeInputDataValid;
+
+	if (bCreateInputMapOnTheFly)
+	{
+		if (!bIsBuilderActive)
+		{
+			// If we're not using a builder, that means the metasound cannot change and that the runtime input data should have been cached. 
+			UE_LOG(LogMetaSound, Warning, TEXT("Creating a Parameter Transmiiter on uninitialized UMetaSoundSource %s will result in slower performance. UMetaSoundSource::InitResources should finish executing on the game thread before attempting to call UMetaSoundSource::CreateParameterTransmitter(...)"), *GetOwningAssetName());
+		}
+
+		return CreateParameterTransmitterInternal(CreateRuntimeInputMap(), InParams);
 	}
-
-	// Build list of parameters that can be set at runtime.
-	TArray<FName> ValidParameters;
-	for (const TPair<FVertexName, FMetasoundAssetBase::FRuntimeInput>& Pair : RuntimeData.PublicInputMap)
+	else
 	{
-		if (Pair.Value.bIsTransmittable && (Pair.Value.AccessType == EMetasoundFrontendVertexAccessType::Reference))
-		{
-			ValidParameters.Add(Pair.Value.Name);
-		}
+		return CreateParameterTransmitterInternal(RuntimeInputData.InputMap, InParams);
 	}
-
-	FParameterRouter& Router = GetParameterRouter();
-	TSharedPtr<TSpscQueue<FMetaSoundParameterTransmitter::FParameter>> DataChannel = Router.FindOrCreateDataChannelForWriter(InParams.AudioDeviceID, InParams.InstanceID);
-
-	Metasound::FMetaSoundParameterTransmitter::FInitParams InitParams(GetOperatorSettings(InParams.SampleRate), InParams.InstanceID, MoveTemp(InParams.DefaultParams), MoveTemp(ValidParameters), DataChannel);
-	InitParams.DebugMetaSoundName = GetFName();
-
-	return MakeShared<Metasound::FMetaSoundParameterTransmitter>(MoveTemp(InitParams));
 }
 
 Metasound::FOperatorSettings UMetaSoundSource::GetOperatorSettings(Metasound::FSampleRate InSampleRate) const
@@ -1283,6 +1352,12 @@
 	WaitForAsyncGraphRegistration();
 
 	bIsBuilderActive = true;
+
+	// Currently we do not have information on whether inputs were added or removed
+	// from the document. We invalidate the cached runtime inputs just in case. 
+	// MetaSounds which have an active builder should not be using cached runtime
+	// input data until the builder is no longer active. 
+	InvalidateCachedRuntimeInputData();
 }
 
 void UMetaSoundSource::OnFinishActiveBuilder()
@@ -1337,4 +1412,78 @@
 {
 	return DynamicTransactor;
 }
+
+Metasound::TSortedVertexNameMap<UMetaSoundSource::FRuntimeInput> UMetaSoundSource::CreateRuntimeInputMap() const
+{
+	using namespace Metasound;
+	using namespace Metasound::Frontend;
+
+	METASOUND_TRACE_CPUPROFILER_EVENT_SCOPE(UMetaSoundSource::CreateRuntimeInputMap);
+
+	auto GetInputName = [](const FMetasoundFrontendClassInput& InInput) { return InInput.Name; };
+
+	IDataTypeRegistry& Registry = IDataTypeRegistry::Get();
+	const FMetasoundFrontendDocument& Doc = GetConstDocument();
+
+	TArray<const IInterfaceRegistryEntry*> Interfaces;
+	FMetaSoundFrontendDocumentBuilder::FindDeclaredInterfaces(Doc, Interfaces);
+
+	// Inputs which are controlled by an interface are private unless
+	// their router name is `Audio::IParameterTransmitter::RouterName`
+	TSet<FVertexName> PrivateInputs;
+	for (const IInterfaceRegistryEntry* InterfaceEntry : Interfaces)
+	{
+		if (InterfaceEntry)
+		{
+			if (InterfaceEntry->GetRouterName() != Audio::IParameterTransmitter::RouterName)
+			{
+				const FMetasoundFrontendInterface& Interface = InterfaceEntry->GetInterface();
+				Algo::Transform(Interface.Inputs, PrivateInputs, GetInputName);
+			}
+		}
+	}
+
+	// Cache all inputs which are not private inputs.
+	TSortedVertexNameMap<FRuntimeInput> PublicInputs;
+	for (const FMetasoundFrontendClassInput& Input : Doc.RootGraph.Interface.Inputs)
+	{
+		if (!PrivateInputs.Contains(Input.Name))
+		{
+			bool bIsTransmittable = false;
+			if (const IDataTypeRegistryEntry* RegistryEntry = Registry.FindDataTypeRegistryEntry(Input.TypeName))
+			{
+				bIsTransmittable = RegistryEntry->GetDataTypeInfo().bIsTransmittable;	
+			}
+			else
+			{
+				UE_LOG(LogMetaSound, Warning, TEXT("Failed to find data type '%s' in registry. Assuming data type is not transmittable"), *Input.TypeName.ToString());
+
+			}
+			PublicInputs.Add(Input.Name, FRuntimeInput{Input.Name, Input.TypeName, Input.AccessType, Input.DefaultLiteral, bIsTransmittable});
+		}
+	}
+
+	// Add the parameter pack input that ALL Metasounds have
+	FMetasoundFrontendClassInput ParameterPackInput = UMetasoundParameterPack::GetClassInput();
+	PublicInputs.Add(ParameterPackInput.Name, FRuntimeInput{ParameterPackInput.Name, ParameterPackInput.TypeName, ParameterPackInput.AccessType, ParameterPackInput.DefaultLiteral, true /* bIsTransmittable */});
+	
+	return PublicInputs;
+}
+
+void UMetaSoundSource::CacheRuntimeInputData()
+{
+	if (bIsBuilderActive)
+	{
+		UE_LOG(LogMetaSound, Warning, TEXT("Skipping caching of runtime inputs for UMetaSoundSource %s because there is an active builder"), *GetOwningAssetName());
+	}
+
+	RuntimeInputData.InputMap = CreateRuntimeInputMap();
+	RuntimeInputData.bIsValid.store(true);
+}
+
+void UMetaSoundSource::InvalidateCachedRuntimeInputData()
+{
+	RuntimeInputData.bIsValid.store(false);
+}
+
 #undef LOCTEXT_NAMESPACE // MetaSound
diff --git a/Engine/Plugins/Runtime/Metasound/Source/MetasoundEngine/Public/MetasoundSource.h b/Engine/Plugins/Runtime/Metasound/Source/MetasoundEngine/Public/MetasoundSource.h
--- a/Engine/Plugins/Runtime/Metasound/Source/MetasoundEngine/Public/MetasoundSource.h
+++ b/Engine/Plugins/Runtime/Metasound/Source/MetasoundEngine/Public/MetasoundSource.h
@@ -13,9 +13,12 @@
 #include "MetasoundOperatorSettings.h"
 #include "MetasoundParameterTransmitter.h"
 #include "MetasoundRouter.h"
+#include "MetasoundVertex.h"
 #include "Sound/SoundWaveProcedural.h"
 #include "UObject/MetaData.h"
 
+#include <atomic>
+
 #include "MetasoundSource.generated.h"
 
 
@@ -56,6 +59,27 @@ class METASOUNDENGINE_API UMetaSoundSource : public USoundWaveProcedural, public
 	friend struct Metasound::FMetaSoundEngineAssetHelper;
 	friend class UMetaSoundSourceBuilder;
 
+	// FRuntimeInput represents an input to a MetaSound which can be manipulated.
+	struct FRuntimeInput
+	{
+		// Name of input vertex
+		FName Name;
+		// Data type name of input vertex.
+		FName TypeName;
+		// Access type of input vertex.
+		EMetasoundFrontendVertexAccessType AccessType;
+		// Default literal of input vertex.
+		FMetasoundFrontendLiteral DefaultLiteral;
+		// True if the data type is transmittable. False otherwise.
+		bool bIsTransmittable;
+	};
+
+	struct FRuntimeInputData
+	{
+		std::atomic<bool> bIsValid = false;
+		Metasound::TSortedVertexNameMap<FRuntimeInput> InputMap;
+	};
+
 protected:
 	UPROPERTY(EditAnywhere, Category = CustomView)
 	FMetasoundFrontendDocument RootMetasoundDocument;
@@ -209,9 +233,9 @@ public:
 	virtual void InitParameters(TArray<FAudioParameter>& ParametersToInit, FName InFeatureName) override;
 
 	virtual void InitResources() override;
+	virtual void RegisterGraphWithFrontend(Metasound::Frontend::FMetaSoundAssetRegistrationOptions InRegistrationOptions = Metasound::Frontend::FMetaSoundAssetRegistrationOptions()) override;
 
 	virtual bool IsPlayable() const override;
-	virtual bool SupportsSubtitles() const override;
 	virtual float GetDuration() const override;
 	virtual bool ImplementsParameterInterface(Audio::FParameterInterfacePtr InInterface) const override;
 	virtual ISoundGeneratorPtr CreateSoundGenerator(const FSoundGeneratorInitParams& InParams, TArray<FAudioParameter>&& InDefaultParameters) override;
@@ -251,7 +275,8 @@ private:
 	virtual void OnBeginActiveBuilder() override;
 	virtual void OnFinishActiveBuilder() override;
 
-	bool IsParameterValid(const FAudioParameter& InParameter, const FName& InTypeName, Metasound::Frontend::IDataTypeRegistry& InDataTypeRegistry) const;
+	void InitParametersInternal(const Metasound::TSortedVertexNameMap<FRuntimeInput>& InputMap, TArray<FAudioParameter>& ParametersToInit, FName InFeatureName) const;
+	bool IsParameterValidInternal(const FAudioParameter& InParameter, const FName& InTypeName, Metasound::Frontend::IDataTypeRegistry& InDataTypeRegistry) const;
 
 	static Metasound::SourcePrivate::FParameterRouter& GetParameterRouter();
 
@@ -266,6 +291,12 @@ private:
 	void TrackGenerator(uint64 Id, TSharedPtr<Metasound::FMetasoundGenerator> Generator);
 	void ForgetGenerator(ISoundGeneratorPtr Generator);
 
+	Metasound::TSortedVertexNameMap<FRuntimeInput> CreateRuntimeInputMap() const;
+	void CacheRuntimeInputData();
+	void InvalidateCachedRuntimeInputData();
+
+	FRuntimeInputData RuntimeInputData;
+
 	/** Enable/disable dynamic generator.
 	 *
 	 * Once a dynamic generator is enabled, all changes to the MetaSound should be applied to the
diff --git a/Engine/Plugins/Runtime/Metasound/Source/MetasoundFrontend/Private/MetasoundAssetBase.cpp b/Engine/Plugins/Runtime/Metasound/Source/MetasoundFrontend/Private/MetasoundAssetBase.cpp
--- a/Engine/Plugins/Runtime/Metasound/Source/MetasoundFrontend/Private/MetasoundAssetBase.cpp
+++ b/Engine/Plugins/Runtime/Metasound/Source/MetasoundFrontend/Private/MetasoundAssetBase.cpp
@@ -121,61 +121,7 @@ namespace Metasound
 				}
 			}
 
-			TSortedVertexNameMap<FMetasoundAssetBase::FRuntimeInput> GetPublicClassInputs(const FMetasoundFrontendDocument& InDoc)
-			{
-				using namespace Metasound;
-				using namespace Metasound::Frontend;
-
-				METASOUND_TRACE_CPUPROFILER_EVENT_SCOPE(AssetBasePrivate::GetPublicClassInputs);
 
-				auto GetInputName = [](const FMetasoundFrontendClassInput& InInput) { return InInput.Name; };
-
-				IDataTypeRegistry& Registry = IDataTypeRegistry::Get();
-
-				TArray<const IInterfaceRegistryEntry*> Interfaces;
-				FMetaSoundFrontendDocumentBuilder::FindDeclaredInterfaces(InDoc, Interfaces);
-
-				// Inputs which are controlled by an interface are private unless
-				// their router name is `Audio::IParameterTransmitter::RouterName`
-				TSet<FVertexName> PrivateInputs;
-				for (const IInterfaceRegistryEntry* InterfaceEntry : Interfaces)
-				{
-					if (InterfaceEntry)
-					{
-						if (InterfaceEntry->GetRouterName() != Audio::IParameterTransmitter::RouterName)
-						{
-							const FMetasoundFrontendInterface& Interface = InterfaceEntry->GetInterface();
-							Algo::Transform(Interface.Inputs, PrivateInputs, GetInputName);
-						}
-					}
-				}
-
-				// Cache all inputs which are not private inputs.
-				TSortedVertexNameMap<FMetasoundAssetBase::FRuntimeInput> PublicInputs;
-				for (const FMetasoundFrontendClassInput& Input : InDoc.RootGraph.Interface.Inputs)
-				{
-					if (!PrivateInputs.Contains(Input.Name))
-					{
-						bool bIsTransmittable = false;
-						if (const IDataTypeRegistryEntry* RegistryEntry = Registry.FindDataTypeRegistryEntry(Input.TypeName))
-						{
-							bIsTransmittable = RegistryEntry->GetDataTypeInfo().bIsTransmittable;	
-						}
-						else
-						{
-							UE_LOG(LogMetaSound, Warning, TEXT("Failed to find data type '%s' in registry. Assuming data type is not transmittable"), *Input.TypeName.ToString());
-
-						}
-						PublicInputs.Add(Input.Name, FMetasoundAssetBase::FRuntimeInput{Input.Name, Input.TypeName, Input.AccessType, Input.DefaultLiteral, bIsTransmittable});
-					}
-				}
-
-				// Add the parameter pack input that ALL Metasounds have
-				FMetasoundFrontendClassInput ParameterPackInput = UMetasoundParameterPack::GetClassInput();
-				PublicInputs.Add(ParameterPackInput.Name, FMetasoundAssetBase::FRuntimeInput{ParameterPackInput.Name, ParameterPackInput.TypeName, ParameterPackInput.AccessType, ParameterPackInput.DefaultLiteral, true /* bIsTransmittable */});
-				
-				return PublicInputs;
-			}
 
 			// Registers node by copying document. Updates to document require re-registration.
 			// This registry entry does not support node creation as it is only intended to be
@@ -233,6 +179,18 @@ namespace Metasound
 			};
 		} // namespace AssetBasePrivate
 
+		namespace ConsoleVariables
+		{
+			static bool bDisableAsyncGraphRegistration = false;
+		}
+
+		FAutoConsoleVariableRef CVarMetaSoundDisableAsyncGraphRegistration(
+			TEXT("au.MetaSound.Experimental.DisableAsyncGraphRegistration"),
+			Metasound::Frontend::ConsoleVariables::bDisableAsyncGraphRegistration,
+			TEXT("Disables async registration of MetaSound graphs\n")
+			TEXT("Default: false"),
+			ECVF_Default);
+
 		FAutoConsoleVariableRef CVarMetaSoundBlockRate(
 			TEXT("au.MetaSound.BlockRate"),
 			AssetBasePrivate::BlockRate,
@@ -271,6 +229,10 @@ void FMetasoundAssetBase::RegisterGraphWithFrontend(Metasound::Frontend::FMetaSo
 	using namespace Metasound;
 	using namespace Metasound::Frontend;
 
+	// Graph registration must only happen on one thread to avoid race conditions on graph registration.
+	checkf(IsInGameThread(), TEXT("MetaSound %s graph can only be registered on the GameThread"), *GetOwningAssetName());
+
+
 	METASOUND_TRACE_CPUPROFILER_EVENT_SCOPE(MetaSoundAssetBase::RegisterGraphWithFrontend);
 	if (!InRegistrationOptions.bForceReregister)
 	{
@@ -762,36 +724,7 @@ void FMetasoundAssetBase::ConvertFromPreset()
 
 TArray<FMetasoundAssetBase::FSendInfoAndVertexName> FMetasoundAssetBase::GetSendInfos(uint64 InInstanceID) const
 {
-	using namespace Metasound;
-	using namespace Metasound::Frontend;
-	using FSendInfo = FMetaSoundParameterTransmitter::FSendInfo;
-
-	check(IsInGameThread() || IsInAudioThread());
-
-	const FRuntimeData& RuntimeData = GetRuntimeData();
-	if(!RuntimeData.IsValid())
-	{
-		UE_LOG(LogMetaSound, Warning, TEXT("Send infos may be incorrect. Accessing invalid runtime data on MetaSound %s"), *GetOwningAssetName());
-	}
-
-	TArray<FSendInfoAndVertexName> SendInfos;
-
-	for (const TTuple<FVertexName, FRuntimeInput>& Entry: RuntimeData.PublicInputMap)
-	{
-		const FRuntimeInput& RuntimeInput = Entry.Get<1>();
-		FSendInfoAndVertexName Info;
-
-PRAGMA_DISABLE_DEPRECATION_WARNINGS
-		Info.SendInfo.Address = FMetaSoundParameterTransmitter::CreateSendAddressFromInstanceID(InInstanceID, RuntimeInput.Name, RuntimeInput.TypeName);
-PRAGMA_ENABLE_DEPRECATION_WARNINGS
-		Info.SendInfo.ParameterName = RuntimeInput.Name;
-		Info.SendInfo.TypeName = RuntimeInput.TypeName;
-		Info.VertexName = RuntimeInput.Name;
-
-		SendInfos.Add(Info);
-	}
-
-	return SendInfos;
+	return TArray<FSendInfoAndVertexName>();
 }
 
 #if WITH_EDITOR
@@ -1004,20 +937,15 @@ Metasound::Frontend::FNodeRegistryKey FMetasoundAssetBase::CacheRuntimeData(cons
 	METASOUND_TRACE_CPUPROFILER_EVENT_SCOPE(MetaSoundAssetBase::CacheRuntimeData);
 
 	const FMetasoundFrontendDocument& Document = InDocumentInterface->GetConstDocument();
-	Metasound::TSortedVertexNameMap<FMetasoundAssetBase::FRuntimeInput> PublicInputMap = AssetBasePrivate::GetPublicClassInputs(Document);
 
 	// If the IMetaSoundDocumentInterface actively modified by a builder, then 
 	// we must build synchronously to avoid a race condition on reading/writing 
 	// the IMetaSoundDocumentInterface
-	//const bool bAsync = !IsBuilderActive();
-	constexpr bool bAsync = false; // TODO: Active crash needs addressing where graph registration and async metasound building happen on same thread. AsyncBuilder deadlocks while waiting for task to complete. 
+	const bool bAsync = !(IsBuilderActive() || ConsoleVariables::bDisableAsyncGraphRegistration);
 
 	FNodeClassInfo NodeClassInfo { Document.RootGraph, FSoftObjectPath(GetOwningAsset()) };
 	FNodeRegistryKey NewRegistryKey = FMetasoundFrontendRegistryContainer::Get()->RegisterGraph(NodeClassInfo, InDocumentInterface, bAsync);
 
-	CachedRuntimeData.ChangeID = FGuid::NewGuid();
-	CachedRuntimeData.PublicInputMap = MoveTemp(PublicInputMap);
-	CachedRuntimeData.Graph = nullptr; // Graph will be retrieved later. Retrieving the graph forces async graph registration to complete. 
 
 	return NewRegistryKey;
 }
@@ -1031,9 +959,12 @@ void FMetasoundAssetBase::WaitForAsyncGraphRegistration()
 	}
 }
 
+PRAGMA_DISABLE_DEPRECATION_WARNINGS
 const FMetasoundAssetBase::FRuntimeData& FMetasoundAssetBase::GetRuntimeData() const
 {
-	return CachedRuntimeData;
+	static const FRuntimeData PlaceholderForDeprecatedMethod;
+	return PlaceholderForDeprecatedMethod;
 }
+PRAGMA_ENABLE_DEPRECATION_WARNINGS
 
 #undef LOCTEXT_NAMESPACE // "MetaSound"
diff --git a/Engine/Plugins/Runtime/Metasound/Source/MetasoundFrontend/Public/MetasoundAssetBase.h b/Engine/Plugins/Runtime/Metasound/Source/MetasoundFrontend/Public/MetasoundAssetBase.h
--- a/Engine/Plugins/Runtime/Metasound/Source/MetasoundFrontend/Public/MetasoundAssetBase.h
+++ b/Engine/Plugins/Runtime/Metasound/Source/MetasoundFrontend/Public/MetasoundAssetBase.h
@@ -47,21 +47,6 @@ class METASOUNDFRONTEND_API FMetasoundAssetBase
 public:
 	static const FString FileExtension;
 
-	// FRuntimeInput represents an input to a MetaSound which can be manipulated.
-	struct FRuntimeInput
-	{
-		// Name of input vertex
-		FName Name;
-		// Data type name of input vertex.
-		FName TypeName;
-		// Access type of input vertex.
-		EMetasoundFrontendVertexAccessType AccessType;
-		// Default literal of input vertex.
-		FMetasoundFrontendLiteral DefaultLiteral;
-		// True if the data type is transmittable. False otherwise.
-		bool bIsTransmittable;
-	};
-
 	FMetasoundAssetBase() = default;
 	virtual ~FMetasoundAssetBase() = default;
 
@@ -92,7 +77,7 @@ public:
 	virtual bool ConformObjectDataToInterfaces() = 0;
 
 	// Registers the root graph of the given asset with the MetaSound Frontend.
-	void RegisterGraphWithFrontend(Metasound::Frontend::FMetaSoundAssetRegistrationOptions InRegistrationOptions = Metasound::Frontend::FMetaSoundAssetRegistrationOptions());
+	virtual void RegisterGraphWithFrontend(Metasound::Frontend::FMetaSoundAssetRegistrationOptions InRegistrationOptions = Metasound::Frontend::FMetaSoundAssetRegistrationOptions());
 
 	// Unregisters the root graph of the given asset with the MetaSound Frontend.
 	void UnregisterGraphWithFrontend();
@@ -230,7 +215,7 @@ protected:
 protected:
 	
 	// Container for runtime data of MetaSound graph.
-	struct FRuntimeData
+	struct FRuntimeData_DEPRECATED
 	{
 		// Current ID of graph.
 		FGuid ChangeID;
@@ -238,29 +223,27 @@ protected:
 		// Array of inputs which can be set for construction. 
 		TArray<FMetasoundFrontendClassInput> PublicInputs;
 
-		// Map of runtime inputs keyed by vertex names. 
-		Metasound::TSortedVertexNameMap<FRuntimeInput> PublicInputMap;
-
 		// Array of inputs which can be transmitted to.
 		TArray<FMetasoundFrontendClassInput> TransmittableInputs;
 
 		// Core graph.
 		TSharedPtr<Metasound::FGraph, ESPMode::ThreadSafe> Graph;
+	};
 
-		bool IsValid() const 
-		{
-			return ChangeID.IsValid();
-		}
+	struct UE_DEPRECATED(5.4, "FRuntimeData is no longer used to store runtime graphs and inputs. Runtime graphs are stored in the node registry. Runtime inputs are stored on the UMetaSoundSoruce") FRuntimeData : FRuntimeData_DEPRECATED
+	{
 	};
 
 	// Returns the cached runtime data.
+	PRAGMA_DISABLE_DEPRECATION_WARNINGS
+	UE_DEPRECATED(5.4, "Access to graph and public inputs has moved. Use the node registry to access the graph and GetPublicClassInputs() to access public inputs")
 	const FRuntimeData& GetRuntimeData() const;
+	PRAGMA_ENABLE_DEPRECATION_WARNINGS
 
 	// Returns all public class inputs.  This is a potentially expensive.
 	// Prefer accessing public class inputs using CacheRuntimeData.
 	TArray<FMetasoundFrontendClassInput> GetPublicClassInputs() const;
 
-
 	bool AutoUpdate(bool bInLogWarningsOnDroppedConnection);
 	void CookReferencedMetaSounds();
 
@@ -286,7 +269,6 @@ private:
 
 	Metasound::Frontend::FNodeRegistryKey RegistryKey;
 
-	FRuntimeData CachedRuntimeData;
 
 	TSharedPtr<Metasound::FGraph, ESPMode::ThreadSafe> BuildMetasoundDocument(const FMetasoundFrontendDocument& InPreprocessDoc, const Metasound::Frontend::FProxyDataCache& InProxies) const;
 };
diff --git a/Engine/Source/Runtime/Engine/Private/SoundCue.cpp b/Engine/Source/Runtime/Engine/Private/SoundCue.cpp
--- a/Engine/Source/Runtime/Engine/Private/SoundCue.cpp
+++ b/Engine/Source/Runtime/Engine/Private/SoundCue.cpp
@@ -646,7 +646,7 @@ void USoundCue::Parse(FAudioDevice* AudioDevice, const UPTRINT NodeWaveInstanceH
 				{
 					TArray<FAudioParameter> Params = Transmitter->ParamsToSet;
 
-					if (USoundBase* Sound = Instance->WaveData)
+					if (USoundWave* Sound = Instance->WaveData)
 					{
 						Sound->InitParameters(Params);
 					}
@@ -892,4 +892,4 @@ bool FSoundCueParameterTransmitter::SetParameters(TArray<FAudioParameter>&& InPa
 	});
 		
 	return Audio::FParameterTransmitterBase::SetParameters(MoveTemp(InParameters));
-}
\ No newline at end of file
+}
