Disable async graph registration again
#rnx

diff --git a/Engine/Plugins/Runtime/Metasound/Source/MetasoundFrontend/Private/MetasoundAssetBase.cpp b/Engine/Plugins/Runtime/Metasound/Source/MetasoundFrontend/Private/MetasoundAssetBase.cpp
--- a/Engine/Plugins/Runtime/Metasound/Source/MetasoundFrontend/Private/MetasoundAssetBase.cpp
+++ b/Engine/Plugins/Runtime/Metasound/Source/MetasoundFrontend/Private/MetasoundAssetBase.cpp
@@ -181,14 +181,14 @@ namespace Metasound
 
 		namespace ConsoleVariables
 		{
-			static bool bDisableAsyncGraphRegistration = false;
+			static bool bDisableAsyncGraphRegistration = true;
 		}
 
 		FAutoConsoleVariableRef CVarMetaSoundDisableAsyncGraphRegistration(
 			TEXT("au.MetaSound.Experimental.DisableAsyncGraphRegistration"),
 			Metasound::Frontend::ConsoleVariables::bDisableAsyncGraphRegistration,
 			TEXT("Disables async registration of MetaSound graphs\n")
-			TEXT("Default: false"),
+			TEXT("Default: true"),
 			ECVF_Default);
 
 		FAutoConsoleVariableRef CVarMetaSoundBlockRate(
