Compressor node optimizations (and node tests)

diff --git a/Engine/Plugins/Runtime/Metasound/Source/MetasoundEngineTest/Private/CompressorNodeTests.cpp b/Engine/Plugins/Runtime/Metasound/Source/MetasoundEngineTest/Private/CompressorNodeTests.cpp
new file mode 100644
--- /dev/null
+++ b/Engine/Plugins/Runtime/Metasound/Source/MetasoundEngineTest/Private/CompressorNodeTests.cpp
@@ -0,0 +1,120 @@
+// Copyright Epic Games, Inc. All Rights Reserved.
+
+#include "MetasoundGenerator.h"
+#include "MetasoundStandardNodesNames.h"
+#include "NodeTestGraphBuilder.h"
+
+#include "Misc/AutomationTest.h"
+
+#if WITH_DEV_AUTOMATION_TESTS
+namespace Metasound::Test::Nodes::Compressor
+{
+	IMPLEMENT_SIMPLE_AUTOMATION_TEST(
+		FMetasoundCompressorNodeBasicTest,
+		"Audio.Metasound.Nodes.Compressor.Basic",
+		EAutomationTestFlags::EditorContext | EAutomationTestFlags::EngineFilter)
+		bool FMetasoundCompressorNodeBasicTest::RunTest(const FString&)
+	{
+		// Make a graph with a sine oscillator going into a compressor as the audio input
+		TUniquePtr<FMetasoundGenerator> Generator;
+		{
+			FNodeTestGraphBuilder Builder;
+
+			const Frontend::FNodeHandle SineOscillator = Builder.AddNode({ StandardNodes::Namespace, TEXT("Sine"), StandardNodes::AudioVariant }, 1);
+			const Frontend::FNodeHandle Compressor = Builder.AddNode({ StandardNodes::Namespace, TEXT("Compressor"), StandardNodes::AudioVariant }, 1);
+			FNodeTestGraphBuilder::ConnectNodes(SineOscillator, "Audio", Compressor, "Audio");
+
+			// Expose the input bypass pin 
+			Builder.AddAndConnectDataReferenceInput(Compressor, "Bypass", GetMetasoundDataTypeName<bool>());
+			//output audio and output envelopes pin
+			Builder.AddAndConnectDataReferenceOutput(Compressor, "Gain Envelope", GetMetasoundDataTypeName<FAudioBuffer>());
+			Builder.AddAndConnectDataReferenceOutput(Compressor, "Audio", GetMetasoundDataTypeName<FAudioBuffer>());
+
+			Generator = Builder.BuildGenerator();
+		}
+		UTEST_NOT_NULL("Created generator", Generator.Get());
+
+		//output refs
+		TOptional<TDataReadReference<FAudioBuffer>> OutputEnvelope = Generator->GetOutputReadReference<FAudioBuffer>("Gain Envelope");
+		UTEST_TRUE("Got float envelope output", OutputEnvelope.IsSet());
+		TOptional<TDataReadReference<FAudioBuffer>> OutputAudio = Generator->GetOutputReadReference<FAudioBuffer>("Audio");
+		UTEST_TRUE("Got audio envelope output", OutputAudio.IsSet());
+
+		//2 buffers for comparing bypassed and compressed audio outputs
+		const int32 BlockSize = Generator->OperatorSettings.GetNumFramesPerBlock();
+		FAudioBuffer TestBufferBypassed{ BlockSize };
+		FAudioBuffer TestBufferWithCompression{ BlockSize };
+
+		//set a low threshold for triggering compression
+		Generator->SetInputValue("Threshold dB", -10.0f);
+
+		//bypass the compressor, expect output envelopes to be 0
+		Generator->SetInputValue("Bypass", true);
+
+		//generate a block
+		Generator->OnGenerateAudio(TestBufferBypassed.GetData(), TestBufferBypassed.Num());
+
+		//check results
+		{
+			UTEST_EQUAL ("Output Envelope outputs 0 if bypassed", (*OutputEnvelope)->GetData()[0], 0.0f);
+			bool HadNonZeroEnvelope = false;
+			for (int32 i = 0; i < BlockSize; ++i)
+			{
+				const float Value = (*OutputEnvelope)->GetData()[i];
+				if (Value != 0.0f)
+				{
+					HadNonZeroEnvelope = true;
+					break;
+				}
+			}
+			UTEST_FALSE("Output Envelope should output 0 if bypassed", HadNonZeroEnvelope);
+		}
+
+		//enable dynamics compression (set bypass to false), expect non-zero output envelope
+		Generator->SetInputValue("Bypass", false);
+
+		//generate a few blocks
+		constexpr int32 NumBlocksToRender = 4;
+		for (int32 i = 0; i < NumBlocksToRender; ++i)
+		{
+			Generator->OnGenerateAudio(TestBufferWithCompression.GetData(), TestBufferWithCompression.Num());
+		}
+
+		//check results
+		{
+			UTEST_TRUE("Output Envelope should have non-zero values if enabled", (*OutputEnvelope)->GetData()[0] != 0.0f);
+			bool HadNonZeroEnvelope = true;
+			for (int32 i = 0; i < BlockSize; ++i)
+			{
+				const float Value = (*OutputEnvelope)->GetData()[i];
+				if (Value == 0.0f)
+				{
+					HadNonZeroEnvelope = false;
+					break;
+				}
+			}
+			UTEST_TRUE("Output Envelope should have non-zero values if enabled", HadNonZeroEnvelope);
+		}
+
+		//compare bypassed and compression-enabled buffers, expect buffers to contain different values after compression
+		{
+			bool BuffersContainDifferentOutputs = false;
+			for (int32 i = 0; i < BlockSize; ++i)
+			{
+				const float ByPassedValue = TestBufferBypassed.GetData()[i];
+				const float CompressedValue = TestBufferWithCompression.GetData()[i];
+				if (ByPassedValue != CompressedValue)
+				{
+					BuffersContainDifferentOutputs = true;
+					break;
+				}
+
+			}
+			
+			UTEST_TRUE("Output audio buffer with compression enabled and bypassed should produce different values", BuffersContainDifferentOutputs);
+		}
+
+		return true;
+	}
+}
+#endif
\ No newline at end of file
diff --git a/Engine/Plugins/Runtime/Metasound/Source/MetasoundStandardNodes/Private/MetasoundCompressorNode.cpp b/Engine/Plugins/Runtime/Metasound/Source/MetasoundStandardNodes/Private/MetasoundCompressorNode.cpp
--- a/Engine/Plugins/Runtime/Metasound/Source/MetasoundStandardNodes/Private/MetasoundCompressorNode.cpp
+++ b/Engine/Plugins/Runtime/Metasound/Source/MetasoundStandardNodes/Private/MetasoundCompressorNode.cpp
@@ -84,10 +84,15 @@ namespace Metasound
 			, PrevAttackTime(FMath::Max(FTime::ToMilliseconds(*InAttackTime), 0.0))
 			, PrevReleaseTime(FMath::Max(FTime::ToMilliseconds(*InReleaseTime), 0.0))
 			, PrevLookaheadTime(FMath::Max(FTime::ToMilliseconds(*InLookaheadTime), 0.0))
+			, PrevKneeInput(*KneeInput)
+			, PrevPeakMode(*EnvelopeModeInput)
+			, bPrevIsUpwardsInput(*bIsUpwardsInput)
+			, PrevClampedRatio(GetClampedRatio())
+			, PrevThreshold(*ThresholdDbInput)
 		{
 			Compressor.Init(InSettings.GetSampleRate(), 1);
 			Compressor.SetKeyNumChannels(1);
-			Compressor.SetRatio(FMath::Max(*InRatio, 1.0f));
+			Compressor.SetRatio(GetClampedRatio());
 			Compressor.SetThreshold(*ThresholdDbInput);
 			Compressor.SetAttackTime(PrevAttackTime);
 			Compressor.SetReleaseTime(PrevReleaseTime);
@@ -116,7 +121,6 @@ namespace Metasound
 				Compressor.SetPeakMode(Audio::EPeakMode::Peak);
 				break;
 			}
-
 		}
 
 		static const FNodeClassMetadata& GetNodeInfo()
@@ -261,12 +265,15 @@ namespace Metasound
 			// Initialize compressor
 			Compressor.Init(InParams.OperatorSettings.GetSampleRate(), 1);
 			Compressor.SetKeyNumChannels(1);
-			Compressor.SetRatio(FMath::Max(*RatioInput, 1.0f));
+			PrevClampedRatio = GetClampedRatio();
+			Compressor.SetRatio(PrevClampedRatio);
 			Compressor.SetThreshold(*ThresholdDbInput);
+			PrevThreshold = *ThresholdDbInput;
 			Compressor.SetAttackTime(PrevAttackTime);
 			Compressor.SetReleaseTime(PrevReleaseTime);
 			Compressor.SetLookaheadMsec(PrevLookaheadTime);
 			Compressor.SetKneeBandwidth(*KneeInput);
+			PrevKneeInput = *KneeInput;
 
 			if (*bIsUpwardsInput)
 			{
@@ -276,6 +283,7 @@ namespace Metasound
 			{
 				Compressor.SetProcessingMode(Audio::EDynamicsProcessingMode::Compressor);
 			}
+			bPrevIsUpwardsInput = *bIsUpwardsInput;
 
 			switch (*EnvelopeModeInput)
 			{
@@ -290,6 +298,7 @@ namespace Metasound
 				Compressor.SetPeakMode(Audio::EPeakMode::Peak);
 				break;
 			}
+			PrevPeakMode = *EnvelopeModeInput;
 		}
 
 		void Execute()
@@ -308,8 +317,17 @@ namespace Metasound
 			/* Update parameters */
 			
 			// For a compressor, ratio values should be 1 or greater
-			Compressor.SetRatio(FMath::Max(*RatioInput, 1.0f));
-			Compressor.SetThreshold(*ThresholdDbInput);
+			float UpdatedClampedRatio = GetClampedRatio();
+			if (PrevClampedRatio != UpdatedClampedRatio)
+			{
+				Compressor.SetRatio(UpdatedClampedRatio);
+				PrevClampedRatio = UpdatedClampedRatio;
+			}
+			if (PrevThreshold != *ThresholdDbInput)
+			{
+				Compressor.SetThreshold(*ThresholdDbInput);
+				PrevThreshold = *ThresholdDbInput;
+			}
 
 			// Attack time cannot be negative
 			float CurrAttack = FMath::Max(FTime::ToMilliseconds(*AttackTimeInput), 0.0f);
@@ -335,44 +353,55 @@ namespace Metasound
 			}
 			InputDelay.SetDelayLengthSamples(CurrLookahead * MsToSamples);
 
-			Compressor.SetKneeBandwidth(*KneeInput);
-
-			if (*bIsUpwardsInput)
+			if (*KneeInput != PrevKneeInput)
 			{
-				Compressor.SetProcessingMode(Audio::EDynamicsProcessingMode::UpwardsCompressor);
+				Compressor.SetKneeBandwidth(*KneeInput);
+				PrevKneeInput = *KneeInput;
 			}
-			else
+
+			if (*bIsUpwardsInput != bPrevIsUpwardsInput)
 			{
-				Compressor.SetProcessingMode(Audio::EDynamicsProcessingMode::Compressor);
+				if (*bIsUpwardsInput)
+				{
+					Compressor.SetProcessingMode(Audio::EDynamicsProcessingMode::UpwardsCompressor);
+				}
+				else
+				{
+					Compressor.SetProcessingMode(Audio::EDynamicsProcessingMode::Compressor);
+				}
+				bPrevIsUpwardsInput = *bIsUpwardsInput;
 			}
-				
-			switch (*EnvelopeModeInput)
-			{
-			default:
-			case EEnvelopePeakMode::MeanSquared:
-				Compressor.SetPeakMode(Audio::EPeakMode::MeanSquared);
-				break;
-
-			case EEnvelopePeakMode::RootMeanSquared:
-				Compressor.SetPeakMode(Audio::EPeakMode::RootMeanSquared);
-				break;
 
-			case EEnvelopePeakMode::Peak:
-				Compressor.SetPeakMode(Audio::EPeakMode::Peak);
-				break;
+			if (*EnvelopeModeInput != PrevPeakMode)
+			{
+				switch (*EnvelopeModeInput)
+				{
+				default:
+				case EEnvelopePeakMode::MeanSquared:
+					Compressor.SetPeakMode(Audio::EPeakMode::MeanSquared);
+					break;
+
+				case EEnvelopePeakMode::RootMeanSquared:
+					Compressor.SetPeakMode(Audio::EPeakMode::RootMeanSquared);
+					break;
+
+				case EEnvelopePeakMode::Peak:
+					Compressor.SetPeakMode(Audio::EPeakMode::Peak);
+					break;
+				}
+				PrevPeakMode = *EnvelopeModeInput;
 			}
 
 			// Apply lookahead delay to dry signal
 			InputDelay.ProcessAudio(*AudioInput, DelayedInputSignal);
 
-			if (bUseSidechain)
-			{
-				Compressor.ProcessAudio(AudioInput->GetData(), AudioInput->Num(), AudioOutput->GetData(), SidechainInput->GetData(), EnvelopeOutput->GetData());
-			}
-			else
-			{
-				Compressor.ProcessAudio(AudioInput->GetData(), AudioInput->Num(), AudioOutput->GetData(), nullptr, EnvelopeOutput->GetData());
-			}
+			const float* InSamples = AudioInput->GetData();
+			float* OutSamples      = AudioOutput->GetData();
+			const float* InKey     = SidechainInput->GetData();
+			float* OutEnvelope     = EnvelopeOutput->GetData();
+
+			Compressor.ProcessAudio(&InSamples, AudioInput->Num(), &OutSamples, bUseSidechain ? &InKey : nullptr, &OutEnvelope);
+
 			bEnvelopeOutputIsZero = false;
 
 			// Calculate Wet/Dry mix
@@ -384,6 +413,11 @@ namespace Metasound
 			
 		}
 
+		float GetClampedRatio() const
+		{
+			return FMath::Max(*RatioInput, 1.0f);
+		}
+
 	private:
 		// Audio input and output
 		FBoolReadRef bIsBypassedInput;
@@ -420,10 +454,18 @@ namespace Metasound
 		// Conversion from milliseconds to samples
 		float MsToSamples;
 
-		// Cached variables
+		// Cached variables to minimize updating the underlying
+		// DynamicsProcessor when there are no changes. It DOES NOT
+		// "early out" if it is told about "new" settngs that actually
+		// match its existing settings. 
 		float PrevAttackTime;
 		float PrevReleaseTime;
 		float PrevLookaheadTime;
+		float PrevKneeInput;
+		EEnvelopePeakMode PrevPeakMode;
+		bool bPrevIsUpwardsInput;
+		float PrevClampedRatio;
+		float PrevThreshold;
 	};
 
 	// Node Class
