Fix logspam on failure to initialize metasound parameters.
#rnx

diff --git a/Engine/Plugins/Runtime/Metasound/Source/MetasoundEngine/Private/MetasoundSource.cpp b/Engine/Plugins/Runtime/Metasound/Source/MetasoundEngine/Private/MetasoundSource.cpp
--- a/Engine/Plugins/Runtime/Metasound/Source/MetasoundEngine/Private/MetasoundSource.cpp
+++ b/Engine/Plugins/Runtime/Metasound/Source/MetasoundEngine/Private/MetasoundSource.cpp
@@ -52,6 +52,20 @@ namespace Metasound
 {
 	namespace SourcePrivate
 	{
+		// Holds onto a global static TSet for tracking which error/warning logs have been
+		// trigger in order to avoid log spam.
+		bool HasNotBeenLoggedForThisObject(const UMetaSoundSource& InMetaSound, uint32 InLogLineNumber)
+		{
+			using FObjectAddressAndLineNum = TTuple<const void*, uint32>;
+
+			static TSet<FObjectAddressAndLineNum> LoggedSet;
+
+			bool bIsAlreadyInSet = false;
+			LoggedSet.Add(FObjectAddressAndLineNum(&InMetaSound, InLogLineNumber), &bIsAlreadyInSet);
+
+			return !bIsAlreadyInSet;
+		}
+
 		Frontend::FMetaSoundAssetRegistrationOptions GetInitRegistrationOptions()
 		{
 			Frontend::FMetaSoundAssetRegistrationOptions RegOptions;
@@ -460,6 +474,8 @@ void UMetaSoundSource::ResolveQualitySettings(const UMetaSoundSettings* Settings
 
 void UMetaSoundSource::InitParameters(TArray<FAudioParameter>& ParametersToInit, FName InFeatureName)
 {
+	using namespace Metasound::SourcePrivate;
+
 	METASOUND_LLM_SCOPE;
 	METASOUND_TRACE_CPUPROFILER_EVENT_SCOPE(UMetaSoundSource::InitParameters);
 
@@ -481,7 +497,7 @@ void UMetaSoundSource::InitParameters(TArray<FAudioParameter>& ParametersToInit,
 			// that InitResources has been called before this method executes or else
 			// suffer the consequences of incurring significant performance losses 
 			// each time a parameter is set on the MetaSound. 
-			UE_LOG(LogMetaSound, Warning, TEXT("Initializing parameters on uninitialized UMetaSoundSource %s will result in slower performance. UMetaSoundSource::InitResources should finish executing on the game thread before attempting to call UMetaSoundSource::InitParameters(...)"), *GetOwningAssetName());
+			UE_CLOG(HasNotBeenLoggedForThisObject(*this, __LINE__), LogMetaSound, Warning, TEXT("Initializing parameters on uninitialized UMetaSoundSource %s will result in slower performance. UMetaSoundSource::InitResources should finish executing on the game thread before attempting to call UMetaSoundSource::InitParameters(...)"), *GetOwningAssetName());
 			InitParametersInternal(CreateRuntimeInputMap(), ParametersToInit, InFeatureName);
 		}
 	}
@@ -691,7 +707,7 @@ bool UMetaSoundSource::GetAllDefaultParameters(TArray<FAudioParameter>& OutParam
 	
 	if(!RuntimeInputData.bIsValid.load())
 	{
-		UE_LOG(LogMetaSound, Warning, TEXT("Default parameters will be ommitted. Accessing invalid runtime data on MetaSound %s. Ensure that UMetaSoundSource::InitResources() is executed on the game thread before calling UMetaSoundSource::GetAllDefaultParameters(...)"), *GetOwningAssetName());
+		UE_CLOG(SourcePrivate::HasNotBeenLoggedForThisObject(*this, __LINE__), LogMetaSound, Warning, TEXT("Default parameters will be ommitted. Accessing invalid runtime data on MetaSound %s. Ensure that UMetaSoundSource::InitResources() is executed on the game thread before calling UMetaSoundSource::GetAllDefaultParameters(...)"), *GetOwningAssetName());
 		return false;
 	}
 
