Make the CVARs a global override instead of just the defaults

diff --git a/Engine/Plugins/Runtime/Metasound/Config/DefaultMetaSound.ini b/Engine/Plugins/Runtime/Metasound/Config/DefaultMetaSound.ini
--- a/Engine/Plugins/Runtime/Metasound/Config/DefaultMetaSound.ini
+++ b/Engine/Plugins/Runtime/Metasound/Config/DefaultMetaSound.ini
@@ -1,5 +1,3 @@
 [/Script/MetasoundEngine.MetaSoundSettings]
-+QualitySettings=(UniqueId=A6909AF846FBBE271647E0A2E54855AC,Name="High",SampleRate=48000,BlockRate=100.000000)
 
 [/Script/MetasoundEngine.MetaSoundSource]
-QualitySettingGuid=A6909AF846FBBE271647E0A2E54855AC
diff --git a/Engine/Plugins/Runtime/Metasound/Source/MetasoundEditor/Private/MetasoundEditor.cpp b/Engine/Plugins/Runtime/Metasound/Source/MetasoundEditor/Private/MetasoundEditor.cpp
--- a/Engine/Plugins/Runtime/Metasound/Source/MetasoundEditor/Private/MetasoundEditor.cpp
+++ b/Engine/Plugins/Runtime/Metasound/Source/MetasoundEditor/Private/MetasoundEditor.cpp
@@ -1644,8 +1644,10 @@ namespace Metasound
 					const FSampleRate SampleRate = static_cast<FSampleRate>(AudioDevice->GetSampleRate());
 					const uint32 PlayOrder = PreviewComp->GetLastPlayOrder();
 					const uint64 TransmitterID = Audio::GetTransmitterID(PreviewComp->GetAudioComponentID(), 0, PlayOrder);
-					
-					GraphConnectionManager = MakeUnique<FGraphConnectionManager>(*MetasoundAsset, *PreviewComp, TransmitterID, SampleRate);
+
+					UMetaSoundSource* Source = CastChecked<UMetaSoundSource>(Metasound);
+					GraphConnectionManager = MakeUnique<FGraphConnectionManager>(
+						*MetasoundAsset, *PreviewComp, TransmitterID, Source->GetOperatorSettings(SampleRate));
 				}
 
 				MetasoundGraphEditor->RegisterActiveTimer(0.0f,
diff --git a/Engine/Plugins/Runtime/Metasound/Source/MetasoundEditor/Private/MetasoundEditorGraphConnectionManager.cpp b/Engine/Plugins/Runtime/Metasound/Source/MetasoundEditor/Private/MetasoundEditorGraphConnectionManager.cpp
--- a/Engine/Plugins/Runtime/Metasound/Source/MetasoundEditor/Private/MetasoundEditorGraphConnectionManager.cpp
+++ b/Engine/Plugins/Runtime/Metasound/Source/MetasoundEditor/Private/MetasoundEditorGraphConnectionManager.cpp
@@ -14,9 +14,9 @@ namespace Metasound
 {
 	namespace Editor
 	{
-		FGraphConnectionManager::FGraphConnectionManager(const FMetasoundAssetBase& InAssetBase, const UAudioComponent& InAudioComponent, uint64 InTransmitterID, FSampleRate InSampleRate)
+		FGraphConnectionManager::FGraphConnectionManager(const FMetasoundAssetBase& InAssetBase, const UAudioComponent& InAudioComponent, uint64 InTransmitterID, const FOperatorSettings& InOperatorSettings)
 			: AudioComponent(&InAudioComponent)
-			, GraphAnalyzerView(MakeUnique<Frontend::FMetasoundGraphAnalyzerView>(InAssetBase, InTransmitterID, InSampleRate))
+			, GraphAnalyzerView(MakeUnique<Frontend::FMetasoundGraphAnalyzerView>(InAssetBase, InTransmitterID, InOperatorSettings))
 		{
 			using namespace Frontend;
 
diff --git a/Engine/Plugins/Runtime/Metasound/Source/MetasoundEditor/Private/MetasoundEditorGraphConnectionManager.h b/Engine/Plugins/Runtime/Metasound/Source/MetasoundEditor/Private/MetasoundEditorGraphConnectionManager.h
--- a/Engine/Plugins/Runtime/Metasound/Source/MetasoundEditor/Private/MetasoundEditorGraphConnectionManager.h
+++ b/Engine/Plugins/Runtime/Metasound/Source/MetasoundEditor/Private/MetasoundEditorGraphConnectionManager.h
@@ -84,9 +84,9 @@ namespace Metasound
 
 		public:
 			FGraphConnectionManager() = default;
-			FGraphConnectionManager(const FMetasoundAssetBase& InAssetBase, const UAudioComponent& InAudioComponent, uint64 InTransmitterID, FSampleRate InSampleRate);
 			~FGraphConnectionManager() = default;
 
+			FGraphConnectionManager(const FMetasoundAssetBase& InAssetBase, const UAudioComponent& InAudioComponent, uint64 InTransmitterID, const FOperatorSettings& InOperatorSettings);
 			bool GetValue(const FGuid& InNodeID, FVertexName InOutputName, float& OutValue) const;
 			bool GetValue(const FGuid& InNodeID, FVertexName InOutputName, bool& OutValue) const;
 			bool GetValue(const FGuid& InNodeID, FVertexName InOutputName, int32& OutValue) const;
diff --git a/Engine/Plugins/Runtime/Metasound/Source/MetasoundEngine/Private/MetasoundSource.cpp b/Engine/Plugins/Runtime/Metasound/Source/MetasoundEngine/Private/MetasoundSource.cpp
--- a/Engine/Plugins/Runtime/Metasound/Source/MetasoundEngine/Private/MetasoundSource.cpp
+++ b/Engine/Plugins/Runtime/Metasound/Source/MetasoundEngine/Private/MetasoundSource.cpp
@@ -421,6 +421,23 @@ void UMetaSoundSource::PostLoadQualitySettings()
 				WeakSource->ResolveQualitySettings(CastChecked<UMetaSoundSettings>(InObj));
 			}
 		});
+
+		// Register for changes from the CVars that control overrides.
+		// We cache the OperatorSettings, so reset when these change.
+		auto ResetOperatorSettings = [WeakSource = MakeWeakObjectPtr(this)](IConsoleVariable* Var)
+		{
+			if (WeakSource.IsValid())
+			{
+				WeakSource->ResolveQualitySettings(GetMutableDefault<UMetaSoundSettings>());
+			
+				WeakSource->OperatorSettings.Reset();
+						
+				// Override SampleRate with the Operator settings version which uses our Quality settings.
+				WeakSource->SampleRate = WeakSource->GetOperatorSettings(WeakSource->SampleRate).GetSampleRate();
+			}
+		};
+		Metasound::Frontend::GetBlockRateOverrideChangedDelegate().AddWeakLambda(this, ResetOperatorSettings);
+		Metasound::Frontend::GetSampleRateOverrideChangedDelegate().AddWeakLambda(this, ResetOperatorSettings);
 	}
 #endif //WITH_EDITORONLY_DATA
 
@@ -1187,48 +1204,34 @@ Metasound::FOperatorSettings UMetaSoundSource::GetOperatorSettings(Metasound::FS
 	{
 		using namespace Metasound;
 		using namespace Metasound::SourcePrivate;
-
+		
 		// Lazy Query and cache on the optional.
 		auto QueryQualitySettings = [&](Metasound::FSampleRate InSampleRate) -> Metasound::FOperatorSettings
 		{
-			static const int32 DefaultSampleRateConstant = 48000;
 			static const float DefaultBlockRateConstant = 100.f;
 			
 			// 1. Sensible defaults.
-			FSampleRate SampleRate = DefaultSampleRateConstant;
+			FSampleRate SampleRate = InSampleRate;
 			float BlockRate = DefaultBlockRateConstant;
 
-			// 2. Query CVars. (Override with CVars if they are > 0)
-			const float BlockRateCVar = Metasound::Frontend::GetDefaultBlockRate();
-			const int32 SampleRateCvar = Metasound::Frontend::GetDefaultSampleRate();
-
-			if (SampleRateCvar != INDEX_NONE)
-			{
-				SampleRate = SampleRateCvar;
-			}
-			if (BlockRateCVar > 0)
-			{
-				BlockRate = BlockRateCVar;
-			}
-
-			// 3. Query our quality settings.
+			// 2. Query our quality settings.
 			if (const UMetaSoundSettings* Settings = GetDefault<UMetaSoundSettings>())
 			{
 				if (const FMetaSoundQualitySettings* Found = Settings->QualitySettings.FindByPredicate([&QT = QualitySetting](const FMetaSoundQualitySettings& Q) -> bool { return Q.Name == QT; }))
 				{
 					// Allow partial applications of settings, if some are non-zero.
-					if (Found->BlockRate > 0.f)
+					if (const float Value = Found->BlockRate.GetValue(); Value > 0.f)
 					{
-						BlockRate = Found->BlockRate;
+						BlockRate = Value;
 					}
-					if (Found->SampleRate > 0.f)
+					if (const float Value = Found->SampleRate.GetValue(); Value > 0.f)
 					{
-						SampleRate = Found->SampleRate;
+						SampleRate = Value;
 					}
 				}
 			}
 
-			// 4. Do per asset overrides.
+			// 3. Do per asset overrides.
 			if (const float SerializedBlockRate = BlockRateOverride.GetValue(); SerializedBlockRate > 0.0f)
 			{
 				BlockRate = SerializedBlockRate;
@@ -1238,6 +1241,26 @@ Metasound::FOperatorSettings UMetaSoundSource::GetOperatorSettings(Metasound::FS
 				SampleRate = SerializedSampleRate;
 			}
 
+			// 4. Query CVars. (Override with CVars if they are > 0)
+			using namespace Metasound::Frontend;
+			const float BlockRateCVar = GetBlockRateOverride();
+			const int32 SampleRateCvar = GetSampleRateOverride();
+
+			if (SampleRateCvar > 0)
+			{
+				SampleRate = SampleRateCvar;
+			}
+			if (BlockRateCVar > 0)
+			{
+				BlockRate = BlockRateCVar;
+			}
+
+			// 5. Sanity clamps.
+			TRange<float> BlockRange = GetBlockRateClampRange();
+			TRange<int32> RateRange = GetSampleRateClampRange();
+			BlockRate = FMath::Clamp(BlockRate, BlockRange.GetLowerBoundValue(), BlockRange.GetUpperBoundValue());
+			SampleRate = FMath::Clamp(SampleRate, RateRange.GetLowerBoundValue(), RateRange.GetUpperBoundValue());
+
 			return Metasound::FOperatorSettings(SampleRate, BlockRate);
 		};
 		OperatorSettings = QueryQualitySettings(InSampleRate);
diff --git a/Engine/Plugins/Runtime/Metasound/Source/MetasoundEngine/Public/MetasoundSettings.h b/Engine/Plugins/Runtime/Metasound/Source/MetasoundEngine/Public/MetasoundSettings.h
--- a/Engine/Plugins/Runtime/Metasound/Source/MetasoundEngine/Public/MetasoundSettings.h
+++ b/Engine/Plugins/Runtime/Metasound/Source/MetasoundEngine/Public/MetasoundSettings.h
@@ -7,6 +7,7 @@
 #include "UObject/NoExportTypes.h"
 #include "UObject/Object.h"
 #include "UObject/ObjectMacros.h"
+#include "PerPlatformProperties.h"
 
 #include "MetasoundSettings.generated.h"
 
@@ -39,9 +40,16 @@ struct METASOUNDENGINE_API FMetaSoundQualitySettings
 	GENERATED_BODY()
 
 #if WITH_EDITORONLY_DATA
+
 	/** A hidden GUID that will be generated once when adding a new entry. This prevents orphaning of renamed entries. **/
 	UPROPERTY()
-	FGuid UniqueId = FGuid::NewGuid();
+	FGuid UniqueId;
+
+	FMetaSoundQualitySettings()
+		: UniqueId(FGuid::NewGuid())
+	{}
+#else	
+	FMetaSoundQualitySettings() = default;
 #endif //WITH_EDITORONLY_DATA
 
 	/** Name of this quality setting. This will appear in the quality dropdown list.
@@ -51,11 +59,11 @@ struct METASOUNDENGINE_API FMetaSoundQualitySettings
 
 	/** Sample Rate (in Hz). NOTE: A Zero value will force the default. **/
 	UPROPERTY(EditAnywhere, Category = "Quality", meta = (ClampMin = "0", ClampMax="96000"))
-	int32 SampleRate = 0;
+	FPerPlatformInt SampleRate = 0;
 
 	/** Block Rate (in Hz). NOTE: A Zero value will force the default.  **/
-	UPROPERTY(EditAnywhere, Category = "Quality", meta = (ClampMin = "0", ClampMax="100"))
-	float BlockRate = 0.f;
+	UPROPERTY(EditAnywhere, Category = "Quality", meta = (ClampMin = "0", ClampMax="1000"))
+	FPerPlatformFloat BlockRate = 0.f;
 };
 
 
diff --git a/Engine/Plugins/Runtime/Metasound/Source/MetasoundEngine/Public/MetasoundSource.h b/Engine/Plugins/Runtime/Metasound/Source/MetasoundEngine/Public/MetasoundSource.h
--- a/Engine/Plugins/Runtime/Metasound/Source/MetasoundEngine/Public/MetasoundSource.h
+++ b/Engine/Plugins/Runtime/Metasound/Source/MetasoundEngine/Public/MetasoundSource.h
@@ -116,11 +116,11 @@ public:
 #endif //WITH_EDITOR_DATA
 
 	// Override the BlockRate for this Sound (overrides the Platforms Quality Settings for this asset)
-	UPROPERTY(EditAnywhere, BlueprintReadOnly, AdvancedDisplay, Category = Metasound, meta = (UIMin = 1, UIMax = 1000, DisplayAfter="OutputFormat", DisplayName = "Override Block Rate (in Hz)"))
+	UPROPERTY(EditAnywhere, BlueprintReadOnly, AdvancedDisplay, Category = Metasound, meta = (UIMin = 0, UIMax = 1000.0, DisplayAfter="OutputFormat", DisplayName = "Override Block Rate (in Hz)"))
 	FPerPlatformFloat BlockRateOverride = 0.f;
 
-	// Override the SampleRate for this Sound (overrides the 
-	UPROPERTY(EditAnywhere, BlueprintReadOnly, AdvancedDisplay, Category = Metasound, meta = (UIMin = 100, UIMax = 96000, DisplayName = "Override Sample Rate (in Hz)"))
+	// Override the SampleRate for this Sound (overrides the Platforms Quality Settings for this asset)
+	UPROPERTY(EditAnywhere, BlueprintReadOnly, AdvancedDisplay, Category = Metasound, meta = (UIMin = 0, UIMax = 96000, DisplayName = "Override Sample Rate (in Hz)"))
 	FPerPlatformInt SampleRateOverride = 0;
 
 	UPROPERTY(AssetRegistrySearchable)
@@ -249,6 +249,7 @@ public:
 	TWeakPtr<Metasound::FMetasoundGenerator> GetGeneratorForAudioComponent(uint64 ComponentId) const;
 	FOnGeneratorInstanceCreated OnGeneratorInstanceCreated;
 	FOnGeneratorInstanceDestroyed OnGeneratorInstanceDestroyed;
+	Metasound::FOperatorSettings GetOperatorSettings(Metasound::FSampleRate InSampleRate) const;
 
 	virtual const FMetasoundFrontendDocument& GetConstDocument() const override;
 
@@ -279,8 +280,7 @@ private:
 	bool IsParameterValidInternal(const FAudioParameter& InParameter, const FName& InTypeName, Metasound::Frontend::IDataTypeRegistry& InDataTypeRegistry) const;
 
 	static Metasound::SourcePrivate::FParameterRouter& GetParameterRouter();
-
-	Metasound::FOperatorSettings GetOperatorSettings(Metasound::FSampleRate InSampleRate) const;
+	
 	Metasound::FMetasoundEnvironment CreateEnvironment() const;
 	Metasound::FMetasoundEnvironment CreateEnvironment(const FSoundGeneratorInitParams& InParams) const;
 	Metasound::FMetasoundEnvironment CreateEnvironment(const Audio::FParameterTransmitterInitParams& InParams) const;
diff --git a/Engine/Plugins/Runtime/Metasound/Source/MetasoundFrontend/Private/Analysis/MetasoundFrontendGraphAnalyzerView.cpp b/Engine/Plugins/Runtime/Metasound/Source/MetasoundFrontend/Private/Analysis/MetasoundFrontendGraphAnalyzerView.cpp
--- a/Engine/Plugins/Runtime/Metasound/Source/MetasoundFrontend/Private/Analysis/MetasoundFrontendGraphAnalyzerView.cpp
+++ b/Engine/Plugins/Runtime/Metasound/Source/MetasoundFrontend/Private/Analysis/MetasoundFrontendGraphAnalyzerView.cpp
@@ -53,10 +53,10 @@ namespace Metasound
 			}
 		} // namespace GraphAnalyzerViewPrivate
 
-		FMetasoundGraphAnalyzerView::FMetasoundGraphAnalyzerView(const FMetasoundAssetBase& InAssetBase, uint64 InInstanceID, FSampleRate InSampleRate)
+		FMetasoundGraphAnalyzerView::FMetasoundGraphAnalyzerView(const FMetasoundAssetBase& InAssetBase, uint64 InInstanceID, const FOperatorSettings& InOperatorSettings)
 			: InstanceID(InInstanceID)
 			, MetaSoundAsset(&InAssetBase)
-			, OperatorSettings({ InSampleRate, GetDefaultBlockRate() })
+			, OperatorSettings(InOperatorSettings)
 		{
 			// TODO: Mirrored from FMetaSoundParameterTransmitter. Fix here when this is refactored.
 			const float DelayTimeInSeconds = 0.1f;
diff --git a/Engine/Plugins/Runtime/Metasound/Source/MetasoundFrontend/Private/MetasoundAssetBase.cpp b/Engine/Plugins/Runtime/Metasound/Source/MetasoundFrontend/Private/MetasoundAssetBase.cpp
--- a/Engine/Plugins/Runtime/Metasound/Source/MetasoundFrontend/Private/MetasoundAssetBase.cpp
+++ b/Engine/Plugins/Runtime/Metasound/Source/MetasoundFrontend/Private/MetasoundAssetBase.cpp
@@ -46,10 +46,9 @@ namespace Metasound
 	{
 		namespace AssetBasePrivate
 		{
-			static float BlockRate = 100.f;
-
-			// Negative value means we use the audio mixer device sample rate
-			static int32 SampleRate = INDEX_NONE;
+			// Zero values means, that these don't do anything.
+			static float BlockRateOverride = 0;
+			static int32 SampleRateOverride = 0;
 
 			TScriptInterface<IMetaSoundDocumentInterface> BuildRegistryDocument(TScriptInterface<IMetaSoundDocumentInterface> DocumentInterface)
 			{
@@ -190,34 +189,67 @@ namespace Metasound
 			TEXT("Disables async registration of MetaSound graphs\n")
 			TEXT("Default: true"),
 			ECVF_Default);
+		FConsoleVariableMulticastDelegate CVarMetaSoundBlockRateChanged;
 
 		FAutoConsoleVariableRef CVarMetaSoundBlockRate(
 			TEXT("au.MetaSound.BlockRate"),
-			AssetBasePrivate::BlockRate,
+			AssetBasePrivate::BlockRateOverride,
 			TEXT("Sets block rate (blocks per second) of MetaSounds.\n")
 			TEXT("Default: 100.0f, Min: 1.0f, Max: 1000.0f"),
+			FConsoleVariableDelegate::CreateLambda([](IConsoleVariable* Var) { CVarMetaSoundBlockRateChanged.Broadcast(Var); }),
 			ECVF_Default);
 
+		FConsoleVariableMulticastDelegate CVarMetaSoundSampleRateChanged;
 		FAutoConsoleVariableRef CVarMetaSoundSampleRate(
 			TEXT("au.MetaSound.SampleRate"),
-			AssetBasePrivate::SampleRate,
+			AssetBasePrivate::SampleRateOverride,
 			TEXT("Overrides the sample rate of metasounds. Negative values default to audio mixer sample rate.\n")
-			TEXT("Default: -1, Min: 8000, Max: 48000"),
+			TEXT("Default: 0, Min: 8000, Max: 48000"),
+			FConsoleVariableDelegate::CreateLambda([](IConsoleVariable* Var) { CVarMetaSoundSampleRateChanged.Broadcast(Var); }),
 			ECVF_Default);
 
+		float GetBlockRateOverride()
+		{
+			if(AssetBasePrivate::BlockRateOverride > 0)
+			{
+				return FMath::Clamp(AssetBasePrivate::BlockRateOverride, 
+					GetBlockRateClampRange().GetLowerBoundValue(), 
+					GetBlockRateClampRange().GetUpperBoundValue()
+				);
+			}
+			return AssetBasePrivate::BlockRateOverride;
+		}
 
-		float GetDefaultBlockRate()
+		FConsoleVariableMulticastDelegate& GetBlockRateOverrideChangedDelegate()
 		{
-			return FMath::Clamp(AssetBasePrivate::BlockRate, 1.0f, 1000.0f);
+			return CVarMetaSoundBlockRateChanged;
 		}
 
-		int32 GetDefaultSampleRate()
+		int32 GetSampleRateOverride()
 		{
-			if (AssetBasePrivate::SampleRate > 0)
+			if (AssetBasePrivate::SampleRateOverride > 0)
 			{
-				return FMath::Clamp(AssetBasePrivate::SampleRate, 8000, 48000);
+				return FMath::Clamp(AssetBasePrivate::SampleRateOverride, 
+					GetSampleRateClampRange().GetLowerBoundValue(),
+					GetSampleRateClampRange().GetUpperBoundValue()
+				);
 			}
-			return AssetBasePrivate::SampleRate;
+			return AssetBasePrivate::SampleRateOverride;
+		}
+
+		FConsoleVariableMulticastDelegate& GetSampleRateOverrideChangedDelegate()
+		{
+			return CVarMetaSoundSampleRateChanged;
+		}
+		
+		TRange<float> GetBlockRateClampRange()
+		{
+			return TRange<float>(100.f,1000.f);
+		}
+
+		TRange<int32> GetSampleRateClampRange()
+		{
+			return TRange<int32>(8000, 96000);
 		}
 	} // namespace Frontend
 } // namespace Metasound
diff --git a/Engine/Plugins/Runtime/Metasound/Source/MetasoundFrontend/Public/Analysis/MetasoundFrontendGraphAnalyzerView.h b/Engine/Plugins/Runtime/Metasound/Source/MetasoundFrontend/Public/Analysis/MetasoundFrontendGraphAnalyzerView.h
--- a/Engine/Plugins/Runtime/Metasound/Source/MetasoundFrontend/Public/Analysis/MetasoundFrontendGraphAnalyzerView.h
+++ b/Engine/Plugins/Runtime/Metasound/Source/MetasoundFrontend/Public/Analysis/MetasoundFrontendGraphAnalyzerView.h
@@ -44,7 +44,7 @@ namespace Metasound
 		public:
 			FMetasoundGraphAnalyzerView() = default;
 
-			FMetasoundGraphAnalyzerView(const FMetasoundAssetBase& InAssetBase, uint64 InInstanceID, FSampleRate InSampleRate);
+			FMetasoundGraphAnalyzerView(const FMetasoundAssetBase& InAssetBase, uint64 InInstanceID, const FOperatorSettings& InOperatorSettings);
 			~FMetasoundGraphAnalyzerView();
 
 			void AddAnalyzerForAllSupportedOutputs(FName InAnalyzerName, bool bInRequiresConnection = true);
diff --git a/Engine/Plugins/Runtime/Metasound/Source/MetasoundFrontend/Public/MetasoundAssetBase.h b/Engine/Plugins/Runtime/Metasound/Source/MetasoundFrontend/Public/MetasoundAssetBase.h
--- a/Engine/Plugins/Runtime/Metasound/Source/MetasoundFrontend/Public/MetasoundAssetBase.h
+++ b/Engine/Plugins/Runtime/Metasound/Source/MetasoundFrontend/Public/MetasoundAssetBase.h
@@ -19,7 +19,8 @@
 // Forward Declarations
 class UEdGraph;
 struct FMetaSoundFrontendDocumentBuilder;
-
+class IConsoleVariable;
+typedef TMulticastDelegate<void(IConsoleVariable*), FDefaultDelegateUserPolicy> FConsoleVariableMulticastDelegate;
 
 namespace Metasound
 {
@@ -28,9 +29,13 @@ namespace Metasound
 		// Forward Declarations
 		class IInterfaceRegistryEntry;
 
-		METASOUNDFRONTEND_API float GetDefaultBlockRate();
-		METASOUNDFRONTEND_API int32 GetDefaultSampleRate();
+		METASOUNDFRONTEND_API TRange<float> GetBlockRateClampRange();
+		METASOUNDFRONTEND_API float GetBlockRateOverride();
+		METASOUNDFRONTEND_API FConsoleVariableMulticastDelegate& GetBlockRateOverrideChangedDelegate();
 
+		METASOUNDFRONTEND_API TRange<int32> GetSampleRateClampRange();
+		METASOUNDFRONTEND_API int32 GetSampleRateOverride();
+		METASOUNDFRONTEND_API FConsoleVariableMulticastDelegate& GetSampleRateOverrideChangedDelegate();
 		class FProxyDataCache;
 
 	} // namespace Frontend
