Make Node Registration more robust to graph build errors.
#rnx

diff --git a/Engine/Plugins/Runtime/Metasound/Source/MetasoundFrontend/Private/MetasoundAssetBase.cpp b/Engine/Plugins/Runtime/Metasound/Source/MetasoundFrontend/Private/MetasoundAssetBase.cpp
--- a/Engine/Plugins/Runtime/Metasound/Source/MetasoundFrontend/Private/MetasoundAssetBase.cpp
+++ b/Engine/Plugins/Runtime/Metasound/Source/MetasoundFrontend/Private/MetasoundAssetBase.cpp
@@ -378,6 +378,7 @@ void FMetasoundAssetBase::UnregisterGraphWithFrontend()
 	using namespace Metasound::Frontend;
 	METASOUND_TRACE_CPUPROFILER_EVENT_SCOPE(MetaSoundAssetBase::UnregisterGraphWithFrontend);
 
+	check(IsInGameThread());
 	if (!NodeRegistryKey::IsValid(RegistryKey))
 	{
 		return;
@@ -389,7 +390,12 @@ void FMetasoundAssetBase::UnregisterGraphWithFrontend()
 		return;
 	}
 
-	ensureAlways(FMetasoundFrontendRegistryContainer::Get()->UnregisterNode(RegistryKey));
+	const bool bSuccess = FMetasoundFrontendRegistryContainer::Get()->UnregisterNode(RegistryKey);
+	if (!bSuccess)
+	{
+		UE_LOG(LogMetaSound, Warning, TEXT("Failed to unregister node with key %s for asset %s. No registry entry exists with that key."), *RegistryKey, *GetOwningAssetName());
+	}
+
 	RegistryKey = FNodeRegistryKey();
 }
 
@@ -648,6 +654,7 @@ const FMetasoundFrontendDocumentModifyContext& FMetasoundAssetBase::GetModifyCon
 bool FMetasoundAssetBase::IsRegistered() const
 {
 	using namespace Metasound::Frontend;
+	check(IsInGameThread());
 
 	return NodeRegistryKey::IsValid(RegistryKey);
 }
@@ -828,6 +835,7 @@ const FMetasoundFrontendDocument& FMetasoundAssetBase::GetDocumentChecked() cons
 
 const Metasound::Frontend::FNodeRegistryKey& FMetasoundAssetBase::GetRegistryKey() const
 {
+	check(IsInGameThread());
 	return RegistryKey;
 }
 
@@ -952,6 +960,7 @@ Metasound::Frontend::FNodeRegistryKey FMetasoundAssetBase::CacheRuntimeData(cons
 
 void FMetasoundAssetBase::WaitForAsyncGraphRegistration()
 {
+	check(IsInGameThread());
 	using namespace Metasound::Frontend;
 	if (NodeRegistryKey::IsValid(RegistryKey))
 	{
diff --git a/Engine/Plugins/Runtime/Metasound/Source/MetasoundFrontend/Private/MetasoundFrontendNodeRegistryPrivate.h b/Engine/Plugins/Runtime/Metasound/Source/MetasoundFrontend/Private/MetasoundFrontendNodeRegistryPrivate.h
--- a/Engine/Plugins/Runtime/Metasound/Source/MetasoundFrontend/Private/MetasoundFrontendNodeRegistryPrivate.h
+++ b/Engine/Plugins/Runtime/Metasound/Source/MetasoundFrontend/Private/MetasoundFrontendNodeRegistryPrivate.h
@@ -141,7 +141,7 @@ namespace Metasound
 
 			static FRegistryContainerImpl* LazySingleton;
 
-			void RegisterGraph(const FNodeRegistryKey& InKey, TSharedRef<const FGraph> InGraph);
+			void RegisterGraph(const FNodeRegistryKey& InKey, TSharedPtr<const FGraph> InGraph);
 			const INodeRegistryEntry* FindNodeEntry(const FNodeRegistryKey& InKey) const;
 
 			const INodeRegistryTemplateEntry* FindNodeTemplateEntry(const FNodeRegistryKey& InKey) const;
@@ -162,7 +162,7 @@ namespace Metasound
 			// Registry in which we keep all information about dynamically-generated templated nodes via in C++.
 			TMap<FNodeRegistryKey, TSharedRef<INodeRegistryTemplateEntry, ESPMode::ThreadSafe>> RegisteredNodeTemplates;
 
-			TMap<FNodeRegistryKey, TSharedRef<const FGraph>> RegisteredGraphs;
+			TMap<FNodeRegistryKey, TSharedPtr<const FGraph>> RegisteredGraphs;
 
 			// Registry in which we keep lists of possible nodes to use to convert between two datatypes
 			TMap<FConverterNodeRegistryKey, FConverterNodeRegistryValue> ConverterNodeRegistry;
diff --git a/Engine/Plugins/Runtime/Metasound/Source/MetasoundFrontend/Private/MetasoundFrontendRegistries.cpp b/Engine/Plugins/Runtime/Metasound/Source/MetasoundFrontend/Private/MetasoundFrontendRegistries.cpp
--- a/Engine/Plugins/Runtime/Metasound/Source/MetasoundFrontend/Private/MetasoundFrontendRegistries.cpp
+++ b/Engine/Plugins/Runtime/Metasound/Source/MetasoundFrontend/Private/MetasoundFrontendRegistries.cpp
@@ -184,7 +184,7 @@ namespace Metasound
 			class FDocumentNodeRegistryEntry : public INodeRegistryEntry
 			{
 			public:
-				FDocumentNodeRegistryEntry(const FMetasoundFrontendGraphClass& InGraphClass, const TSet<FMetasoundFrontendVersion>& InInterfaces, const FNodeClassInfo& InNodeClassInfo, TSharedRef<const IGraph> InGraph)
+				FDocumentNodeRegistryEntry(const FMetasoundFrontendGraphClass& InGraphClass, const TSet<FMetasoundFrontendVersion>& InInterfaces, const FNodeClassInfo& InNodeClassInfo, TSharedPtr<const IGraph> InGraph)
 				: FrontendClass(InGraphClass)
 				, Interfaces(InInterfaces)
 				, ClassInfo(InNodeClassInfo)
@@ -206,7 +206,15 @@ namespace Metasound
 
 				virtual TUniquePtr<INode> CreateNode(const FNodeInitData& InNodeInitData) const override
 				{
-					return MakeUnique<FGraphNode>(InNodeInitData, Graph);
+					if (Graph.IsValid())
+					{
+						return MakeUnique<FGraphNode>(InNodeInitData, Graph.ToSharedRef());
+					}
+					else
+					{
+						UE_LOG(LogMetaSound, Error, TEXT("Cannot create MetaSound node from asset %s due to prior failure to build graph"), *ClassInfo.AssetPath.ToString());
+						return TUniquePtr<INode>();
+					}
 				}
 
 				virtual TUniquePtr<INode> CreateNode(FDefaultLiteralNodeConstructorParams&&) const override { return nullptr; }
@@ -238,7 +246,7 @@ namespace Metasound
 				FMetasoundFrontendClass FrontendClass;
 				TSet<FMetasoundFrontendVersion> Interfaces;
 				FNodeClassInfo ClassInfo;
-				TSharedRef<const IGraph> Graph;
+				TSharedPtr<const IGraph> Graph;
 			};
 
 			// Builds an IGraph from a FMetasoundFrontendDocument and registers it in the node registry.
@@ -250,10 +258,9 @@ namespace Metasound
 				if (!FrontendGraph.IsValid())
 				{
 					UE_LOG(LogMetaSound, Error, TEXT("Failed to build MetaSound graph in asset '%s'"), *InNodeClassInfo.AssetPath.ToString());
-					return;
 				}
 
-				TSharedRef<const FGraph> GraphToRegister = MakeShareable<const FGraph>(FrontendGraph.Release());
+				TSharedPtr<const FGraph> GraphToRegister = MakeShareable<const FGraph>(FrontendGraph.Release());
 
 				TUniquePtr<INodeRegistryEntry> RegistryEntry = MakeUnique<FDocumentNodeRegistryEntry>(
 						InDocument.RootGraph, 
@@ -404,6 +411,7 @@ namespace Metasound
 			using namespace UE;
 
 			check(InDocumentInterface);
+			check(IsInGameThread());
 
 			const FMetasoundFrontendDocument& Document = InDocumentInterface->GetConstDocument();
 			FNodeRegistryKey RegistryKey = NodeRegistryKey::CreateKey(Document.RootGraph);
@@ -415,7 +423,7 @@ namespace Metasound
 
 			if (bAsync)
 			{	
-				// Wait for any async tasks that are inflight which correspond to the same graph
+				// Wait for any async tasks that are in flight which correspond to the same graph
 				FScopeLock LockActiveReg(&ActiveRegistrationTasksCriticalSection);
 				if (const FActiveRegistrationTaskInfo* ActiveTaskInfo = ActiveRegistrationTasks.Find(RegistryKey))
 				{
@@ -465,11 +473,11 @@ namespace Metasound
 			return RegistryKey;
 		}
 
-		void FRegistryContainerImpl::RegisterGraph(const FNodeRegistryKey& InKey, TSharedRef<const FGraph> InGraph)
+		void FRegistryContainerImpl::RegisterGraph(const FNodeRegistryKey& InKey, TSharedPtr<const FGraph> InGraph)
 		{
 			FScopeLock Lock(&RegistryMapsCriticalSection);
 
-			if (const TSharedRef<const FGraph>* ExistingGraph = RegisteredGraphs.Find(InKey))
+			if (const TSharedPtr<const FGraph>* ExistingGraph = RegisteredGraphs.Find(InKey))
 			{
 				UE_LOG(LogMetaSound, Error, TEXT("Multiple graphs are registered with the same registry key (%s). The existing registered graph (%s) will be replaced with the new graph (%s)."),  *InKey, *((*ExistingGraph)->GetInstanceName().ToString()), *(InGraph->GetInstanceName().ToString()));
 			}
@@ -484,7 +492,7 @@ namespace Metasound
 			TSharedPtr<const FGraph> Graph;
 			{
 				FScopeLock Lock(&RegistryMapsCriticalSection);
-				if (const TSharedRef<const FGraph>* RegisteredGraph = RegisteredGraphs.Find(InRegistryKey))
+				if (const TSharedPtr<const FGraph>* RegisteredGraph = RegisteredGraphs.Find(InRegistryKey))
 				{
 					Graph = *RegisteredGraph;
 				}
@@ -577,7 +585,8 @@ namespace Metasound
 		bool FRegistryContainerImpl::UnregisterNode(const FNodeRegistryKey& InKey)
 		{
 			METASOUND_LLM_SCOPE;
-
+			
+			check(IsInGameThread());
 			if (NodeRegistryKey::IsValid(InKey))
 			{
 				if (const INodeRegistryEntry* Entry = FindNodeEntry(InKey))
