CIS fix crash in MetaSoundAssetBase.cpp check(IsInGameThread())
#rnx

diff --git a/Engine/Plugins/Runtime/Metasound/Source/MetasoundFrontend/Private/MetasoundAssetBase.cpp b/Engine/Plugins/Runtime/Metasound/Source/MetasoundFrontend/Private/MetasoundAssetBase.cpp
--- a/Engine/Plugins/Runtime/Metasound/Source/MetasoundFrontend/Private/MetasoundAssetBase.cpp
+++ b/Engine/Plugins/Runtime/Metasound/Source/MetasoundFrontend/Private/MetasoundAssetBase.cpp
@@ -654,7 +654,6 @@ const FMetasoundFrontendDocumentModifyContext& FMetasoundAssetBase::GetModifyCon
 bool FMetasoundAssetBase::IsRegistered() const
 {
 	using namespace Metasound::Frontend;
-	check(IsInGameThread());
 
 	return NodeRegistryKey::IsValid(RegistryKey);
 }
@@ -835,7 +834,6 @@ const FMetasoundFrontendDocument& FMetasoundAssetBase::GetDocumentChecked() cons
 
 const Metasound::Frontend::FNodeRegistryKey& FMetasoundAssetBase::GetRegistryKey() const
 {
-	check(IsInGameThread());
 	return RegistryKey;
 }
 
@@ -960,7 +958,6 @@ Metasound::Frontend::FNodeRegistryKey FMetasoundAssetBase::CacheRuntimeData(cons
 
 void FMetasoundAssetBase::WaitForAsyncGraphRegistration()
 {
-	check(IsInGameThread());
 	using namespace Metasound::Frontend;
 	if (NodeRegistryKey::IsValid(RegistryKey))
 	{
