Fix editor crash for invalid MetaSound graph connection view when previewing sound

diff --git a/Engine/Plugins/Runtime/Metasound/Source/MetasoundEditor/Private/MetasoundAssetDefinitions.cpp b/Engine/Plugins/Runtime/Metasound/Source/MetasoundEditor/Private/MetasoundAssetDefinitions.cpp
--- a/Engine/Plugins/Runtime/Metasound/Source/MetasoundEditor/Private/MetasoundAssetDefinitions.cpp
+++ b/Engine/Plugins/Runtime/Metasound/Source/MetasoundEditor/Private/MetasoundAssetDefinitions.cpp
@@ -255,7 +255,35 @@ bool UAssetDefinition_MetaSoundSource::CanExecuteSoloCommand(const FToolMenuCont
 
 TSharedPtr<SWidget> UAssetDefinition_MetaSoundSource::GetThumbnailOverlay(const FAssetData& InAssetData) const
 {
-	return UAssetDefinition_SoundBase::GetSoundBaseThumbnailOverlay(InAssetData);
+	auto OnClickedLambdaOverride = [InAssetData]() -> FReply
+	{
+		TSharedPtr<Metasound::Editor::FEditor> Editor = Metasound::Editor::FGraphBuilder::GetEditorForMetasound(*InAssetData.GetAsset());
+		if (UE::AudioEditor::IsSoundPlaying(InAssetData))
+		{
+			if (Editor.IsValid())
+			{
+				Editor->Stop();
+			}
+			else
+			{
+				UE::AudioEditor::StopSound();
+			}
+		}
+		else
+		{
+			if (Editor.IsValid())
+			{
+				Editor->Play();
+			}
+			else
+			{
+				// Load and play sound
+				UE::AudioEditor::PlaySound(Cast<USoundBase>(InAssetData.GetAsset()));
+			}
+		}
+		return FReply::Handled();
+	};
+	return UAssetDefinition_SoundBase::GetSoundBaseThumbnailOverlay(InAssetData, MoveTemp(OnClickedLambdaOverride));
 }
 
 EAssetCommandResult UAssetDefinition_MetaSoundSource::ActivateAssets(const FAssetActivateArgs& ActivateArgs) const
diff --git a/Engine/Plugins/Runtime/Metasound/Source/MetasoundEditor/Private/MetasoundEditorGraphConnectionManager.h b/Engine/Plugins/Runtime/Metasound/Source/MetasoundEditor/Private/MetasoundEditorGraphConnectionManager.h
--- a/Engine/Plugins/Runtime/Metasound/Source/MetasoundEditor/Private/MetasoundEditorGraphConnectionManager.h
+++ b/Engine/Plugins/Runtime/Metasound/Source/MetasoundEditor/Private/MetasoundEditorGraphConnectionManager.h
@@ -119,24 +119,26 @@ namespace Metasound
 				using namespace Frontend;
 
 				const Frontend::FMetasoundGraphAnalyzerView* GraphView = GraphAnalyzerView.Get();
-				TArray<const FMetasoundAnalyzerView*> Views = GraphView->GetAnalyzerViewsForOutput(InNodeID, InOutputName, ForwardValueAnalyzerClass::GetAnalyzerName());
-				if (!Views.IsEmpty())
+				if (GraphView)
 				{
-					const FMetasoundAnalyzerView** View = Views.FindByPredicate([](const FMetasoundAnalyzerView* Candidate)
+					TArray<const FMetasoundAnalyzerView*> Views = GraphView->GetAnalyzerViewsForOutput(InNodeID, InOutputName, ForwardValueAnalyzerClass::GetAnalyzerName());
+					if (!Views.IsEmpty())
 					{
-						check(Candidate);
-						return Candidate->AnalyzerAddress.AnalyzerMemberName == ForwardValueAnalyzerClass::FOutputs::GetValue().Name;
-					});
+						const FMetasoundAnalyzerView** View = Views.FindByPredicate([](const FMetasoundAnalyzerView* Candidate)
+						{
+							check(Candidate);
+							return Candidate->AnalyzerAddress.AnalyzerMemberName == ForwardValueAnalyzerClass::FOutputs::GetValue().Name;
+						});
 
-					if (View)
-					{
-						check(*View);
-						const FAnalyzerAddress& Address = (*View)->AnalyzerAddress;
-						const FString AnalyzerKey = Address.ToString();
-						return AnalyzerKey;
+						if (View)
+						{
+							check(*View);
+							const FAnalyzerAddress& Address = (*View)->AnalyzerAddress;
+							const FString AnalyzerKey = Address.ToString();
+							return AnalyzerKey;
+						}
 					}
 				}
-
 				return { };
 			}
 
diff --git a/Engine/Source/Editor/AudioEditor/Private/AssetTypeActions/AssetDefinition_SoundBase.cpp b/Engine/Source/Editor/AudioEditor/Private/AssetTypeActions/AssetDefinition_SoundBase.cpp
--- a/Engine/Source/Editor/AudioEditor/Private/AssetTypeActions/AssetDefinition_SoundBase.cpp
+++ b/Engine/Source/Editor/AudioEditor/Private/AssetTypeActions/AssetDefinition_SoundBase.cpp
@@ -91,7 +91,7 @@ EAssetCommandResult UAssetDefinition_SoundBase::ActivateSoundBase(const FAssetAc
 	return EAssetCommandResult::Unhandled;
 }
 
-TSharedPtr<SWidget> UAssetDefinition_SoundBase::GetSoundBaseThumbnailOverlay(const FAssetData& InAssetData)
+TSharedPtr<SWidget> UAssetDefinition_SoundBase::GetSoundBaseThumbnailOverlay(const FAssetData& InAssetData, TUniqueFunction<FReply()>&& OnClickedLambdaOverride)
 {
 	auto OnGetDisplayBrushLambda = [InAssetData]() -> const FSlateBrush*
 	{
@@ -175,7 +175,20 @@ EAssetCommandResult UAssetDefinition_SoundBase::ActivateAssets(const FAssetActiv
 
 TSharedPtr<SWidget> UAssetDefinition_SoundBase::GetThumbnailOverlay(const FAssetData& InAssetData) const
 {
-	return GetSoundBaseThumbnailOverlay(InAssetData);
+	auto OnClickedLambda = [InAssetData]() -> FReply
+	{
+		if (UE::AudioEditor::IsSoundPlaying(InAssetData))
+		{
+			UE::AudioEditor::StopSound();
+		}
+		else
+		{
+			// Load and play sound
+			UE::AudioEditor::PlaySound(Cast<USoundBase>(InAssetData.GetAsset()));
+		}
+		return FReply::Handled();
+	};
+	return GetSoundBaseThumbnailOverlay(InAssetData, MoveTemp(OnClickedLambda));
 }
 
 // Menu Extensions
diff --git a/Engine/Source/Editor/AudioEditor/Public/AssetTypeActions/AssetDefinition_SoundBase.h b/Engine/Source/Editor/AudioEditor/Public/AssetTypeActions/AssetDefinition_SoundBase.h
--- a/Engine/Source/Editor/AudioEditor/Public/AssetTypeActions/AssetDefinition_SoundBase.h
+++ b/Engine/Source/Editor/AudioEditor/Public/AssetTypeActions/AssetDefinition_SoundBase.h
@@ -50,7 +50,6 @@ public:
 	static bool CanExecuteSoloCommand(const FToolMenuContext& InContext);
 
 	// Asset definition static utilities
-	static TSharedPtr<SWidget> GetSoundBaseThumbnailOverlay(const FAssetData& InAssetData);
+	static TSharedPtr<SWidget> GetSoundBaseThumbnailOverlay(const FAssetData& InAssetData, TUniqueFunction<FReply()>&& OnClickedLambda);
 	static EAssetCommandResult ActivateSoundBase(const FAssetActivateArgs& ActivateArgs);
-
 };
