Fix deadlock in metasound node registration when updating assets
#rnx

diff --git a/Engine/Plugins/Runtime/Metasound/Source/MetasoundEngine/Private/MetasoundSource.cpp b/Engine/Plugins/Runtime/Metasound/Source/MetasoundEngine/Private/MetasoundSource.cpp
--- a/Engine/Plugins/Runtime/Metasound/Source/MetasoundEngine/Private/MetasoundSource.cpp
+++ b/Engine/Plugins/Runtime/Metasound/Source/MetasoundEngine/Private/MetasoundSource.cpp
@@ -663,7 +663,7 @@ ISoundGeneratorPtr UMetaSoundSource::CreateSoundGenerator(const FSoundGeneratorI
 	}
 	else
 	{
-		TSharedPtr<const IGraph> MetasoundGraph = FMetasoundFrontendRegistryContainer::Get()->GetGraph(GetRegistryKey());
+		TSharedPtr<const IGraph> MetasoundGraph = GetRegisteredGraph();
 		if (!MetasoundGraph.IsValid())
 		{
 			return ISoundGeneratorPtr(nullptr);
@@ -1400,7 +1400,7 @@ TSharedPtr<Metasound::DynamicGraph::FDynamicOperatorTransactor> UMetaSoundSource
 			// graph to see if any FGraph already exists. 
 			if (IsRegistered())
 			{
-				TSharedPtr<const FGraph> CurrentGraph = FMetasoundFrontendRegistryContainer::Get()->GetGraph(GetRegistryKey());
+				TSharedPtr<const FGraph> CurrentGraph = GetRegisteredGraph();
 				if (CurrentGraph.IsValid())
 				{
 					DynamicTransactor = MakeShared<FDynamicOperatorTransactor>(*CurrentGraph);
diff --git a/Engine/Plugins/Runtime/Metasound/Source/MetasoundFrontend/Private/MetasoundAssetBase.cpp b/Engine/Plugins/Runtime/Metasound/Source/MetasoundFrontend/Private/MetasoundAssetBase.cpp
--- a/Engine/Plugins/Runtime/Metasound/Source/MetasoundFrontend/Private/MetasoundAssetBase.cpp
+++ b/Engine/Plugins/Runtime/Metasound/Source/MetasoundFrontend/Private/MetasoundAssetBase.cpp
@@ -28,6 +28,7 @@
 #include "MetasoundFrontendRegistries.h"
 #include "MetasoundFrontendSearchEngine.h"
 #include "MetasoundFrontendTransform.h"
+#include "MetasoundGraph.h"
 #include "MetasoundJsonBackend.h"
 #include "MetasoundLog.h"
 #include "MetasoundParameterPack.h"
@@ -181,14 +182,14 @@ namespace Metasound
 
 		namespace ConsoleVariables
 		{
-			static bool bDisableAsyncGraphRegistration = true;
+			static bool bDisableAsyncGraphRegistration = false;
 		}
 
 		FAutoConsoleVariableRef CVarMetaSoundDisableAsyncGraphRegistration(
 			TEXT("au.MetaSound.Experimental.DisableAsyncGraphRegistration"),
 			Metasound::Frontend::ConsoleVariables::bDisableAsyncGraphRegistration,
 			TEXT("Disables async registration of MetaSound graphs\n")
-			TEXT("Default: true"),
+			TEXT("Default: false"),
 			ECVF_Default);
 
 		FAutoConsoleVariableRef CVarMetaSoundBlockRate(
@@ -308,6 +309,11 @@ void FMetasoundAssetBase::RegisterGraphWithFrontend(Metasound::Frontend::FMetaSo
 	}
 }
 
+TSharedPtr<const Metasound::FGraph> FMetasoundAssetBase::GetRegisteredGraph() const
+{
+	return FMetasoundFrontendRegistryContainer::Get()->GetGraph(RegistryKey, RegisteredGraphAssetPath);
+}
+
 void FMetasoundAssetBase::CookMetaSound()
 {
 	using namespace Metasound;
@@ -390,13 +396,23 @@ void FMetasoundAssetBase::UnregisterGraphWithFrontend()
 		return;
 	}
 
-	const bool bSuccess = FMetasoundFrontendRegistryContainer::Get()->UnregisterNode(RegistryKey);
+	bool bSuccess = false;
+	if (RegisteredGraphAssetPath.IsValid())
+	{
+		bSuccess = FMetasoundFrontendRegistryContainer::Get()->UnregisterGraph(RegistryKey, RegisteredGraphAssetPath);
+	}
+	else
+	{
+		bSuccess = FMetasoundFrontendRegistryContainer::Get()->UnregisterNode(RegistryKey);
+	}
+
 	if (!bSuccess)
 	{
 		UE_LOG(LogMetaSound, Warning, TEXT("Failed to unregister node with key %s for asset %s. No registry entry exists with that key."), *RegistryKey, *GetOwningAssetName());
 	}
 
 	RegistryKey = FNodeRegistryKey();
+	RegisteredGraphAssetPath = FSoftObjectPath();
 }
 
 void FMetasoundAssetBase::SetMetadata(FMetasoundFrontendClassMetadata& InMetadata)
@@ -949,8 +965,8 @@ Metasound::Frontend::FNodeRegistryKey FMetasoundAssetBase::CacheRuntimeData(cons
 	// the IMetaSoundDocumentInterface
 	const bool bAsync = !(IsBuilderActive() || ConsoleVariables::bDisableAsyncGraphRegistration);
 
-	FNodeClassInfo NodeClassInfo { Document.RootGraph, FSoftObjectPath(GetOwningAsset()) };
-	FNodeRegistryKey NewRegistryKey = FMetasoundFrontendRegistryContainer::Get()->RegisterGraph(NodeClassInfo, InDocumentInterface, bAsync);
+	RegisteredGraphAssetPath = FSoftObjectPath(GetOwningAsset());
+	FNodeRegistryKey NewRegistryKey = FMetasoundFrontendRegistryContainer::Get()->RegisterGraph(RegisteredGraphAssetPath, InDocumentInterface, bAsync);
 
 
 	return NewRegistryKey;
@@ -961,7 +977,7 @@ void FMetasoundAssetBase::WaitForAsyncGraphRegistration()
 	using namespace Metasound::Frontend;
 	if (NodeRegistryKey::IsValid(RegistryKey))
 	{
-		FMetasoundFrontendRegistryContainer::Get()->WaitForAsyncGraphRegistration(RegistryKey);
+		FMetasoundFrontendRegistryContainer::Get()->WaitForAsyncGraphRegistration(RegistryKey, RegisteredGraphAssetPath);
 	}
 }
 
diff --git a/Engine/Plugins/Runtime/Metasound/Source/MetasoundFrontend/Private/MetasoundFrontendNodeRegistryPrivate.h b/Engine/Plugins/Runtime/Metasound/Source/MetasoundFrontend/Private/MetasoundFrontendNodeRegistryPrivate.h
--- a/Engine/Plugins/Runtime/Metasound/Source/MetasoundFrontend/Private/MetasoundFrontendNodeRegistryPrivate.h
+++ b/Engine/Plugins/Runtime/Metasound/Source/MetasoundFrontend/Private/MetasoundFrontendNodeRegistryPrivate.h
@@ -84,15 +84,18 @@ namespace Metasound
 			virtual void RegisterPendingNodes() override;
 
 			// Register a graph from an IMetaSoundDocumentInterface
-			virtual FNodeRegistryKey RegisterGraph(const FNodeClassInfo& InAssetPath, const TScriptInterface<IMetaSoundDocumentInterface>& InDocument, bool bAsync) override;
+			virtual FNodeRegistryKey RegisterGraph(const FSoftObjectPath& InAssetPath, const TScriptInterface<IMetaSoundDocumentInterface>& InDocument, bool bAsync) override;
 			
 			// Wait for async graph registration to complete for a specific graph
-			virtual void WaitForAsyncGraphRegistration(const FNodeRegistryKey& InRegistryKey) const override;
+			virtual void WaitForAsyncGraphRegistration(const FNodeRegistryKey& InRegistryKey, const FSoftObjectPath& InAssetPath) const override;
 			
+			// Unregister a graph 
+			virtual bool UnregisterGraph(const FNodeRegistryKey& InNodeRegistryKey, const FSoftObjectPath& InAssetPath) override;
+
 			// Retrieve a registered graph. 
 			//
 			// If the graph is registered asynchronously, this will wait until the registration task has completed.
-			virtual TSharedPtr<const FGraph> GetGraph(const FNodeRegistryKey& InRegistryKey) const override;
+			virtual TSharedPtr<const FGraph> GetGraph(const FNodeRegistryKey& InRegistryKey, const FSoftObjectPath& InAssetPath) const override;
 
 			/** Register external node with the frontend.
 			 *
@@ -141,7 +144,11 @@ namespace Metasound
 
 			static FRegistryContainerImpl* LazySingleton;
 
-			void RegisterGraph(const FNodeRegistryKey& InKey, TSharedPtr<const FGraph> InGraph);
+			using FGraphRegistryKey = TTuple<FNodeRegistryKey, FSoftObjectPath>;
+
+			void WaitForAsyncRegistrationInternal(const FNodeRegistryKey& InRegistryKey, const FSoftObjectPath* InAssetPath) const;
+			void RegisterGraphInternal(const FNodeRegistryKey& InKey, const FSoftObjectPath& InAssetPath, TSharedPtr<const FGraph> InGraph);
+
 			const INodeRegistryEntry* FindNodeEntry(const FNodeRegistryKey& InKey) const;
 
 			const INodeRegistryTemplateEntry* FindNodeTemplateEntry(const FNodeRegistryKey& InKey) const;
@@ -157,12 +164,13 @@ namespace Metasound
 			FCriticalSection LazyInitCommandCritSection;
 
 			// Registry in which we keep all information about nodes implemented in C++.
-			TMap<FNodeRegistryKey, TSharedRef<INodeRegistryEntry, ESPMode::ThreadSafe>> RegisteredNodes;
+			TMultiMap<FNodeRegistryKey, TSharedRef<INodeRegistryEntry, ESPMode::ThreadSafe>> RegisteredNodes;
 
 			// Registry in which we keep all information about dynamically-generated templated nodes via in C++.
 			TMap<FNodeRegistryKey, TSharedRef<INodeRegistryTemplateEntry, ESPMode::ThreadSafe>> RegisteredNodeTemplates;
 
-			TMap<FNodeRegistryKey, TSharedPtr<const FGraph>> RegisteredGraphs;
+			// Map of all registered graphs
+			TMap<FGraphRegistryKey, TSharedPtr<const FGraph>> RegisteredGraphs;
 
 			// Registry in which we keep lists of possible nodes to use to convert between two datatypes
 			TMap<FConverterNodeRegistryKey, FConverterNodeRegistryValue> ConverterNodeRegistry;
diff --git a/Engine/Plugins/Runtime/Metasound/Source/MetasoundFrontend/Private/MetasoundFrontendRegistries.cpp b/Engine/Plugins/Runtime/Metasound/Source/MetasoundFrontend/Private/MetasoundFrontendRegistries.cpp
--- a/Engine/Plugins/Runtime/Metasound/Source/MetasoundFrontend/Private/MetasoundFrontendRegistries.cpp
+++ b/Engine/Plugins/Runtime/Metasound/Source/MetasoundFrontend/Private/MetasoundFrontendRegistries.cpp
@@ -270,7 +270,7 @@ namespace Metasound
 				FRegistryContainerImpl& Registry = FRegistryContainerImpl::Get();
 
 				FNodeRegistryKey RegistryKey = Registry.RegisterNode(MoveTemp(RegistryEntry));
-				Registry.RegisterGraph(RegistryKey, MoveTemp(GraphToRegister));
+				Registry.RegisterGraphInternal(RegistryKey, InNodeClassInfo.AssetPath, MoveTemp(GraphToRegister));
 			}
 		}
 
@@ -406,7 +406,7 @@ namespace Metasound
 			return MakeUnique<FNodeRegistryTransactionStream>(TransactionBuffer);
 		}
 
-		FNodeRegistryKey FRegistryContainerImpl::RegisterGraph(const FNodeClassInfo& InNodeClassInfo, const TScriptInterface<IMetaSoundDocumentInterface>& InDocumentInterface, bool bAsync)
+		FNodeRegistryKey FRegistryContainerImpl::RegisterGraph(const FSoftObjectPath& InAssetPath, const TScriptInterface<IMetaSoundDocumentInterface>& InDocumentInterface, bool bAsync)
 		{
 			using namespace UE;
 
@@ -415,6 +415,7 @@ namespace Metasound
 
 			const FMetasoundFrontendDocument& Document = InDocumentInterface->GetConstDocument();
 			FNodeRegistryKey RegistryKey = NodeRegistryKey::CreateKey(Document.RootGraph);
+			FNodeClassInfo NodeClassInfo(Document.RootGraph, InAssetPath);
 
 			// Proxies are created synchronously to avoid creating proxies in async tasks. Proxies
 			// are created from UObjects which need to be protected from GC and non-GT access.
@@ -427,13 +428,13 @@ namespace Metasound
 				FScopeLock LockActiveReg(&ActiveRegistrationTasksCriticalSection);
 				if (const FActiveRegistrationTaskInfo* ActiveTaskInfo = ActiveRegistrationTasks.Find(RegistryKey))
 				{
-					UE_LOG(LogMetaSound, Warning, TEXT("Waiting for async registration task to finish before beginning new registration task for same registration key (%s) with asset (%s))"), *RegistryKey, *InNodeClassInfo.AssetPath.ToString());
+					UE_LOG(LogMetaSound, Warning, TEXT("Waiting for async registration task to finish before beginning new registration task for same registration key (%s) with asset (%s))"), *RegistryKey, *InAssetPath.ToString());
 					ActiveTaskInfo->Task.Wait(FTimespan::FromSeconds(1.));
 				}
 			
 				Tasks::FTask BuildAndRegisterTask = AsyncRegistrationPipe.Launch(
 					UE_SOURCE_LOCATION,
-					[RegistryKey, NodeClassInfo=InNodeClassInfo, DocumentInterface=InDocumentInterface, ProxyDataCache=MoveTemp(ProxyDataCache)]()
+					[RegistryKey, NodeClassInfo=NodeClassInfo, DocumentInterface=InDocumentInterface, ProxyDataCache=MoveTemp(ProxyDataCache)]()
 					{
 						// Build and add the graph to the node registry
 						MetasoundFrontendRegistryPrivate::BuildAndRegisterGraphFromDocument(DocumentInterface->GetConstDocument(), ProxyDataCache, NodeClassInfo);
@@ -467,32 +468,52 @@ namespace Metasound
 			else
 			{
 				// Build and register graph synchronously
-				MetasoundFrontendRegistryPrivate::BuildAndRegisterGraphFromDocument(InDocumentInterface->GetConstDocument(), ProxyDataCache, InNodeClassInfo);
+				MetasoundFrontendRegistryPrivate::BuildAndRegisterGraphFromDocument(InDocumentInterface->GetConstDocument(), ProxyDataCache, NodeClassInfo);
 			}
 
 			return RegistryKey;
 		}
 
-		void FRegistryContainerImpl::RegisterGraph(const FNodeRegistryKey& InKey, TSharedPtr<const FGraph> InGraph)
+		void FRegistryContainerImpl::RegisterGraphInternal(const FNodeRegistryKey& InKey, const FSoftObjectPath& InAssetPath, TSharedPtr<const FGraph> InGraph)
 		{
 			FScopeLock Lock(&RegistryMapsCriticalSection);
 
-			if (const TSharedPtr<const FGraph>* ExistingGraph = RegisteredGraphs.Find(InKey))
+			FGraphRegistryKey GraphRegistryKey{InKey, InAssetPath};
+
+			if (const TSharedPtr<const FGraph>* ExistingEntry = RegisteredGraphs.Find(GraphRegistryKey))
 			{
-				UE_LOG(LogMetaSound, Error, TEXT("Multiple graphs are registered with the same registry key (%s). The existing registered graph (%s) will be replaced with the new graph (%s)."),  *InKey, *((*ExistingGraph)->GetInstanceName().ToString()), *(InGraph->GetInstanceName().ToString()));
+				UE_LOG(LogMetaSound, Warning, TEXT("Multiple graphs are registered with the same registry key (%s) and asset path. The existing registered graph will be replaced with the new graph."),  *InKey, *InAssetPath.ToString());
 			}
 
-			RegisteredGraphs.Add(InKey, MoveTemp(InGraph));
+			RegisteredGraphs.Add(GraphRegistryKey, MoveTemp(InGraph));
+		}
+
+		bool FRegistryContainerImpl::UnregisterGraph(const FNodeRegistryKey& InNodeRegistryKey, const FSoftObjectPath& InAssetPath)
+		{
+			WaitForAsyncGraphRegistration(InNodeRegistryKey, InAssetPath);
+
+			{
+				FScopeLock Lock(&RegistryMapsCriticalSection);
+
+				int32 NumRemoved = RegisteredGraphs.Remove(FGraphRegistryKey { InNodeRegistryKey, InAssetPath });
+
+				if (NumRemoved)
+				{
+					UnregisterNode(InNodeRegistryKey);
+					return true;
+				}
+				return false;
+			}
 		}
 
-		TSharedPtr<const Metasound::FGraph> FRegistryContainerImpl::GetGraph(const Metasound::Frontend::FNodeRegistryKey& InRegistryKey) const
+		TSharedPtr<const Metasound::FGraph> FRegistryContainerImpl::GetGraph(const FNodeRegistryKey& InNodeRegistryKey, const FSoftObjectPath& InAssetPath) const
 		{
-			WaitForAsyncGraphRegistration(InRegistryKey);
+			WaitForAsyncGraphRegistration(InNodeRegistryKey, InAssetPath);
 
 			TSharedPtr<const FGraph> Graph;
 			{
 				FScopeLock Lock(&RegistryMapsCriticalSection);
-				if (const TSharedPtr<const FGraph>* RegisteredGraph = RegisteredGraphs.Find(InRegistryKey))
+				if (const TSharedPtr<const FGraph>* RegisteredGraph = RegisteredGraphs.Find(FGraphRegistryKey{ InNodeRegistryKey, InAssetPath }))
 				{
 					Graph = *RegisteredGraph;
 				}
@@ -500,7 +521,7 @@ namespace Metasound
 
 			if (!Graph)
 			{
-				UE_LOG(LogMetaSound, Error, TEXT("Could not find graph with registry key (%s)."),  *InRegistryKey);
+				UE_LOG(LogMetaSound, Error, TEXT("Could not find graph with registry key (%s) and asset (%s)."),  *InNodeRegistryKey, *InAssetPath.ToString());
 			}
 
 			return Graph;
@@ -524,13 +545,11 @@ namespace Metasound
 					FScopeLock Lock(&RegistryMapsCriticalSection);
 
 					// check to see if an identical node was already registered, and log
-					if (RegisteredNodes.Contains(Key))
+					if (const TSharedRef<INodeRegistryEntry>* ExistingEntry = RegisteredNodes.Find(Key))
 					{
-						TSharedRef<INodeRegistryEntry, ESPMode::ThreadSafe> Node = RegisteredNodes[Key];
-						const FNodeClassInfo& ClassInfo = Node->GetClassInfo();
-						UE_LOG(LogMetaSound, Error, TEXT("Node with registry key '%s' already registered by asset '%s'. The previously registered node will be overwritten. " 
-							"This can happen if two classes share the same name or if METASOUND_REGISTER_NODE is defined in a public header."
-						), *Key, *ClassInfo.AssetPath.ToString());
+						const FNodeClassInfo& ClassInfo = (*ExistingEntry)->GetClassInfo();
+						UE_LOG(LogMetaSound, Error, TEXT("Node with registry key '%s' already registered by asset '%s' encountered while registering node with asset %s. MetaSounds which depend on these assets may utilize incorrect asset." 
+						), *Key, *ClassInfo.AssetPath.ToString(), *Entry->GetClassInfo().AssetPath.ToString());
 					}
 
 					// Store registry elements in map so nodes can be queried using registry key.
@@ -589,17 +608,17 @@ namespace Metasound
 			check(IsInGameThread());
 			if (NodeRegistryKey::IsValid(InKey))
 			{
-				if (const INodeRegistryEntry* Entry = FindNodeEntry(InKey))
+				FScopeLock Lock(&RegistryMapsCriticalSection);
+				if (const TSharedRef<INodeRegistryEntry, ESPMode::ThreadSafe>* EntryPtr = RegisteredNodes.Find(InKey))
 				{
-					FNodeRegistryTransaction::FTimeType Timestamp = FPlatformTime::Cycles64();
+					const TSharedRef<INodeRegistryEntry>& Entry = *EntryPtr;
 
+					FNodeRegistryTransaction::FTimeType Timestamp = FPlatformTime::Cycles64();
 					TransactionBuffer->AddTransaction(FNodeRegistryTransaction(FNodeRegistryTransaction::ETransactionType::NodeUnregistration, Entry->GetClassInfo(), Timestamp));
 
-					{
-						FScopeLock Lock(&RegistryMapsCriticalSection);
-						RegisteredNodes.Remove(InKey);
-						RegisteredGraphs.Remove(InKey);
-					}
+					uint32 NumRemoved = RegisteredNodes.RemoveSingle(InKey, Entry);
+					check(NumRemoved == 1);
+
 					return true;
 				}
 			}
@@ -653,11 +672,20 @@ namespace Metasound
 
 		bool FRegistryContainerImpl::IsNodeRegistered(const FNodeRegistryKey& InKey) const
 		{
-			WaitForAsyncGraphRegistration(InKey);
-
+			auto IsNodeRegisteredInternal = [this, &InKey]() -> bool
 			{
 				FScopeLock Lock(&RegistryMapsCriticalSection);
 				return RegisteredNodes.Contains(InKey) || RegisteredNodeTemplates.Contains(InKey);
+			};
+
+			if (IsNodeRegisteredInternal())
+			{
+				return true;
+			}
+			else
+			{
+				WaitForAsyncRegistrationInternal(InKey, nullptr /* InAssetPath */);
+				return IsNodeRegisteredInternal();
 			}
 		}
 
@@ -822,10 +850,9 @@ namespace Metasound
 
 		const INodeRegistryEntry* FRegistryContainerImpl::FindNodeEntry(const FNodeRegistryKey& InKey) const
 		{
-			WaitForAsyncGraphRegistration(InKey);
-
+			auto TryFindNodeEntry = [this, &InKey]() -> const INodeRegistryEntry*
 			{
-				// This scope lock proctects against race conditions manipulating the `RegisteredNodes` map, but it does not
+				// This scope lock protects against race conditions manipulating the `RegisteredNodes` map, but it does not
 				// protect against the individual INodeRegistryEntry from being removed after this function returns. Generally
 				// this is not an issue as the acess to the INodeRegistryEntry pointer happens on the same thread as calls to 
 				// remove node registry entries. 
@@ -834,9 +861,18 @@ namespace Metasound
 				{
 					return &Entry->Get();
 				}
-			}
+				return nullptr;
+			};
 
-			return nullptr;
+			if (const INodeRegistryEntry* Entry = TryFindNodeEntry())
+			{
+				return Entry;
+			}
+			else
+			{
+				WaitForAsyncRegistrationInternal(InKey, nullptr /* InAssetPath */);
+				return TryFindNodeEntry();
+			}
 		}
 
 
@@ -851,14 +887,31 @@ namespace Metasound
 			return nullptr;
 		}
 
-		void FRegistryContainerImpl::WaitForAsyncGraphRegistration(const FNodeRegistryKey& InRegistryKey) const
+		void FRegistryContainerImpl::WaitForAsyncGraphRegistration(const FNodeRegistryKey& InRegistryKey, const FSoftObjectPath& InAssetPath) const
+		{
+			WaitForAsyncRegistrationInternal(InRegistryKey, &InAssetPath);
+		}
+
+		void FRegistryContainerImpl::WaitForAsyncRegistrationInternal(const FNodeRegistryKey& InRegistryKey, const FSoftObjectPath* InAssetPath) const
 		{
 			UE::Tasks::FTask ActiveRegistrationTask;
 			{
 				FScopeLock Lock(&ActiveRegistrationTasksCriticalSection);
 				if (const FActiveRegistrationTaskInfo* FoundTask = ActiveRegistrationTasks.Find(InRegistryKey))
 				{
-					ActiveRegistrationTask = FoundTask->Task;
+					if (InAssetPath)
+					{
+						// Filter by asset path
+						if (*InAssetPath == FSoftObjectPath(FoundTask->OwningObject))
+						{
+							ActiveRegistrationTask = FoundTask->Task;
+						}
+					}
+					else
+					{
+						// ignore asset path
+						ActiveRegistrationTask = FoundTask->Task;
+					}
 				}
 			}
 
diff --git a/Engine/Plugins/Runtime/Metasound/Source/MetasoundFrontend/Public/MetasoundAssetBase.h b/Engine/Plugins/Runtime/Metasound/Source/MetasoundFrontend/Public/MetasoundAssetBase.h
--- a/Engine/Plugins/Runtime/Metasound/Source/MetasoundFrontend/Public/MetasoundAssetBase.h
+++ b/Engine/Plugins/Runtime/Metasound/Source/MetasoundFrontend/Public/MetasoundAssetBase.h
@@ -23,6 +23,7 @@ struct FMetaSoundFrontendDocumentBuilder;
 
 namespace Metasound
 {
+	class FGraph;
 	namespace Frontend
 	{
 		// Forward Declarations
@@ -79,6 +80,9 @@ public:
 	// Registers the root graph of the given asset with the MetaSound Frontend.
 	virtual void RegisterGraphWithFrontend(Metasound::Frontend::FMetaSoundAssetRegistrationOptions InRegistrationOptions = Metasound::Frontend::FMetaSoundAssetRegistrationOptions());
 
+	// Return the graph registered for this asset. 
+	TSharedPtr<const Metasound::FGraph> GetRegisteredGraph() const;
+
 	// Unregisters the root graph of the given asset with the MetaSound Frontend.
 	void UnregisterGraphWithFrontend();
 
@@ -268,7 +272,7 @@ private:
 	Metasound::Frontend::FNodeRegistryKey CacheRuntimeData(const TScriptInterface<IMetaSoundDocumentInterface>& InDoc);
 
 	Metasound::Frontend::FNodeRegistryKey RegistryKey;
-
+	FSoftObjectPath RegisteredGraphAssetPath; // This is cached to avoid reparsing path strings to create FSoftObjectPath from UObject
 
 	TSharedPtr<Metasound::FGraph, ESPMode::ThreadSafe> BuildMetasoundDocument(const FMetasoundFrontendDocument& InPreprocessDoc, const Metasound::Frontend::FProxyDataCache& InProxies) const;
 };
diff --git a/Engine/Plugins/Runtime/Metasound/Source/MetasoundFrontend/Public/MetasoundFrontendRegistries.h b/Engine/Plugins/Runtime/Metasound/Source/MetasoundFrontend/Public/MetasoundFrontendRegistries.h
--- a/Engine/Plugins/Runtime/Metasound/Source/MetasoundFrontend/Public/MetasoundFrontendRegistries.h
+++ b/Engine/Plugins/Runtime/Metasound/Source/MetasoundFrontend/Public/MetasoundFrontendRegistries.h
@@ -322,15 +322,18 @@ public:
 	virtual void RegisterPendingNodes() = 0;
 
 	// Register a graph from an IMetaSoundDocumentInterface
-	virtual FNodeRegistryKey RegisterGraph(const FNodeClassInfo& InAssetPath, const TScriptInterface<IMetaSoundDocumentInterface>& InDocument, bool bAsync=true) = 0;
+	virtual FNodeRegistryKey RegisterGraph(const FSoftObjectPath& InAssetPath, const TScriptInterface<IMetaSoundDocumentInterface>& InDocument, bool bAsync=true) = 0;
 
 	// Wait for async graph registration to complete for a specific graph
-	virtual void WaitForAsyncGraphRegistration(const FNodeRegistryKey& InRegistryKey) const = 0;
+	virtual void WaitForAsyncGraphRegistration(const FNodeRegistryKey& InRegistryKey, const FSoftObjectPath& InAssetPath) const = 0;
+
+	// Unregister an asset based graph from the frontend. 
+	virtual bool UnregisterGraph(const FNodeRegistryKey& InNodeRegistryKey, const FSoftObjectPath& InAssetPath) = 0;
 
 	// Retrieve a registered graph. 
 	//
 	// If the graph is registered asynchronously, this will wait until the registration task has completed.
-	virtual TSharedPtr<const Metasound::FGraph> GetGraph(const Metasound::Frontend::FNodeRegistryKey& InRegistryKey) const = 0;
+	virtual TSharedPtr<const Metasound::FGraph> GetGraph(const FNodeRegistryKey& InRegistryKey, const FSoftObjectPath& InAssetPath) const = 0;
 
 	/** Register an external node with the frontend.
 	 *
