Missed a file. Fixing metasound deadlock in node registration
#rnx

diff --git a/Engine/Plugins/Runtime/Metasound/Source/MetasoundFrontend/Private/MetasoundFrontendRegistries.cpp b/Engine/Plugins/Runtime/Metasound/Source/MetasoundFrontend/Private/MetasoundFrontendRegistries.cpp
--- a/Engine/Plugins/Runtime/Metasound/Source/MetasoundFrontend/Private/MetasoundFrontendRegistries.cpp
+++ b/Engine/Plugins/Runtime/Metasound/Source/MetasoundFrontend/Private/MetasoundFrontendRegistries.cpp
@@ -894,6 +894,12 @@ namespace Metasound
 
 		void FRegistryContainerImpl::WaitForAsyncRegistrationInternal(const FNodeRegistryKey& InRegistryKey, const FSoftObjectPath* InAssetPath) const
 		{
+			if (AsyncRegistrationPipe.IsInContext())
+			{
+				// It is not safe to wait for an async registration task from within the async registration pipe because it will result in a deadlock. 
+				return;
+			}
+
 			UE::Tasks::FTask ActiveRegistrationTask;
 			{
 				FScopeLock Lock(&ActiveRegistrationTasksCriticalSection);
