- Make FMetasoundFrontendClassName GetTypeHash performant by removing string copy/format
- Add InvalidFrontendClassName to test validity against instance/if class name has been set
- Minor clean-up
#rnx

diff --git a/Engine/Plugins/Runtime/Metasound/Source/MetasoundFrontend/Private/MetasoundFrontendDocument.cpp b/Engine/Plugins/Runtime/Metasound/Source/MetasoundFrontend/Private/MetasoundFrontendDocument.cpp
--- a/Engine/Plugins/Runtime/Metasound/Source/MetasoundFrontend/Private/MetasoundFrontendDocument.cpp
+++ b/Engine/Plugins/Runtime/Metasound/Source/MetasoundFrontend/Private/MetasoundFrontendDocument.cpp
@@ -363,6 +363,8 @@ const FMetasoundFrontendInterfaceUClassOptions* FMetasoundFrontendInterface::Fin
 	return UClassOptions.FindByPredicate(FindClassOptionsPredicate);
 }
 
+const FMetasoundFrontendClassName FMetasoundFrontendClassName::InvalidClassName;
+
 FMetasoundFrontendClassName::FMetasoundFrontendClassName(const FName& InNamespace, const FName& InName, const FName& InVariant)
 : Namespace(InNamespace)
 , Name(InName)
@@ -385,6 +387,17 @@ FName FMetasoundFrontendClassName::GetFullName() const
 	return Metasound::FNodeClassName::FormatFullName(Namespace, Name, Variant);
 }
 
+bool FMetasoundFrontendClassName::IsValid() const
+{
+	return *this != InvalidClassName;
+}
+
+// Returns NodeClassName version of full name
+Metasound::FNodeClassName FMetasoundFrontendClassName::ToNodeClassName() const
+{
+	return { Namespace, Name, Variant };
+}
+
 FString FMetasoundFrontendClassName::ToString() const
 {
 	return GetFullName().ToString();
diff --git a/Engine/Plugins/Runtime/Metasound/Source/MetasoundFrontend/Public/MetasoundFrontendDocument.h b/Engine/Plugins/Runtime/Metasound/Source/MetasoundFrontend/Public/MetasoundFrontendDocument.h
--- a/Engine/Plugins/Runtime/Metasound/Source/MetasoundFrontend/Public/MetasoundFrontendDocument.h
+++ b/Engine/Plugins/Runtime/Metasound/Source/MetasoundFrontend/Public/MetasoundFrontendDocument.h
@@ -1238,18 +1238,22 @@ struct METASOUNDFRONTEND_API FMetasoundFrontendClassName
 	// Returns scoped name representing namespace and name. 
 	FName GetScopedName() const;
 
+	// Invalid form of class name (i.e. empty namespace, name, and variant)
+	static const FMetasoundFrontendClassName InvalidClassName;
+
+	// Whether or not this instance of a class name is a valid name.
+	bool IsValid() const;
+
 	// Returns NodeClassName version of full name
-	Metasound::FNodeClassName ToNodeClassName() const
-	{
-		return { Namespace, Name, Variant };
-	}
+	Metasound::FNodeClassName ToNodeClassName() const;
 
 	// Return string version of full name.
 	FString ToString() const;
 
 	friend FORCEINLINE uint32 GetTypeHash(const FMetasoundFrontendClassName& ClassName)
 	{
-		return GetTypeHash(ClassName.GetFullName());
+		const int32 NameHash = HashCombineFast(GetTypeHash(ClassName.Namespace), GetTypeHash(ClassName.Name));
+		return HashCombineFast(NameHash, GetTypeHash(ClassName.Variant));
 	}
 
 	METASOUNDFRONTEND_API friend bool operator==(const FMetasoundFrontendClassName& InLHS, const FMetasoundFrontendClassName& InRHS);
@@ -1572,6 +1576,7 @@ struct METASOUNDFRONTEND_API FMetasoundFrontendGraphClass : public FMetasoundFro
 	FMetasoundFrontendGraphClassPresetOptions PresetOptions;
 };
 
+
 USTRUCT()
 struct METASOUNDFRONTEND_API FMetasoundFrontendDocumentMetadata
 {
