Enable MetaSound deterministic id generation by default for moving autoupdate from runtime to cook

diff --git a/Engine/Plugins/Runtime/Metasound/Source/MetasoundFrontend/Private/MetasoundFrontendDocumentIdGenerator.cpp b/Engine/Plugins/Runtime/Metasound/Source/MetasoundFrontend/Private/MetasoundFrontendDocumentIdGenerator.cpp
--- a/Engine/Plugins/Runtime/Metasound/Source/MetasoundFrontend/Private/MetasoundFrontendDocumentIdGenerator.cpp
+++ b/Engine/Plugins/Runtime/Metasound/Source/MetasoundFrontend/Private/MetasoundFrontendDocumentIdGenerator.cpp
@@ -6,12 +6,12 @@ namespace Metasound
 {
 	namespace Frontend
 	{
-		int32 MetaSoundEnableCookDeterministicIDGeneration = 0;
+		int32 MetaSoundEnableCookDeterministicIDGeneration = 1;
 		FAutoConsoleVariableRef CVarMetaSoundEnableCookDeterministicIDGeneration(
 			TEXT("au.MetaSound.EnableCookDeterministicIDGeneration"),
 			MetaSoundEnableCookDeterministicIDGeneration,
-			TEXT("Enable moving MetaSound registration operations like AutoUpdate and PreprocessDocument from runtime to cook using deterministic ID generation \n")
-			TEXT("0: Disabled (default), !0: Enabled"),
+			TEXT("Enable moving MetaSound registration operations like AutoUpdate and some template node transformations from runtime to cook using deterministic ID generation \n")
+			TEXT("0: Disabled, !0: Enabled (default)"),
 			ECVF_Default);
 
 		FDocumentIDGenerator::FScopeDeterminism::FScopeDeterminism(bool bInIsDeterministic)
