Change MetaSound AutoUpdate warning verbosity level for cook

diff --git a/Engine/Plugins/Runtime/Metasound/Source/MetasoundFrontend/Private/MetasoundFrontendTransform.cpp b/Engine/Plugins/Runtime/Metasound/Source/MetasoundFrontend/Private/MetasoundFrontendTransform.cpp
--- a/Engine/Plugins/Runtime/Metasound/Source/MetasoundFrontend/Private/MetasoundFrontendTransform.cpp
+++ b/Engine/Plugins/Runtime/Metasound/Source/MetasoundFrontend/Private/MetasoundFrontendTransform.cpp
@@ -15,13 +15,28 @@
 #include "MetasoundTrace.h"
 #include "Misc/App.h"
 
-
 namespace Metasound
 {
 	namespace Frontend
 	{
 		namespace DocumentTransform
 		{
+			void LogAutoUpdateWarning(const FString& LogMessage)
+			{
+				// These should eventually move back to warning on cook 
+				// but are temporarily downgraded to prevent 
+				// warnings on things like unused test content from 
+				// blocking code checkins 
+				if (IsRunningCookCommandlet())
+				{
+					UE_LOG(LogMetaSound, Display, TEXT("%s"), *LogMessage);
+				}
+				else
+				{
+					UE_LOG(LogMetaSound, Warning, TEXT("%s"), *LogMessage);
+				}
+			}
+
 #if WITH_EDITOR
 			FGetNodeDisplayNameProjection NodeDisplayNameProjection;
 
@@ -691,12 +706,12 @@ namespace Metasound
 
 							for (const FVertexNameAndType& InputPin : DisconnectedInputs)
 							{
-								UE_LOG(LogMetaSound, Warning, TEXT("Auto-Updating '%s' node class '%s (%s)': Previously connected input '%s' with data type '%s' no longer exists."), *DebugAssetPath, *NodeClassName, *NewClassVersion, *InputPin.Get<0>().ToString(), *InputPin.Get<1>().ToString());
+								DocumentTransform::LogAutoUpdateWarning(FString::Printf(TEXT("Auto-Updating '%s' node class '%s (%s)': Previously connected input '%s' with data type '%s' no longer exists."), *DebugAssetPath, *NodeClassName, *NewClassVersion, *InputPin.Get<0>().ToString(), *InputPin.Get<1>().ToString()));
 							}
 
 							for (const FVertexNameAndType& OutputPin : DisconnectedOutputs)
 							{
-								UE_LOG(LogMetaSound, Warning, TEXT("Auto-Updating '%s' node class '%s (%s)': Previously connected output '%s' with data type '%s' no longer exists."), *DebugAssetPath, *NodeClassName, *NewClassVersion, *OutputPin.Get<0>().ToString(), *OutputPin.Get<1>().ToString());
+								DocumentTransform::LogAutoUpdateWarning(FString::Printf(TEXT("Auto-Updating '%s' node class '%s (%s)': Previously connected output '%s' with data type '%s' no longer exists."), *DebugAssetPath, *NodeClassName, *NewClassVersion, *OutputPin.Get<0>().ToString(), *OutputPin.Get<1>().ToString()));
 							}
 						}
 					}
