Fix for Metasounds not using device rate when Quality Settings are disabled

diff --git a/Engine/Plugins/Runtime/Metasound/Source/MetasoundEditor/Private/MetasoundEditor.cpp b/Engine/Plugins/Runtime/Metasound/Source/MetasoundEditor/Private/MetasoundEditor.cpp
--- a/Engine/Plugins/Runtime/Metasound/Source/MetasoundEditor/Private/MetasoundEditor.cpp
+++ b/Engine/Plugins/Runtime/Metasound/Source/MetasoundEditor/Private/MetasoundEditor.cpp
@@ -1641,13 +1641,13 @@ namespace Metasound
 					check(AudioDevice);
 
 					const FName& AudioBufferTypeName = GetMetasoundDataTypeName<FAudioBuffer>();
-					const FSampleRate SampleRate = static_cast<FSampleRate>(AudioDevice->GetSampleRate());
+					const FSampleRate DeviceSampleRate = static_cast<FSampleRate>(AudioDevice->GetSampleRate());
 					const uint32 PlayOrder = PreviewComp->GetLastPlayOrder();
 					const uint64 TransmitterID = Audio::GetTransmitterID(PreviewComp->GetAudioComponentID(), 0, PlayOrder);
 
 					UMetaSoundSource* Source = CastChecked<UMetaSoundSource>(Metasound);
 					GraphConnectionManager = MakeUnique<FGraphConnectionManager>(
-						*MetasoundAsset, *PreviewComp, TransmitterID, Source->GetOperatorSettings(SampleRate));
+						*MetasoundAsset, *PreviewComp, TransmitterID, Source->GetOperatorSettings(DeviceSampleRate));
 				}
 
 				MetasoundGraphEditor->RegisterActiveTimer(0.0f,
diff --git a/Engine/Plugins/Runtime/Metasound/Source/MetasoundEngine/Private/MetasoundSource.cpp b/Engine/Plugins/Runtime/Metasound/Source/MetasoundEngine/Private/MetasoundSource.cpp
--- a/Engine/Plugins/Runtime/Metasound/Source/MetasoundEngine/Private/MetasoundSource.cpp
+++ b/Engine/Plugins/Runtime/Metasound/Source/MetasoundEngine/Private/MetasoundSource.cpp
@@ -223,10 +223,26 @@ bool UMetaSoundSource::CanEditChange(const FProperty* InProperty) const
 	// Allow changes to quality if we don't have any overrides.
 	if (InProperty->GetFName() == GET_MEMBER_NAME_CHECKED(UMetaSoundSource, QualitySetting))
 	{
-		const bool bBlockRateIsZero = BlockRateOverride.GetValue() == 0;
-		const bool bSampleRateIsZero = SampleRateOverride.GetValue() == 0;
+		const TArray<FName> Platforms = FDataDrivenPlatformInfoRegistry::GetSortedPlatformNames(EPlatformInfoType::AllPlatformInfos);
+		const int32 DefaultBlockRate = BlockRateOverride.GetDefault();
+		const float DefaultSampleRate = BlockRateOverride.GetDefault();
 
-		return bBlockRateIsZero && bSampleRateIsZero;
+		if (DefaultBlockRate > 0 && DefaultSampleRate > 0)
+		{
+			return false;
+		}
+
+		for (const FName Platform : Platforms)
+		{
+			if (BlockRateOverride.GetValueForPlatform(Platform) != DefaultBlockRate )
+			{
+				return false;
+			}
+			if (!FMath::IsNearlyEqual(SampleRateOverride.GetValueForPlatform(Platform), DefaultSampleRate))
+			{
+				return false;
+			}
+		}
 	}
 		
 	return true;
@@ -271,7 +287,7 @@ void UMetaSoundSource::PostEditChangeQualitySettings()
 	OperatorSettings.Reset();
 
 	// Refresh the SampleRate (which is what the engine sees from the operator settings).
-	SampleRate = GetOperatorSettings(SampleRate).GetSampleRate();
+	SampleRate = GetOperatorSettings(CachedAudioDeviceSampleRate).GetSampleRate();
 
 	// Always refresh the GUID with the selection.
 	if (const UMetaSoundSettings* Settings = GetDefault<UMetaSoundSettings>())	
@@ -433,16 +449,13 @@ void UMetaSoundSource::PostLoadQualitySettings()
 				WeakSource->OperatorSettings.Reset();
 						
 				// Override SampleRate with the Operator settings version which uses our Quality settings.
-				WeakSource->SampleRate = WeakSource->GetOperatorSettings(WeakSource->SampleRate).GetSampleRate();
+				WeakSource->SampleRate = WeakSource->GetOperatorSettings(WeakSource->CachedAudioDeviceSampleRate).GetSampleRate();
 			}
 		};
 		Metasound::Frontend::GetBlockRateOverrideChangedDelegate().AddWeakLambda(this, ResetOperatorSettings);
 		Metasound::Frontend::GetSampleRateOverrideChangedDelegate().AddWeakLambda(this, ResetOperatorSettings);
 	}
 #endif //WITH_EDITORONLY_DATA
-
-	// Override SampleRate with the Operator settings version which uses our Quality settings.
-	SampleRate = GetOperatorSettings(SampleRate).GetSampleRate();
 }
 
 void UMetaSoundSource::ResolveQualitySettings(const UMetaSoundSettings* Settings)
@@ -1199,8 +1212,14 @@ TSharedPtr<Audio::IParameterTransmitter> UMetaSoundSource::CreateParameterTransm
 	}
 }
 
-Metasound::FOperatorSettings UMetaSoundSource::GetOperatorSettings(Metasound::FSampleRate InSampleRate) const
+Metasound::FOperatorSettings UMetaSoundSource::GetOperatorSettings(Metasound::FSampleRate InDeviceSampleRate) const
 {	
+	// We should recache the operator settings if the device rate has changed.
+	if (InDeviceSampleRate != CachedAudioDeviceSampleRate)
+	{
+		OperatorSettings.Reset();
+	}
+	
 	if (!OperatorSettings)
 	{
 		using namespace Metasound;
@@ -1210,10 +1229,11 @@ Metasound::FOperatorSettings UMetaSoundSource::GetOperatorSettings(Metasound::FS
 		auto QueryQualitySettings = [&](Metasound::FSampleRate InSampleRate) -> Metasound::FOperatorSettings
 		{
 			static const float DefaultBlockRateConstant = 100.f;
+			static const float DefaultSampleRateConstant = 48000.f;
 			
-			// 1. Sensible defaults.
-			FSampleRate SampleRate = InSampleRate;
-			float BlockRate = DefaultBlockRateConstant;
+			// 1. Sensible defaults. (If Device SampleRate is sensible use that as default).
+			FSampleRate MetasoundSampleRate = InSampleRate > 0 ? InDeviceSampleRate : DefaultSampleRateConstant;
+			float MetasoundBlockRate = DefaultBlockRateConstant;
 
 			// 2. Query our quality settings.
 			if (const UMetaSoundSettings* Settings = GetDefault<UMetaSoundSettings>())
@@ -1223,11 +1243,11 @@ Metasound::FOperatorSettings UMetaSoundSource::GetOperatorSettings(Metasound::FS
 					// Allow partial applications of settings, if some are non-zero.
 					if (const float Value = Found->BlockRate.GetValue(); Value > 0.f)
 					{
-						BlockRate = Value;
+						MetasoundBlockRate = Value;
 					}
 					if (const float Value = Found->SampleRate.GetValue(); Value > 0.f)
 					{
-						SampleRate = Value;
+						MetasoundSampleRate = Value;
 					}
 				}
 			}
@@ -1235,11 +1255,11 @@ Metasound::FOperatorSettings UMetaSoundSource::GetOperatorSettings(Metasound::FS
 			// 3. Do per asset overrides.
 			if (const float SerializedBlockRate = BlockRateOverride.GetValue(); SerializedBlockRate > 0.0f)
 			{
-				BlockRate = SerializedBlockRate;
+				MetasoundBlockRate = SerializedBlockRate;
 			}
 			if (const int32 SerializedSampleRate = SampleRateOverride.GetValue(); SerializedSampleRate > 0)
 			{
-				SampleRate = SerializedSampleRate;
+				MetasoundSampleRate = SerializedSampleRate;
 			}
 
 			// 4. Query CVars. (Override with CVars if they are > 0)
@@ -1249,22 +1269,22 @@ Metasound::FOperatorSettings UMetaSoundSource::GetOperatorSettings(Metasound::FS
 
 			if (SampleRateCvar > 0)
 			{
-				SampleRate = SampleRateCvar;
+				MetasoundSampleRate = SampleRateCvar;
 			}
 			if (BlockRateCVar > 0)
 			{
-				BlockRate = BlockRateCVar;
+				MetasoundBlockRate = BlockRateCVar;
 			}
 
 			// 5. Sanity clamps.
-			TRange<float> BlockRange = GetBlockRateClampRange();
-			TRange<int32> RateRange = GetSampleRateClampRange();
-			BlockRate = FMath::Clamp(BlockRate, BlockRange.GetLowerBoundValue(), BlockRange.GetUpperBoundValue());
-			SampleRate = FMath::Clamp(SampleRate, RateRange.GetLowerBoundValue(), RateRange.GetUpperBoundValue());
+			const TRange<float> BlockRange = GetBlockRateClampRange();
+			const TRange<int32> RateRange = GetSampleRateClampRange();
+			MetasoundBlockRate = FMath::Clamp(MetasoundBlockRate, BlockRange.GetLowerBoundValue(), BlockRange.GetUpperBoundValue());
+			MetasoundSampleRate = FMath::Clamp(MetasoundSampleRate, RateRange.GetLowerBoundValue(), RateRange.GetUpperBoundValue());
 
-			return Metasound::FOperatorSettings(SampleRate, BlockRate);
+			return Metasound::FOperatorSettings(MetasoundSampleRate, MetasoundBlockRate);
 		};
-		OperatorSettings = QueryQualitySettings(InSampleRate);
+		OperatorSettings = QueryQualitySettings(InDeviceSampleRate);
 	}
 	return *OperatorSettings;
 }
diff --git a/Engine/Plugins/Runtime/Metasound/Source/MetasoundEngine/Public/MetasoundSettings.h b/Engine/Plugins/Runtime/Metasound/Source/MetasoundEngine/Public/MetasoundSettings.h
--- a/Engine/Plugins/Runtime/Metasound/Source/MetasoundEngine/Public/MetasoundSettings.h
+++ b/Engine/Plugins/Runtime/Metasound/Source/MetasoundEngine/Public/MetasoundSettings.h
@@ -52,11 +52,11 @@ struct METASOUNDENGINE_API FMetaSoundQualitySettings
 	UPROPERTY(EditAnywhere, Category = "Quality")
 	FName Name;
 
-	/** Sample Rate (in Hz). NOTE: A Zero value will force the default. **/
+	/** Sample Rate (in Hz). NOTE: A Zero value will have no effect and use the Device Rate. **/
 	UPROPERTY(EditAnywhere, Category = "Quality", meta = (ClampMin = "0", ClampMax="96000"))
 	FPerPlatformInt SampleRate = 0;
 
-	/** Block Rate (in Hz). NOTE: A Zero value will force the default.  **/
+	/** Block Rate (in Hz). NOTE: A Zero value will have no effect and use the Default (100)  **/
 	UPROPERTY(EditAnywhere, Category = "Quality", meta = (ClampMin = "0", ClampMax="1000"))
 	FPerPlatformFloat BlockRate = 0.f;
 };
diff --git a/Engine/Plugins/Runtime/Metasound/Source/MetasoundEngine/Public/MetasoundSource.h b/Engine/Plugins/Runtime/Metasound/Source/MetasoundEngine/Public/MetasoundSource.h
--- a/Engine/Plugins/Runtime/Metasound/Source/MetasoundEngine/Public/MetasoundSource.h
+++ b/Engine/Plugins/Runtime/Metasound/Source/MetasoundEngine/Public/MetasoundSource.h
@@ -115,11 +115,11 @@ public:
 	FGuid QualitySettingGuid;
 #endif //WITH_EDITOR_DATA
 
-	// Override the BlockRate for this Sound (overrides the Platforms Quality Settings for this asset)
+	// Override the BlockRate for this Sound (overrides Quality). NOTE: A Zero value will have no effect and use either the Quality setting (if set), or the defaults.
 	UPROPERTY(EditAnywhere, BlueprintReadOnly, AdvancedDisplay, Category = Metasound, meta = (UIMin = 0, UIMax = 1000.0, DisplayAfter="OutputFormat", DisplayName = "Override Block Rate (in Hz)"))
 	FPerPlatformFloat BlockRateOverride = 0.f;
 
-	// Override the SampleRate for this Sound (overrides the Platforms Quality Settings for this asset)
+	// Override the SampleRate for this Sound (overrides Quality). NOTE: A Zero value will have no effect and use either the Quality setting (if set), or the Device Rate
 	UPROPERTY(EditAnywhere, BlueprintReadOnly, AdvancedDisplay, Category = Metasound, meta = (UIMin = 0, UIMax = 96000, DisplayName = "Override Sample Rate (in Hz)"))
 	FPerPlatformInt SampleRateOverride = 0;
 
@@ -249,7 +249,7 @@ public:
 	TWeakPtr<Metasound::FMetasoundGenerator> GetGeneratorForAudioComponent(uint64 ComponentId) const;
 	FOnGeneratorInstanceCreated OnGeneratorInstanceCreated;
 	FOnGeneratorInstanceDestroyed OnGeneratorInstanceDestroyed;
-	Metasound::FOperatorSettings GetOperatorSettings(Metasound::FSampleRate InSampleRate) const;
+	Metasound::FOperatorSettings GetOperatorSettings(Metasound::FSampleRate InDeviceSampleRate) const;
 
 	virtual const FMetasoundFrontendDocument& GetConstDocument() const override;
 
@@ -319,6 +319,9 @@ private:
 	 * Lazy (Cached) Operator Settings. Built in GetOperatorSettings
 	 */
 	mutable TOptional<Metasound::FOperatorSettings> OperatorSettings;
+
+	// Cache the AudioDevice Samplerate. (so that if we have to regenerate operator settings without the device rate we can use this).
+	mutable Metasound::FSampleRate CachedAudioDeviceSampleRate = 0;
 	
 	bool bIsBuilderActive = false;
 };
diff --git a/Engine/Source/Runtime/Engine/Private/SoundWave.cpp b/Engine/Source/Runtime/Engine/Private/SoundWave.cpp
--- a/Engine/Source/Runtime/Engine/Private/SoundWave.cpp
+++ b/Engine/Source/Runtime/Engine/Private/SoundWave.cpp
@@ -3732,7 +3732,14 @@ float USoundWave::GetSampleRateForCurrentPlatform() const
 {
 	if (bProcedural)
 	{
-		return SampleRate;
+		if (SampleRate > 0)
+		{
+			return SampleRate;
+		}
+
+		// Never return a 0 if we've been serialized with one.
+		static const float DefaultSampleRate = 48000.f;
+		return DefaultSampleRate;
 	}
 
 #if WITH_EDITORONLY_DATA
