- Make FNodeRegistryKey an actual struct as opposed to a string alias
    - Avoids string format/copies/conversions in most cases where its just used as a key
- Deprecate NodeRegistryKey namespace calls

diff --git a/Engine/Plugins/Runtime/Metasound/Source/MetasoundEditor/Private/MetasoundEditor.cpp b/Engine/Plugins/Runtime/Metasound/Source/MetasoundEditor/Private/MetasoundEditor.cpp
--- a/Engine/Plugins/Runtime/Metasound/Source/MetasoundEditor/Private/MetasoundEditor.cpp
+++ b/Engine/Plugins/Runtime/Metasound/Source/MetasoundEditor/Private/MetasoundEditor.cpp
@@ -32,6 +32,7 @@
 #include "Kismet2/KismetEditorUtilities.h"
 #include "Logging/TokenizedMessage.h"
 #include "Metasound.h"
+#include "MetasoundAssetSubsystem.h"
 #include "MetasoundBuilderSubsystem.h"
 #include "MetasoundDocumentInterface.h"
 #include "MetasoundEditorCommands.h"
@@ -1824,18 +1825,11 @@ namespace Metasound
 			if (UMetasoundEditorGraphExternalNode* ExternalNode = Cast<UMetasoundEditorGraphExternalNode>(InNode))
 			{
 				FConstNodeHandle NodeHandle = ExternalNode->GetConstNodeHandle();
-				FNodeRegistryKey Key = FMetasoundFrontendRegistryContainer::Get()->GetRegistryKey(NodeHandle->GetClassMetadata());
+				const FNodeRegistryKey Key(NodeHandle->GetClassMetadata());
 
-				FNodeClassInfo ClassInfo;
-				if (FMetasoundFrontendRegistryContainer::Get()->FindNodeClassInfoFromRegistered(Key, ClassInfo))
+				if (const FSoftObjectPath* AssetObject = UMetaSoundAssetSubsystem::GetChecked().FindObjectPathFromKey(Key))
 				{
-					if (ClassInfo.AssetClassID.IsValid())
-					{
-						if (UObject* AssetObject = ClassInfo.LoadAsset())
-						{
-							GEditor->GetEditorSubsystem<UAssetEditorSubsystem>()->OpenEditorForAsset(AssetObject);
-						}
-					}
+					GEditor->GetEditorSubsystem<UAssetEditorSubsystem>()->OpenEditorForAsset(*AssetObject);
 				}
 			}
 		}
@@ -2477,7 +2471,7 @@ namespace Metasound
 					FMetasoundFrontendClassMetadata LookupMetadata;
 					LookupMetadata.SetClassName(ExternalNode->GetClassName());
 					LookupMetadata.SetType(EMetasoundFrontendClassType::External);
-					const FNodeRegistryKey PastedRegistryKey = NodeRegistryKey::CreateKey(LookupMetadata);
+					const FNodeRegistryKey PastedRegistryKey = FNodeRegistryKey(LookupMetadata);
 					if (const FSoftObjectPath* AssetPath = IMetaSoundAssetManager::GetChecked().FindObjectPathFromKey(PastedRegistryKey))
 					{
 						if (MetasoundAsset->AddingReferenceCausesLoop(*AssetPath))
@@ -2948,7 +2942,7 @@ namespace Metasound
 					const bool bHasNewVersion = HighestVersion.IsValid() && HighestVersion > Metadata.GetVersion();
 
 					// Check for non-native classes
-					const FNodeRegistryKey RegistryKey = NodeRegistryKey::CreateKey(Metadata);
+					const FNodeRegistryKey RegistryKey = FNodeRegistryKey(Metadata);
 					const bool bIsClassNative = FMetasoundFrontendRegistryContainer::Get()->IsNodeNative(RegistryKey);
 
 					if (bHasNewVersion || !bIsClassNative)
diff --git a/Engine/Plugins/Runtime/Metasound/Source/MetasoundEditor/Private/MetasoundEditorGraphBuilder.cpp b/Engine/Plugins/Runtime/Metasound/Source/MetasoundEditor/Private/MetasoundEditorGraphBuilder.cpp
--- a/Engine/Plugins/Runtime/Metasound/Source/MetasoundEditor/Private/MetasoundEditorGraphBuilder.cpp
+++ b/Engine/Plugins/Runtime/Metasound/Source/MetasoundEditor/Private/MetasoundEditorGraphBuilder.cpp
@@ -246,7 +246,7 @@ namespace Metasound
 				DisplayName = Metadata.GetDisplayName();
 				if (DisplayName.IsEmptyOrWhitespace())
 				{
-					const FNodeRegistryKey RegistryKey = NodeRegistryKey::CreateKey(Metadata);
+					const FNodeRegistryKey RegistryKey = FNodeRegistryKey(Metadata);
 					bool bIsClassNative = FMetasoundFrontendRegistryContainer::Get()->IsNodeNative(RegistryKey);
 					if (!bIsClassNative)
 					{
@@ -435,7 +435,7 @@ namespace Metasound
 			NewGraphNode = NodeCreator.CreateNode(bInSelectNewNode);
 			if (ensure(NewGraphNode))
 			{
-				const FNodeRegistryKey RegistryKey = NodeRegistryKey::CreateKey(InNodeHandle->GetClassMetadata());
+				const FNodeRegistryKey RegistryKey = FNodeRegistryKey(InNodeHandle->GetClassMetadata());
 				NewGraphNode->bIsClassNative = FMetasoundFrontendRegistryContainer::Get()->IsNodeNative(RegistryKey);
 				NewGraphNode->ClassName = InNodeHandle->GetClassMetadata().GetClassName();
 				NewGraphNode->CacheTitle();
diff --git a/Engine/Plugins/Runtime/Metasound/Source/MetasoundEditor/Private/MetasoundEditorGraphNode.cpp b/Engine/Plugins/Runtime/Metasound/Source/MetasoundEditor/Private/MetasoundEditorGraphNode.cpp
--- a/Engine/Plugins/Runtime/Metasound/Source/MetasoundEditor/Private/MetasoundEditorGraphNode.cpp
+++ b/Engine/Plugins/Runtime/Metasound/Source/MetasoundEditor/Private/MetasoundEditorGraphNode.cpp
@@ -853,7 +853,7 @@ void UMetasoundEditorGraphExternalNode::Validate(Metasound::Editor::FGraphNodeVa
 	const FMetasoundFrontendClassMetadata& Metadata = NodeHandle->GetClassMetadata();
 
 	// 1. Validate referenced graph recursively if defined as asset node class
-	const FNodeRegistryKey RegistryKey = NodeRegistryKey::CreateKey(Metadata);
+	const FNodeRegistryKey RegistryKey = FNodeRegistryKey(Metadata);
 	if (IMetaSoundAssetManager* AssetManager = IMetaSoundAssetManager::Get())
 	{
 		if (const FSoftObjectPath* Path = AssetManager->FindObjectPathFromKey(RegistryKey))
@@ -892,7 +892,7 @@ void UMetasoundEditorGraphExternalNode::Validate(Metasound::Editor::FGraphNodeVa
 	// 2. Validate template nodes
 	if (Metadata.GetType() == EMetasoundFrontendClassType::Template)
 	{
-		const FNodeRegistryKey Key = NodeRegistryKey::CreateKey(Metadata);
+		const FNodeRegistryKey Key = FNodeRegistryKey(Metadata);
 		if (const INodeTemplate* Template = INodeTemplateRegistry::Get().FindTemplate(Key))
 		{
 			const bool bIsValidInterface = Template->IsValidNodeInterface(NodeHandle->GetNodeInterface());
diff --git a/Engine/Plugins/Runtime/Metasound/Source/MetasoundEditor/Private/MetasoundEditorGraphSchema.cpp b/Engine/Plugins/Runtime/Metasound/Source/MetasoundEditor/Private/MetasoundEditorGraphSchema.cpp
--- a/Engine/Plugins/Runtime/Metasound/Source/MetasoundEditor/Private/MetasoundEditorGraphSchema.cpp
+++ b/Engine/Plugins/Runtime/Metasound/Source/MetasoundEditor/Private/MetasoundEditorGraphSchema.cpp
@@ -408,7 +408,7 @@ const FSlateBrush* FMetasoundGraphSchemaAction_NewNode::GetIconBrush() const
 {
 	using namespace Metasound::Frontend;
 
-	const FNodeRegistryKey RegistryKey = NodeRegistryKey::CreateKey(ClassMetadata);
+	const FNodeRegistryKey RegistryKey = FNodeRegistryKey(ClassMetadata);
 	const bool bIsClassNative = FMetasoundFrontendRegistryContainer::Get()->IsNodeNative(RegistryKey);
 	if (bIsClassNative)
 	{
@@ -424,7 +424,7 @@ const FLinearColor& FMetasoundGraphSchemaAction_NewNode::GetIconColor() const
 
 	if (const UMetasoundEditorSettings* EditorSettings = GetDefault<UMetasoundEditorSettings>())
 	{
-		const FNodeRegistryKey RegistryKey = NodeRegistryKey::CreateKey(ClassMetadata);
+		const FNodeRegistryKey RegistryKey = FNodeRegistryKey(ClassMetadata);
 		const bool bIsClassNative = FMetasoundFrontendRegistryContainer::Get()->IsNodeNative(RegistryKey);
 		if (bIsClassNative)
 		{
@@ -1177,7 +1177,7 @@ void UMetasoundEditorGraphSchema::GetContextMenuActions(class UToolMenu* Menu, c
 			const FMetasoundFrontendClassMetadata& Metadata = NodeHandle->GetClassMetadata();
 			const bool bHasNewVersion = HighestVersion.IsValid() && HighestVersion > Metadata.GetVersion();
 
-			const FNodeRegistryKey RegistryKey = NodeRegistryKey::CreateKey(Metadata);
+			const FNodeRegistryKey RegistryKey = FNodeRegistryKey(Metadata);
 			const bool bIsClassNative = FMetasoundFrontendRegistryContainer::Get()->IsNodeNative(RegistryKey);
 
 			if (bHasNewVersion || !bIsClassNative)
@@ -1783,8 +1783,8 @@ void UMetasoundEditorGraphSchema::DroppedAssetsOnGraph(const TArray<FAssetData>&
 
 			// This may not be necessary as dropping an asset on the graph may load it, thus triggering the registration from the MetaSoundAssetManager.
 			const FMetasoundFrontendDocument& DroppedDoc = DroppedMetaSoundAsset->GetDocumentChecked();
-			const FNodeRegistryKey RegistryKey = NodeRegistryKey::CreateKey(DroppedDoc.RootGraph);
-			if (ensure(NodeRegistryKey::IsValid(RegistryKey)))
+			const FNodeRegistryKey RegistryKey = FNodeRegistryKey(DroppedDoc.RootGraph);
+			if (ensure(RegistryKey.IsValid()))
 			{
 				FMetaSoundAssetRegistrationOptions RegOptions;
 				RegOptions.bForceReregister = false;
@@ -1963,7 +1963,7 @@ void UMetasoundEditorGraphSchema::GetFunctionActions(FGraphActionMenuBuilder& Ac
 	{
 		FMetasoundFrontendClassMetadata AssetMetadata = InGraphHandle->GetGraphMetadata();
 		AssetMetadata.SetType(EMetasoundFrontendClassType::External);
-		const FNodeRegistryKey RegistryKey = NodeRegistryKey::CreateKey(AssetMetadata);
+		const FNodeRegistryKey RegistryKey = FNodeRegistryKey(AssetMetadata);
 		ParentAsset = AssetManager.TryLoadAssetFromKey(RegistryKey);
 	}
 
@@ -1994,7 +1994,7 @@ void UMetasoundEditorGraphSchema::GetFunctionActions(FGraphActionMenuBuilder& Ac
 		bool bIsAsset = false;
 		if (ParentAsset)
 		{
-			const FNodeRegistryKey RegistryKey = NodeRegistryKey::CreateKey(FrontendClass.Metadata);
+			const FNodeRegistryKey RegistryKey = FNodeRegistryKey(FrontendClass.Metadata);
 			if (const FSoftObjectPath* Path = AssetManager.FindObjectPathFromKey(RegistryKey))
 			{
 				// The path can be not an asset if it was transiently created (eg. in PIE if generated by a runtime MetaSoundDocumentBuilder)
diff --git a/Engine/Plugins/Runtime/Metasound/Source/MetasoundEditor/Private/MetasoundEditorModule.cpp b/Engine/Plugins/Runtime/Metasound/Source/MetasoundEditor/Private/MetasoundEditorModule.cpp
--- a/Engine/Plugins/Runtime/Metasound/Source/MetasoundEditor/Private/MetasoundEditorModule.cpp
+++ b/Engine/Plugins/Runtime/Metasound/Source/MetasoundEditor/Private/MetasoundEditorModule.cpp
@@ -303,7 +303,7 @@ namespace Metasound
 						const FNodeRegistryKey RegistryKey = AssetSubsystem->AddOrUpdateAsset(InAssetData);
 
 						// Can be invalid if being called for the first time on an asset before FRenameRootGraphClass is called
-						if (NodeRegistryKey::IsValid(RegistryKey))
+						if (RegistryKey.IsValid())
 						{
 							const bool bPrimeRequested = AssetPrimeStatus > EAssetPrimeStatus::NotRequested;
 							const bool bIsRegistered = FMetasoundFrontendRegistryContainer::Get()->IsNodeRegistered(RegistryKey);
diff --git a/Engine/Plugins/Runtime/Metasound/Source/MetasoundEngine/Private/MetasoundAssetSubsystem.cpp b/Engine/Plugins/Runtime/Metasound/Source/MetasoundEngine/Private/MetasoundAssetSubsystem.cpp
--- a/Engine/Plugins/Runtime/Metasound/Source/MetasoundEngine/Private/MetasoundAssetSubsystem.cpp
+++ b/Engine/Plugins/Runtime/Metasound/Source/MetasoundEngine/Private/MetasoundAssetSubsystem.cpp
@@ -96,7 +96,7 @@ namespace Metasound
 				}
 				else
 				{
-					UE_LOG(LogMetaSound, VeryVerbose, TEXT("Object paths do not match. Skipping removal of %s:%s from the asset subsystem."), *InKeyToRemove, *InPathToRemove.ToString());
+					UE_LOG(LogMetaSound, VeryVerbose, TEXT("Object paths do not match. Skipping removal of %s:%s from the asset subsystem."), *InKeyToRemove.ToString(), *InPathToRemove.ToString());
 				}
 			}
 			return false;
@@ -143,7 +143,7 @@ void UMetaSoundAssetSubsystem::AddAssetReferences(FMetasoundAssetBase& InAssetBa
 	using namespace Metasound::Frontend;
 
 	const FMetasoundFrontendDocument& Document = InAssetBase.GetDocumentChecked();
-	const FNodeRegistryKey AssetClassKey = NodeRegistryKey::CreateKey(Document.RootGraph);
+	const FNodeRegistryKey AssetClassKey = FNodeRegistryKey(Document.RootGraph);
 
 	if (!ContainsKey(AssetClassKey))
 	{
@@ -153,11 +153,13 @@ void UMetaSoundAssetSubsystem::AddAssetReferences(FMetasoundAssetBase& InAssetBa
 
 	bool bAddFromReferencedAssets = false;
 	const TSet<FString>& ReferencedAssetClassKeys = InAssetBase.GetReferencedAssetClassKeys();
-	for (const FString& ReferencedAssetClassKey : ReferencedAssetClassKeys)
+	for (const FString& KeyString : ReferencedAssetClassKeys)
 	{
-		if (!ContainsKey(ReferencedAssetClassKey))
+		FNodeRegistryKey RegistryKey;
+		const bool bIsKey = FNodeRegistryKey::Parse(KeyString, RegistryKey);
+		if (!bIsKey || !ContainsKey(RegistryKey))
 		{
-			UE_LOG(LogMetaSound, Verbose, TEXT("Missing referenced class '%s' asset entry."), *ReferencedAssetClassKey);
+			UE_LOG(LogMetaSound, Verbose, TEXT("Missing referenced class '%s' asset entry."), *KeyString);
 			bAddFromReferencedAssets = true;
 		}
 	}
@@ -176,13 +178,13 @@ void UMetaSoundAssetSubsystem::AddAssetReferences(FMetasoundAssetBase& InAssetBa
 		if (Asset)
 		{
 			const FMetasoundFrontendDocument& RefDocument = Asset->GetDocumentChecked();
-			const FNodeRegistryKey ClassKey = NodeRegistryKey::CreateKey(RefDocument.RootGraph);
+			const FNodeRegistryKey ClassKey = FNodeRegistryKey(RefDocument.RootGraph);
 			if (!ContainsKey(ClassKey))
 			{
 				UE_LOG(LogMetaSound, Verbose,
 					TEXT("Preemptive load of class '%s' due to early "
 						"registration request (asset scan likely not complete)."),
-					*ClassKey);
+					*ClassKey.ToString());
 
 				UObject* MetaSoundObject = Asset->GetOwningAsset();
 				if (ensureAlways(MetaSoundObject))
@@ -211,9 +213,9 @@ Metasound::Frontend::FNodeRegistryKey UMetaSoundAssetSubsystem::AddOrUpdateAsset
 	check(MetaSoundAsset);
 
 	const FMetasoundFrontendDocument& Document = MetaSoundAsset->GetDocumentChecked();
-	const FNodeRegistryKey RegistryKey = NodeRegistryKey::CreateKey(Document.RootGraph);
+	const FNodeRegistryKey RegistryKey = FNodeRegistryKey(Document.RootGraph);
 
-	if (NodeRegistryKey::IsValid(RegistryKey))
+	if (RegistryKey.IsValid())
 	{
 		PathMap.FindOrAdd(RegistryKey) = InObject.GetPathName();
 	}
@@ -257,8 +259,8 @@ Metasound::Frontend::FNodeRegistryKey UMetaSoundAssetSubsystem::AddOrUpdateAsset
 
 	if (ClassInfo.AssetClassID.IsValid())
 	{
-		const FNodeRegistryKey RegistryKey = NodeRegistryKey::CreateKey(ClassInfo);
-		if (NodeRegistryKey::IsValid(RegistryKey))
+		const FNodeRegistryKey RegistryKey = FNodeRegistryKey(ClassInfo);
+		if (RegistryKey.IsValid())
 		{
 			PathMap.FindOrAdd(RegistryKey) = InAssetData.GetSoftObjectPath();
 		}
@@ -268,7 +270,7 @@ Metasound::Frontend::FNodeRegistryKey UMetaSoundAssetSubsystem::AddOrUpdateAsset
 
 	// Invalid ClassID means the node could not be registered.
 	// Let caller report or ensure as necessary.
-	return NodeRegistryKey::GetInvalid();
+	return FNodeRegistryKey::GetInvalid();
 }
 
 bool UMetaSoundAssetSubsystem::CanAutoUpdate(const FMetasoundFrontendClassName& InClassName) const
@@ -333,7 +335,7 @@ TSet<UMetaSoundAssetSubsystem::FAssetInfo> UMetaSoundAssetSubsystem::GetReferenc
 	const FMetasoundFrontendDocument& Document = InAssetBase.GetDocumentChecked();
 	for (const FMetasoundFrontendClass& Class : Document.Dependencies)
 	{
-		const FNodeRegistryKey Key = NodeRegistryKey::CreateKey(Class.Metadata);
+		const FNodeRegistryKey Key = FNodeRegistryKey(Class.Metadata);
 		if (const FSoftObjectPath* ObjectPath = PathMap.Find(Key))
 		{
 			OutAssetInfos.Add(FAssetInfo{Key, *ObjectPath});
@@ -361,11 +363,11 @@ TSet<UMetaSoundAssetSubsystem::FAssetInfo> UMetaSoundAssetSubsystem::GetReferenc
 			{
 				if (bIsInitialAssetScanComplete)
 				{
-					UE_LOG(LogMetaSound, Warning, TEXT("MetaSound Node Class with registry key '%s' not registered when gathering referenced asset classes from '%s': Retrieving all asset classes may not be comprehensive."), *Key, *InAssetBase.GetOwningAssetName());
+					UE_LOG(LogMetaSound, Warning, TEXT("MetaSound Node Class with registry key '%s' not registered when gathering referenced asset classes from '%s': Retrieving all asset classes may not be comprehensive."), *Key.ToString(), *InAssetBase.GetOwningAssetName());
 				}
 				else
 				{
-					UE_LOG(LogMetaSound, Warning, TEXT("Attempt to get registered dependent asset with key '%s' from MetaSound asset '%s' before asset scan has completed: Asset class cannot be provided"), *Key, *InAssetBase.GetOwningAssetName());
+					UE_LOG(LogMetaSound, Warning, TEXT("Attempt to get registered dependent asset with key '%s' from MetaSound asset '%s' before asset scan has completed: Asset class cannot be provided"), *Key.ToString(), *InAssetBase.GetOwningAssetName());
 				}
 			}
 		}
@@ -399,15 +401,17 @@ bool UMetaSoundAssetSubsystem::TryLoadReferencedAssets(const FMetasoundAssetBase
 
 	TArray<FMetasoundAssetBase*> ReferencedAssets;
 	const TSet<FString>& AssetClassKeys = InAssetBase.GetReferencedAssetClassKeys();
-	for (const FNodeRegistryKey& Key : AssetClassKeys)
+	for (const FString& KeyString : AssetClassKeys)
 	{
+		FNodeRegistryKey Key;
+		FNodeRegistryKey::Parse(KeyString, Key);
 		if (FMetasoundAssetBase* MetaSound = TryLoadAssetFromKey(Key))
 		{
 			OutReferencedAssets.Add(MetaSound);
 		}
 		else
 		{
-			UE_LOG(LogMetaSound, Error, TEXT("Failed to find referenced MetaSound asset with key '%s'"), *Key);
+			UE_LOG(LogMetaSound, Error, TEXT("Failed to find referenced MetaSound asset with key '%s'"), *KeyString);
 			bSucceeded = false;
 		}
 	}
@@ -599,7 +603,7 @@ void UMetaSoundAssetSubsystem::RemoveAsset(const UObject& InObject)
 	{
 		const FSoftObjectPath ObjectPath(&InObject);
 		const FMetasoundFrontendDocument& Document = MetaSoundAsset->GetDocumentChecked();
-		FNodeRegistryKey RegistryKey = NodeRegistryKey::CreateKey(Document.RootGraph);
+		FNodeRegistryKey RegistryKey = FNodeRegistryKey(Document.RootGraph);
 		if (AssetSubsystemPrivate::RemoveIfExactMatch(PathMap, RegistryKey, ObjectPath))
 		{
 			UMetaSoundBuilderSubsystem::GetChecked().DetachBuilderFromAsset(Document.RootGraph.Metadata.GetClassName());
@@ -615,7 +619,7 @@ void UMetaSoundAssetSubsystem::RemoveAsset(const FAssetData& InAssetData)
 	FNodeClassInfo ClassInfo;
 	if (ensureAlways(AssetSubsystemPrivate::GetAssetClassInfo(InAssetData, ClassInfo)))
 	{
-		FNodeRegistryKey RegistryKey = FMetasoundFrontendRegistryContainer::Get()->GetRegistryKey(ClassInfo);
+		const FNodeRegistryKey RegistryKey(ClassInfo);
 		const FSoftObjectPath ObjectPath = InAssetData.GetSoftObjectPath();
 		if (AssetSubsystemPrivate::RemoveIfExactMatch(PathMap, RegistryKey, ObjectPath))
 		{
@@ -740,7 +744,7 @@ void UMetaSoundAssetSubsystem::UnregisterAssetClassesInDirectories(const TArray<
 			FNodeClassInfo AssetClassInfo;
 			if (ensureAlways(AssetSubsystemPrivate::GetAssetClassInfo(AssetData, AssetClassInfo)))
 			{
-				const FNodeRegistryKey RegistryKey = NodeRegistryKey::CreateKey(AssetClassInfo);
+				const FNodeRegistryKey RegistryKey = FNodeRegistryKey(AssetClassInfo);
 				const bool bIsRegistered = FMetasoundFrontendRegistryContainer::Get()->IsNodeRegistered(RegistryKey);
 				if (bIsRegistered)
 				{
diff --git a/Engine/Plugins/Runtime/Metasound/Source/MetasoundEngine/Private/MetasoundBuilderSubsystem.cpp b/Engine/Plugins/Runtime/Metasound/Source/MetasoundEngine/Private/MetasoundBuilderSubsystem.cpp
--- a/Engine/Plugins/Runtime/Metasound/Source/MetasoundEngine/Private/MetasoundBuilderSubsystem.cpp
+++ b/Engine/Plugins/Runtime/Metasound/Source/MetasoundEngine/Private/MetasoundBuilderSubsystem.cpp
@@ -667,7 +667,7 @@ UObject* UMetaSoundBuilderBase::GetReferencedPresetAsset() const
 	if (Node != nullptr)
 	{
 		const FMetasoundFrontendClass* NodeClass = Builder.FindDependency(Node->ClassID);
-		const FNodeRegistryKey NodeClassRegistryKey = NodeRegistryKey::CreateKey(NodeClass->Metadata);
+		const FNodeRegistryKey NodeClassRegistryKey = FNodeRegistryKey(NodeClass->Metadata);
 		if (FMetasoundAssetBase* Asset = IMetaSoundAssetManager::GetChecked().TryLoadAssetFromKey(NodeClassRegistryKey))
 		{
 			return Asset->GetOwningAsset();
@@ -1156,7 +1156,7 @@ void UMetaSoundSourceBuilder::OnNodeAdded(int32 NodeIndex) const
 		const FMetasoundFrontendClass* NodeClass = Builder.FindDependency(AddedNode.ClassID);
 		checkf(NodeClass, TEXT("Node successfully added to graph but document is missing associated dependency"));
 
-		const FNodeRegistryKey& ClassKey = NodeRegistryKey::CreateKey(NodeClass->Metadata);
+		const FNodeRegistryKey& ClassKey = FNodeRegistryKey(NodeClass->Metadata);
 		FMetasoundFrontendRegistryContainer& NodeRegistry = *FMetasoundFrontendRegistryContainer::Get();
 		IDataTypeRegistry& DataTypeRegistry = IDataTypeRegistry::Get();
 
diff --git a/Engine/Plugins/Runtime/Metasound/Source/MetasoundEngine/Private/MetasoundEngineAsset.h b/Engine/Plugins/Runtime/Metasound/Source/MetasoundEngine/Private/MetasoundEngineAsset.h
--- a/Engine/Plugins/Runtime/Metasound/Source/MetasoundEngine/Private/MetasoundEngineAsset.h
+++ b/Engine/Plugins/Runtime/Metasound/Source/MetasoundEngine/Private/MetasoundEngineAsset.h
@@ -49,7 +49,7 @@ namespace Metasound
 
 			for (const IMetaSoundAssetManager::FAssetInfo& AssetClass : InAssetClasses)
 			{
-				InMetaSound.ReferencedAssetClassKeys.Add(AssetClass.RegistryKey);
+				InMetaSound.ReferencedAssetClassKeys.Add(AssetClass.RegistryKey.ToString());
 				if (UObject* Object = AssetClass.AssetPath.TryLoad())
 				{
 					InMetaSound.ReferencedAssetClassObjects.Add(Object);
diff --git a/Engine/Plugins/Runtime/Metasound/Source/MetasoundEngine/Public/MetasoundBuilderSubsystem.h b/Engine/Plugins/Runtime/Metasound/Source/MetasoundEngine/Public/MetasoundBuilderSubsystem.h
--- a/Engine/Plugins/Runtime/Metasound/Source/MetasoundEngine/Public/MetasoundBuilderSubsystem.h
+++ b/Engine/Plugins/Runtime/Metasound/Source/MetasoundEngine/Public/MetasoundBuilderSubsystem.h
@@ -424,7 +424,7 @@ protected:
 					DocClassName = &ExistingDoc.RootGraph.Metadata.GetClassName();
 				}
 				const FNodeRegistryKey& RegistryKey = MetaSound->GetRegistryKey();
-				if (NodeRegistryKey::IsValid(RegistryKey))
+				if (RegistryKey.IsValid())
 				{
 					if (FMetasoundFrontendRegistryContainer::Get()->IsNodeRegistered(RegistryKey))
 					{
diff --git a/Engine/Plugins/Runtime/Metasound/Source/MetasoundEngineTest/Private/EngineTestMetaSoundAutomatedNodeTest.cpp b/Engine/Plugins/Runtime/Metasound/Source/MetasoundEngineTest/Private/EngineTestMetaSoundAutomatedNodeTest.cpp
--- a/Engine/Plugins/Runtime/Metasound/Source/MetasoundEngineTest/Private/EngineTestMetaSoundAutomatedNodeTest.cpp
+++ b/Engine/Plugins/Runtime/Metasound/Source/MetasoundEngineTest/Private/EngineTestMetaSoundAutomatedNodeTest.cpp
@@ -668,7 +668,7 @@ namespace Metasound::EngineTest{
 		FMetasoundFrontendClass NodeClass;
 		if (!NodeRegistry->FindFrontendClassFromRegistered(InNodeRegistryKey, NodeClass))
 		{
-			UE_LOG(LogMetaSound, Error, TEXT("Failed to find registered class with registry key %s"), *InNodeRegistryKey);
+			UE_LOG(LogMetaSound, Error, TEXT("Failed to find registered class with registry key %s"), *InNodeRegistryKey.ToString());
 			return MoveTemp(Node);
 		}
 
@@ -753,11 +753,11 @@ namespace Metasound::EngineTest{
 				continue;
 			}
 
-			Frontend::FNodeRegistryKey NodeRegistryKey = NodeRegistry->GetRegistryKey(NodeClass.Metadata);
-			
 			OutBeautifiedNames.Add(FString::Printf(TEXT("%s %s"), *NodeClass.Metadata.GetClassName().ToString(), *NodeClass.Metadata.GetVersion().ToString()));
+
 			// Test commands are node registry keys
-			OutTestCommands.Add(NodeRegistryKey);
+			Frontend::FNodeRegistryKey NodeRegistryKey(NodeClass.Metadata);
+			OutTestCommands.Add(NodeRegistryKey.ToString());
 		}
 	}
 
diff --git a/Engine/Plugins/Runtime/Metasound/Source/MetasoundEngineTest/Private/EngineTestMetaSoundAutomatedNodeTest.h b/Engine/Plugins/Runtime/Metasound/Source/MetasoundEngineTest/Private/EngineTestMetaSoundAutomatedNodeTest.h
--- a/Engine/Plugins/Runtime/Metasound/Source/MetasoundEngineTest/Private/EngineTestMetaSoundAutomatedNodeTest.h
+++ b/Engine/Plugins/Runtime/Metasound/Source/MetasoundEngineTest/Private/EngineTestMetaSoundAutomatedNodeTest.h
@@ -7,6 +7,7 @@
 #include "Containers/UnrealString.h"
 #include "CoreFwd.h"
 #include "MetasoundDataReference.h"
+#include "MetasoundFrontendRegistryKey.h"
 #include "MetasoundLog.h"
 #include "MetasoundOperatorSettings.h"
 #include "MetasoundVertex.h"
@@ -30,11 +31,6 @@ namespace Metasound
 	class INode;
 	class FOutputVertexInterfaceData;
 	class FInputVertexInterfaceData;
-
-	namespace Frontend
-	{
-		typedef FString FNodeRegistryKey;
-	}
 }
 
 namespace Metasound::EngineTest 
diff --git a/Engine/Plugins/Runtime/Metasound/Source/MetasoundEngineTest/Private/EngineTestMetaSoundAutomatedNodeTest_Bind.cpp b/Engine/Plugins/Runtime/Metasound/Source/MetasoundEngineTest/Private/EngineTestMetaSoundAutomatedNodeTest_Bind.cpp
--- a/Engine/Plugins/Runtime/Metasound/Source/MetasoundEngineTest/Private/EngineTestMetaSoundAutomatedNodeTest_Bind.cpp
+++ b/Engine/Plugins/Runtime/Metasound/Source/MetasoundEngineTest/Private/EngineTestMetaSoundAutomatedNodeTest_Bind.cpp
@@ -9,6 +9,7 @@
 #include "MetasoundFrontendDataTypeRegistry.h"
 #include "MetasoundOperatorInterface.h"
 #include "MetasoundLog.h"
+#include "MetasoundFrontendRegistryKey.h"
 #include "MetasoundVertex.h"
 #include "MetasoundVertexData.h"
 #include "Templates/UniquePtr.h"
@@ -27,18 +28,26 @@ void FMetasoundAutomatedNodeTest_Bind::GetTests(TArray<FString>& OutBeautifiedNa
 	UE_LOG(LogMetaSound, Verbose, TEXT("Found %d metasound nodes to test"), OutTestCommands.Num());
 }
 
-bool FMetasoundAutomatedNodeTest_Bind::RunTest(const FString& InRegistryKey)
+bool FMetasoundAutomatedNodeTest_Bind::RunTest(const FString& InRegistryKeyString)
 {
 	using namespace Metasound;
 	using namespace Metasound::EngineTest;
+	using namespace Metasound::Frontend;
 
 	static const FOperatorSettings OperatorSettings{48000  /* samplerate */, 100.f /* block rate */};
 	static const FMetasoundEnvironment SourceEnvironment = GetSourceEnvironmentForTest();
 
-	TUniquePtr<INode> Node = CreateNodeFromRegistry(InRegistryKey);
+	FNodeRegistryKey RegistryKey;
+	if (!FNodeRegistryKey::Parse(InRegistryKeyString, RegistryKey))
+	{
+		AddError(FString::Printf(TEXT("Failed to parse registry key string %s"), *InRegistryKeyString));
+		return false;
+	}
+
+	TUniquePtr<INode> Node = CreateNodeFromRegistry(RegistryKey);
 	if (!Node.IsValid())
 	{
-		AddError(FString::Printf(TEXT("Failed to create node %s from registry"), *InRegistryKey));
+		AddError(FString::Printf(TEXT("Failed to create node %s from registry"), *InRegistryKeyString));
 		return false;
 	}
 
@@ -64,7 +73,7 @@ bool FMetasoundAutomatedNodeTest_Bind::RunTest(const FString& InRegistryKey)
 	TUniquePtr<IOperator> Operator = Node->GetDefaultOperatorFactory()->CreateOperator(BuildParams, BuildResults);
 	if (!Operator.IsValid())
 	{
-		AddError(FString::Printf(TEXT("Failed to create operator from node %s - %s."), *InRegistryKey, *GetPrettyName(InRegistryKey)));
+		AddError(FString::Printf(TEXT("Failed to create operator from node %s - %s."), *InRegistryKeyString, *GetPrettyName(RegistryKey)));
 	}
 
 	// Query the data on the operator
@@ -85,7 +94,7 @@ bool FMetasoundAutomatedNodeTest_Bind::RunTest(const FString& InRegistryKey)
 			// is not bound with any data. If data is sent to the CreateOperator(...)
 			// InputData, it should be reflected in the inputs of the IOperator through
 			// the BindInputs(...) method.
-			AddError(FString::Printf(TEXT("Input vertex (%s) of node (%s) does not bind supplied input data. Check that CreateOperator(...) and BindInputs(...) are correctly implemented for the node's operator."), *InputVertex.VertexName.ToString(), *GetPrettyName(InRegistryKey)));
+			AddError(FString::Printf(TEXT("Input vertex (%s) of node (%s) does not bind supplied input data. Check that CreateOperator(...) and BindInputs(...) are correctly implemented for the node's operator."), *InputVertex.VertexName.ToString(), *GetPrettyName(RegistryKey)));
 		}
 		else if (bIsVertexBoundDuringBuild && bIsVertexBoundAfterBuild)
 		{
@@ -110,7 +119,7 @@ bool FMetasoundAutomatedNodeTest_Bind::RunTest(const FString& InRegistryKey)
 					// is bound to different data. If data is sent to the CreateOperator(...)
 					// InputData, it should be reflected in the inputs of the IOperator through
 					// the BindInputs(...) method and be unchanged in this scenario.
-					AddError(FString::Printf(TEXT("Input vertex (%s) of node (%s) binds different data than the supplied at input data. Check that CreateOperator(...) and BindInputs(...) are correctly implemented for the node's operator."), *InputVertex.VertexName.ToString(), *GetPrettyName(InRegistryKey)));
+					AddError(FString::Printf(TEXT("Input vertex (%s) of node (%s) binds different data than the supplied at input data. Check that CreateOperator(...) and BindInputs(...) are correctly implemented for the node's operator."), *InputVertex.VertexName.ToString(), *GetPrettyName(RegistryKey)));
 				}
 			}
 		}
@@ -154,7 +163,7 @@ bool FMetasoundAutomatedNodeTest_Bind::RunTest(const FString& InRegistryKey)
 			// on subsequent calls to IOperator::BindInputs(...).  If new data is 
 			// bound to the operator inputs, it should reflect that new data on subsequent
 			// calls to IOperator::BindInputs(...)
-			AddWarning(FString::Printf(TEXT("Input vertex (%s) of node (%s) does not reflect updated data reference when rebound. This node will not operator as expected in a dynamic graph. Check that CreateOperator(...) and BindInputs(...) are correctly implemented for the node's operator."), *InputVertex.VertexName.ToString(), *GetPrettyName(InRegistryKey)));
+			AddWarning(FString::Printf(TEXT("Input vertex (%s) of node (%s) does not reflect updated data reference when rebound. This node will not operator as expected in a dynamic graph. Check that CreateOperator(...) and BindInputs(...) are correctly implemented for the node's operator."), *InputVertex.VertexName.ToString(), *GetPrettyName(RegistryKey)));
 		}
 		else if (!bIsUpdateExpected && bIsUpdateActualized)
 		{
@@ -163,7 +172,7 @@ bool FMetasoundAutomatedNodeTest_Bind::RunTest(const FString& InRegistryKey)
 			// on subsequent calls to IOperator::BindInputs(...).  If new data is 
 			// bound to the operator inputs, it should reflect that new data on subsequent
 			// calls to IOperator::BindInputs(...)
-			AddWarning(FString::Printf(TEXT("Input vertex (%s) of node (%s) reflects unintended updated of data reference when rebound. This node will not operator as expected in a dynamic graph. Check that CreateOperator(...) and BindInputs(...) are correctly implemented for the node's operator."), *InputVertex.VertexName.ToString(), *GetPrettyName(InRegistryKey)));
+			AddWarning(FString::Printf(TEXT("Input vertex (%s) of node (%s) reflects unintended updated of data reference when rebound. This node will not operator as expected in a dynamic graph. Check that CreateOperator(...) and BindInputs(...) are correctly implemented for the node's operator."), *InputVertex.VertexName.ToString(), *GetPrettyName(RegistryKey)));
 		}
 	}
 
diff --git a/Engine/Plugins/Runtime/Metasound/Source/MetasoundEngineTest/Private/EngineTestMetaSoundAutomatedNodeTest_Reset.cpp b/Engine/Plugins/Runtime/Metasound/Source/MetasoundEngineTest/Private/EngineTestMetaSoundAutomatedNodeTest_Reset.cpp
--- a/Engine/Plugins/Runtime/Metasound/Source/MetasoundEngineTest/Private/EngineTestMetaSoundAutomatedNodeTest_Reset.cpp
+++ b/Engine/Plugins/Runtime/Metasound/Source/MetasoundEngineTest/Private/EngineTestMetaSoundAutomatedNodeTest_Reset.cpp
@@ -27,18 +27,26 @@ void FMetasoundAutomatedNodeTest_Reset::GetTests(TArray<FString>& OutBeautifiedN
 	UE_LOG(LogMetaSound, Verbose, TEXT("Found %d metasound nodes to test"), OutTestCommands.Num());
 }
 
-bool FMetasoundAutomatedNodeTest_Reset::RunTest(const FString& InRegistryKey)
+bool FMetasoundAutomatedNodeTest_Reset::RunTest(const FString& InRegistryKeyString)
 {
 	using namespace Metasound;
 	using namespace Metasound::EngineTest;
+	using namespace Metasound::Frontend;
 
 	static const FOperatorSettings OperatorSettings{48000  /* samplerate */, 100.f /* block rate */};
 	static const FMetasoundEnvironment SourceEnvironment = GetSourceEnvironmentForTest();
 
-	TUniquePtr<INode> Node = CreateNodeFromRegistry(InRegistryKey);
+	FNodeRegistryKey RegistryKey;
+	if (!FNodeRegistryKey::Parse(InRegistryKeyString, RegistryKey))
+	{
+		AddError(FString::Printf(TEXT("Failed to parse registry key string %s"), *InRegistryKeyString));
+		return false;
+	}
+
+	TUniquePtr<INode> Node = CreateNodeFromRegistry(RegistryKey);
 	if (!Node.IsValid())
 	{
-		AddError(FString::Printf(TEXT("Failed to create node %s from registry"), *InRegistryKey));
+		AddError(FString::Printf(TEXT("Failed to create node %s from registry"), *InRegistryKeyString));
 		return false;
 	}
 
@@ -74,7 +82,7 @@ bool FMetasoundAutomatedNodeTest_Reset::RunTest(const FString& InRegistryKey)
 
 		if (!Operator.IsValid())
 		{
-			AddError(FString::Printf(TEXT("Failed to create operator from node %s - %s."), *InRegistryKey, *GetPrettyName(InRegistryKey)));
+			AddError(FString::Printf(TEXT("Failed to create operator from node %s - %s."), *InRegistryKeyString, *GetPrettyName(RegistryKey)));
 		}
 
 		// Store a copy of the input values data values so the inputs
@@ -123,14 +131,14 @@ bool FMetasoundAutomatedNodeTest_Reset::RunTest(const FString& InRegistryKey)
 		}
 		else if (OpExecFunc)
 		{
-			AddError(FString::Printf(TEXT("Missing initialize function when execute function exists for node %s - %s"), *InRegistryKey, *GetPrettyName(InRegistryKey)));
+			AddError(FString::Printf(TEXT("Missing initialize function when execute function exists for node %s - %s"), *InRegistryKeyString, *GetPrettyName(RegistryKey)));
 		}
 
 		// Check that after returning all inputs to their original state and calling
 		// reset on the operator, that all output values have returned to their initail state. 
 		if (!OutputTester.AreAllOutputValuesEqualToCapturedValues())
 		{
-			AddError(FString::Printf(TEXT("Reset function resulted in different starting conditions for node %s - %s"), *InRegistryKey, *GetPrettyName(InRegistryKey)));
+			AddError(FString::Printf(TEXT("Reset function resulted in different starting conditions for node %s - %s"), *InRegistryKeyString, *GetPrettyName(RegistryKey)));
 		}
 	};
 
diff --git a/Engine/Plugins/Runtime/Metasound/Source/MetasoundEngineTest/Private/EngineTestMetaSoundAutomatedNodeTest_ValidOutputs.cpp b/Engine/Plugins/Runtime/Metasound/Source/MetasoundEngineTest/Private/EngineTestMetaSoundAutomatedNodeTest_ValidOutputs.cpp
--- a/Engine/Plugins/Runtime/Metasound/Source/MetasoundEngineTest/Private/EngineTestMetaSoundAutomatedNodeTest_ValidOutputs.cpp
+++ b/Engine/Plugins/Runtime/Metasound/Source/MetasoundEngineTest/Private/EngineTestMetaSoundAutomatedNodeTest_ValidOutputs.cpp
@@ -27,18 +27,26 @@ void FMetasoundAutomatedNodeTest_ValidOutputs::GetTests(TArray<FString>& OutBeau
 	UE_LOG(LogMetaSound, Verbose, TEXT("Found %d metasound nodes to test"), OutTestCommands.Num());
 }
 
-bool FMetasoundAutomatedNodeTest_ValidOutputs::RunTest(const FString& InRegistryKey)
+bool FMetasoundAutomatedNodeTest_ValidOutputs::RunTest(const FString& InRegistryKeyString)
 {
 	using namespace Metasound;
 	using namespace Metasound::EngineTest;
+	using namespace Metasound::Frontend;
 
 	static const FOperatorSettings OperatorSettings{48000  /* samplerate */, 100.f /* block rate */};
 	static const FMetasoundEnvironment SourceEnvironment = GetSourceEnvironmentForTest();
 
-	TUniquePtr<INode> Node = CreateNodeFromRegistry(InRegistryKey);
+	FNodeRegistryKey RegistryKey;
+	if (!FNodeRegistryKey::Parse(InRegistryKeyString, RegistryKey))
+	{
+		AddError(FString::Printf(TEXT("Failed to parse registry key string %s"), *InRegistryKeyString));
+		return false;
+	}
+
+	TUniquePtr<INode> Node = CreateNodeFromRegistry(RegistryKey);
 	if (!Node.IsValid())
 	{
-		AddError(FString::Printf(TEXT("Failed to create node %s from registry"), *InRegistryKey));
+		AddError(FString::Printf(TEXT("Failed to create node %s from registry"), *InRegistryKeyString));
 		return false;
 	}
 
@@ -74,7 +82,7 @@ bool FMetasoundAutomatedNodeTest_ValidOutputs::RunTest(const FString& InRegistry
 
 		if (!Operator.IsValid())
 		{
-			AddError(FString::Printf(TEXT("Failed to create operator from node %s - %s."), *InRegistryKey, *GetPrettyName(InRegistryKey)));
+			AddError(FString::Printf(TEXT("Failed to create operator from node %s - %s."), *InRegistryKeyString, *GetPrettyName(RegistryKey)));
 		}
 
 
@@ -95,7 +103,7 @@ bool FMetasoundAutomatedNodeTest_ValidOutputs::RunTest(const FString& InRegistry
 		{
 			if (!OutputTester.AreAllAnalyzableOutputsValid())
 			{
-				AddError(FString::Printf(TEXT("Invalid output value encountered from node %s - %s."), *InRegistryKey, *GetPrettyName(InRegistryKey)));
+				AddError(FString::Printf(TEXT("Invalid output value encountered from node %s - %s."), *InRegistryKeyString, *GetPrettyName(RegistryKey)));
 			}
 		};
 		// Initial values should all be valid.
diff --git a/Engine/Plugins/Runtime/Metasound/Source/MetasoundFrontend/Private/MetasoundAssetBase.cpp b/Engine/Plugins/Runtime/Metasound/Source/MetasoundFrontend/Private/MetasoundAssetBase.cpp
--- a/Engine/Plugins/Runtime/Metasound/Source/MetasoundFrontend/Private/MetasoundAssetBase.cpp
+++ b/Engine/Plugins/Runtime/Metasound/Source/MetasoundFrontend/Private/MetasoundAssetBase.cpp
@@ -326,7 +326,7 @@ void FMetasoundAssetBase::RegisterGraphWithFrontend(Metasound::Frontend::FMetaSo
 		RegistryKey = CacheRuntimeData(RegistryDocInterface);
 	}
 
-	if (NodeRegistryKey::IsValid(RegistryKey))
+	if (RegistryKey.IsValid())
 	{
 #if WITH_EDITORONLY_DATA
 		UpdateAssetRegistry();
@@ -396,7 +396,7 @@ void FMetasoundAssetBase::CookMetaSound()
 		RegistryKey = FMetasoundFrontendRegistryContainer::Get()->RegisterNode(MoveTemp(RegistryEntry));
 	}
 
-	if (NodeRegistryKey::IsValid(RegistryKey))
+	if (RegistryKey.IsValid())
 	{
 #if WITH_EDITORONLY_DATA
 		UpdateAssetRegistry();
@@ -417,7 +417,7 @@ void FMetasoundAssetBase::UnregisterGraphWithFrontend()
 	METASOUND_TRACE_CPUPROFILER_EVENT_SCOPE(MetaSoundAssetBase::UnregisterGraphWithFrontend);
 
 	check(IsInGameThread());
-	if (!NodeRegistryKey::IsValid(RegistryKey))
+	if (!RegistryKey.IsValid())
 	{
 		return;
 	}
@@ -440,7 +440,7 @@ void FMetasoundAssetBase::UnregisterGraphWithFrontend()
 
 	if (!bSuccess)
 	{
-		UE_LOG(LogMetaSound, Warning, TEXT("Failed to unregister node with key %s for asset %s. No registry entry exists with that key."), *RegistryKey, *GetOwningAssetName());
+		UE_LOG(LogMetaSound, Warning, TEXT("Failed to unregister node with key %s for asset %s. No registry entry exists with that key."), *RegistryKey.ToString(), *GetOwningAssetName());
 	}
 
 	RegistryKey = FNodeRegistryKey();
@@ -703,7 +703,7 @@ bool FMetasoundAssetBase::IsRegistered() const
 {
 	using namespace Metasound::Frontend;
 
-	return NodeRegistryKey::IsValid(RegistryKey);
+	return RegistryKey.IsValid();
 }
 
 bool FMetasoundAssetBase::IsReferencedAsset(const FMetasoundAssetBase& InAsset) const
@@ -1007,7 +1007,7 @@ Metasound::Frontend::FNodeRegistryKey FMetasoundAssetBase::CacheRuntimeData(cons
 void FMetasoundAssetBase::WaitForAsyncGraphRegistration()
 {
 	using namespace Metasound::Frontend;
-	if (NodeRegistryKey::IsValid(RegistryKey))
+	if (RegistryKey.IsValid())
 	{
 		FMetasoundFrontendRegistryContainer::Get()->WaitForAsyncGraphRegistration(RegistryKey, RegisteredGraphAssetPath);
 	}
diff --git a/Engine/Plugins/Runtime/Metasound/Source/MetasoundFrontend/Private/MetasoundFrontend.cpp b/Engine/Plugins/Runtime/Metasound/Source/MetasoundFrontend/Private/MetasoundFrontend.cpp
--- a/Engine/Plugins/Runtime/Metasound/Source/MetasoundFrontend/Private/MetasoundFrontend.cpp
+++ b/Engine/Plugins/Runtime/Metasound/Source/MetasoundFrontend/Private/MetasoundFrontend.cpp
@@ -52,7 +52,7 @@ namespace Metasound
 			if (ensure(nullptr != Registry))
 			{
 				bool bSuccess = Registry->FindFrontendClassFromRegistered(InKey, OutClass);
-				ensureAlwaysMsgf(bSuccess, TEXT("Cannot generate description of unregistered node [RegistryKey:%s]"), *InKey);
+				ensureAlwaysMsgf(bSuccess, TEXT("Cannot generate description of unregistered node [RegistryKey:%s]"), *InKey.ToString());
 			}
 
 			return OutClass;
@@ -109,7 +109,7 @@ class FMetasoundFrontendModule : public IModuleInterface
 	{
 		using namespace Metasound::Frontend;
 
-		UnregisterNodeTemplate(FRerouteNodeTemplate::Version);
+		UnregisterNodeTemplate(FRerouteNodeTemplate::ClassName, FRerouteNodeTemplate::VersionNumber);
 	}
 };
 
diff --git a/Engine/Plugins/Runtime/Metasound/Source/MetasoundFrontend/Private/MetasoundFrontendDocument.cpp b/Engine/Plugins/Runtime/Metasound/Source/MetasoundFrontend/Private/MetasoundFrontendDocument.cpp
--- a/Engine/Plugins/Runtime/Metasound/Source/MetasoundFrontend/Private/MetasoundFrontendDocument.cpp
+++ b/Engine/Plugins/Runtime/Metasound/Source/MetasoundFrontend/Private/MetasoundFrontendDocument.cpp
@@ -36,6 +36,21 @@ namespace Metasound
 				const FVector2D DefaultOffsetY { 0.0f, 80.0f };
 			} // namespace NodeLayout
 		} // namespace DisplayStyle
+
+		namespace ClassTypePrivate
+		{
+			static const FString External = TEXT("External");
+			static const FString Graph = TEXT("Graph");
+			static const FString Input = TEXT("Input");
+			static const FString Output = TEXT("Output");
+			static const FString Literal = TEXT("Literal");
+			static const FString Variable = TEXT("Variable");
+			static const FString VariableDeferredAccessor = TEXT("Variable (Deferred Accessor)");
+			static const FString VariableAccessor = TEXT("Variable (Accessor)");
+			static const FString VariableMutator = TEXT("Variable (Mutator)");
+			static const FString Template = TEXT("Template");
+			static const FString Invalid = TEXT("Invalid");
+		}
 	} // namespace Frontend
 
 	namespace DocumentPrivate
@@ -122,6 +137,7 @@ namespace Metasound
 	} // namespace DocumentPrivate
 } // namespace Metasound
 
+
 #if WITH_EDITORONLY_DATA
 void FMetasoundFrontendDocumentModifyContext::ClearDocumentModified()
 {
@@ -365,6 +381,12 @@ const FMetasoundFrontendInterfaceUClassOptions* FMetasoundFrontendInterface::Fin
 
 const FMetasoundFrontendClassName FMetasoundFrontendClassName::InvalidClassName;
 
+FMetasoundFrontendClassName::FMetasoundFrontendClassName(const FName& InNamespace, const FName& InName)
+	: Namespace(InNamespace)
+	, Name(InName)
+{
+}
+
 FMetasoundFrontendClassName::FMetasoundFrontendClassName(const FName& InNamespace, const FName& InName, const FName& InVariant)
 : Namespace(InNamespace)
 , Name(InName)
@@ -403,14 +425,28 @@ FString FMetasoundFrontendClassName::ToString() const
 	return GetFullName().ToString();
 }
 
-bool operator==(const FMetasoundFrontendClassName& InLHS, const FMetasoundFrontendClassName& InRHS)
+bool FMetasoundFrontendClassName::Parse(const FString& InClassName, FMetasoundFrontendClassName& OutClassName)
 {
-	return (InLHS.Namespace == InRHS.Namespace) && (InLHS.Name == InRHS.Name) && (InLHS.Variant == InRHS.Variant);
-}
+	OutClassName = { };
+	TArray<FString> Tokens;
+	InClassName.ParseIntoArray(Tokens, TEXT("."));
 
-bool operator!=(const FMetasoundFrontendClassName& InLHS, const FMetasoundFrontendClassName& InRHS)
-{
-	return !(InLHS == InRHS);
+	// Name is required, which in turn requires at least "None" is serialized for the namespace
+	if (Tokens.Num() < 2)
+	{
+		return false;
+	}
+
+	OutClassName.Namespace = FName(*Tokens[0]);
+	OutClassName.Name = FName(*Tokens[1]);
+
+	// Variant is optional
+	if (Tokens.Num() > 2)
+	{
+		OutClassName.Variant = FName(*Tokens[2]);
+	}
+
+	return true;
 }
 
 FMetasoundFrontendClassInterface FMetasoundFrontendClassInterface::GenerateClassInterface(const Metasound::FVertexInterface& InVertexInterface)
@@ -644,7 +680,7 @@ bool FMetasoundFrontendClass::CacheGraphDependencyMetadataFromRegistry(FMetasoun
 {
 	using namespace Metasound::Frontend;
 
-	const FNodeRegistryKey Key = NodeRegistryKey::CreateKey(InOutDependency.Metadata);
+	const FNodeRegistryKey Key = FNodeRegistryKey(InOutDependency.Metadata);
 	FMetasoundFrontendClass RegistryClass;
 
 	FMetasoundFrontendRegistryContainer* Registry = FMetasoundFrontendRegistryContainer::Get();
@@ -833,30 +869,32 @@ FMetasoundFrontendDocument::FMetasoundFrontendDocument()
 
 const TCHAR* LexToString(EMetasoundFrontendClassType InClassType)
 {
+	using namespace Metasound::Frontend;
+
 	switch (InClassType)
 	{
 		case EMetasoundFrontendClassType::External:
-			return TEXT("External");
+			return *ClassTypePrivate::External;
 		case EMetasoundFrontendClassType::Graph:
-			return TEXT("Graph");
+			return *ClassTypePrivate::Graph;
 		case EMetasoundFrontendClassType::Input:
-			return TEXT("Input");
+			return *ClassTypePrivate::Input;
 		case EMetasoundFrontendClassType::Output:
-			return TEXT("Output");
+			return *ClassTypePrivate::Output;
 		case EMetasoundFrontendClassType::Literal:
-			return TEXT("Literal");
+			return *ClassTypePrivate::Literal;
 		case EMetasoundFrontendClassType::Variable:
-			return TEXT("Variable");
+			return *ClassTypePrivate::Variable;
 		case EMetasoundFrontendClassType::VariableDeferredAccessor:
-			return TEXT("Variable (Deferred Accessor)");
+			return *ClassTypePrivate::VariableDeferredAccessor;
 		case EMetasoundFrontendClassType::VariableAccessor:
-			return TEXT("Variable (Accessor)");
+			return *ClassTypePrivate::VariableAccessor;
 		case EMetasoundFrontendClassType::VariableMutator:
-			return TEXT("Variable (Mutator)");
+			return *ClassTypePrivate::VariableMutator;
 		case EMetasoundFrontendClassType::Template:
-			return TEXT("Template");
+			return *ClassTypePrivate::Template;
 		case EMetasoundFrontendClassType::Invalid:
-			return TEXT("Invalid");
+			return *ClassTypePrivate::Invalid;
 		default:
 			static_assert(static_cast<int32>(EMetasoundFrontendClassType::Invalid) == 10, "Possible missed EMetasoundFrontendClassType case coverage");
 			return nullptr;
@@ -879,60 +917,117 @@ const TCHAR* LexToString(EMetasoundFrontendVertexAccessType InVertexAccess)
 	}
 }
 
-namespace Metasound
+namespace Metasound::Frontend
 {
-	namespace Frontend
+	bool StringToClassType(const FString& InString, EMetasoundFrontendClassType& OutClassType)
 	{
-		void ForEachLiteral(const FMetasoundFrontendDocument& InDoc, FForEachLiteralFunctionRef OnLiteral)
+		if (InString == ClassTypePrivate::External)
 		{
-			ForEachLiteral(InDoc.RootGraph, OnLiteral);
-			for (const FMetasoundFrontendGraphClass& GraphClass : InDoc.Subgraphs)
-			{
-				ForEachLiteral(GraphClass, OnLiteral);
-			}
-			for (const FMetasoundFrontendClass& Dependency : InDoc.Dependencies)
-			{
-				ForEachLiteral(Dependency, OnLiteral);
-			}
+			OutClassType = EMetasoundFrontendClassType::External;
+			return true;
 		}
 
-		void ForEachLiteral(const FMetasoundFrontendGraphClass& InGraphClass, FForEachLiteralFunctionRef OnLiteral)
+		if (InString == ClassTypePrivate::Graph)
 		{
-			ForEachLiteral(static_cast<const FMetasoundFrontendClass&>(InGraphClass), OnLiteral);
+			OutClassType = EMetasoundFrontendClassType::Graph;
+			return true;
+		}
 
-			for (const FMetasoundFrontendNode& Node : InGraphClass.Graph.Nodes)
-			{
-				ForEachLiteral(Node, OnLiteral);
-			}
+		if (InString == ClassTypePrivate::Input)
+		{
+			OutClassType = EMetasoundFrontendClassType::Input;
+			return true;
+		}
 
-			for (const FMetasoundFrontendVariable& Variable : InGraphClass.Graph.Variables)
-			{
-				OnLiteral(Variable.TypeName, Variable.Literal);
-			}
+		if (InString == ClassTypePrivate::Invalid)
+		{
+			OutClassType = EMetasoundFrontendClassType::Invalid;
+			return true;
 		}
 
-		void ForEachLiteral(const FMetasoundFrontendClass& InClass, FForEachLiteralFunctionRef OnLiteral)
+		if (InString == ClassTypePrivate::Literal)
 		{
-			for (const FMetasoundFrontendClassInput& ClassInput : InClass.Interface.Inputs)
-			{
-				OnLiteral(ClassInput.TypeName, ClassInput.DefaultLiteral);
-			}
+			OutClassType = EMetasoundFrontendClassType::Literal;
+			return true;
 		}
 
-		void ForEachLiteral(const FMetasoundFrontendNode& InNode, FForEachLiteralFunctionRef OnLiteral)
+		if (InString == ClassTypePrivate::Output)
 		{
-			for (const FMetasoundFrontendVertexLiteral& VertexLiteral : InNode.InputLiterals)
-			{
-				auto HasEqualVertexID = [&VertexLiteral](const FMetasoundFrontendVertex& InVertex) -> bool
-				{ 
-					return InVertex.VertexID == VertexLiteral.VertexID; 
-				};
+			OutClassType = EMetasoundFrontendClassType::Output;
+			return true;
+		}
 
-				if (const FMetasoundFrontendVertex* InputVertex = InNode.Interface.Inputs.FindByPredicate(HasEqualVertexID))
-				{
-					OnLiteral(InputVertex->TypeName, VertexLiteral.Value);
-				}
+		if (InString == ClassTypePrivate::Template)
+		{
+			OutClassType = EMetasoundFrontendClassType::Template;
+			return true;
+		}
+
+		if (InString == ClassTypePrivate::Variable)
+		{
+			OutClassType = EMetasoundFrontendClassType::Variable;
+			return true;
+		}
+
+		if (InString == ClassTypePrivate::VariableDeferredAccessor)
+		{
+			OutClassType = EMetasoundFrontendClassType::VariableDeferredAccessor;
+			return true;
+		}
+
+		OutClassType = EMetasoundFrontendClassType::Invalid;
+		return false;
+	}
+
+	void ForEachLiteral(const FMetasoundFrontendDocument& InDoc, FForEachLiteralFunctionRef OnLiteral)
+	{
+		ForEachLiteral(InDoc.RootGraph, OnLiteral);
+		for (const FMetasoundFrontendGraphClass& GraphClass : InDoc.Subgraphs)
+		{
+			ForEachLiteral(GraphClass, OnLiteral);
+		}
+		for (const FMetasoundFrontendClass& Dependency : InDoc.Dependencies)
+		{
+			ForEachLiteral(Dependency, OnLiteral);
+		}
+	}
+
+	void ForEachLiteral(const FMetasoundFrontendGraphClass& InGraphClass, FForEachLiteralFunctionRef OnLiteral)
+	{
+		ForEachLiteral(static_cast<const FMetasoundFrontendClass&>(InGraphClass), OnLiteral);
+
+		for (const FMetasoundFrontendNode& Node : InGraphClass.Graph.Nodes)
+		{
+			ForEachLiteral(Node, OnLiteral);
+		}
+
+		for (const FMetasoundFrontendVariable& Variable : InGraphClass.Graph.Variables)
+		{
+			OnLiteral(Variable.TypeName, Variable.Literal);
+		}
+	}
+
+	void ForEachLiteral(const FMetasoundFrontendClass& InClass, FForEachLiteralFunctionRef OnLiteral)
+	{
+		for (const FMetasoundFrontendClassInput& ClassInput : InClass.Interface.Inputs)
+		{
+			OnLiteral(ClassInput.TypeName, ClassInput.DefaultLiteral);
+		}
+	}
+
+	void ForEachLiteral(const FMetasoundFrontendNode& InNode, FForEachLiteralFunctionRef OnLiteral)
+	{
+		for (const FMetasoundFrontendVertexLiteral& VertexLiteral : InNode.InputLiterals)
+		{
+			auto HasEqualVertexID = [&VertexLiteral](const FMetasoundFrontendVertex& InVertex) -> bool
+			{ 
+				return InVertex.VertexID == VertexLiteral.VertexID; 
+			};
+
+			if (const FMetasoundFrontendVertex* InputVertex = InNode.Interface.Inputs.FindByPredicate(HasEqualVertexID))
+			{
+				OnLiteral(InputVertex->TypeName, VertexLiteral.Value);
 			}
 		}
 	}
-}
+} // namespace Metasound::Frontend
diff --git a/Engine/Plugins/Runtime/Metasound/Source/MetasoundFrontend/Private/MetasoundFrontendDocumentAccessPtr.cpp b/Engine/Plugins/Runtime/Metasound/Source/MetasoundFrontend/Private/MetasoundFrontendDocumentAccessPtr.cpp
--- a/Engine/Plugins/Runtime/Metasound/Source/MetasoundFrontend/Private/MetasoundFrontendDocumentAccessPtr.cpp
+++ b/Engine/Plugins/Runtime/Metasound/Source/MetasoundFrontend/Private/MetasoundFrontendDocumentAccessPtr.cpp
@@ -407,12 +407,12 @@ namespace Metasound
 
 				bool operator()(const FMetasoundFrontendClass& InClass) const
 				{
-					return NodeRegistryKey::IsEqual(RegistryKey, NodeRegistryKey::CreateKey(InClass.Metadata));
+					return RegistryKey == FNodeRegistryKey(InClass.Metadata);
 				};
 
 				bool operator()(const FMetasoundFrontendGraphClass& InGraphClass) const
 				{
-					return NodeRegistryKey::IsEqual(RegistryKey, NodeRegistryKey::CreateKey(InGraphClass));
+					return RegistryKey == FNodeRegistryKey(InGraphClass);
 				};
 			private:
 				FNodeRegistryKey RegistryKey;
diff --git a/Engine/Plugins/Runtime/Metasound/Source/MetasoundFrontend/Private/MetasoundFrontendDocumentBuilder.cpp b/Engine/Plugins/Runtime/Metasound/Source/MetasoundFrontend/Private/MetasoundFrontendDocumentBuilder.cpp
--- a/Engine/Plugins/Runtime/Metasound/Source/MetasoundFrontend/Private/MetasoundFrontendDocumentBuilder.cpp
+++ b/Engine/Plugins/Runtime/Metasound/Source/MetasoundFrontend/Private/MetasoundFrontendDocumentBuilder.cpp
@@ -19,6 +19,7 @@
 #include "MetasoundFrontendDocumentModifyDelegates.h"
 #include "MetasoundFrontendNodeTemplateRegistry.h"
 #include "MetasoundFrontendRegistries.h"
+#include "MetasoundFrontendRegistryKey.h"
 #include "MetasoundFrontendSearchEngine.h"
 #include "MetasoundFrontendTransform.h"
 #include "MetasoundTrace.h"
@@ -134,7 +135,7 @@ namespace Metasound::Frontend
 		{
 			const FMetasoundFrontendClassMetadata& Metadata = DocumentInterface.GetConstDocument().RootGraph.Metadata;
 			const FMetasoundFrontendVersionNumber& Version = Metadata.GetVersion();
-			const FNodeRegistryKey RegKey = NodeRegistryKey::CreateKey(EMetasoundFrontendClassType::External, Metadata.GetClassName().ToString(), Version.Major, Version.Minor);
+			const FNodeRegistryKey RegKey(EMetasoundFrontendClassType::External, Metadata.GetClassName(), Version.Major, Version.Minor);
 			if (const FMetasoundAssetBase* Asset = IMetaSoundAssetManager::GetChecked().TryLoadAssetFromKey(RegKey))
 			{
 				return Asset->GetOwningAssetName();
@@ -1125,7 +1126,7 @@ bool FMetaSoundFrontendDocumentBuilder::AddInterface(FName InterfaceName)
 
 const FMetasoundFrontendNode* FMetaSoundFrontendDocumentBuilder::AddGraphNode(const FMetasoundFrontendGraphClass& InGraphClass, FGuid InNodeID)
 {
-	auto FinalizeNode = [](const FMetasoundFrontendNode& Node, const Metasound::Frontend::FNodeRegistryKey& ClassKey)
+	auto FinalizeNode = [](FMetasoundFrontendNode& InOutNode, const Metasound::Frontend::FNodeRegistryKey& ClassKey)
 	{
 #if WITH_EDITOR
 		using namespace Metasound::Frontend;
@@ -1135,15 +1136,15 @@ const FMetasoundFrontendNode* FMetaSoundFrontendDocumentBuilder::AddGraphNode(co
 		{
 			if (const FSoftObjectPath* Path = AssetManager->FindObjectPathFromKey(ClassKey))
 			{
-				return FName(*Path->GetAssetName());
+				InOutNode.Name = Path->GetAssetFName();
+				return;
 			}
 		}
-#endif // WITH_EDITOR
 
-		return Node.Name;
+		InOutNode.Name = ClassKey.ClassName.GetFullName();
+#endif // WITH_EDITOR
 	};
 
-
 	// Dependency is considered "External" when looked up or added on another graph
 	FMetasoundFrontendClassMetadata NewClassMetadata = InGraphClass.Metadata;
 	NewClassMetadata.SetType(EMetasoundFrontendClassType::External);
@@ -1198,7 +1199,7 @@ FMetasoundFrontendNode* FMetaSoundFrontendDocumentBuilder::AddNodeInternal(const
 
 	using namespace Metasound::Frontend;
 
-	const FNodeRegistryKey ClassKey = NodeRegistryKey::CreateKey(InClassMetadata);
+	const FNodeRegistryKey ClassKey = FNodeRegistryKey(InClassMetadata);
 	if (const FMetasoundFrontendClass* Dependency = DocumentCache->FindDependency(ClassKey))
 	{
 		FMetasoundFrontendDocument& Document = GetDocument();
@@ -1381,7 +1382,7 @@ const FMetasoundFrontendClass* FMetaSoundFrontendDocumentBuilder::FindDependency
 
 	checkf(InMetadata.GetType() != EMetasoundFrontendClassType::Graph,
 		TEXT("Dependencies are never listed as 'Graph' types. Graphs are considered 'External' from the perspective of the parent document to allow for nativization."));
-	const FNodeRegistryKey RegistryKey = NodeRegistryKey::CreateKey(InMetadata);
+	const FNodeRegistryKey RegistryKey = FNodeRegistryKey(InMetadata);
 	return DocumentCache->FindDependency(RegistryKey);
 }
 
@@ -1530,7 +1531,7 @@ const TSet<FMetasoundFrontendVersion>* FMetaSoundFrontendDocumentBuilder::FindNo
 		if (const FMetasoundFrontendClass* NodeClass = DocumentCache->FindDependency(Node->ClassID))
 		{
 			// 1. May be a serialized asset, so first check with asset manager.
-			const FNodeRegistryKey NodeClassRegistryKey = NodeRegistryKey::CreateKey(NodeClass->Metadata);
+			const FNodeRegistryKey NodeClassRegistryKey = FNodeRegistryKey(NodeClass->Metadata);
 			return FMetasoundFrontendRegistryContainer::Get()->FindImplementedInterfacesFromRegistered(NodeClassRegistryKey);
 		}
 	}
@@ -1678,9 +1679,9 @@ EMetasoundFrontendVertexAccessType FMetaSoundFrontendDocumentBuilder::GetNodeInp
 			{
 				case EMetasoundFrontendClassType::Template:
 				{
-					const FNodeRegistryKey Key = NodeRegistryKey::CreateKey(Class->Metadata);
+					const FNodeRegistryKey Key = FNodeRegistryKey(Class->Metadata);
 					const INodeTemplate* Template = INodeTemplateRegistry::Get().FindTemplate(Key);
-					if (ensureMsgf(Template, TEXT("Failed to find MetaSound node template registered with key '%s'"), *Key))
+					if (ensureMsgf(Template, TEXT("Failed to find MetaSound node template registered with key '%s'"), *Key.ToString()))
 					{
 						if (Template->IsInputAccessTypeDynamic())
 						{
@@ -1803,9 +1804,9 @@ EMetasoundFrontendVertexAccessType FMetaSoundFrontendDocumentBuilder::GetNodeOut
 			{
 				case EMetasoundFrontendClassType::Template:
 				{
-					const FNodeRegistryKey Key = NodeRegistryKey::CreateKey(Class->Metadata);
+					const FNodeRegistryKey Key = FNodeRegistryKey(Class->Metadata);
 					const INodeTemplate* Template = INodeTemplateRegistry::Get().FindTemplate(Key);
-					if (ensureMsgf(Template, TEXT("Failed to find MetaSound node template registered with key '%s'"), *Key))
+					if (ensureMsgf(Template, TEXT("Failed to find MetaSound node template registered with key '%s'"), *Key.ToString()))
 					{
 						if (Template->IsOutputAccessTypeDynamic())
 						{
@@ -2057,7 +2058,7 @@ bool FMetaSoundFrontendDocumentBuilder::TransformTemplateNodes()
 	{
 		if (Dependency.Metadata.GetType() == EMetasoundFrontendClassType::Template)
 		{
-			const FNodeRegistryKey Key = NodeRegistryKey::CreateKey(Dependency.Metadata);
+			const FNodeRegistryKey Key = FNodeRegistryKey(Dependency.Metadata);
 			const INodeTemplate* Template = INodeTemplateRegistry::Get().FindTemplate(Key);
 			ensureMsgf(Template, TEXT("Template not found for template class reference '%s'"), *Dependency.Metadata.GetClassName().ToString());
 			TemplateParams.Add(Dependency.ID, FTemplateTransformParams { Template });
@@ -2174,7 +2175,7 @@ bool FMetaSoundFrontendDocumentBuilder::RemoveDependency(EMetasoundFrontendClass
 	FMetasoundFrontendDocument& Document = GetDocument();
 	const IDocumentGraphNodeCache& NodeCache = DocumentCache->GetNodeCache();
 
-	const FNodeRegistryKey ClassKey = NodeRegistryKey::CreateKey(ClassType, InClassName.GetFullName().ToString(), InClassVersionNumber.Major, InClassVersionNumber.Minor);
+	const FNodeRegistryKey ClassKey(ClassType, InClassName, InClassVersionNumber);
 	if (const int32* IndexPtr = DocumentCache->FindDependencyIndex(ClassKey))
 	{
 		TArray<FMetasoundFrontendClass>& Dependencies = Document.Dependencies;
diff --git a/Engine/Plugins/Runtime/Metasound/Source/MetasoundFrontend/Private/MetasoundFrontendDocumentCache.cpp b/Engine/Plugins/Runtime/Metasound/Source/MetasoundFrontend/Private/MetasoundFrontendDocumentCache.cpp
--- a/Engine/Plugins/Runtime/Metasound/Source/MetasoundFrontend/Private/MetasoundFrontendDocumentCache.cpp
+++ b/Engine/Plugins/Runtime/Metasound/Source/MetasoundFrontend/Private/MetasoundFrontendDocumentCache.cpp
@@ -209,7 +209,7 @@ namespace Metasound::Frontend
 		using namespace Metasound::Frontend;
 
 		const FMetasoundFrontendClass& Dependency = GetDocument().Dependencies[NewIndex];
-		const FNodeRegistryKey Key = NodeRegistryKey::CreateKey(Dependency.Metadata);
+		const FNodeRegistryKey Key = FNodeRegistryKey(Dependency.Metadata);
 
 		FDocumentDependencyCache& Cache = GetDependencyCache();
 		Cache.KeyToIndex.Add(Key, NewIndex);
@@ -224,10 +224,14 @@ namespace Metasound::Frontend
 		const FMetasoundFrontendClass& SwapDependency = Dependencies[SwapIndex];
 		const FMetasoundFrontendClass& LastDependency = Dependencies[LastIndex];
 
-		const FNodeRegistryKey SwapKey = NodeRegistryKey::CreateKey(SwapDependency.Metadata);
-		const FNodeRegistryKey LastKey = NodeRegistryKey::CreateKey(LastDependency.Metadata);
+		const FNodeRegistryKey SwapKey = FNodeRegistryKey(SwapDependency.Metadata);
+		const FNodeRegistryKey LastKey = FNodeRegistryKey(LastDependency.Metadata);
 
 		FDocumentDependencyCache& Cache = GetDependencyCache();
+
+		const bool contains = Cache.KeyToIndex.Contains(SwapKey);
+		check(contains);
+
 		DocumentCachePrivate::RemoveSwapMapIndexChecked(SwapKey, LastKey, Cache.KeyToIndex);
 		DocumentCachePrivate::RemoveSwapMapIndexChecked(SwapDependency.ID, LastDependency.ID, Cache.IDToIndex);
 	}
@@ -237,10 +241,10 @@ namespace Metasound::Frontend
 		using namespace Metasound::Frontend;
 
 		const FMetasoundFrontendClass& DependencyBeingRemoved = GetDocument().Dependencies[IndexBeingRenamed];
-		const FNodeRegistryKey OldKey = NodeRegistryKey::CreateKey(DependencyBeingRemoved.Metadata);
+		const FNodeRegistryKey OldKey = FNodeRegistryKey(DependencyBeingRemoved.Metadata);
 		FMetasoundFrontendClassMetadata NewMetadata = DependencyBeingRemoved.Metadata;
 		NewMetadata.SetClassName(NewName);
-		const FNodeRegistryKey NewKey = NodeRegistryKey::CreateKey(NewMetadata);
+		const FNodeRegistryKey NewKey = FNodeRegistryKey(NewMetadata);
 
 		FDocumentDependencyCache& Cache = GetDependencyCache();
 		Cache.KeyToIndex.Add(NewKey, IndexBeingRenamed);
@@ -822,7 +826,7 @@ namespace Metasound::Frontend
 		{
 			const FMetasoundFrontendClass& Class = InDocument.Dependencies[Index];
 			const FMetasoundFrontendClassMetadata& Metadata = Class.Metadata;
-			FNodeRegistryKey Key = NodeRegistryKey::CreateKey(Metadata);
+			FNodeRegistryKey Key = FNodeRegistryKey(Metadata);
 			KeyToIndex.Add(MoveTemp(Key), Index);
 			IDToIndex.Add(Class.ID, Index);
 		}
diff --git a/Engine/Plugins/Runtime/Metasound/Source/MetasoundFrontend/Private/MetasoundFrontendDocumentController.cpp b/Engine/Plugins/Runtime/Metasound/Source/MetasoundFrontend/Private/MetasoundFrontendDocumentController.cpp
--- a/Engine/Plugins/Runtime/Metasound/Source/MetasoundFrontend/Private/MetasoundFrontendDocumentController.cpp
+++ b/Engine/Plugins/Runtime/Metasound/Source/MetasoundFrontend/Private/MetasoundFrontendDocumentController.cpp
@@ -370,7 +370,7 @@ namespace Metasound
 						case EMetasoundFrontendClassType::Output:
 						{
 							FMetasoundFrontendClass NewClass;
-							FNodeRegistryKey Key = NodeRegistryKey::CreateKey(InMetadata);
+							FNodeRegistryKey Key = FNodeRegistryKey(InMetadata);
 
 							if (FRegistry::GetFrontendClassFromRegistered(Key, NewClass))
 							{
@@ -382,7 +382,7 @@ namespace Metasound
 #if WITH_EDITOR
 								UE_LOG(LogMetaSound, Error,
 									TEXT("Cannot add external dependency. No Metasound class found with matching registry key [Key:%s, Name:%s, Version:%s]. Suggested solution \"%s\" by %s."),
-									*Key,
+									*Key.ToString(),
 									*InMetadata.GetClassName().GetFullName().ToString(),
 									*InMetadata.GetVersion().ToString(),
 									*InMetadata.GetPromptIfMissing().ToString(),
@@ -390,7 +390,7 @@ namespace Metasound
 #else
 								UE_LOG(LogMetaSound, Error,
 									TEXT("Cannot add external dependency. No Metasound class found with matching registry key [Key:%s, Name:%s, Version:%s]."),
-									*Key,
+									*Key.ToString(),
 									*InMetadata.GetClassName().GetFullName().ToString(),
 									*InMetadata.GetVersion().ToString());
 #endif // !WITH_EDITOR
@@ -470,9 +470,8 @@ namespace Metasound
 			{
 				for (FMetasoundFrontendClass& Class : Document->Dependencies)
 				{
-					FNodeRegistryKey RegistryKey = NodeRegistryKey::CreateKey(Class.Metadata);
-
 					FMetasoundFrontendClass RegistryVersion;
+					FNodeRegistryKey RegistryKey(Class.Metadata);
 					if (FMetasoundFrontendRegistryContainer::Get()->FindFrontendClassFromRegistered(RegistryKey, RegistryVersion))
 					{
 						if (Class.Metadata.GetChangeID() != RegistryVersion.Metadata.GetChangeID())
diff --git a/Engine/Plugins/Runtime/Metasound/Source/MetasoundFrontend/Private/MetasoundFrontendGraph.cpp b/Engine/Plugins/Runtime/Metasound/Source/MetasoundFrontend/Private/MetasoundFrontendGraph.cpp
--- a/Engine/Plugins/Runtime/Metasound/Source/MetasoundFrontend/Private/MetasoundFrontendGraph.cpp
+++ b/Engine/Plugins/Runtime/Metasound/Source/MetasoundFrontend/Private/MetasoundFrontendGraph.cpp
@@ -241,7 +241,7 @@ namespace Metasound
 		// TODO: handle check to see if node interface conforms to class interface here. 
 		// TODO: check to see if external object supports class interface.
 
-		const FNodeRegistryKey Key = NodeRegistryKey::CreateKey(InClass.Metadata);
+		const FNodeRegistryKey Key = FNodeRegistryKey(InClass.Metadata);
 		return FMetasoundFrontendRegistryContainer::Get()->CreateNode(Key, InitData);
 	}
 
diff --git a/Engine/Plugins/Runtime/Metasound/Source/MetasoundFrontend/Private/MetasoundFrontendGraphController.cpp b/Engine/Plugins/Runtime/Metasound/Source/MetasoundFrontend/Private/MetasoundFrontendGraphController.cpp
--- a/Engine/Plugins/Runtime/Metasound/Source/MetasoundFrontend/Private/MetasoundFrontendGraphController.cpp
+++ b/Engine/Plugins/Runtime/Metasound/Source/MetasoundFrontend/Private/MetasoundFrontendGraphController.cpp
@@ -1411,7 +1411,7 @@ namespace Metasound
 				return AddNode(Class, InNodeGuid);
 			}
 
-			UE_LOG(LogMetaSound, Warning, TEXT("Failed to find or add node class info with registry key [Key:%s]"), *InKey);
+			UE_LOG(LogMetaSound, Warning, TEXT("Failed to find or add node class info with registry key [Key:%s]"), *InKey.ToString());
 			return INodeController::GetInvalidHandle();
 		}
 
@@ -1421,7 +1421,7 @@ namespace Metasound
 				TEXT("Cannot implement '%s' template node using 'AddNode'. Template nodes must always "
 				"be added using AddTemplateNode function and supply the interface to be implemented"),
 				*InClassMetadata.GetClassName().ToString());
-			return AddNode(NodeRegistryKey::CreateKey(InClassMetadata), InNodeGuid);
+			return AddNode(FNodeRegistryKey(InClassMetadata), InNodeGuid);
 		}
 
 		FNodeHandle FGraphController::AddTemplateNode(const FNodeRegistryKey& InKey, FMetasoundFrontendNodeInterface&& InNodeInterface, FGuid InNodeGuid)
@@ -1429,7 +1429,7 @@ namespace Metasound
 			if (const INodeTemplate* Template = INodeTemplateRegistry::Get().FindTemplate(InKey))
 			{
 				const bool bIsValidInterface = Template->IsValidNodeInterface(InNodeInterface);
-				if (ensureAlwaysMsgf(bIsValidInterface, TEXT("Cannot implement interface when attempting to add node using template with key '%s'"), *InKey))
+				if (ensureAlwaysMsgf(bIsValidInterface, TEXT("Cannot implement interface when attempting to add node using template with key '%s'"), *InKey.ToString()))
 				{
 					// Construct a FNodeClassInfo from this lookup key.
 					FConstClassAccessPtr Class = OwningDocument->FindOrAddClass(InKey);
@@ -1452,7 +1452,7 @@ namespace Metasound
 				}
 			}
 
-			UE_LOG(LogMetaSound, Warning, TEXT("Failed to find or add node template class info with registry key [Key:%s]"), *InKey);
+			UE_LOG(LogMetaSound, Warning, TEXT("Failed to find or add node template class info with registry key [Key:%s]"), *InKey.ToString());
 			return INodeController::GetInvalidHandle();
 		}
 
@@ -1656,7 +1656,7 @@ namespace Metasound
 					{
 						if (IMetaSoundAssetManager* AssetManager = IMetaSoundAssetManager::Get())
 						{
-							const FNodeRegistryKey RegistryKey = NodeRegistryKey::CreateKey(NodeClass->Metadata);
+							const FNodeRegistryKey RegistryKey = FNodeRegistryKey(NodeClass->Metadata);
 							if (const FSoftObjectPath* Path = AssetManager->FindObjectPathFromKey(RegistryKey))
 							{
 								const FString& AssetName = Path->GetAssetName();
diff --git a/Engine/Plugins/Runtime/Metasound/Source/MetasoundFrontend/Private/MetasoundFrontendNodeController.cpp b/Engine/Plugins/Runtime/Metasound/Source/MetasoundFrontend/Private/MetasoundFrontendNodeController.cpp
--- a/Engine/Plugins/Runtime/Metasound/Source/MetasoundFrontend/Private/MetasoundFrontendNodeController.cpp
+++ b/Engine/Plugins/Runtime/Metasound/Source/MetasoundFrontend/Private/MetasoundFrontendNodeController.cpp
@@ -737,7 +737,7 @@ namespace Metasound
 			const FMetasoundFrontendClassMetadata Metadata = GetClassMetadata();
 
 			// Lookup new version in node registry
-			FNodeRegistryKey NewVersionRegistryKey = NodeRegistryKey::CreateKey(Metadata.GetType(), Metadata.GetClassName().ToString(), InNewVersion.Major, InNewVersion.Minor);
+			FNodeRegistryKey NewVersionRegistryKey = FNodeRegistryKey(Metadata.GetType(), Metadata.GetClassName(), InNewVersion);
 			FMetasoundFrontendRegistryContainer* Registry = FMetasoundFrontendRegistryContainer::Get();
 			checkf(nullptr != Registry, TEXT("The metasound node registry should always be available if the metasound plugin is loaded"));
 
@@ -827,7 +827,7 @@ namespace Metasound
 			// Make sure classes are up-to-date with registered versions of class.
 			// Note that this may break other nodes in the graph that have stale
 			// class API, but that's on the caller to fix-up or report invalid state.
-			const FNodeRegistryKey RegistryKey = NodeRegistryKey::CreateKey(NewMetasoundClass.Metadata);
+			const FNodeRegistryKey RegistryKey = FNodeRegistryKey(NewMetasoundClass.Metadata);
 			FDocumentHandle Document = GetOwningGraph()->GetOwningDocument();
 
 			constexpr bool bRefreshFromRegistry = true;
@@ -946,7 +946,7 @@ namespace Metasound
 				FMetasoundFrontendRegistryContainer* Registry = FMetasoundFrontendRegistryContainer::Get();
 				checkf(nullptr != Registry, TEXT("The metasound node registry should always be available if the metasound plugin is loaded"));
 				FMetasoundFrontendClass RegisteredClass;
-				bool bFoundRegisteredClass = Registry->FindFrontendClassFromRegistered(NodeRegistryKey::CreateKey(GetClassMetadata()), OutInterfaceUpdates.RegistryClass);
+				bool bFoundRegisteredClass = Registry->FindFrontendClassFromRegistered(FNodeRegistryKey(GetClassMetadata()), OutInterfaceUpdates.RegistryClass);
 
 				if (!bFoundRegisteredClass)
 				{
@@ -1039,7 +1039,7 @@ namespace Metasound
 				// define changes in native vs asset class definitions. Need to
 				// hash a changeID from natively defined node classes in order to
 				// merge branches.
-				const FNodeRegistryKey RegistryKey = NodeRegistryKey::CreateKey(RegistryClass.Metadata);
+				const FNodeRegistryKey RegistryKey = FNodeRegistryKey(RegistryClass.Metadata);
 				const bool bIsClassNative = FMetasoundFrontendRegistryContainer::Get()->IsNodeNative(RegistryKey);
 				if (bIsClassNative)
 				{
diff --git a/Engine/Plugins/Runtime/Metasound/Source/MetasoundFrontend/Private/MetasoundFrontendNodeRegistryPrivate.h b/Engine/Plugins/Runtime/Metasound/Source/MetasoundFrontend/Private/MetasoundFrontendNodeRegistryPrivate.h
--- a/Engine/Plugins/Runtime/Metasound/Source/MetasoundFrontend/Private/MetasoundFrontendNodeRegistryPrivate.h
+++ b/Engine/Plugins/Runtime/Metasound/Source/MetasoundFrontend/Private/MetasoundFrontendNodeRegistryPrivate.h
@@ -117,7 +117,10 @@ namespace Metasound
 			// Find Frontend class data.
 			virtual bool FindFrontendClassFromRegistered(const FNodeRegistryKey& InKey, FMetasoundFrontendClass& OutClass) override;
 			virtual const TSet<FMetasoundFrontendVersion>* FindImplementedInterfacesFromRegistered(const Metasound::Frontend::FNodeRegistryKey& InKey) const override;
+
+			UE_DEPRECATED(5.4, "Use FindFrontendClassFromRegistered instead")
 			virtual bool FindNodeClassInfoFromRegistered(const Metasound::Frontend::FNodeRegistryKey& InKey, FNodeClassInfo& OutInfo) override;
+
 			virtual bool FindInputNodeRegistryKeyForDataType(const FName& InDataTypeName, const EMetasoundFrontendVertexAccessType InAccessType, FNodeRegistryKey& OutKey) override;
 			virtual bool FindVariableNodeRegistryKeyForDataType(const FName& InDataTypeName, FNodeRegistryKey& OutKey) override;
 			virtual bool FindOutputNodeRegistryKeyForDataType(const FName& InDataTypeName, const EMetasoundFrontendVertexAccessType InAccessType, FNodeRegistryKey& OutKey) override;
diff --git a/Engine/Plugins/Runtime/Metasound/Source/MetasoundFrontend/Private/MetasoundFrontendNodeTemplateRegistry.cpp b/Engine/Plugins/Runtime/Metasound/Source/MetasoundFrontend/Private/MetasoundFrontendNodeTemplateRegistry.cpp
--- a/Engine/Plugins/Runtime/Metasound/Source/MetasoundFrontend/Private/MetasoundFrontendNodeTemplateRegistry.cpp
+++ b/Engine/Plugins/Runtime/Metasound/Source/MetasoundFrontend/Private/MetasoundFrontendNodeTemplateRegistry.cpp
@@ -56,9 +56,8 @@ namespace Metasound::Frontend
 	{
 		if (ensure(InTemplate.IsValid()))
 		{
-			const FMetasoundFrontendVersion& Version = InTemplate->GetVersion();
-			const FNodeRegistryKey Key = NodeRegistryKey::CreateKey(InTemplate->GetFrontendClass().Metadata);
-			if (ensure(NodeRegistryKey::IsValid(Key)))
+			const FNodeRegistryKey Key = FNodeRegistryKey(InTemplate->GetFrontendClass().Metadata);
+			if (ensure(Key.IsValid()))
 			{
 				Templates.Add(Key, MoveTemp(InTemplate));
 			}
@@ -122,8 +121,20 @@ namespace Metasound::Frontend
 
 	void UnregisterNodeTemplate(const FMetasoundFrontendVersion& InVersion)
 	{
-		const FNodeRegistryKey Key = NodeRegistryKey::CreateKey(EMetasoundFrontendClassType::Template, InVersion.Name.ToString(), InVersion.Number.Major, InVersion.Number.Minor);
-		if (ensure(NodeRegistryKey::IsValid(Key)))
+		FMetasoundFrontendClassName ClassName;
+		FMetasoundFrontendClassName::Parse(InVersion.Name.ToString(), ClassName);
+		const FNodeRegistryKey Key = FNodeRegistryKey(EMetasoundFrontendClassType::Template, ClassName, InVersion.Number);
+		if (ensure(Key.IsValid()))
+		{
+			FRegistryContainerImpl::Get().UnregisterNodeTemplate(Key);
+			static_cast<FNodeTemplateRegistry&>(INodeTemplateRegistry::Get()).Unregister(Key);
+		}
+	}
+
+	void UnregisterNodeTemplate(const FMetasoundFrontendClassName& InClassName, const FMetasoundFrontendVersionNumber& InTemplateVersion)
+	{
+		const FNodeRegistryKey Key = FNodeRegistryKey(EMetasoundFrontendClassType::Template, InClassName, InTemplateVersion);
+		if (ensure(Key.IsValid()))
 		{
 			FRegistryContainerImpl::Get().UnregisterNodeTemplate(Key);
 			static_cast<FNodeTemplateRegistry&>(INodeTemplateRegistry::Get()).Unregister(Key);
diff --git a/Engine/Plugins/Runtime/Metasound/Source/MetasoundFrontend/Private/MetasoundFrontendQuerySteps.cpp b/Engine/Plugins/Runtime/Metasound/Source/MetasoundFrontend/Private/MetasoundFrontendQuerySteps.cpp
--- a/Engine/Plugins/Runtime/Metasound/Source/MetasoundFrontend/Private/MetasoundFrontendQuerySteps.cpp
+++ b/Engine/Plugins/Runtime/Metasound/Source/MetasoundFrontend/Private/MetasoundFrontendQuerySteps.cpp
@@ -61,7 +61,7 @@ namespace Metasound
 			RegistryKey = InEntry.Value.Get<FNodeRegistryTransaction>().GetNodeRegistryKey();
 		}
 
-		return FFrontendQueryKey(RegistryKey);
+		return FFrontendQueryKey(RegistryKey.ToString());
 	}
 
 	void FReduceRegistrationEventsToCurrentStatus::Reduce(const FFrontendQueryKey& InKey, FFrontendQueryPartition& InOutEntries) const
diff --git a/Engine/Plugins/Runtime/Metasound/Source/MetasoundFrontend/Private/MetasoundFrontendRegistries.cpp b/Engine/Plugins/Runtime/Metasound/Source/MetasoundFrontend/Private/MetasoundFrontendRegistries.cpp
--- a/Engine/Plugins/Runtime/Metasound/Source/MetasoundFrontend/Private/MetasoundFrontendRegistries.cpp
+++ b/Engine/Plugins/Runtime/Metasound/Source/MetasoundFrontend/Private/MetasoundFrontendRegistries.cpp
@@ -33,58 +33,6 @@ namespace Metasound
 	{
 		namespace MetasoundFrontendRegistryPrivate
 		{
-			const FString& GetClassTypeString(EMetasoundFrontendClassType InType)
-			{
-				static const FString InputType(TEXT("Input"));
-				static const FString OutputType(TEXT("Output"));
-				static const FString ExternalType(TEXT("External"));
-				static const FString VariableAccessorType(TEXT("VariableAccessor"));
-				static const FString VariableDeferredAccessorType(TEXT("VariableDeferredAccessor"));
-				static const FString VariableMutatorType(TEXT("VariableMutator"));
-				static const FString VariableType(TEXT("Variable"));
-				static const FString LiteralType(TEXT("Literal"));
-				static const FString GraphType(TEXT("Graph"));
-				static const FString TemplateType(TEXT("Template"));
-				static const FString InvalidType(TEXT("Invalid"));
-
-				switch (InType)
-				{
-					case EMetasoundFrontendClassType::Input:
-						return InputType;
-
-					case EMetasoundFrontendClassType::Output:
-						return OutputType;
-
-					case EMetasoundFrontendClassType::External:
-						return ExternalType;
-
-					case EMetasoundFrontendClassType::Literal:
-						return LiteralType;
-
-					case EMetasoundFrontendClassType::Variable:
-						return VariableType;
-						
-					case EMetasoundFrontendClassType::VariableDeferredAccessor:
-						return VariableDeferredAccessorType;
-
-					case EMetasoundFrontendClassType::VariableAccessor:
-						return VariableAccessorType;
-
-					case EMetasoundFrontendClassType::VariableMutator:
-						return VariableMutatorType;
-
-					case EMetasoundFrontendClassType::Template:
-						return TemplateType;
-
-					case EMetasoundFrontendClassType::Graph:
-						return GraphType;
-
-					default:
-						static_assert(static_cast<uint8>(EMetasoundFrontendClassType::Invalid) == 10, "Missing EMetasoundFrontendClassType case coverage");
-						return InvalidType;
-				}
-			}
-
 			// FGraphNode is used to create unique INodes based off of a IGraph. 
 			//
 			// Individual nodes need to reflect their InstanceName and InstanceID, but otherwise
@@ -348,7 +296,7 @@ namespace Metasound
 			// Because creation of external nodes can rely on assets being unavailable due to errors in loading order, asset(s)
 			// missing, etc. only log error and don't throw ensure to avoid blocking start-up if assets are missing. All other
 			// CreateNode calls are natively managed and thus better suited to throw ensures.
-			UE_LOG(LogMetaSound, Error, TEXT("Could not find node [RegistryKey:%s]"), *InKey);
+			UE_LOG(LogMetaSound, Error, TEXT("Could not find node [RegistryKey:%s]"), *InKey.ToString());
 			return nullptr;
 		}
 
@@ -356,7 +304,7 @@ namespace Metasound
 		{
 			const INodeRegistryEntry* Entry = FindNodeEntry(InKey);
 
-			if (ensureAlwaysMsgf(nullptr != Entry, TEXT("Could not find node [RegistryKey:%s]"), *InKey))
+			if (ensureAlwaysMsgf(nullptr != Entry, TEXT("Could not find node [RegistryKey:%s]"), *InKey.ToString()))
 			{
 				return Entry->CreateNode(MoveTemp(InParams));
 			}
@@ -368,7 +316,7 @@ namespace Metasound
 		{
 			const INodeRegistryEntry* Entry = FindNodeEntry(InKey);
 
-			if (ensureAlwaysMsgf(nullptr != Entry, TEXT("Could not find node [RegistryKey:%s]"), *InKey))
+			if (ensureAlwaysMsgf(nullptr != Entry, TEXT("Could not find node [RegistryKey:%s]"), *InKey.ToString()))
 			{
 				return Entry->CreateNode(MoveTemp(InParams));
 			}
@@ -380,7 +328,7 @@ namespace Metasound
 		{
 			const INodeRegistryEntry* Entry = FindNodeEntry(InKey);
 
-			if (ensureAlwaysMsgf(nullptr != Entry, TEXT("Could not find node [RegistryKey:%s]"), *InKey))
+			if (ensureAlwaysMsgf(nullptr != Entry, TEXT("Could not find node [RegistryKey:%s]"), *InKey.ToString()))
 			{
 				return Entry->CreateNode(MoveTemp(InParams));
 			}
@@ -408,13 +356,15 @@ namespace Metasound
 
 		FNodeRegistryKey FRegistryContainerImpl::RegisterGraph(const FSoftObjectPath& InAssetPath, const TScriptInterface<IMetaSoundDocumentInterface>& InDocumentInterface, bool bAsync)
 		{
+			METASOUND_TRACE_CPUPROFILER_EVENT_SCOPE(Metasound::FRegistryContainerImpl::RegisterGraph);
+
 			using namespace UE;
 
 			check(InDocumentInterface);
 			check(IsInGameThread());
 
 			const FMetasoundFrontendDocument& Document = InDocumentInterface->GetConstDocument();
-			FNodeRegistryKey RegistryKey = NodeRegistryKey::CreateKey(Document.RootGraph);
+			FNodeRegistryKey RegistryKey = FNodeRegistryKey(Document.RootGraph);
 			FNodeClassInfo NodeClassInfo(Document.RootGraph, InAssetPath);
 
 			// Proxies are created synchronously to avoid creating proxies in async tasks. Proxies
@@ -428,7 +378,7 @@ namespace Metasound
 				FScopeLock LockActiveReg(&ActiveRegistrationTasksCriticalSection);
 				if (const FActiveRegistrationTaskInfo* ActiveTaskInfo = ActiveRegistrationTasks.Find(RegistryKey))
 				{
-					UE_LOG(LogMetaSound, Warning, TEXT("Waiting for async registration task to finish before beginning new registration task for same registration key (%s) with asset (%s))"), *RegistryKey, *InAssetPath.ToString());
+					UE_LOG(LogMetaSound, Warning, TEXT("Waiting for async registration task to finish before beginning new registration task for same registration key (%s) with asset (%s))"), *RegistryKey.ToString(), *InAssetPath.ToString());
 					ActiveTaskInfo->Task.Wait(FTimespan::FromSeconds(1.));
 				}
 			
@@ -452,7 +402,7 @@ namespace Metasound
 							int32 NumRemoved = Registry.ActiveRegistrationTasks.Remove(RegistryKey);
 							if (NumRemoved != 1)
 							{
-								UE_LOG(LogMetaSound, Warning, TEXT("Multiple active registration tasks for the same asset may result in incorrect MetaSound graphs being instantiated during runtime. Registry Key: (%s) Asset Path: (%s)"), *RegistryKey, *NodeClassInfo.AssetPath.ToString());
+								UE_LOG(LogMetaSound, Warning, TEXT("Multiple active registration tasks for the same asset may result in incorrect MetaSound graphs being instantiated during runtime. Registry Key: (%s) Asset Path: (%s)"), *RegistryKey.ToString(), *NodeClassInfo.AssetPath.ToString());
 							}
 						}
 					}
@@ -482,7 +432,7 @@ namespace Metasound
 
 			if (const TSharedPtr<const FGraph>* ExistingEntry = RegisteredGraphs.Find(GraphRegistryKey))
 			{
-				UE_LOG(LogMetaSound, Warning, TEXT("Multiple graphs are registered with the same registry key (%s) and asset path. The existing registered graph will be replaced with the new graph."),  *InKey, *InAssetPath.ToString());
+				UE_LOG(LogMetaSound, Warning, TEXT("Multiple graphs are registered with the same registry key (%s) and asset path. The existing registered graph will be replaced with the new graph."),  *InKey.ToString(), *InAssetPath.ToString());
 			}
 
 			RegisteredGraphs.Add(GraphRegistryKey, MoveTemp(InGraph));
@@ -521,7 +471,7 @@ namespace Metasound
 
 			if (!Graph)
 			{
-				UE_LOG(LogMetaSound, Error, TEXT("Could not find graph with registry key (%s) and asset (%s)."),  *InNodeRegistryKey, *InAssetPath.ToString());
+				UE_LOG(LogMetaSound, Error, TEXT("Could not find graph with registry key (%s) and asset (%s)."),  *InNodeRegistryKey.ToString(), *InAssetPath.ToString());
 			}
 
 			return Graph;
@@ -539,7 +489,7 @@ namespace Metasound
 
 				FNodeRegistryTransaction::FTimeType Timestamp = FPlatformTime::Cycles64();
 
-				Key = NodeRegistryKey::CreateKey(Entry->GetClassInfo());
+				Key = FNodeRegistryKey(Entry->GetClassInfo());
 
 				{
 					FScopeLock Lock(&RegistryMapsCriticalSection);
@@ -549,7 +499,7 @@ namespace Metasound
 					{
 						const FNodeClassInfo& ClassInfo = (*ExistingEntry)->GetClassInfo();
 						UE_LOG(LogMetaSound, Error, TEXT("Node with registry key '%s' already registered by asset '%s' encountered while registering node with asset %s. MetaSounds which depend on these assets may utilize incorrect asset." 
-						), *Key, *ClassInfo.AssetPath.ToString(), *Entry->GetClassInfo().AssetPath.ToString());
+						), *Key.ToString(), *ClassInfo.AssetPath.ToString(), *Entry->GetClassInfo().AssetPath.ToString());
 					}
 
 					// Store registry elements in map so nodes can be queried using registry key.
@@ -577,7 +527,7 @@ namespace Metasound
 
 				FNodeRegistryTransaction::FTimeType Timestamp = FPlatformTime::Cycles64();
 
-				Key = NodeRegistryKey::CreateKey(Entry->GetClassInfo());
+				Key = FNodeRegistryKey(Entry->GetClassInfo());
 
 				{
 					FScopeLock Lock(&RegistryMapsCriticalSection);
@@ -585,7 +535,7 @@ namespace Metasound
 					ensureAlwaysMsgf(
 						!RegisteredNodeTemplates.Contains(Key),
 						TEXT("Node template with registry key '%s' already registered. The previously registered node will be overwritten."),
-						*Key);
+						*Key.ToString());
 
 
 					// Store registry elements in map so nodes can be queried using registry key.
@@ -606,7 +556,7 @@ namespace Metasound
 			METASOUND_LLM_SCOPE;
 			
 			check(IsInGameThread());
-			if (NodeRegistryKey::IsValid(InKey))
+			if (InKey.IsValid())
 			{
 				FScopeLock Lock(&RegistryMapsCriticalSection);
 				if (const TSharedRef<INodeRegistryEntry, ESPMode::ThreadSafe>* EntryPtr = RegisteredNodes.Find(InKey))
@@ -630,7 +580,7 @@ namespace Metasound
 		{
 			METASOUND_LLM_SCOPE;
 
-			if (NodeRegistryKey::IsValid(InKey))
+			if (InKey.IsValid())
 			{
 				if (const INodeRegistryTemplateEntry* Entry = FindNodeTemplateEntry(InKey))
 				{
@@ -757,7 +707,7 @@ namespace Metasound
 				{
 					if (IDataTypeRegistry::Get().GetFrontendInputClass(InDataTypeName, Class))
 					{
-						OutKey = NodeRegistryKey::CreateKey(Class.Metadata);
+						OutKey = FNodeRegistryKey(Class.Metadata);
 						return true;
 					}
 				}
@@ -767,7 +717,7 @@ namespace Metasound
 				{
 					if (IDataTypeRegistry::Get().GetFrontendConstructorInputClass(InDataTypeName, Class))
 					{
-						OutKey = NodeRegistryKey::CreateKey(Class.Metadata);
+						OutKey = FNodeRegistryKey(Class.Metadata);
 						return true;
 					}
 				}
@@ -789,7 +739,7 @@ namespace Metasound
 			FMetasoundFrontendClass Class;
 			if (IDataTypeRegistry::Get().GetFrontendLiteralClass(InDataTypeName, Class))
 			{
-				OutKey = NodeRegistryKey::CreateKey(Class.Metadata);
+				OutKey = FNodeRegistryKey(Class.Metadata);
 				return true;
 			}
 			return false;
@@ -804,7 +754,7 @@ namespace Metasound
 				{
 					if (IDataTypeRegistry::Get().GetFrontendOutputClass(InDataTypeName, Class))
 					{
-						OutKey = NodeRegistryKey::CreateKey(Class.Metadata);
+						OutKey = FNodeRegistryKey(Class.Metadata);
 						return true;
 					}
 				}
@@ -814,7 +764,7 @@ namespace Metasound
 				{
 					if (IDataTypeRegistry::Get().GetFrontendConstructorOutputClass(InDataTypeName, Class))
 					{
-						OutKey = NodeRegistryKey::CreateKey(Class.Metadata);
+						OutKey = FNodeRegistryKey(Class.Metadata);
 						return true;
 					}
 				}
@@ -875,7 +825,6 @@ namespace Metasound
 			}
 		}
 
-
 		const INodeRegistryTemplateEntry* FRegistryContainerImpl::FindNodeTemplateEntry(const FNodeRegistryKey& InKey) const
 		{
 			FScopeLock Lock(&RegistryMapsCriticalSection);
@@ -947,7 +896,7 @@ namespace Metasound
 
 		FNodeRegistryKey FNodeRegistryTransaction::GetNodeRegistryKey() const
 		{
-			return NodeRegistryKey::CreateKey(NodeClassInfo);
+			return FNodeRegistryKey(NodeClassInfo);
 		}
 
 		FNodeRegistryTransaction::FTimeType FNodeRegistryTransaction::GetTimestamp() const
@@ -957,7 +906,6 @@ namespace Metasound
 
 		namespace NodeRegistryKey
 		{
-			// All registry keys should be created through this function to ensure consistency.
 			FNodeRegistryKey CreateKey(EMetasoundFrontendClassType InType, const FString& InFullClassName, int32 InMajorVersion, int32 InMinorVersion)
 			{
 				using namespace MetasoundFrontendRegistryPrivate;
@@ -966,19 +914,20 @@ namespace Metasound
 					// No graphs are registered. Any registered graph should be registered as an external node.
 					InType = EMetasoundFrontendClassType::External;
 				}
-				const FString RegistryKey = FString::Format(TEXT("{0}_{1}_{2}.{3}"), {*GetClassTypeString(InType), *InFullClassName, InMajorVersion, InMinorVersion});
-				return RegistryKey;
+
+				FMetasoundFrontendClassName ClassName;
+				FMetasoundFrontendClassName::Parse(InFullClassName, ClassName);
+				return FNodeRegistryKey(InType, ClassName, InMajorVersion, InMinorVersion);
 			}
 
 			const FNodeRegistryKey& GetInvalid()
 			{
-				static const FNodeRegistryKey InvalidKey;
-				return InvalidKey;
+				return FNodeRegistryKey::GetInvalid();
 			}
 
 			bool IsValid(const FNodeRegistryKey& InKey)
 			{
-				return InKey != GetInvalid();
+				return InKey.IsValid();
 			}
 
 			bool IsEqual(const FNodeRegistryKey& InLHS, const FNodeRegistryKey& InRHS)
@@ -1018,61 +967,25 @@ namespace Metasound
 
 			FNodeRegistryKey CreateKey(const FNodeClassMetadata& InNodeMetadata)
 			{
-				return CreateKey(EMetasoundFrontendClassType::External, InNodeMetadata.ClassName.GetFullName().ToString(), InNodeMetadata.MajorVersion, InNodeMetadata.MinorVersion);
+				return FNodeRegistryKey(InNodeMetadata);
 			}
 
 			FNodeRegistryKey CreateKey(const FMetasoundFrontendClassMetadata& InNodeMetadata)
 			{
 				checkf(InNodeMetadata.GetType() != EMetasoundFrontendClassType::Graph, TEXT("Cannot create key from 'graph' type. Likely meant to use CreateKey overload that is provided FMetasoundFrontendGraphClass"));
-				return CreateKey(InNodeMetadata.GetType(), InNodeMetadata.GetClassName().GetFullName().ToString(), InNodeMetadata.GetVersion().Major, InNodeMetadata.GetVersion().Minor);
+				return FNodeRegistryKey(InNodeMetadata);
 			}
 
 			FNodeRegistryKey CreateKey(const FMetasoundFrontendGraphClass& InGraphClass)
 			{
-				return CreateKey(
-					EMetasoundFrontendClassType::External, // Overridden as all graphs are considered the same as an external in the registry
-					InGraphClass.Metadata.GetClassName().GetFullName().ToString(),
-					InGraphClass.Metadata.GetVersion().Major,
-					InGraphClass.Metadata.GetVersion().Minor
-				);
+				return FNodeRegistryKey(InGraphClass);
 			}
 
 			FNodeRegistryKey CreateKey(const FNodeClassInfo& InClassInfo)
 			{
-				checkf(InClassInfo.Type != EMetasoundFrontendClassType::Graph, TEXT("Cannot create key from 'graph' type. Likely meant to use CreateKey overload that is provided FMetasoundFrontendGraphClass"));
-				return CreateKey(InClassInfo.Type, InClassInfo.ClassName.GetFullName().ToString(), InClassInfo.Version.Major, InClassInfo.Version.Minor);
+				return FNodeRegistryKey(InClassInfo);
 			}
 		}
-
-		FNodeClassInfo::FNodeClassInfo(const FMetasoundFrontendClassMetadata& InMetadata)
-			: ClassName(InMetadata.GetClassName())
-			, Type(InMetadata.GetType())
-			, Version(InMetadata.GetVersion())
-		{
-		}
-
-		FNodeClassInfo::FNodeClassInfo(const FMetasoundFrontendGraphClass& InClass, const FSoftObjectPath& InAssetPath)
-			: ClassName(InClass.Metadata.GetClassName())
-			, Type(EMetasoundFrontendClassType::External) // Overridden as it is considered the same as an external class in the registry
-			, AssetClassID(FGuid(ClassName.Name.ToString()))
-			, AssetPath(InAssetPath)
-			, Version(InClass.Metadata.GetVersion())
-		{
-			ensure(!AssetPath.IsNull());
-#if WITH_EDITORONLY_DATA
-			for (const FMetasoundFrontendClassInput& Input : InClass.Interface.Inputs)
-			{
-				InputTypes.Add(Input.TypeName);
-			}
-
-			for (const FMetasoundFrontendClassOutput& Output : InClass.Interface.Outputs)
-			{
-				OutputTypes.Add(Output.TypeName);
-			}
-
-			bIsPreset = InClass.PresetOptions.bIsPreset;
-#endif // WITH_EDITORONLY_DATA
-		}
 	} // namespace Frontend
 } // namespace Metasound
 
@@ -1089,17 +1002,17 @@ void FMetasoundFrontendRegistryContainer::ShutdownMetasoundFrontend()
 
 Metasound::Frontend::FNodeRegistryKey FMetasoundFrontendRegistryContainer::GetRegistryKey(const FNodeClassMetadata& InNodeMetadata)
 {
-	return Metasound::Frontend::NodeRegistryKey::CreateKey(InNodeMetadata);
+	return Metasound::Frontend::FNodeRegistryKey(InNodeMetadata);
 }
 
 Metasound::Frontend::FNodeRegistryKey FMetasoundFrontendRegistryContainer::GetRegistryKey(const FMetasoundFrontendClassMetadata& InNodeMetadata)
 {
-	return Metasound::Frontend::NodeRegistryKey::CreateKey(InNodeMetadata);
+	return Metasound::Frontend::FNodeRegistryKey(InNodeMetadata);
 }
 
 Metasound::Frontend::FNodeRegistryKey FMetasoundFrontendRegistryContainer::GetRegistryKey(const FNodeClassInfo& InClassInfo)
 {
-	return Metasound::Frontend::NodeRegistryKey::CreateKey(InClassInfo);
+	return Metasound::Frontend::FNodeRegistryKey(InClassInfo);
 }
 
 bool FMetasoundFrontendRegistryContainer::GetFrontendClassFromRegistered(const FNodeRegistryKey& InKey, FMetasoundFrontendClass& OutClass)
@@ -1118,7 +1031,9 @@ bool FMetasoundFrontendRegistryContainer::GetNodeClassInfoFromRegistered(const F
 {
 	if (FMetasoundFrontendRegistryContainer* Registry = FMetasoundFrontendRegistryContainer::Get())
 	{
+		PRAGMA_DISABLE_DEPRECATION_WARNINGS
 		return Registry->FindNodeClassInfoFromRegistered(InKey, OutInfo);
+		PRAGMA_ENABLE_DEPRECATION_WARNINGS
 	}
 	return false;
 }
diff --git a/Engine/Plugins/Runtime/Metasound/Source/MetasoundFrontend/Private/MetasoundFrontendRegistryKey.cpp b/Engine/Plugins/Runtime/Metasound/Source/MetasoundFrontend/Private/MetasoundFrontendRegistryKey.cpp
new file mode 100644
--- /dev/null
+++ b/Engine/Plugins/Runtime/Metasound/Source/MetasoundFrontend/Private/MetasoundFrontendRegistryKey.cpp
@@ -0,0 +1,162 @@
+// Copyright Epic Games, Inc. All Rights Reserved.
+
+#include "MetasoundFrontendRegistryKey.h"
+
+
+namespace Metasound::Frontend
+{
+	FNodeClassInfo::FNodeClassInfo(const FMetasoundFrontendClassMetadata& InMetadata)
+		: ClassName(InMetadata.GetClassName())
+		, Type(InMetadata.GetType())
+		, Version(InMetadata.GetVersion())
+	{
+	}
+
+	FNodeClassInfo::FNodeClassInfo(const FMetasoundFrontendGraphClass& InClass, const FSoftObjectPath& InAssetPath)
+		: ClassName(InClass.Metadata.GetClassName())
+		, Type(EMetasoundFrontendClassType::External) // Overridden as it is considered the same as an external class in registries
+		, AssetClassID(FGuid(ClassName.Name.ToString()))
+		, AssetPath(InAssetPath)
+		, Version(InClass.Metadata.GetVersion())
+	{
+		ensure(!AssetPath.IsNull());
+#if WITH_EDITORONLY_DATA
+		for (const FMetasoundFrontendClassInput& Input : InClass.Interface.Inputs)
+		{
+			InputTypes.Add(Input.TypeName);
+		}
+
+		for (const FMetasoundFrontendClassOutput& Output : InClass.Interface.Outputs)
+		{
+			OutputTypes.Add(Output.TypeName);
+		}
+
+		bIsPreset = InClass.PresetOptions.bIsPreset;
+#endif // WITH_EDITORONLY_DATA
+	}
+
+	UObject* FNodeClassInfo::LoadAsset() const
+	{
+		if (ensure(Type == EMetasoundFrontendClassType::External))
+		{
+			FSoftObjectPath SoftObjectPath(AssetPath);
+			return SoftObjectPath.TryLoad();
+		}
+
+		return nullptr;
+	}
+
+	FNodeRegistryKey::FNodeRegistryKey(EMetasoundFrontendClassType InType, const FMetasoundFrontendClassName& InClassName, int32 InMajorVersion, int32 InMinorVersion)
+		: Type(InType)
+		, ClassName(InClassName)
+		, Version({ InMajorVersion, InMinorVersion })
+	{
+	}
+
+	FNodeRegistryKey::FNodeRegistryKey(EMetasoundFrontendClassType InType, const FMetasoundFrontendClassName& InClassName, const FMetasoundFrontendVersionNumber& InVersion)
+		: Type(InType)
+		, ClassName(InClassName)
+		, Version(InVersion)
+	{
+	}
+
+	FNodeRegistryKey::FNodeRegistryKey(const FNodeClassMetadata& InNodeMetadata)
+		: Type(EMetasoundFrontendClassType::External) // Overridden as it is considered the same as an external class in registries
+		, ClassName(InNodeMetadata.ClassName)
+		, Version({ InNodeMetadata.MajorVersion, InNodeMetadata.MinorVersion })
+	{
+	}
+
+	FNodeRegistryKey::FNodeRegistryKey(const FMetasoundFrontendClassMetadata& InNodeMetadata)
+		: Type(InNodeMetadata.GetType())
+		, ClassName(InNodeMetadata.GetClassName())
+		, Version(InNodeMetadata.GetVersion())
+	{
+		checkf(InNodeMetadata.GetType() != EMetasoundFrontendClassType::Graph, TEXT("Cannot create key from 'graph' type. Likely meant to use FNodeRegistryKey ctor that is provided FMetasoundFrontendGraphClass"));
+	}
+
+	FNodeRegistryKey::FNodeRegistryKey(const FMetasoundFrontendGraphClass& InGraphClass)
+		: Type(EMetasoundFrontendClassType::External) // Type overridden as all graphs are considered the same as an external class in the registry
+		, ClassName(InGraphClass.Metadata.GetClassName())
+		, Version(InGraphClass.Metadata.GetVersion())
+	{
+	}
+
+	FNodeRegistryKey::FNodeRegistryKey(const FNodeClassInfo& InClassInfo)
+		: Type(InClassInfo.Type)
+		, ClassName(InClassInfo.ClassName)
+		, Version(InClassInfo.Version)
+	{
+		checkf(InClassInfo.Type != EMetasoundFrontendClassType::Graph, TEXT("Cannot create key from 'graph' type. Likely meant to use FNodeRegistryKey ctor that is provided FMetasoundFrontendGraphClass"));
+	}
+
+	FNodeRegistryKey::FNodeRegistryKey(const FString& InKeyString)
+	{
+		Parse(InKeyString, *this);
+	}
+
+	const FNodeRegistryKey& FNodeRegistryKey::GetInvalid()
+	{
+		static const FNodeRegistryKey InvalidKey;
+		return InvalidKey;
+	}
+
+	bool FNodeRegistryKey::IsValid() const
+	{
+		return Type != EMetasoundFrontendClassType::Invalid && ClassName.IsValid() && Version.IsValid();
+	}
+
+	void FNodeRegistryKey::Reset()
+	{
+		Type = EMetasoundFrontendClassType::Invalid;
+		ClassName = { };
+		Version = { };
+	}
+
+	FString FNodeRegistryKey::ToString() const
+	{
+		TRACE_CPUPROFILER_EVENT_SCOPE(Metasound::Frontend::FNodeRegistryKey::ToString);
+
+		TStringBuilder<128> KeyStringBuilder;
+		KeyStringBuilder.Append(LexToString(Type));
+		KeyStringBuilder.AppendChar('_');
+		KeyStringBuilder.Append(ClassName.GetFullName().ToString());
+		KeyStringBuilder.AppendChar('_');
+		KeyStringBuilder.Append(FString::FromInt(Version.Major));
+		KeyStringBuilder.AppendChar('.');
+		KeyStringBuilder.Append(FString::FromInt(Version.Minor));
+		return KeyStringBuilder.ToString();
+	}
+
+	bool FNodeRegistryKey::Parse(const FString& InKeyString, FNodeRegistryKey& OutKey)
+	{
+		TRACE_CPUPROFILER_EVENT_SCOPE(Metasound::Frontend::FNodeRegistryKey::Parse);
+
+		TArray<FString> Tokens;
+		InKeyString.ParseIntoArray(Tokens, TEXT("_"));
+		if (Tokens.Num() == 3)
+		{
+			EMetasoundFrontendClassType Type;
+			if (Metasound::Frontend::StringToClassType(Tokens[0], Type))
+			{
+				FMetasoundFrontendClassName ClassName;
+				if (FMetasoundFrontendClassName::Parse(Tokens[1], ClassName))
+				{
+					FMetasoundFrontendVersionNumber Version;
+					FString MajorVersionString;
+					FString MinorVersionString;
+					if (Tokens[2].Split(TEXT("."), &MajorVersionString, &MinorVersionString))
+					{
+						Version.Major = FCString::Atoi(*MajorVersionString);
+						Version.Minor = FCString::Atoi(*MinorVersionString);
+
+						OutKey = FNodeRegistryKey(Type, ClassName, Version.Major, Version.Minor);
+						return true;
+					}
+				}
+			}
+		}
+
+		return false;
+	}
+} // namespace Metasound::Frontend
diff --git a/Engine/Plugins/Runtime/Metasound/Source/MetasoundFrontend/Private/MetasoundFrontendSearchEngineCore.cpp b/Engine/Plugins/Runtime/Metasound/Source/MetasoundFrontend/Private/MetasoundFrontendSearchEngineCore.cpp
--- a/Engine/Plugins/Runtime/Metasound/Source/MetasoundFrontend/Private/MetasoundFrontendSearchEngineCore.cpp
+++ b/Engine/Plugins/Runtime/Metasound/Source/MetasoundFrontend/Private/MetasoundFrontendSearchEngineCore.cpp
@@ -316,7 +316,7 @@ namespace Metasound::Frontend
 		using namespace SearchEngineQuerySteps;
 
 		FMetasoundFrontendVersionNumber HighestVersionNumber{-1, -1};
-		FNodeRegistryKey Key = NodeRegistryKey::GetInvalid();
+		FNodeRegistryKey Key = FNodeRegistryKey::GetInvalid();
 		for (const FFrontendQueryEntry& InEntry : InPartition)
 		{
 			if (InEntry.Value.IsType<FNodeRegistryTransaction>())
diff --git a/Engine/Plugins/Runtime/Metasound/Source/MetasoundFrontend/Private/MetasoundFrontendTransform.cpp b/Engine/Plugins/Runtime/Metasound/Source/MetasoundFrontend/Private/MetasoundFrontendTransform.cpp
--- a/Engine/Plugins/Runtime/Metasound/Source/MetasoundFrontend/Private/MetasoundFrontendTransform.cpp
+++ b/Engine/Plugins/Runtime/Metasound/Source/MetasoundFrontend/Private/MetasoundFrontendTransform.cpp
@@ -616,7 +616,7 @@ namespace Metasound
 
 				if (bIsPreset)
 				{
-					const FNodeRegistryKey RegistryKey = FMetasoundFrontendRegistryContainer::Get()->GetRegistryKey(ClassMetadata);
+					const FNodeRegistryKey RegistryKey(ClassMetadata);
 					PresetReferencedMetaSoundAsset = IMetaSoundAssetManager::GetChecked().TryLoadAssetFromKey(RegistryKey);
 					if (!PresetReferencedMetaSoundAsset)
 					{
@@ -670,7 +670,7 @@ namespace Metasound
 					{
 						FMetasoundFrontendClassMetadata PresetMetadata = InDocument->GetRootGraphClass().Metadata;
 						PresetMetadata.SetType(EMetasoundFrontendClassType::External);
-						const FNodeRegistryKey RegistryKey = FMetasoundFrontendRegistryContainer::Get()->GetRegistryKey(PresetMetadata);
+						const FNodeRegistryKey RegistryKey(PresetMetadata);
 						FMetasoundAssetBase* PresetMetaSoundAsset = IMetaSoundAssetManager::GetChecked().TryLoadAssetFromKey(RegistryKey);
 						if (ensure(PresetMetaSoundAsset))
 						{
diff --git a/Engine/Plugins/Runtime/Metasound/Source/MetasoundFrontend/Private/NodeTemplates/MetasoundFrontendNodeTemplateReroute.cpp b/Engine/Plugins/Runtime/Metasound/Source/MetasoundFrontend/Private/NodeTemplates/MetasoundFrontendNodeTemplateReroute.cpp
--- a/Engine/Plugins/Runtime/Metasound/Source/MetasoundFrontend/Private/NodeTemplates/MetasoundFrontendNodeTemplateReroute.cpp
+++ b/Engine/Plugins/Runtime/Metasound/Source/MetasoundFrontend/Private/NodeTemplates/MetasoundFrontendNodeTemplateReroute.cpp
@@ -89,7 +89,12 @@ namespace Metasound::Frontend
 
 	const FMetasoundFrontendClassName FRerouteNodeTemplate::ClassName { "UE", "Reroute", "" };
 
-	const FMetasoundFrontendVersion FRerouteNodeTemplate::Version { ClassName.GetFullName(), { 1, 0 } };
+	const FMetasoundFrontendVersionNumber FRerouteNodeTemplate::VersionNumber { 1, 0 };
+
+	const FMetasoundFrontendClassName& FRerouteNodeTemplate::GetClassName() const
+	{
+		return ClassName;
+	}
 
 	TUniquePtr<INodeTransform> FRerouteNodeTemplate::GenerateNodeTransform(FMetasoundFrontendDocument& InPreprocessedDocument) const
 	{
@@ -123,7 +128,7 @@ namespace Metasound::Frontend
 #endif // WITH_EDITOR
 
 			Class.Metadata.SetType(EMetasoundFrontendClassType::Template);
-			Class.Metadata.SetVersion(Version.Number);
+			Class.Metadata.SetVersion(VersionNumber);
 
 
 			return Class;
@@ -201,18 +206,13 @@ namespace Metasound::Frontend
 
 	const FNodeRegistryKey& FRerouteNodeTemplate::GetRegistryKey()
 	{
-		static const FNodeRegistryKey RegistryKey = NodeRegistryKey::CreateKey(
-			EMetasoundFrontendClassType::Template,
-			ClassName.ToString(), 
-			Version.Number.Major, 
-			Version.Number.Minor);
-
+		static const FNodeRegistryKey RegistryKey = FNodeRegistryKey(EMetasoundFrontendClassType::Template, ClassName, VersionNumber);
 		return RegistryKey;
 	}
 
-	const FMetasoundFrontendVersion& FRerouteNodeTemplate::GetVersion() const
+	const FMetasoundFrontendVersionNumber& FRerouteNodeTemplate::GetVersionNumber() const
 	{
-		return Version;
+		return VersionNumber;
 	}
 
 #if WITH_EDITOR
diff --git a/Engine/Plugins/Runtime/Metasound/Source/MetasoundFrontend/Public/MetasoundAssetBase.h b/Engine/Plugins/Runtime/Metasound/Source/MetasoundFrontend/Public/MetasoundAssetBase.h
--- a/Engine/Plugins/Runtime/Metasound/Source/MetasoundFrontend/Public/MetasoundAssetBase.h
+++ b/Engine/Plugins/Runtime/Metasound/Source/MetasoundFrontend/Public/MetasoundAssetBase.h
@@ -7,6 +7,7 @@
 #include "MetasoundFrontendController.h"
 #include "MetasoundFrontendDocument.h"
 #include "MetasoundFrontendDocumentAccessPtr.h"
+#include "MetasoundFrontendRegistryKey.h"
 #include "MetasoundGraph.h"
 #include "MetasoundLog.h"
 #include "MetasoundParameterTransmitter.h"
diff --git a/Engine/Plugins/Runtime/Metasound/Source/MetasoundFrontend/Public/MetasoundAssetManager.h b/Engine/Plugins/Runtime/Metasound/Source/MetasoundFrontend/Public/MetasoundAssetManager.h
--- a/Engine/Plugins/Runtime/Metasound/Source/MetasoundFrontend/Public/MetasoundAssetManager.h
+++ b/Engine/Plugins/Runtime/Metasound/Source/MetasoundFrontend/Public/MetasoundAssetManager.h
@@ -2,6 +2,7 @@
 #pragma once
 
 #include "MetasoundVertex.h"
+#include "MetasoundFrontendRegistryKey.h"
 #include "UObject/SoftObjectPath.h"
 #include "UObject/WeakObjectPtrTemplates.h"
 
@@ -14,8 +15,6 @@ namespace Metasound
 {
 	namespace Frontend
 	{
-		using FNodeRegistryKey = FString;
-
 		namespace AssetTags
 		{
 			extern const FString METASOUNDFRONTEND_API ArrayDelim;
diff --git a/Engine/Plugins/Runtime/Metasound/Source/MetasoundFrontend/Public/MetasoundConverterNodeRegistrationMacro.h b/Engine/Plugins/Runtime/Metasound/Source/MetasoundFrontend/Public/MetasoundConverterNodeRegistrationMacro.h
--- a/Engine/Plugins/Runtime/Metasound/Source/MetasoundFrontend/Public/MetasoundConverterNodeRegistrationMacro.h
+++ b/Engine/Plugins/Runtime/Metasound/Source/MetasoundFrontend/Public/MetasoundConverterNodeRegistrationMacro.h
@@ -7,6 +7,7 @@
 
 #include "MetasoundDataReference.h"
 #include "MetasoundFrontendRegistries.h"
+#include "MetasoundFrontendRegistryKey.h"
 #include "MetasoundInputNode.h"
 #include "MetasoundLog.h"
 #include "MetasoundOperatorInterface.h"
@@ -40,7 +41,7 @@ bool RegisterConversionNode(const Metasound::FVertexName& FromPin, const Metasou
 	{
 		FromPin,
 		ToPin,
-		FMetasoundFrontendRegistryContainer::GetRegistryKey(InNodeMetadata)
+		Metasound::Frontend::FNodeRegistryKey(InNodeMetadata)
 	};
 
 	return FMetasoundFrontendRegistryContainer::Get()->RegisterConversionNode(RegistryKey, ConverterNodeInfo);
diff --git a/Engine/Plugins/Runtime/Metasound/Source/MetasoundFrontend/Public/MetasoundDataTypeRegistrationMacro.h b/Engine/Plugins/Runtime/Metasound/Source/MetasoundFrontend/Public/MetasoundDataTypeRegistrationMacro.h
--- a/Engine/Plugins/Runtime/Metasound/Source/MetasoundFrontend/Public/MetasoundDataTypeRegistrationMacro.h
+++ b/Engine/Plugins/Runtime/Metasound/Source/MetasoundFrontend/Public/MetasoundDataTypeRegistrationMacro.h
@@ -143,7 +143,7 @@ namespace Metasound
 			if (TEnableAutoConverterNodeRegistration<TFromDataType, TToDataType>::Value)
 			{
 				const FNodeClassMetadata& Metadata = FConverterNode::GetAutoConverterNodeMetadata();
-				const Metasound::Frontend::FNodeRegistryKey Key = FMetasoundFrontendRegistryContainer::GetRegistryKey(Metadata);
+				const Metasound::Frontend::FNodeRegistryKey Key(Metadata);
 
 				if (!std::is_same<TFromDataType, TToDataType>::value && !FMetasoundFrontendRegistryContainer::Get()->IsNodeRegistered(Key))
 				{
diff --git a/Engine/Plugins/Runtime/Metasound/Source/MetasoundFrontend/Public/MetasoundFrontendDocument.h b/Engine/Plugins/Runtime/Metasound/Source/MetasoundFrontend/Public/MetasoundFrontendDocument.h
--- a/Engine/Plugins/Runtime/Metasound/Source/MetasoundFrontend/Public/MetasoundFrontendDocument.h
+++ b/Engine/Plugins/Runtime/Metasound/Source/MetasoundFrontend/Public/MetasoundFrontendDocument.h
@@ -23,6 +23,8 @@
 struct FMetasoundFrontendClass;
 struct FMetasoundFrontendClassInterface;
 
+enum class EMetasoundFrontendClassType : uint8;
+
 
 namespace Metasound
 {
@@ -154,6 +156,7 @@ enum class EMetasoundFrontendClassType : uint8
 	Invalid UMETA(Hidden)
 };
 
+
 // General purpose version number for Metasound Frontend objects.
 USTRUCT(BlueprintType)
 struct METASOUNDFRONTEND_API FMetasoundFrontendVersionNumber
@@ -1215,9 +1218,8 @@ struct METASOUNDFRONTEND_API FMetasoundFrontendClassName
 	GENERATED_BODY()
 
 	FMetasoundFrontendClassName() = default;
-
+	FMetasoundFrontendClassName(const FName& InNamespace, const FName& InName);
 	FMetasoundFrontendClassName(const FName& InNamespace, const FName& InName, const FName& InVariant);
-
 	FMetasoundFrontendClassName(const Metasound::FNodeClassName& InName);
 
 	// Namespace of class.
@@ -1250,15 +1252,34 @@ struct METASOUNDFRONTEND_API FMetasoundFrontendClassName
 	// Return string version of full name.
 	FString ToString() const;
 
+	// Parses string into class name.  For deserialization and debug use only.
+	static bool Parse(const FString& InClassName, FMetasoundFrontendClassName& OutClassName);
+
 	friend FORCEINLINE uint32 GetTypeHash(const FMetasoundFrontendClassName& ClassName)
 	{
 		const int32 NameHash = HashCombineFast(GetTypeHash(ClassName.Namespace), GetTypeHash(ClassName.Name));
 		return HashCombineFast(NameHash, GetTypeHash(ClassName.Variant));
 	}
 
-	METASOUNDFRONTEND_API friend bool operator==(const FMetasoundFrontendClassName& InLHS, const FMetasoundFrontendClassName& InRHS);
+	friend FORCEINLINE bool operator==(const FMetasoundFrontendClassName& InLHS, const FMetasoundFrontendClassName& InRHS)
+	{
+		return (InLHS.Namespace == InRHS.Namespace) && (InLHS.Name == InRHS.Name) && (InLHS.Variant == InRHS.Variant);
+	}
+
+	friend FORCEINLINE bool operator<(const FMetasoundFrontendClassName& InLHS, const FMetasoundFrontendClassName& InRHS)
+	{
+		if (InLHS.Namespace == InRHS.Namespace)
+		{
+			if (InLHS.Name == InRHS.Name)
+			{
+				return InLHS.Variant.FastLess(InRHS.Variant);
+			}
+
+			return InLHS.Name.FastLess(InRHS.Name);
+		}
 
-	METASOUNDFRONTEND_API friend bool operator!=(const FMetasoundFrontendClassName& InLHS, const FMetasoundFrontendClassName& InRHS);
+		return InLHS.Namespace.FastLess(InRHS.Namespace);
+	}
 };
 
 
@@ -1658,10 +1679,10 @@ public:
 METASOUNDFRONTEND_API const TCHAR* LexToString(EMetasoundFrontendClassType InClassType);
 METASOUNDFRONTEND_API const TCHAR* LexToString(EMetasoundFrontendVertexAccessType InVertexAccess);
 
-namespace Metasound
+namespace Metasound::Frontend
 {
-	namespace Frontend
-	{
+		METASOUNDFRONTEND_API bool StringToClassType(const FString& InString, EMetasoundFrontendClassType& OutClassType);
+
 		/** Signature of function called for each found literal. */
 		using FForEachLiteralFunctionRef = TFunctionRef<void(const FName& InDataTypeName, const FMetasoundFrontendLiteral&)>; 
 
@@ -1676,6 +1697,5 @@ namespace Metasound
 
 		/** Execute the provided function for each literal on a FMetasoundFrontendNode.*/
 		METASOUNDFRONTEND_API void ForEachLiteral(const FMetasoundFrontendNode& InNode, FForEachLiteralFunctionRef OnLiteral);
-	}
-}
+} // namespace Metasound::Frontend
 
diff --git a/Engine/Plugins/Runtime/Metasound/Source/MetasoundFrontend/Public/MetasoundFrontendDocumentCacheInterface.h b/Engine/Plugins/Runtime/Metasound/Source/MetasoundFrontend/Public/MetasoundFrontendDocumentCacheInterface.h
--- a/Engine/Plugins/Runtime/Metasound/Source/MetasoundFrontend/Public/MetasoundFrontendDocumentCacheInterface.h
+++ b/Engine/Plugins/Runtime/Metasound/Source/MetasoundFrontend/Public/MetasoundFrontendDocumentCacheInterface.h
@@ -4,6 +4,7 @@
 #include "Containers/Array.h"
 #include "HAL/Platform.h"
 #include "MetasoundFrontendDocumentModifyDelegates.h"
+#include "MetasoundFrontendRegistryKey.h"
 #include "Templates/SharedPointer.h"
 #include "UObject/Interface.h"
 
@@ -21,8 +22,6 @@ struct FMetasoundFrontendVertex;
 
 namespace Metasound::Frontend
 {
-	using FNodeRegistryKey = FString;
-
 	/** Interface for querying cached document nodes. */
 	class IDocumentGraphNodeCache : public TSharedFromThis<IDocumentGraphNodeCache>
 	{
diff --git a/Engine/Plugins/Runtime/Metasound/Source/MetasoundFrontend/Public/MetasoundFrontendNodeTemplateRegistry.h b/Engine/Plugins/Runtime/Metasound/Source/MetasoundFrontend/Public/MetasoundFrontendNodeTemplateRegistry.h
--- a/Engine/Plugins/Runtime/Metasound/Source/MetasoundFrontend/Public/MetasoundFrontendNodeTemplateRegistry.h
+++ b/Engine/Plugins/Runtime/Metasound/Source/MetasoundFrontend/Public/MetasoundFrontendNodeTemplateRegistry.h
@@ -26,6 +26,9 @@ namespace Metasound::Frontend
 	public:
 		virtual ~INodeTemplate() = default;
 
+		// Returns note template class name.
+		virtual const FMetasoundFrontendClassName& GetClassName() const = 0;
+
 		UE_DEPRECATED(5.4, "Use version that does not provide a preprocessed document")
 		virtual TUniquePtr<INodeTransform> GenerateNodeTransform(FMetasoundFrontendDocument& InDocument) const { return nullptr; }
 
@@ -41,8 +44,11 @@ namespace Metasound::Frontend
 		// Returns access type of the given output within the provided builder's document
 		virtual EMetasoundFrontendVertexAccessType GetNodeOutputAccessType(const FMetaSoundFrontendDocumentBuilder& InBuilder, const FGuid& InNodeID, const FGuid& InVertexID) const = 0;
 
-		// Returns the version of the given node class template.
-		virtual const FMetasoundFrontendVersion& GetVersion() const = 0;
+		UE_DEPRECATED(5.4, "Use version number or classname instead")
+		virtual const FMetasoundFrontendVersion& GetVersion() const { const static FMetasoundFrontendVersion NullVersion; return NullVersion; }
+
+		// Returns note template class version.
+		virtual const FMetasoundFrontendVersionNumber& GetVersionNumber() const = 0;
 
 #if WITH_EDITOR
 		// Returns whether or not the given node template has the necessary
@@ -85,5 +91,9 @@ namespace Metasound::Frontend
 	// Register & Unregister are not publicly accessible implementation as the API
 	// is in beta and, currently, only to be used by internal implementation (ex. reroute nodes).
 	void RegisterNodeTemplate(TUniquePtr<INodeTemplate>&& InTemplate);
+
+	UE_DEPRECATED(5.4, "Use version that provides class name and version instead")
 	void UnregisterNodeTemplate(const FMetasoundFrontendVersion& InNodeTemplateVersion);
+
+	void UnregisterNodeTemplate(const FMetasoundFrontendClassName& InClassName, const FMetasoundFrontendVersionNumber& InTemplateVersion);
 } // namespace Metasound::Frontend
diff --git a/Engine/Plugins/Runtime/Metasound/Source/MetasoundFrontend/Public/MetasoundFrontendRegistries.h b/Engine/Plugins/Runtime/Metasound/Source/MetasoundFrontend/Public/MetasoundFrontendRegistries.h
--- a/Engine/Plugins/Runtime/Metasound/Source/MetasoundFrontend/Public/MetasoundFrontendRegistries.h
+++ b/Engine/Plugins/Runtime/Metasound/Source/MetasoundFrontend/Public/MetasoundFrontendRegistries.h
@@ -9,78 +9,27 @@
 #include "MetasoundNodeConstructorParams.h"
 #include "MetasoundNodeInterface.h"
 #include "MetasoundOperatorInterface.h"
+#include "MetasoundFrontendRegistryKey.h"
 #include "MetasoundRouter.h"
 #include "MetasoundVertex.h"
 #include "Templates/Function.h"
 
-class IMetaSoundDocumentInterface;
 
+// Forward Declarations
+class IMetaSoundDocumentInterface;
 template <typename InInterfaceType>
 class TScriptInterface;
-
 class UObject;
 
+
 namespace Metasound
 {
+	// Forward Declarations
 	using FIterateMetasoundFrontendClassFunction = TFunctionRef<void(const FMetasoundFrontendClass&)>;
 	class FGraph;
 
 	namespace Frontend
 	{
-		using FNodeRegistryKey = FString;
-
-		/** FNodeClassInfo contains a minimal set of information needed to find
-		 * and query node classes. 
-		 */
-		struct METASOUNDFRONTEND_API FNodeClassInfo
-		{
-			// ClassName of the given class
-			FMetasoundFrontendClassName ClassName;
-
-			// The type of this node class
-			EMetasoundFrontendClassType Type = EMetasoundFrontendClassType::Invalid;
-
-			// The ID used for the Asset Classes. If zero, class is natively defined.
-			FGuid AssetClassID;
-
-			// Path to asset containing graph if external type and references asset class.
-			FSoftObjectPath AssetPath;
-
-			// Version of the registered class
-			FMetasoundFrontendVersionNumber Version;
-
-#if WITH_EDITORONLY_DATA
-			// Types of class inputs
-			TSet<FName> InputTypes;
-
-			// Types of class outputs
-			TSet<FName> OutputTypes;
-
-			// Whether or not class is preset
-			bool bIsPreset = false;
-#endif // WITH_EDITORONLY_DATA
-
-			FNodeClassInfo() = default;
-
-			// Constructor used to generate NodeClassInfo from a native class' Metadata.
-			FNodeClassInfo(const FMetasoundFrontendClassMetadata& InMetadata);
-
-			// Constructor used to generate NodeClassInfo from an asset
-			FNodeClassInfo(const FMetasoundFrontendGraphClass& InClass, const FSoftObjectPath& InAssetPath);
-
-			// Loads the asset from the provided path, ensuring that the class is of type graph.
-			UObject* LoadAsset() const
-			{
-				if (ensure(Type == EMetasoundFrontendClassType::External))
-				{
-					FSoftObjectPath SoftObjectPath(AssetPath);
-					return SoftObjectPath.TryLoad();
-				}
-
-				return nullptr;
-			}
-		};
-
 		/** INodeRegistryEntry declares the interface for a node registry entry.
 		 * Each node class in the registry must satisfy this interface. 
 		 */
@@ -237,15 +186,13 @@ namespace Metasound
 
 		namespace NodeRegistryKey
 		{
-			// Returns the invalid NodeRegistryKey.
+			UE_DEPRECATED(5.4, "Use FNodeRegistryKey::GetInvalid() instead")
 			METASOUNDFRONTEND_API const FNodeRegistryKey& GetInvalid();
 
-			// Returns true if the registry key is a valid key.
-			//
-			// This does *not* connote that the registry key exists in the registry.
+			UE_DEPRECATED(5.4, "Use FNodeRegistryKey::IsValid() instead")
 			METASOUNDFRONTEND_API bool IsValid(const FNodeRegistryKey& InKey);
 
-			// Returns true if both keys represent the same entry in the node registry.
+			UE_DEPRECATED(5.4, "Use FNodeRegistryKey equality operator instead")
 			METASOUNDFRONTEND_API bool IsEqual(const FNodeRegistryKey& InLHS, const FNodeRegistryKey& InRHS);
 
 			// Returns true if the class metadata represent the same entry in the node registry.
@@ -254,10 +201,19 @@ namespace Metasound
 			// Returns true if the class info and class metadata represent the same entry in the node registry.
 			METASOUNDFRONTEND_API bool IsEqual(const FNodeClassInfo& InLHS, const FMetasoundFrontendClassMetadata& InRHS);
 
+			UE_DEPRECATED(5.4, "Use FNodeRegistryKey constructors instead")
 			METASOUNDFRONTEND_API FNodeRegistryKey CreateKey(EMetasoundFrontendClassType InType, const FString& InFullClassName, int32 InMajorVersion, int32 InMinorVersion);
+
+			UE_DEPRECATED(5.4, "Use FNodeRegistryKey constructors instead")
 			METASOUNDFRONTEND_API FNodeRegistryKey CreateKey(const FNodeClassMetadata& InNodeMetadata);
+
+			UE_DEPRECATED(5.4, "Use FNodeRegistryKey constructor instead")
 			METASOUNDFRONTEND_API FNodeRegistryKey CreateKey(const FMetasoundFrontendClassMetadata& InNodeMetadata);
+
+			UE_DEPRECATED(5.4, "Use FNodeRegistryKey constructor instead")
 			METASOUNDFRONTEND_API FNodeRegistryKey CreateKey(const FMetasoundFrontendGraphClass& InNodeMetadata);
+
+			UE_DEPRECATED(5.4, "Use FNodeRegistryKey constructor instead")
 			METASOUNDFRONTEND_API FNodeRegistryKey CreateKey(const FNodeClassInfo& ClassInfo);
 		}
 	} // namespace Frontend
@@ -295,12 +251,20 @@ public:
 	static FMetasoundFrontendRegistryContainer* Get();
 	static void ShutdownMetasoundFrontend();
 
+	UE_DEPRECATED(5.4, "Use FNodeRegistryKey constructor that takes FNodeClassMetadata instead")
 	static FNodeRegistryKey GetRegistryKey(const FNodeClassMetadata& InNodeMetadata);
+
+	UE_DEPRECATED(5.4, "Use FNodeRegistryKey constructor that takes FMetasoundFrontendClassMetadata instead")
 	static FNodeRegistryKey GetRegistryKey(const FMetasoundFrontendClassMetadata& InNodeMetadata);
+
+	UE_DEPRECATED(5.4, "Use FNodeRegistryKey constructor that takes FNodeClassInfo instead")
 	static FNodeRegistryKey GetRegistryKey(const FNodeClassInfo& ClassInfo);
 
 	static bool GetFrontendClassFromRegistered(const FNodeRegistryKey& InKey, FMetasoundFrontendClass& OutClass);
+
+	UE_DEPRECATED(5.4, "Use GetFrontendClassFromRegistered instead")
 	static bool GetNodeClassInfoFromRegistered(const FNodeRegistryKey& InKey, FNodeClassInfo& OutInfo);
+
 	static bool GetInputNodeRegistryKeyForDataType(const FName& InDataTypeName, const EMetasoundFrontendVertexAccessType InAccessType, FNodeRegistryKey& OutKey);
 	static bool GetVariableNodeRegistryKeyForDataType(const FName& InDataTypeName, FNodeRegistryKey& OutKey);
 	static bool GetOutputNodeRegistryKeyForDataType(const FName& InDataTypeName, const EMetasoundFrontendVertexAccessType InAccessType, FNodeRegistryKey& OutKey);
@@ -365,7 +329,10 @@ public:
 	// Query for MetaSound Frontend document objects.
 	virtual bool FindFrontendClassFromRegistered(const Metasound::Frontend::FNodeRegistryKey& InKey, FMetasoundFrontendClass& OutClass) = 0;
 	virtual const TSet<FMetasoundFrontendVersion>* FindImplementedInterfacesFromRegistered(const Metasound::Frontend::FNodeRegistryKey& InKey) const = 0;
+
+	UE_DEPRECATED(5.4, "Use FindFrontendClassFromRegistered instead")
 	virtual bool FindNodeClassInfoFromRegistered(const Metasound::Frontend::FNodeRegistryKey& InKey, FNodeClassInfo& OutInfo) = 0;
+
 	virtual bool FindInputNodeRegistryKeyForDataType(const FName& InDataTypeName, const EMetasoundFrontendVertexAccessType InAccessType, FNodeRegistryKey& OutKey) = 0;
 	virtual bool FindVariableNodeRegistryKeyForDataType(const FName& InDataTypeName, FNodeRegistryKey& OutKey) = 0;
 	virtual bool FindOutputNodeRegistryKeyForDataType(const FName& InDataTypeName, const EMetasoundFrontendVertexAccessType InAccessType, FNodeRegistryKey& OutKey) = 0;
diff --git a/Engine/Plugins/Runtime/Metasound/Source/MetasoundFrontend/Public/MetasoundFrontendRegistryKey.h b/Engine/Plugins/Runtime/Metasound/Source/MetasoundFrontend/Public/MetasoundFrontendRegistryKey.h
new file mode 100644
--- /dev/null
+++ b/Engine/Plugins/Runtime/Metasound/Source/MetasoundFrontend/Public/MetasoundFrontendRegistryKey.h
@@ -0,0 +1,137 @@
+// Copyright Epic Games, Inc. All Rights Reserved.
+#pragma once
+
+#include "MetasoundFrontendDocument.h"
+#include "MetasoundNodeInterface.h"
+#include "Misc/CoreDefines.h"
+#include "UObject/NoExportTypes.h"
+
+
+namespace Metasound::Frontend
+{
+	/** FNodeClassInfo contains a minimal set of information needed to find
+	 * and query asset node classes without loading the asset.
+	 */
+	struct METASOUNDFRONTEND_API FNodeClassInfo
+	{
+		// ClassName of the given class
+		FMetasoundFrontendClassName ClassName;
+
+		// The type of this node class
+		EMetasoundFrontendClassType Type = EMetasoundFrontendClassType::Invalid;
+
+		// The ID used for the Asset Classes. If zero, class is natively defined.
+		FGuid AssetClassID;
+
+		// Path to asset containing graph if external type and references asset class.
+		FSoftObjectPath AssetPath;
+
+		// Version of the registered class
+		FMetasoundFrontendVersionNumber Version;
+
+#if WITH_EDITORONLY_DATA
+		// Types of class inputs
+		TSet<FName> InputTypes;
+
+		// Types of class outputs
+		TSet<FName> OutputTypes;
+
+		// Whether or not class is preset
+		bool bIsPreset = false;
+#endif // WITH_EDITORONLY_DATA
+
+		FNodeClassInfo() = default;
+
+		// Constructor used to generate NodeClassInfo from a native class' Metadata.
+		FNodeClassInfo(const FMetasoundFrontendClassMetadata& InMetadata);
+
+		// Constructor used to generate NodeClassInfo from an asset
+		FNodeClassInfo(const FMetasoundFrontendGraphClass& InClass, const FSoftObjectPath& InAssetPath);
+
+		// Loads the asset from the provided path, ensuring that the class is of type graph.
+		UObject* LoadAsset() const;
+	};
+
+		/** FNodeClassInfo contains a minimal set of information needed to find
+	 * and query node classes.
+	 */
+
+	struct METASOUNDFRONTEND_API FNodeRegistryKey
+	{
+		EMetasoundFrontendClassType Type = EMetasoundFrontendClassType::Invalid;
+		FMetasoundFrontendClassName ClassName;
+		FMetasoundFrontendVersionNumber Version;
+
+		FNodeRegistryKey() = default;
+		FNodeRegistryKey(EMetasoundFrontendClassType InType, const FMetasoundFrontendClassName& InClassName, int32 InMajorVersion, int32 InMinorVersion);
+		FNodeRegistryKey(EMetasoundFrontendClassType InType, const FMetasoundFrontendClassName& InClassName, const FMetasoundFrontendVersionNumber& InVersion);
+		FNodeRegistryKey(const FNodeClassMetadata& InNodeMetadata);
+		FNodeRegistryKey(const FMetasoundFrontendClassMetadata& InNodeMetadata);
+		FNodeRegistryKey(const FMetasoundFrontendGraphClass& InGraphClass);
+		FNodeRegistryKey(const FNodeClassInfo& InClassInfo);
+
+		UE_DEPRECATED(5.4, "Implicit String ctor is no longer supported.")
+		FNodeRegistryKey(const FString& InKeyString);
+
+		/*FORCEINLINE*/ friend bool operator==(const FNodeRegistryKey& InLHS, const FNodeRegistryKey& InRHS)
+		{
+			return (InLHS.Type == InRHS.Type) && (InLHS.ClassName == InRHS.ClassName) && (InLHS.Version == InRHS.Version);
+		}
+
+		/*FORCEINLINE*/ friend bool operator<(const FNodeRegistryKey& InLHS, const FNodeRegistryKey& InRHS)
+		{
+			if (static_cast<uint8>(InLHS.Type) == static_cast<uint8>(InRHS.Type))
+			{
+				if (InLHS.ClassName == InRHS.ClassName)
+				{
+					return InLHS.Version < InRHS.Version;
+				}
+
+				return InLHS.ClassName < InRHS.ClassName;
+			}
+
+			return static_cast<uint8>(InLHS.Type) < static_cast<uint8>(InRHS.Type);
+		}
+
+		friend FORCEINLINE uint32 GetTypeHash(const FNodeRegistryKey& InKey)
+		{
+			const int32 Hash = HashCombineFast(static_cast<uint32>(InKey.Type), GetTypeHash(InKey.ClassName));
+			return HashCombineFast(Hash, GetTypeHash(InKey.Version));
+		}
+
+		UE_DEPRECATED(5.4, "Implicit String IsEmpty() check is no longer valid. Use IsValid() instead.")
+		bool IsEmpty() const
+		{
+			return !IsValid();
+		}
+
+		// Returns invalid (default constructed) key
+		static const FNodeRegistryKey& GetInvalid();
+
+		// Returns whether or not instance is valid or not
+		bool IsValid() const;
+
+		UE_DEPRECATED(5.4, "Implicit cast to FString from FNodeRegistryKey is deprecated, use ToString().")
+		FORCEINLINE operator FString() const
+		{
+			return ToString();
+		}
+
+		UE_DEPRECATED(5.4, "Implicit cast to TCHAR* is deprecated, use ToString().")
+		FORCEINLINE const TCHAR* operator*() const
+		{
+			static const FString DeprecatedString;
+			return *DeprecatedString;
+		}
+
+		// Resets the key back to an invalid default state
+		void Reset();
+
+		// Returns string representation of key
+		FString ToString() const;
+
+		// Parses string representation of key into registry key.  For debug and deserialization use only.
+		// Returns true if parsed successfully.
+		static bool Parse(const FString& InKeyString, FNodeRegistryKey& OutKey);
+	};
+} // namespace Metasound::Frontend
diff --git a/Engine/Plugins/Runtime/Metasound/Source/MetasoundFrontend/Public/MetasoundNodeRegistrationMacro.h b/Engine/Plugins/Runtime/Metasound/Source/MetasoundFrontend/Public/MetasoundNodeRegistrationMacro.h
--- a/Engine/Plugins/Runtime/Metasound/Source/MetasoundFrontend/Public/MetasoundNodeRegistrationMacro.h
+++ b/Engine/Plugins/Runtime/Metasound/Source/MetasoundFrontend/Public/MetasoundNodeRegistrationMacro.h
@@ -119,7 +119,7 @@ namespace Metasound
 
 
 		Frontend::FNodeRegistryKey Key = FMetasoundFrontendRegistryContainer::Get()->RegisterNode(MakeUnique<FNodeRegistryEntry>(InMetadata));
-		const bool bSuccessfullyRegisteredNode = Frontend::NodeRegistryKey::IsValid(Key);
+		const bool bSuccessfullyRegisteredNode = Key.IsValid();
 		ensureAlwaysMsgf(bSuccessfullyRegisteredNode, TEXT("Registering node class failed. Please check the logs."));
 
 		return bSuccessfullyRegisteredNode;
diff --git a/Engine/Plugins/Runtime/Metasound/Source/MetasoundFrontend/Public/NodeTemplates/MetasoundFrontendNodeTemplateReroute.h b/Engine/Plugins/Runtime/Metasound/Source/MetasoundFrontend/Public/NodeTemplates/MetasoundFrontendNodeTemplateReroute.h
--- a/Engine/Plugins/Runtime/Metasound/Source/MetasoundFrontend/Public/NodeTemplates/MetasoundFrontendNodeTemplateReroute.h
+++ b/Engine/Plugins/Runtime/Metasound/Source/MetasoundFrontend/Public/NodeTemplates/MetasoundFrontendNodeTemplateReroute.h
@@ -16,7 +16,8 @@ namespace Metasound::Frontend
 	{
 	public:
 		static const FMetasoundFrontendClassName ClassName;
-		static const FMetasoundFrontendVersion Version;
+
+		static const FMetasoundFrontendVersionNumber VersionNumber;
 
 		static const FNodeRegistryKey& GetRegistryKey();
 
@@ -24,6 +25,8 @@ namespace Metasound::Frontend
 
 		virtual ~FRerouteNodeTemplate() = default;
 
+		virtual const FMetasoundFrontendClassName& GetClassName() const override;
+
 		UE_DEPRECATED(5.4, "Use version that does not require mutating a provided PreprocessedDocument")
 		virtual TUniquePtr<INodeTransform> GenerateNodeTransform(FMetasoundFrontendDocument& InPreprocessedDocument) const override;
 
@@ -31,7 +34,7 @@ namespace Metasound::Frontend
 		virtual const FMetasoundFrontendClass& GetFrontendClass() const override;
 		virtual EMetasoundFrontendVertexAccessType GetNodeInputAccessType(const FMetaSoundFrontendDocumentBuilder& InBuilder, const FGuid& InNodeID, const FGuid& InVertexID) const override;
 		virtual EMetasoundFrontendVertexAccessType GetNodeOutputAccessType(const FMetaSoundFrontendDocumentBuilder& InBuilder, const FGuid& InNodeID, const FGuid& InVertexID) const override;
-		virtual const FMetasoundFrontendVersion& GetVersion() const override;
+		virtual const FMetasoundFrontendVersionNumber& GetVersionNumber() const override;
 		virtual bool IsInputAccessTypeDynamic() const override;
 		virtual bool IsOutputAccessTypeDynamic() const override;
 		virtual bool IsValidNodeInterface(const FMetasoundFrontendNodeInterface& InNodeInterface) const override;
diff --git a/Engine/Plugins/Runtime/Metasound/Source/MetasoundGraphCore/Public/MetasoundNodeInterface.h b/Engine/Plugins/Runtime/Metasound/Source/MetasoundGraphCore/Public/MetasoundNodeInterface.h
--- a/Engine/Plugins/Runtime/Metasound/Source/MetasoundGraphCore/Public/MetasoundNodeInterface.h
+++ b/Engine/Plugins/Runtime/Metasound/Source/MetasoundGraphCore/Public/MetasoundNodeInterface.h
@@ -83,7 +83,7 @@ namespace Metasound
 	};
 
 	/** Provides metadata for a given node. */
-	struct FNodeClassMetadata
+	struct METASOUNDGRAPHCORE_API FNodeClassMetadata
 	{
 		/** Name of class. Used for registration and lookup. */
 		FNodeClassName ClassName;
