Removing test code not meant for submission in prior CL.
Contributor: Guadian swarm angel
#rnx

diff --git a/Engine/Plugins/Runtime/Metasound/Source/MetasoundFrontend/Private/MetasoundFrontendDocumentCache.cpp b/Engine/Plugins/Runtime/Metasound/Source/MetasoundFrontend/Private/MetasoundFrontendDocumentCache.cpp
--- a/Engine/Plugins/Runtime/Metasound/Source/MetasoundFrontend/Private/MetasoundFrontendDocumentCache.cpp
+++ b/Engine/Plugins/Runtime/Metasound/Source/MetasoundFrontend/Private/MetasoundFrontendDocumentCache.cpp
@@ -229,9 +229,6 @@ namespace Metasound::Frontend
 
 		FDocumentDependencyCache& Cache = GetDependencyCache();
 
-		const bool contains = Cache.KeyToIndex.Contains(SwapKey);
-		check(contains);
-
 		DocumentCachePrivate::RemoveSwapMapIndexChecked(SwapKey, LastKey, Cache.KeyToIndex);
 		DocumentCachePrivate::RemoveSwapMapIndexChecked(SwapDependency.ID, LastDependency.ID, Cache.IDToIndex);
 	}
