[Metasounds] Operator precaching mechanism.

#swarm https://p4-swarm.epicgames.net/reviews/27703580

diff --git a/Engine/Plugins/Runtime/Metasound/Source/MetasoundEngine/Private/MetasoundOperatorCacheSubsystem.cpp b/Engine/Plugins/Runtime/Metasound/Source/MetasoundEngine/Private/MetasoundOperatorCacheSubsystem.cpp
new file mode 100644
--- /dev/null
+++ b/Engine/Plugins/Runtime/Metasound/Source/MetasoundEngine/Private/MetasoundOperatorCacheSubsystem.cpp
@@ -0,0 +1,131 @@
+// Copyright Epic Games, Inc. All Rights Reserved.
+#include "MetasoundOperatorCacheSubsystem.h"
+
+#include "AudioMixerDevice.h"
+#include "MetasoundGenerator.h"
+#include "MetasoundGeneratorModule.h"
+#include "MetasoundOperatorCache.h"
+#include "Modules/ModuleManager.h"
+#include "Misc/Optional.h"
+
+static TOptional<Metasound::FMetasoundGeneratorInitParams> CreateInitParams(UMetaSoundSource* InMetaSound, const FSoundGeneratorInitParams& InParams)
+{
+	using namespace Metasound;
+	using namespace Metasound::Frontend;
+	using namespace Metasound::Engine;
+	using namespace Metasound::SourcePrivate;
+
+	if (!ensure(InMetaSound))
+	{
+		return {}; // InMetaSound was null
+	}
+
+	if (ensure(!InMetaSound->IsDynamic()))
+	{
+		return {}; // we cannot precache dynamic metasounds
+	}
+
+	FOperatorSettings InSettings = InMetaSound->GetOperatorSettings(static_cast<FSampleRate>(InParams.SampleRate));
+	FMetasoundEnvironment Environment = InMetaSound->CreateEnvironment(InParams);
+
+	FOperatorBuilderSettings BuilderSettings = FOperatorBuilderSettings::GetDefaultSettings();
+	// Graph analyzer currently only enabled for preview sounds (but can theoretically be supported for all sounds)
+	BuilderSettings.bPopulateInternalDataReferences = InParams.bIsPreviewSound;
+
+	return FMetasoundGeneratorInitParams
+	{
+		InSettings,
+		MoveTemp(BuilderSettings),
+		{}, // Graph, retrieved from the FrontEnd Registry in FOperatorPool::BuildAndAddOperator()
+		Environment,
+		InMetaSound->GetName(),
+		InMetaSound->GetOutputAudioChannelOrder(),
+		{}, // DefaultParameters
+		true, // bBuildSynchronous
+		{} // DataChannel
+	};
+}
+
+bool UMetaSoundCacheSubsystem::ShouldCreateSubsystem(UObject* Outer) const
+{
+	return !IsRunningDedicatedServer();
+}
+
+void UMetaSoundCacheSubsystem::Initialize(FSubsystemCollectionBase& Collection)
+{
+	using namespace Audio;
+	using namespace Metasound;
+
+	BuildParams = MakeUnique<FSoundGeneratorInitParams>();
+	const FMixerDevice* MixerDevice = GetMixerDevice();
+
+	if (BuildParams && ensure(MixerDevice))
+	{
+		BuildParams->AudioDeviceID = GetAudioDeviceHandle().GetDeviceID();
+		BuildParams->SampleRate = MixerDevice->GetSampleRate();
+		BuildParams->AudioMixerNumOutputFrames = MixerDevice->GetNumOutputFrames();
+		BuildParams->NumChannels = MixerDevice->GetNumDeviceChannels();
+		BuildParams->NumFramesPerCallback = 0;
+		BuildParams->InstanceID = 0;
+	}
+}
+
+void UMetaSoundCacheSubsystem::Update()
+{
+#if METASOUND_OPERATORCACHEPROFILER_ENABLED
+	using namespace Metasound;
+	if (TSharedPtr<FOperatorPool> OperatorPool = FModuleManager::GetModuleChecked<IMetasoundGeneratorModule>("MetasoundGenerator").GetOperatorPool())
+	{
+		OperatorPool->UpdateHitRateTracker();
+	}
+#endif // #if METASOUND_OPERATORCACHEPROFILER_ENABLED
+}
+
+void UMetaSoundCacheSubsystem::PrecacheMetaSound(UMetaSoundSource* InMetaSound, int32 InNumInstances)
+{
+	using namespace Audio;
+	using namespace Metasound;
+
+	TSharedPtr<FOperatorPool> OperatorPool;
+	IMetasoundGeneratorModule* Module = FModuleManager::GetModulePtr<IMetasoundGeneratorModule>("MetasoundGenerator");
+	if (ensure(Module))
+	{
+		OperatorPool = Module->GetOperatorPool();
+	}
+
+	const FMixerDevice* MixerDevice = GetMixerDevice();
+
+	if (!ensure(MixerDevice && InMetaSound && BuildParams && OperatorPool))
+	{
+		return;
+	}
+
+	if (!ensure(InNumInstances > 0))
+	{
+		return;
+	}
+
+	InMetaSound->InitResources();
+	BuildParams->GraphName = InMetaSound->GetOwningAssetName();
+
+	if (InMetaSound->IsDynamic())
+	{
+		return;
+	}
+
+	TOptional<FMetasoundGeneratorInitParams> InitParams = CreateInitParams(InMetaSound, *BuildParams);
+	if (!InitParams.IsSet())
+	{
+		return;
+	}
+	
+
+	TUniquePtr<FOperatorBuildData> Data = MakeUnique<FOperatorBuildData>(
+		  MoveTemp(InitParams.GetValue())
+		, InMetaSound->GetRegistryKey()
+		, FSoftObjectPath(InMetaSound->GetOwningAsset())
+		, InNumInstances
+	);
+
+	OperatorPool->BuildAndAddOperator(MoveTemp(Data));
+}
\ No newline at end of file
diff --git a/Engine/Plugins/Runtime/Metasound/Source/MetasoundEngine/Private/MetasoundSource.cpp b/Engine/Plugins/Runtime/Metasound/Source/MetasoundEngine/Private/MetasoundSource.cpp
--- a/Engine/Plugins/Runtime/Metasound/Source/MetasoundEngine/Private/MetasoundSource.cpp
+++ b/Engine/Plugins/Runtime/Metasound/Source/MetasoundEngine/Private/MetasoundSource.cpp
@@ -29,6 +29,7 @@
 #include "MetasoundGenerator.h"
 #include "MetasoundLog.h"
 #include "MetasoundOperatorBuilderSettings.h"
+#include "MetasoundOperatorCacheSubsystem.h"
 #include "MetasoundOperatorSettings.h"
 #include "MetasoundParameterTransmitter.h"
 #include "MetasoundPrimitives.h"
@@ -509,7 +510,7 @@ void UMetaSoundSource::InitParameters(TArray<FAudioParameter>& ParametersToInit,
 	METASOUND_LLM_SCOPE;
 	METASOUND_TRACE_CPUPROFILER_EVENT_SCOPE(UMetaSoundSource::InitParameters);
 
-	if (bIsBuilderActive)
+	if (bIsBuilderActive)	
 	{
 		InitParametersInternal(CreateRuntimeInputMap(), ParametersToInit, InFeatureName);
 	}
@@ -1384,6 +1385,12 @@ TWeakPtr<Metasound::FMetasoundGenerator> UMetaSoundSource::GetGeneratorForAudioC
 	return *Result;
 }
 
+
+bool UMetaSoundSource::IsDynamic() const
+{
+	return DynamicTransactor.IsValid();
+}
+
 Metasound::SourcePrivate::FParameterRouter& UMetaSoundSource::GetParameterRouter()
 {
 	using namespace Metasound::SourcePrivate;
diff --git a/Engine/Plugins/Runtime/Metasound/Source/MetasoundEngine/Public/MetasoundOperatorCacheSubsystem.h b/Engine/Plugins/Runtime/Metasound/Source/MetasoundEngine/Public/MetasoundOperatorCacheSubsystem.h
new file mode 100644
--- /dev/null
+++ b/Engine/Plugins/Runtime/Metasound/Source/MetasoundEngine/Public/MetasoundOperatorCacheSubsystem.h
@@ -0,0 +1,41 @@
+// Copyright Epic Games, Inc. All Rights Reserved.
+
+#pragma once
+
+#include "Containers/Array.h"
+#include "Containers/Map.h"
+#include "MetasoundOperatorInterface.h"
+#include "Misc/Guid.h"
+#include "HAL/CriticalSection.h"
+#include "Templates/UniquePtr.h"
+#include "Subsystems/AudioEngineSubsystem.h"
+
+#include "MetasoundSource.h"
+
+#include "MetasoundOperatorCacheSubsystem.generated.h"
+
+
+/**
+* 	UMetaSoundCacheSubsystem
+*/
+UCLASS()
+class METASOUNDENGINE_API UMetaSoundCacheSubsystem : public UAudioEngineSubsystem
+{
+	GENERATED_BODY()
+
+public: 
+	//~ Begin USubsystem interface
+	virtual bool ShouldCreateSubsystem(UObject* Outer) const override;
+	virtual void Initialize(FSubsystemCollectionBase& Collection) override;
+	virtual void Update() override; // todo: compile out
+	//~ End USubsystem interface
+	 
+	/* Builds the requested number of MetaSound operators (asynchronously) and puts them in the pool for playback.
+	(If these operators are not yet available when the MetaSound attepmts to play, one will be created Independent of this request.) */
+	UFUNCTION(BlueprintCallable, Category = "Audio")
+	void PrecacheMetaSound(UMetaSoundSource* InMetaSound, int32 InNumInstances = 1);
+
+private:
+	TUniquePtr<FSoundGeneratorInitParams> BuildParams;
+}; //UMetaSoundCacheSubsystem
+
diff --git a/Engine/Plugins/Runtime/Metasound/Source/MetasoundEngine/Public/MetasoundSource.h b/Engine/Plugins/Runtime/Metasound/Source/MetasoundEngine/Public/MetasoundSource.h
--- a/Engine/Plugins/Runtime/Metasound/Source/MetasoundEngine/Public/MetasoundSource.h
+++ b/Engine/Plugins/Runtime/Metasound/Source/MetasoundEngine/Public/MetasoundSource.h
@@ -45,6 +45,11 @@ namespace Audio
 
 class UMetaSoundSettings;
 
+namespace Metasound
+{
+	struct FMetasoundGeneratorInitParams;
+} // namespace Metasound
+
 DECLARE_TS_MULTICAST_DELEGATE_TwoParams(FOnGeneratorInstanceCreated, uint64, TSharedPtr<Metasound::FMetasoundGenerator>);
 DECLARE_TS_MULTICAST_DELEGATE_TwoParams(FOnGeneratorInstanceDestroyed, uint64, TSharedPtr<Metasound::FMetasoundGenerator>);
 
@@ -247,6 +252,8 @@ public:
 	virtual bool EnableSubmixSendsOnPreview() const override { return true; }
 
 	TWeakPtr<Metasound::FMetasoundGenerator> GetGeneratorForAudioComponent(uint64 ComponentId) const;
+	bool IsDynamic() const;
+
 	FOnGeneratorInstanceCreated OnGeneratorInstanceCreated;
 	FOnGeneratorInstanceDestroyed OnGeneratorInstanceDestroyed;
 	Metasound::FOperatorSettings GetOperatorSettings(Metasound::FSampleRate InDeviceSampleRate) const;
@@ -280,11 +287,15 @@ private:
 	bool IsParameterValidInternal(const FAudioParameter& InParameter, const FName& InTypeName, Metasound::Frontend::IDataTypeRegistry& InDataTypeRegistry) const;
 
 	static Metasound::SourcePrivate::FParameterRouter& GetParameterRouter();
+
+public:
+	Metasound::FOperatorSettings GetOperatorSettings(Metasound::FSampleRate InSampleRate) const;
+	Metasound::FMetasoundEnvironment CreateEnvironment(const FSoundGeneratorInitParams& InParams) const;
+	const TArray<Metasound::FVertexName>& GetOutputAudioChannelOrder() const;
+private:
 	
 	Metasound::FMetasoundEnvironment CreateEnvironment() const;
-	Metasound::FMetasoundEnvironment CreateEnvironment(const FSoundGeneratorInitParams& InParams) const;
 	Metasound::FMetasoundEnvironment CreateEnvironment(const Audio::FParameterTransmitterInitParams& InParams) const;
-	const TArray<Metasound::FVertexName>& GetOutputAudioChannelOrder() const;
 
 	mutable FCriticalSection GeneratorMapCriticalSection;
 	TSortedMap<uint64, TWeakPtr<Metasound::FMetasoundGenerator>> Generators;
diff --git a/Engine/Plugins/Runtime/Metasound/Source/MetasoundGenerator/Private/MetasoundGenerator.cpp b/Engine/Plugins/Runtime/Metasound/Source/MetasoundGenerator/Private/MetasoundGenerator.cpp
--- a/Engine/Plugins/Runtime/Metasound/Source/MetasoundGenerator/Private/MetasoundGenerator.cpp
+++ b/Engine/Plugins/Runtime/Metasound/Source/MetasoundGenerator/Private/MetasoundGenerator.cpp
@@ -9,6 +9,7 @@
 #include "MetasoundOperatorBuilder.h"
 #include "MetasoundOperatorInterface.h"
 #include "MetasoundOutputNode.h"
+#include "MetasoundOperatorCache.h"
 #include "MetasoundParameterTransmitter.h"
 #include "MetasoundRouter.h"
 #include "MetasoundTrace.h"
@@ -30,8 +31,8 @@ namespace Metasound
 	namespace ConsoleVariables
 	{
 		static bool bEnableAsyncMetaSoundGeneratorBuilder = true;
-		static bool bEnableExperimentalOneShotOperatorCache = false;
-		static bool bEnableExperimentalOperatorCache = false;
+		static bool bEnableExperimentalOneShotOperatorPool = false;
+		static bool bEnableExperimentalOperatorPool = false;
 #if ENABLE_METASOUNDGENERATOR_INVALID_SAMPLE_VALUE_LOGGING
 		static bool bEnableMetaSoundGeneratorNonFiniteLogging = false;
 		static bool bEnableMetaSoundGeneratorInvalidSampleValueLogging = false;
@@ -195,16 +196,16 @@ FAutoConsoleVariableRef CVarMetaSoundEnableAsyncGeneratorBuilder(
 	TEXT("Default: true"),
 	ECVF_Default);
 
-FAutoConsoleVariableRef CVarMetaSoundEnableExperimentalOneShotOperatorCache(
-	TEXT("au.MetaSound.Experimental.EnableOneShotOperatorCache"),
-	Metasound::ConsoleVariables::bEnableExperimentalOneShotOperatorCache,
+FAutoConsoleVariableRef CVarMetaSoundEnableExperimentalOneShotOperatorPool(
+	TEXT("au.MetaSound.Experimental.EnableOneShotOperatorPool"),
+	Metasound::ConsoleVariables::bEnableExperimentalOneShotOperatorPool,
 	TEXT("Enables caching of MetaSound operators using the OneShot source interface.\n")
 	TEXT("Default: false"),
 	ECVF_Default);
 
-FAutoConsoleVariableRef CVarMetaSoundEnableExperimentalOperatorCache(
-	TEXT("au.MetaSound.Experimental.EnableOperatorCache"),
-	Metasound::ConsoleVariables::bEnableExperimentalOperatorCache,
+FAutoConsoleVariableRef CVarMetaSoundEnableExperimentalOperatorPool(
+	TEXT("au.MetaSound.Experimental.EnableOperatorPool"),
+	Metasound::ConsoleVariables::bEnableExperimentalOperatorPool,
 	TEXT("Enables caching of all MetaSound operators.\n")
 	TEXT("Default: false"),
 	ECVF_Default);
@@ -246,7 +247,6 @@ namespace Metasound
 
 	void FMetasoundGeneratorInitParams::Reset(FMetasoundGeneratorInitParams& InParams)
 	{
-		InParams.Graph.Reset();
 		InParams.Environment = {};
 		InParams.MetaSoundName = {};
 		InParams.AudioOutputNames = {};
@@ -262,6 +262,9 @@ namespace Metasound
 
 	FMetasoundGenerator::FMetasoundGenerator(const FOperatorSettings& InOperatorSettings)
 		: OperatorSettings(InOperatorSettings)
+#if ENABLE_METASOUND_GENERATOR_INSTANCE_COUNTING
+		, InstanceCounter(FModuleManager::GetModuleChecked<FMetasoundGeneratorModule>("MetasoundGenerator").GetOperatorInstanceCounterManager())
+#endif // if ENABLE_METASOUND_GENERATOR_INSTANCE_COUNTING
 		, bIsFinishTriggered(false)
 		, bIsFinished(false)
 		, bPendingGraphTrigger(true)
@@ -277,7 +280,6 @@ namespace Metasound
 #else
 		, bDoRuntimeRenderTiming(false)
 #endif // if ENABLE_METASOUND_GENERATOR_RENDER_TIMING
-
 	{
 	}
 
@@ -305,6 +307,10 @@ namespace Metasound
 		NumChannels = InInitParams.AudioOutputNames.Num();
 		NumSamplesPerExecute = NumChannels * NumFramesPerExecute;
 
+#if ENABLE_METASOUND_GENERATOR_INSTANCE_COUNTING
+		InstanceCounter.Init(MetasoundName);
+#endif // if ENABLE_METASOUND_GENERATOR_INSTANCE_COUNTING
+
 		// Data channels
 		ParameterQueue = InInitParams.DataChannel;
 
@@ -803,14 +809,14 @@ PRAGMA_ENABLE_DEPRECATION_WARNINGS
 		InitBase(InParams);
 		// attempt to use operator cache instead of building a new operator.
 		bool bDidUseCachedOperator = false;
-		const bool bIsOperatorCacheEnabled = ConsoleVariables::bEnableExperimentalOneShotOperatorCache || ConsoleVariables::bEnableExperimentalOperatorCache;
+		const bool bIsOperatorPoolEnabled = ConsoleVariables::bEnableExperimentalOneShotOperatorPool || ConsoleVariables::bEnableExperimentalOperatorPool;
 		// Dynamic operators cannot use the operator cache because they can change their internal structure. 
 		// The operator cache assumes that the operator is unchanged from it's original structure. 
-		if (bIsOperatorCacheEnabled)
+		if (bIsOperatorPoolEnabled)
 		{
-			bUseOperatorCache = ConsoleVariables::bEnableExperimentalOperatorCache || MetasoundGeneratorPrivate::HasOneShotInterface(InParams.Graph->GetVertexInterface());
+			bUseOperatorPool = ConsoleVariables::bEnableExperimentalOperatorPool || MetasoundGeneratorPrivate::HasOneShotInterface(InParams.Graph->GetVertexInterface());
 
-			if (bUseOperatorCache)
+			if (bUseOperatorPool)
 			{
 				bDidUseCachedOperator = TryUseCachedOperator(InParams, true /* bTriggerGenerator */);
 			}
@@ -830,7 +836,7 @@ PRAGMA_ENABLE_DEPRECATION_WARNINGS
 			BuilderTask = nullptr;
 		}
 
-		if (bUseOperatorCache)
+		if (bUseOperatorPool)
 		{
 			ReleaseOperatorToCache();
 		}
@@ -842,12 +848,12 @@ PRAGMA_ENABLE_DEPRECATION_WARNINGS
 		METASOUND_TRACE_CPUPROFILER_EVENT_SCOPE(FMetasoundConstGraphGenerator::TryUseCachedOperator);
 
 		FMetasoundGeneratorModule& Module = FModuleManager::GetModuleChecked<FMetasoundGeneratorModule>("MetasoundGenerator");
-		TSharedPtr<FOperatorCache> OperatorCache = Module.GetOperatorCache();
-		if (OperatorCache.IsValid())
+		TSharedPtr<FOperatorPool> OperatorPool = Module.GetOperatorPool();
+		if (OperatorPool.IsValid())
 		{
 			// See if the cache has an operator with matching OperatorID
 			OperatorID = InInitParams.Graph->GetInstanceID();
-			FOperatorAndInputs GraphOperatorAndInputs = OperatorCache->ClaimCachedOperator(OperatorID);
+			FOperatorAndInputs GraphOperatorAndInputs = OperatorPool->ClaimOperator(OperatorID);
 			if (GraphOperatorAndInputs.Operator.IsValid())
 			{
 				UE_LOG(LogMetasoundGenerator, VeryVerbose, TEXT("Using cached operator %s for MetaSound %s"), *LexToString(OperatorID), *InInitParams.MetaSoundName); 
@@ -882,8 +888,8 @@ PRAGMA_ENABLE_DEPRECATION_WARNINGS
 		if (OperatorID.IsValid())
 		{
 			FMetasoundGeneratorModule& Module = FModuleManager::GetModuleChecked<FMetasoundGeneratorModule>("MetasoundGenerator");
-			TSharedPtr<FOperatorCache> OperatorCache = Module.GetOperatorCache();
-			if (OperatorCache.IsValid())
+			TSharedPtr<FOperatorPool> OperatorPool = Module.GetOperatorPool();
+			if (OperatorPool.IsValid())
 			{
 				TUniquePtr<IOperator> GraphOperator = ReleaseGraphOperator();
 
@@ -891,7 +897,7 @@ PRAGMA_ENABLE_DEPRECATION_WARNINGS
 				{
 					// Release graph operator and input data to the cache
 					UE_LOG(LogMetasoundGenerator, VeryVerbose, TEXT("Caching operator %s"), *LexToString(OperatorID));
-					OperatorCache->AddOperatorToCache(OperatorID, MoveTemp(GraphOperator), ReleaseInputVertexData());
+					OperatorPool->AddOperator(OperatorID, MoveTemp(GraphOperator), ReleaseInputVertexData());
 				}
 			}
 		}
diff --git a/Engine/Plugins/Runtime/Metasound/Source/MetasoundGenerator/Private/MetasoundGeneratorModule.cpp b/Engine/Plugins/Runtime/Metasound/Source/MetasoundGenerator/Private/MetasoundGeneratorModule.cpp
--- a/Engine/Plugins/Runtime/Metasound/Source/MetasoundGenerator/Private/MetasoundGeneratorModule.cpp
+++ b/Engine/Plugins/Runtime/Metasound/Source/MetasoundGenerator/Private/MetasoundGeneratorModule.cpp
@@ -13,8 +13,8 @@
 
 DEFINE_LOG_CATEGORY(LogMetasoundGenerator);
 
-static FAutoConsoleCommand CommandMetaSoundExperimentalOperatorCacheSetMaxNumOperators(
-	TEXT("au.MetaSound.Experimental.OperatorCache.SetMaxNumOperators"),
+static FAutoConsoleCommand CommandMetaSoundExperimentalOperatorPoolSetMaxNumOperators(
+	TEXT("au.MetaSound.Experimental.OperatorPool.SetMaxNumOperators"),
 	TEXT("Set the maximum number of operators in the MetaSound operator cache."),
 	FConsoleCommandWithArgsDelegate::CreateStatic([](const TArray< FString >& Args)
 	{
@@ -33,10 +33,10 @@ static FAutoConsoleCommand CommandMetaSoundExperimentalOperatorCacheSetMaxNumOpe
 		}
 
 		FMetasoundGeneratorModule& Module = FModuleManager::GetModuleChecked<FMetasoundGeneratorModule>("MetasoundGenerator");
-		TSharedPtr<FOperatorCache> OperatorCache = Module.GetOperatorCache();
-		if (OperatorCache.IsValid())
+		TSharedPtr<FOperatorPool> OperatorPool = Module.GetOperatorPool();
+		if (OperatorPool.IsValid())
 		{
-			OperatorCache->SetMaxNumOperators(static_cast<uint32>(MaxNumOperators));
+			OperatorPool->SetMaxNumOperators(static_cast<uint32>(MaxNumOperators));
 			UE_LOG(LogMetasoundGenerator, Display, TEXT("Metasound operator cache size set to %d operators."), MaxNumOperators);
 		}
 	})
@@ -44,22 +44,31 @@ static FAutoConsoleCommand CommandMetaSoundExperimentalOperatorCacheSetMaxNumOpe
 
 namespace Metasound
 {
+	static const FString InstanceCounterCategory(TEXT("Metasound/Active_Generators"));
+
 	void FMetasoundGeneratorModule::StartupModule()
 	{
-		FOperatorCacheSettings OperatorCacheSettings;
-		OperatorCacheSettings.MaxNumOperators = 64;
+		FOperatorPoolSettings OperatorPoolSettings;
+		OperatorPoolSettings.MaxNumOperators = 64;
 
-		OperatorCache = MakeShared<FOperatorCache>(OperatorCacheSettings);
+		OperatorPool = MakeShared<FOperatorPool>(OperatorPoolSettings);
+		OperatorInstanceCounterManager = MakeShared<FConcurrentInstanceCounterManager>(InstanceCounterCategory);
 	}
 
 	void FMetasoundGeneratorModule::ShutdownModule() 
 	{
-		OperatorCache.Reset();
+		OperatorPool.Reset();
+		OperatorInstanceCounterManager.Reset();
+	}
+
+	TSharedPtr<FOperatorPool> FMetasoundGeneratorModule::GetOperatorPool()
+	{
+		return OperatorPool;
 	}
 
-	TSharedPtr<FOperatorCache> FMetasoundGeneratorModule::GetOperatorCache()
+	TSharedPtr<FConcurrentInstanceCounterManager> FMetasoundGeneratorModule::GetOperatorInstanceCounterManager()
 	{
-		return OperatorCache;
+		return OperatorInstanceCounterManager;
 	}
 }
 
diff --git a/Engine/Plugins/Runtime/Metasound/Source/MetasoundGenerator/Private/MetasoundGeneratorModuleImpl.h b/Engine/Plugins/Runtime/Metasound/Source/MetasoundGenerator/Private/MetasoundGeneratorModuleImpl.h
--- a/Engine/Plugins/Runtime/Metasound/Source/MetasoundGenerator/Private/MetasoundGeneratorModuleImpl.h
+++ b/Engine/Plugins/Runtime/Metasound/Source/MetasoundGenerator/Private/MetasoundGeneratorModuleImpl.h
@@ -5,6 +5,7 @@
 #include "MetasoundGeneratorModule.h"
 
 #include "MetasoundOperatorCache.h"
+#include "MetasoundInstanceCounter.h"
 
 namespace Metasound
 {
@@ -15,14 +16,12 @@ namespace Metasound
 		virtual void StartupModule() override;
 		virtual void ShutdownModule() override;
 
-		TSharedPtr<FOperatorCache> GetOperatorCache();
+		virtual TSharedPtr<FOperatorPool> GetOperatorPool() override;
+		virtual TSharedPtr<FConcurrentInstanceCounterManager> GetOperatorInstanceCounterManager() override;
 
 	private:
 
-		TSharedPtr<FOperatorCache> OperatorCache;
+		TSharedPtr<FOperatorPool> OperatorPool;
+		TSharedPtr<FConcurrentInstanceCounterManager> OperatorInstanceCounterManager;
 	};
 }
-
-
-
-
diff --git a/Engine/Plugins/Runtime/Metasound/Source/MetasoundGenerator/Private/MetasoundInstanceCounter.cpp b/Engine/Plugins/Runtime/Metasound/Source/MetasoundGenerator/Private/MetasoundInstanceCounter.cpp
new file mode 100644
--- /dev/null
+++ b/Engine/Plugins/Runtime/Metasound/Source/MetasoundGenerator/Private/MetasoundInstanceCounter.cpp
@@ -0,0 +1,145 @@
+// Copyright Epic Games, Inc. All Rights Reserved.
+
+#include "MetasoundInstanceCounter.h"
+
+namespace Metasound
+{
+	FConcurrentInstanceCounterManager::FConcurrentInstanceCounterManager(const FString InCategoryName)
+	: CategoryName(InCategoryName)
+	{
+	}
+
+	void FConcurrentInstanceCounterManager::Increment(const FName& InstanceName)
+	{
+		FScopeLock Lock(&MapCritSec);
+		GetOrAddStats(InstanceName).Increment();
+	}
+
+	void FConcurrentInstanceCounterManager::Decrement(const FName& InstanceName)
+	{
+		FScopeLock Lock(&MapCritSec);
+		GetOrAddStats(InstanceName).Decrement();
+	}
+
+	int64 FConcurrentInstanceCounterManager::GetCountForName(const FName& InName)
+	{
+		FScopeLock Lock(&MapCritSec);
+		if (FStats* Stats = StatsMap.Find(InName))
+		{
+			return Stats->GetCount();
+		}
+
+		return 0;
+	}
+
+	int64 FConcurrentInstanceCounterManager::GetPeakCountForName(const FName& InName)
+	{
+		FScopeLock Lock(&MapCritSec);
+		if (FStats* Stats = StatsMap.Find(InName))
+		{
+			return Stats->GetPeakCount();
+		}
+
+		return 0;
+	}
+
+#if COUNTERSTRACE_ENABLED
+	FConcurrentInstanceCounterManager::FStats::FStats(const FString& InName)
+	: TraceCounter(MakeUnique<FCountersTrace::FCounterInt>(TraceCounterNameType_Dynamic, *InName, TraceCounterDisplayHint_None))
+	{
+	}
+#endif
+
+	void FConcurrentInstanceCounterManager::FStats::Increment()
+	{
+		ensure(TraceCounter);
+		TraceCounter->Increment();
+		PeakCount = FMath::Max(PeakCount, GetCount());
+	}
+
+
+	void FConcurrentInstanceCounterManager::FStats::Decrement()
+	{
+		ensure(TraceCounter);
+		TraceCounter->Decrement();
+	}
+
+	int64 FConcurrentInstanceCounterManager::FStats::GetCount()
+	{
+		ensure(TraceCounter);
+#if COUNTERSTRACE_ENABLED
+		return TraceCounter->Get();
+#else
+		return TraceCounter->GetValue();
+#endif
+	}
+
+	int64 FConcurrentInstanceCounterManager::FStats::GetPeakCount()
+	{
+		return PeakCount;
+	}
+
+	FConcurrentInstanceCounterManager::FStats& FConcurrentInstanceCounterManager::GetOrAddStats(const FName& InstanceName)
+	{
+		FScopeLock Lock(&MapCritSec);
+
+		// avoid re-constructing the trace counter string unless it's new
+		if (StatsMap.Contains(InstanceName))
+		{
+			return StatsMap[InstanceName];
+		}
+
+#if COUNTERSTRACE_ENABLED
+		return StatsMap.Emplace(InstanceName, FString::Printf(TEXT("%s - %s"), *CategoryName, *InstanceName.ToString()));
+#else
+		return StatsMap.Add(InstanceName);
+#endif
+	}
+
+	FConcurrentInstanceCounter::FConcurrentInstanceCounter(TSharedPtr<FConcurrentInstanceCounterManager> InCounterManager)
+	: ManagerPtr(InCounterManager)
+	{
+		check(ManagerPtr);
+	}
+
+	FConcurrentInstanceCounter::FConcurrentInstanceCounter(const FName& InName, TSharedPtr<FConcurrentInstanceCounterManager> InCounterManager)
+	: InstanceName(InName)
+	, ManagerPtr(InCounterManager)
+	{
+		check(ManagerPtr);
+		GetManagerChecked().Increment(InstanceName);
+	}
+
+	FConcurrentInstanceCounter::FConcurrentInstanceCounter(const FString& InName, TSharedPtr<FConcurrentInstanceCounterManager> InCounterManager)
+	: InstanceName(InName)
+	, ManagerPtr(InCounterManager)
+	{
+		GetManagerChecked().Increment(InstanceName);
+	}
+
+	// dtor
+	FConcurrentInstanceCounter::~FConcurrentInstanceCounter()
+	{
+		GetManagerChecked().Decrement(InstanceName);
+	}
+
+	void FConcurrentInstanceCounter::Init(const FName& InName)
+	{
+		InstanceName = InName;
+		GetManagerChecked().Increment(InstanceName);
+	}
+
+	void FConcurrentInstanceCounter::Init(const FString& InName)
+	{
+		InstanceName = FName(InName);
+		GetManagerChecked().Increment(InstanceName);
+	}
+
+	FConcurrentInstanceCounterManager& FConcurrentInstanceCounter::GetManagerChecked()
+	{
+		check(ManagerPtr);
+		return *ManagerPtr;
+	}
+
+	// static interface
+} // namespace Metasound
diff --git a/Engine/Plugins/Runtime/Metasound/Source/MetasoundGenerator/Private/MetasoundOperatorCache.cpp b/Engine/Plugins/Runtime/Metasound/Source/MetasoundGenerator/Private/MetasoundOperatorCache.cpp
--- a/Engine/Plugins/Runtime/Metasound/Source/MetasoundGenerator/Private/MetasoundOperatorCache.cpp
+++ b/Engine/Plugins/Runtime/Metasound/Source/MetasoundGenerator/Private/MetasoundOperatorCache.cpp
@@ -2,27 +2,38 @@
 
 #include "MetasoundOperatorCache.h"
 
+#include "Async/Async.h"
+#include "Async/TaskGraphInterfaces.h"
 #include "Containers/Array.h"
 #include "Containers/Map.h"
 #include "HAL/CriticalSection.h"
-#include "MetasoundGeneratorModule.h"
+#include "HAL/IConsoleManager.h"
+#include "MetasoundGeneratorModuleImpl.h"
+#include "MetasoundGeneratorBuilder.h"
+#include "MetasoundTrace.h"
 #include "Misc/Guid.h"
 #include "Misc/ScopeLock.h"
+#include "Modules/ModuleManager.h"
 #include "ProfilingDebugging/CountersTrace.h"
 #include "Templates/UniquePtr.h"
 
-#ifndef METASOUND_OPERATORCACHEPROFILER_ENABLED
-#define METASOUND_OPERATORCACHEPROFILER_ENABLED COUNTERSTRACE_ENABLED
-#endif
-
 namespace Metasound
 {
 #if METASOUND_OPERATORCACHEPROFILER_ENABLED
-	TRACE_DECLARE_FLOAT_COUNTER(MetaSound_OperatorCache_HitRatio, TEXT("MetaSound/OperatorCache/HitRatio"));
-	namespace OperatorCachePrivate
+	TRACE_DECLARE_FLOAT_COUNTER(MetaSound_OperatorPool_HitRatio, TEXT("MetaSound/OperatorPool/HitRatio"));
+	TRACE_DECLARE_FLOAT_COUNTER(MetaSound_OperatorPool_WindowedHitRatio, TEXT("MetaSound/OperatorPool/WindowedHitRatio"));
+
+	namespace OperatorPoolPrivate
 	{
 		static std::atomic<uint32> CacheHitCount = 0;
 		static std::atomic<uint32> CacheAttemptCount = 0;
+	 
+		static float MetasoundPoolHitRateWindowSecondsCVar = 1.0f;
+		FAutoConsoleVariableRef CVarMetasoundPoolHitRateWindowSeconds(
+			TEXT("au.MetaSound.OperatorPoolHitRateWindowSeconds"),
+			MetasoundPoolHitRateWindowSecondsCVar,
+			TEXT("Control how long hit/miss results matter for the success rate reporting.\n"),
+			ECVF_Default);
 
 		double GetHitRatio()
 		{
@@ -37,19 +48,133 @@ namespace Metasound
 				return 0.f;
 			}
 		}
-	}
-#endif
 
-	FOperatorCache::FOperatorCache(const FOperatorCacheSettings& InSettings)
+		FWindowedHitRate::FWindowedHitRate()
+		: CurrTTLSeconds(MetasoundPoolHitRateWindowSecondsCVar)
+		{
+		}
+
+		void FWindowedHitRate::Update()
+		{
+			if (bIsFirstUpdate)
+			{
+				bIsFirstUpdate = false;
+				FirstUpdate();
+			}
+
+			if (CurrTTLSeconds != MetasoundPoolHitRateWindowSecondsCVar)
+			{
+				SetWindowLength(MetasoundPoolHitRateWindowSecondsCVar);
+			}
+	
+			// Incorporate latest results
+			// note:	there a sliver of a race condition here between the 2 values
+			// 			but we should be able to afford the occasional off-by-one and
+			// 			avoid mutex contention
+			IntermediateResult Result( { CurrHitCount, CurrTotal });
+			CurrHitCount = 0;
+			CurrTotal = 0;
+			RunningHitCount += Result.NumHits;
+			RunningTotal += Result.Total;
+			Result.TTLSeconds = CurrTTLSeconds;
+			History.Emplace(MoveTemp(Result));
+	
+			// calculate delta time, update time
+			const uint64 CurrentTimeCycles = FPlatformTime::Cycles64();
+			const float DeltaTimeSeconds = FPlatformTime::ToSeconds64(CurrentTimeCycles - PreviousTimeCycles);
+			PreviousTimeCycles = CurrentTimeCycles;
+	
+			// tick down intermediate results, remove any which have expired
+			TickResults(DeltaTimeSeconds);
+			
+			if (RunningTotal > 0)
+			{
+				const float HitRatio = RunningHitCount / static_cast<float>(RunningTotal);
+				TRACE_COUNTER_SET(MetaSound_OperatorPool_WindowedHitRatio, HitRatio);
+			}
+		}
+
+		void FWindowedHitRate::AddHit()
+		{
+			++CurrHitCount;
+			++CurrTotal;
+		}
+	
+		void FWindowedHitRate::AddMiss()
+		{
+			++CurrTotal;
+		}
+	
+		void FWindowedHitRate::SetWindowLength(const float InNewLengthSeconds)
+		{
+			if (!ensure(InNewLengthSeconds))
+			{
+				return;
+			}
+	
+			const float Delta = InNewLengthSeconds - CurrTTLSeconds;
+			CurrTTLSeconds = InNewLengthSeconds;
+	
+			// Delta is positive if the new length is longer than the old length.
+			// perform an "inverse" Tick to adjust the TTLs in the History
+			TickResults(-Delta);
+		}
+	
+		void FWindowedHitRate::FirstUpdate()
+		{
+			PreviousTimeCycles = FPlatformTime::Cycles64();
+		}
+	
+		void FWindowedHitRate::ExpireResult(const IntermediateResult& InResultToExpire)
+		{	
+			RunningHitCount -= InResultToExpire.NumHits;
+			RunningTotal -= InResultToExpire.Total;
+		}
+	
+		void FWindowedHitRate::TickResults(const float DeltaTimeSeconds)
+		{
+			// tick down intermediate results, remove any which have expired
+			const int32 NumEntries = History.Num();
+			for (int i = NumEntries - 1; i >= 0; --i)
+			{
+				IntermediateResult& Result = History[i];
+				Result.TTLSeconds -= DeltaTimeSeconds;
+				if (Result.TTLSeconds < 0.f)
+				{
+					ExpireResult(Result);
+					History.RemoveAtSwap(i); // note: Result ref is no longer valid!
+				}
+			}
+		}
+	} // namespace OperatorPoolPrivate
+#endif // #if METASOUND_OPERATORCACHEPROFILER_ENABLED
+
+
+	FOperatorBuildData::FOperatorBuildData(
+		  FMetasoundGeneratorInitParams&& InInitParams
+		, Frontend::FNodeRegistryKey InRegistryKey
+		, FSoftObjectPath InAssetPath
+		, int32 InNumInstances
+	)
+	: InitParams(InInitParams)
+	, RegistryKey(InRegistryKey)
+	, AssetPath(InAssetPath)
+	, NumInstances(InNumInstances)
+	{}
+
+
+	FOperatorPool::FOperatorPool(const FOperatorPoolSettings& InSettings)
 	: Settings(InSettings)
 	{
 	}
 
-	FOperatorAndInputs FOperatorCache::ClaimCachedOperator(const FGuid& InOperatorID)
+	FOperatorAndInputs FOperatorPool::ClaimOperator(const FGuid& InOperatorID)
 	{
 		FScopeLock Lock(&CriticalSection);
 
 		FOperatorAndInputs OpAndInputs;
+
+		bool bCacheHit = false;
 		if (TArray<FOperatorAndInputs>* OperatorsWithID = Operators.Find(InOperatorID))
 		{
 			if (OperatorsWithID->Num() > 0)
@@ -57,61 +182,108 @@ namespace Metasound
 				OpAndInputs = OperatorsWithID->Pop();
 				Stack.RemoveAt(Stack.FindLast(InOperatorID));
 #if METASOUND_OPERATORCACHEPROFILER_ENABLED
-				OperatorCachePrivate::CacheHitCount++;
+				OperatorPoolPrivate::CacheHitCount++;
+				bCacheHit = true;
 #endif
 			}
 		}
 
 #if METASOUND_OPERATORCACHEPROFILER_ENABLED
-		OperatorCachePrivate::CacheAttemptCount++;
-
-		TRACE_COUNTER_SET(MetaSound_OperatorCache_HitRatio, OperatorCachePrivate::GetHitRatio());
+		bCacheHit? HitRateTracker.AddHit() : HitRateTracker.AddMiss();
+		OperatorPoolPrivate::CacheAttemptCount++;
+		TRACE_COUNTER_SET(MetaSound_OperatorPool_HitRatio, OperatorPoolPrivate::GetHitRatio());
 #endif
 		return OpAndInputs;
 	}
 
-	void FOperatorCache::AddOperatorToCache(const FGuid& InOperatorID, TUniquePtr<IOperator>&& InOperator, FInputVertexInterfaceData&& InInputData)
+	void FOperatorPool::AddOperator(const FGuid& InOperatorID, TUniquePtr<IOperator>&& InOperator, FInputVertexInterfaceData&& InInputData)
+	{
+		AddOperator(InOperatorID, { MoveTemp(InOperator), MoveTemp(InInputData) });
+	}
+
+	void FOperatorPool::AddOperator(const FGuid& InOperatorID, FOperatorAndInputs && OperatorAndInputs)
+	{
+		if (!OperatorAndInputs.Operator.IsValid())
+		{
+			return;
+		}
+
+		FScopeLock Lock(&CriticalSection);
+		Stack.Add(InOperatorID);
+
+		if (TArray<FOperatorAndInputs>* OperatorArray = Operators.Find(InOperatorID))
+		{
+			// add to existing array
+			OperatorArray->Add(MoveTemp(OperatorAndInputs));
+		}
+		else
+		{
+			// create a new array and add it to the map
+			TArray<FOperatorAndInputs> NewOperatorArray;
+			NewOperatorArray.Add(MoveTemp(OperatorAndInputs));
+			Operators.Add(InOperatorID, MoveTemp(NewOperatorArray));
+		}
+
+		Trim();
+	}
+
+	void FOperatorPool::BuildAndAddOperator(TUniquePtr<FOperatorBuildData> InBuildData)
 	{
-		if (InOperator.IsValid())
+		if (!ensure(InBuildData))
 		{
-			FScopeLock Lock(&CriticalSection);
+			return;
+		}
 
-			Stack.Add(InOperatorID);
 
-			FOperatorAndInputs OpAndInputs;
-			OpAndInputs.Operator = MoveTemp(InOperator);
-			OpAndInputs.Inputs = MoveTemp(InInputData);
+		FMetasoundGeneratorModule& Module = FModuleManager::GetModuleChecked<FMetasoundGeneratorModule>("MetasoundGenerator");
+		AsyncTask(ENamedThreads::AnyThread, [PreCacheData = MoveTemp(InBuildData), OperatorPool = Module.GetOperatorPool()] ()
+		{
+			METASOUND_LLM_SCOPE;
+			METASOUND_TRACE_CPUPROFILER_EVENT_SCOPE(Metasound::FOperatorPool::AsyncOperatorPrecache)
 
-			if (TArray<FOperatorAndInputs>* OperatorArray = Operators.Find(InOperatorID))
+			if (!ensure(PreCacheData && OperatorPool))
 			{
-				OperatorArray->Add(MoveTemp(OpAndInputs));
+				return;
 			}
-			else
+
+			// get the metasound graph and add to init params (might wait for async registration to complete)
+			PreCacheData->InitParams.Graph = FMetasoundFrontendRegistryContainer::Get()->GetGraph(PreCacheData->RegistryKey, PreCacheData->AssetPath);
+			ensure(PreCacheData->InitParams.Graph);
+	
+			const int32 NumInstances = PreCacheData->NumInstances;
+			for (int32 i = 0; i < NumInstances; ++i)
 			{
-				TArray<FOperatorAndInputs> NewOperatorArray;
-				NewOperatorArray.Add(MoveTemp(OpAndInputs));
-				Operators.Add(InOperatorID, MoveTemp(NewOperatorArray));
+				FBuildResults BuildResults;
+				FOperatorAndInputs OperatorAndInputs = GeneratorBuilder::BuildGraphOperator(PreCacheData->InitParams.OperatorSettings, PreCacheData->InitParams, BuildResults);
+				GeneratorBuilder::LogBuildErrors(PreCacheData->InitParams.MetaSoundName, BuildResults);
+	
+				OperatorPool->AddOperator(PreCacheData->InitParams.Graph->GetInstanceID(), MoveTemp(OperatorAndInputs));	
 			}
-
-			TrimCache();
-		}
+		});
 	}
 
-	void FOperatorCache::RemoveOperatorsWithID(const FGuid& InOperatorID)
+	void FOperatorPool::RemoveOperatorsWithID(const FGuid& InOperatorID)
 	{
 		FScopeLock Lock(&CriticalSection);
 		Operators.Remove(InOperatorID);
 		Stack.Remove(InOperatorID);
 	}
 
-	void FOperatorCache::SetMaxNumOperators(uint32 InMaxNumOperators)
+	void FOperatorPool::SetMaxNumOperators(uint32 InMaxNumOperators)
 	{
 		FScopeLock Lock(&CriticalSection);
 		Settings.MaxNumOperators = InMaxNumOperators;
-		TrimCache();
+		Trim();
+	}
+
+#if METASOUND_OPERATORCACHEPROFILER_ENABLED
+	void FOperatorPool::UpdateHitRateTracker()
+	{
+		HitRateTracker.Update();
 	}
+#endif // #if METASOUND_OPERATORCACHEPROFILER_ENABLED
 
-	void FOperatorCache::TrimCache()
+	void FOperatorPool::Trim()
 	{
 		int32 NumToTrim = Stack.Num() - Settings.MaxNumOperators;
 		if (NumToTrim > 0)
@@ -137,5 +309,4 @@ namespace Metasound
 			Stack.RemoveAt(0, NumToTrim);
 		}
 	}
-}
-
+} // namespace Metasound
\ No newline at end of file
diff --git a/Engine/Plugins/Runtime/Metasound/Source/MetasoundGenerator/Private/MetasoundOperatorCache.h b/Engine/Plugins/Runtime/Metasound/Source/MetasoundGenerator/Private/MetasoundOperatorCache.h
deleted file mode 100644
diff --git a/Engine/Plugins/Runtime/Metasound/Source/MetasoundGenerator/Public/MetasoundGenerator.h b/Engine/Plugins/Runtime/Metasound/Source/MetasoundGenerator/Public/MetasoundGenerator.h
--- a/Engine/Plugins/Runtime/Metasound/Source/MetasoundGenerator/Public/MetasoundGenerator.h
+++ b/Engine/Plugins/Runtime/Metasound/Source/MetasoundGenerator/Public/MetasoundGenerator.h
@@ -3,6 +3,7 @@
 
 #include "MetasoundExecutableOperator.h"
 #include "MetasoundGraphOperator.h"
+#include "MetasoundInstanceCounter.h"
 #include "MetasoundOperatorBuilder.h"
 #include "MetasoundOperatorInterface.h"
 #include "MetasoundParameterPack.h"
@@ -21,6 +22,9 @@
 #define ENABLE_METASOUND_GENERATOR_RENDER_TIMING WITH_EDITOR
 #endif // ifndef ENABLE_METASOUND_GENERATOR_RENDER_TIMING
 
+#ifndef ENABLE_METASOUND_GENERATOR_INSTANCE_COUNTING
+#define ENABLE_METASOUND_GENERATOR_INSTANCE_COUNTING UE_TRACE_ENABLED && !UE_BUILD_SHIPPING
+#endif // ifndef ENABLE_METASOUND_GENERATOR_INSTANCE_COUNTING
 namespace Metasound
 {
 	namespace DynamicGraph
@@ -240,6 +244,9 @@ namespace Metasound
 
 		virtual TUniquePtr<IOperator> ReleaseGraphOperator();
 		FInputVertexInterfaceData ReleaseInputVertexData();
+#if ENABLE_METASOUND_GENERATOR_INSTANCE_COUNTING
+		FConcurrentInstanceCounter InstanceCounter; 
+#endif // if ENABLE_METASOUND_GENERATOR_INSTANCE_COUNTING
 
 		/** Release the graph operator and remove any references to data owned by
 		 * the graph operator.
@@ -360,7 +367,7 @@ namespace Metasound
 
 		TUniquePtr<FAsyncTaskBase> BuilderTask;
 		FGuid OperatorID;
-		bool bUseOperatorCache = false;
+		bool bUseOperatorPool = false;
 	};
 
 	struct METASOUNDGENERATOR_API FMetasoundDynamicGraphGeneratorInitParams : FMetasoundGeneratorInitParams
diff --git a/Engine/Plugins/Runtime/Metasound/Source/MetasoundGenerator/Public/MetasoundGeneratorModule.h b/Engine/Plugins/Runtime/Metasound/Source/MetasoundGenerator/Public/MetasoundGeneratorModule.h
--- a/Engine/Plugins/Runtime/Metasound/Source/MetasoundGenerator/Public/MetasoundGeneratorModule.h
+++ b/Engine/Plugins/Runtime/Metasound/Source/MetasoundGenerator/Public/MetasoundGeneratorModule.h
@@ -9,8 +9,15 @@ DECLARE_LOG_CATEGORY_EXTERN(LogMetasoundGenerator, Log, All);
 
 namespace Metasound
 {
+	// forward
+	class FOperatorPool;
+	class FConcurrentInstanceCounterManager;
+
 	class METASOUNDGENERATOR_API IMetasoundGeneratorModule : public IModuleInterface
 	{
+	public:
+		virtual TSharedPtr<FOperatorPool> GetOperatorPool() = 0;
+		virtual TSharedPtr<FConcurrentInstanceCounterManager> GetOperatorInstanceCounterManager() = 0;
 	};
 }
 
diff --git a/Engine/Plugins/Runtime/Metasound/Source/MetasoundGenerator/Public/MetasoundInstanceCounter.h b/Engine/Plugins/Runtime/Metasound/Source/MetasoundGenerator/Public/MetasoundInstanceCounter.h
new file mode 100644
--- /dev/null
+++ b/Engine/Plugins/Runtime/Metasound/Source/MetasoundGenerator/Public/MetasoundInstanceCounter.h
@@ -0,0 +1,91 @@
+// Copyright Epic Games, Inc. All Rights Reserved.
+#pragma once
+
+#include "CoreTypes.h"
+#include "HAL/PlatformAtomics.h"
+#include "ProfilingDebugging/CountersTrace.h"
+#include "ProfilingDebugging/CsvProfiler.h"
+#include "Stats/Stats.h"
+#include "Templates/UniquePtr.h"
+#include "UObject/NameTypes.h"
+
+// todo: comments
+
+namespace Metasound
+{
+#if COUNTERSTRACE_ENABLED
+	using CounterType = FCountersTrace::FCounterInt;
+#else
+	using CounterType = FThreadSafeCounter64;
+#endif
+
+	// forward
+	class FConcurrentInstanceCounter;
+
+	class FConcurrentInstanceCounterManager
+	{
+	public:
+		FConcurrentInstanceCounterManager(const FString InCategoryName);
+
+		void Increment(const FName& InstanceName);
+		void Decrement(const FName& IntsanceName);
+
+		int64 GetCountForName(const FName& InName);
+		int64 GetPeakCountForName(const FName& InName);
+
+	private:
+		struct FStats
+		{
+		public:
+			// ctor
+#if COUNTERSTRACE_ENABLED
+			FStats(const FString& InName);
+#else
+			FStats() = default;
+#endif
+	
+			void Increment();
+			void Decrement();
+	
+			int64 GetCount();
+			int64 GetPeakCount();
+	
+		private:
+			TUniquePtr<CounterType> TraceCounter;
+			int64 PeakCount;
+	
+		}; // struct FStats
+
+		FStats& GetOrAddStats(const FName& InstanceName);
+
+		TMap<FName, FStats> StatsMap;
+		FCriticalSection MapCritSec;
+		FString CategoryName;
+
+	}; // class FConcurrentInstanceCounterManager
+
+
+
+	class FConcurrentInstanceCounter
+	{
+	public:
+		// ctor
+		FConcurrentInstanceCounter() = delete; // must init ManagerPtr
+		FConcurrentInstanceCounter(TSharedPtr<FConcurrentInstanceCounterManager> InCounterManager);
+		FConcurrentInstanceCounter(const FName& InName, TSharedPtr<FConcurrentInstanceCounterManager> InCounterManager);
+		FConcurrentInstanceCounter(const FString& InName, TSharedPtr<FConcurrentInstanceCounterManager> InCounterManager);
+	
+		// dtor
+		virtual ~FConcurrentInstanceCounter();
+	
+		// for non-RAII clients
+		void Init(const FName& InName);
+		void Init(const FString& InName);	
+	
+	private:
+		FConcurrentInstanceCounterManager& GetManagerChecked();
+
+		FName InstanceName;
+		TSharedPtr<FConcurrentInstanceCounterManager> ManagerPtr;
+	}; // class FConcurrentInstanceCounter
+} // namespace Metasound
diff --git a/Engine/Plugins/Runtime/Metasound/Source/MetasoundGenerator/Public/MetasoundOperatorCache.h b/Engine/Plugins/Runtime/Metasound/Source/MetasoundGenerator/Public/MetasoundOperatorCache.h
new file mode 100644
--- /dev/null
+++ b/Engine/Plugins/Runtime/Metasound/Source/MetasoundGenerator/Public/MetasoundOperatorCache.h
@@ -0,0 +1,117 @@
+// Copyright Epic Games, Inc. All Rights Reserved.
+
+#pragma once
+
+#include "Containers/Array.h"
+#include "Containers/Map.h"
+#include "MetasoundOperatorInterface.h"
+#include "MetasoundGenerator.h"
+#include "Misc/Guid.h"
+#include "HAL/CriticalSection.h"
+#include "Templates/UniquePtr.h"
+
+#ifndef METASOUND_OPERATORCACHEPROFILER_ENABLED
+#define METASOUND_OPERATORCACHEPROFILER_ENABLED COUNTERSTRACE_ENABLED
+#endif
+namespace Metasound
+{
+#if METASOUND_OPERATORCACHEPROFILER_ENABLED
+	namespace OperatorPoolPrivate
+	{
+		class FWindowedHitRate
+		{
+		public:
+			// ctor
+			FWindowedHitRate();
+			void Update();
+			void AddHit();
+			void AddMiss();
+	
+		private:
+			struct IntermediateResult
+			{
+				uint32 NumHits = 0;
+				uint32 Total = 0;
+				float TTLSeconds;
+			};
+	
+			TArray<IntermediateResult> History;
+	
+			uint32 CurrHitCount = 0;
+			uint32 CurrTotal = 0;
+			uint32 RunningHitCount = 0;
+			uint32 RunningTotal = 0;
+	
+			float CurrTTLSeconds = 0.f;
+	
+			uint64 PreviousTimeCycles = 0;
+			bool bIsFirstUpdate = true;
+	
+			void FirstUpdate();
+			void SetWindowLength(const float InNewLengthSeconds);
+			void AddResult(IntermediateResult&& InResultToAdd);
+			void ExpireResult(const IntermediateResult& InResultToExpire);
+			void TickResults(const float DeltaTimeSeconds);
+	
+		}; // class FWindowedHitRate
+	} // namespace OperatorPoolPrivate
+#endif // #if METASOUND_OPERATORCACHEPROFILER_ENABLED
+
+	struct FOperatorPoolSettings
+	{
+		uint32 MaxNumOperators = 64;
+	};
+
+	// Data required to build an operator without immediately playing it
+	struct METASOUNDGENERATOR_API FOperatorBuildData
+	{
+		FMetasoundGeneratorInitParams InitParams;
+		Frontend::FNodeRegistryKey RegistryKey;
+		FSoftObjectPath AssetPath;
+		int32 NumInstances;
+
+		FOperatorBuildData() = delete;
+		FOperatorBuildData(
+			  FMetasoundGeneratorInitParams&& InInitParams
+			, Frontend::FNodeRegistryKey InRegistryKey
+			, FSoftObjectPath InAssetPath
+			, int32 InNumInstances = 1
+		);
+
+	}; // struct FOperatorPrecacheData
+
+	// Pool of re-useable metasound operators to be used / put back by the metasound generator
+	// operators can also be pre-constructed via the UMetasoundCacheSubsystem BP api.
+	class METASOUNDGENERATOR_API FOperatorPool
+	{
+	public:
+		FOperatorPool(const FOperatorPoolSettings& InSettings);
+		FOperatorAndInputs ClaimOperator(const FGuid& InOperatorID);
+		void AddOperator(const FGuid& InOperatorID, TUniquePtr<IOperator>&& InOperator, FInputVertexInterfaceData&& InputData);
+		void AddOperator(const FGuid& InOperatorID, FOperatorAndInputs && OperatorAndInputs);
+		void BuildAndAddOperator(TUniquePtr<FOperatorBuildData> InBuildData);
+		void RemoveOperatorsWithID(const FGuid& InOperatorID);
+		void SetMaxNumOperators(uint32 InMaxNumOperators);
+#if METASOUND_OPERATORCACHEPROFILER_ENABLED
+		void UpdateHitRateTracker();
+#endif // #if METASOUND_OPERATORCACHEPROFILER_ENABLED
+
+	private:
+		void Trim();
+
+		FOperatorPoolSettings Settings;
+		FCriticalSection CriticalSection;
+
+#if METASOUND_OPERATORCACHEPROFILER_ENABLED
+		OperatorPoolPrivate::FWindowedHitRate HitRateTracker;
+#endif // #if METASOUND_OPERATORCACHEPROFILER_ENABLED
+
+		TMap<FGuid, TArray<FOperatorAndInputs>> Operators;
+		TArray<FGuid> Stack;
+	};
+
+} // namespace Metasound
+
+
+
+
diff --git a/Engine/Plugins/Runtime/Metasound/Source/MetasoundGraphCore/Public/MetasoundOperatorInterface.h b/Engine/Plugins/Runtime/Metasound/Source/MetasoundGraphCore/Public/MetasoundOperatorInterface.h
--- a/Engine/Plugins/Runtime/Metasound/Source/MetasoundGraphCore/Public/MetasoundOperatorInterface.h
+++ b/Engine/Plugins/Runtime/Metasound/Source/MetasoundGraphCore/Public/MetasoundOperatorInterface.h
@@ -181,4 +181,11 @@ namespace Metasound
 		 */
 		virtual FPostExecuteFunction GetPostExecuteFunction() = 0;
 	};
+
+	// todo: better place to put this? needs to be in core
+	struct FOperatorAndInputs
+	{
+		TUniquePtr<IOperator> Operator;
+		FInputVertexInterfaceData Inputs;
+	};
 }
