[CIS] added missing include.

diff --git a/Engine/Plugins/Runtime/Metasound/Source/MetasoundEngine/Private/MetasoundOperatorCacheSubsystem.cpp b/Engine/Plugins/Runtime/Metasound/Source/MetasoundEngine/Private/MetasoundOperatorCacheSubsystem.cpp
--- a/Engine/Plugins/Runtime/Metasound/Source/MetasoundEngine/Private/MetasoundOperatorCacheSubsystem.cpp
+++ b/Engine/Plugins/Runtime/Metasound/Source/MetasoundEngine/Private/MetasoundOperatorCacheSubsystem.cpp
@@ -7,6 +7,7 @@
 #include "MetasoundOperatorCache.h"
 #include "Modules/ModuleManager.h"
 #include "Misc/Optional.h"
+#include "Sound/SoundGenerator.h"
 
 static TOptional<Metasound::FMetasoundGeneratorInitParams> CreateInitParams(UMetaSoundSource* InMetaSound, const FSoundGeneratorInitParams& InParams)
 {
