[CIS fix] issues using TUniquePtr<FSoundGeneratorInitParams>

diff --git a/Engine/Plugins/Runtime/Metasound/Source/MetasoundEngine/Private/MetasoundOperatorCacheSubsystem.cpp b/Engine/Plugins/Runtime/Metasound/Source/MetasoundEngine/Private/MetasoundOperatorCacheSubsystem.cpp
--- a/Engine/Plugins/Runtime/Metasound/Source/MetasoundEngine/Private/MetasoundOperatorCacheSubsystem.cpp
+++ b/Engine/Plugins/Runtime/Metasound/Source/MetasoundEngine/Private/MetasoundOperatorCacheSubsystem.cpp
@@ -7,7 +7,6 @@
 #include "MetasoundOperatorCache.h"
 #include "Modules/ModuleManager.h"
 #include "Misc/Optional.h"
-#include "Sound/SoundGenerator.h"
 
 static TOptional<Metasound::FMetasoundGeneratorInitParams> CreateInitParams(UMetaSoundSource* InMetaSound, const FSoundGeneratorInitParams& InParams)
 {
@@ -57,17 +56,16 @@ void UMetaSoundCacheSubsystem::Initialize(FSubsystemCollectionBase& Collection)
 	using namespace Audio;
 	using namespace Metasound;
 
-	BuildParams = MakeUnique<FSoundGeneratorInitParams>();
 	const FMixerDevice* MixerDevice = GetMixerDevice();
 
-	if (BuildParams && ensure(MixerDevice))
+	if (ensure(MixerDevice))
 	{
-		BuildParams->AudioDeviceID = GetAudioDeviceHandle().GetDeviceID();
-		BuildParams->SampleRate = MixerDevice->GetSampleRate();
-		BuildParams->AudioMixerNumOutputFrames = MixerDevice->GetNumOutputFrames();
-		BuildParams->NumChannels = MixerDevice->GetNumDeviceChannels();
-		BuildParams->NumFramesPerCallback = 0;
-		BuildParams->InstanceID = 0;
+		BuildParams.AudioDeviceID = GetAudioDeviceHandle().GetDeviceID();
+		BuildParams.SampleRate = MixerDevice->GetSampleRate();
+		BuildParams.AudioMixerNumOutputFrames = MixerDevice->GetNumOutputFrames();
+		BuildParams.NumChannels = MixerDevice->GetNumDeviceChannels();
+		BuildParams.NumFramesPerCallback = 0;
+		BuildParams.InstanceID = 0;
 	}
 }
 
@@ -96,7 +94,7 @@ void UMetaSoundCacheSubsystem::PrecacheMetaSound(UMetaSoundSource* InMetaSound,
 
 	const FMixerDevice* MixerDevice = GetMixerDevice();
 
-	if (!ensure(MixerDevice && InMetaSound && BuildParams && OperatorPool))
+	if (!ensure(MixerDevice && InMetaSound && OperatorPool))
 	{
 		return;
 	}
@@ -107,14 +105,14 @@ void UMetaSoundCacheSubsystem::PrecacheMetaSound(UMetaSoundSource* InMetaSound,
 	}
 
 	InMetaSound->InitResources();
-	BuildParams->GraphName = InMetaSound->GetOwningAssetName();
+	BuildParams.GraphName = InMetaSound->GetOwningAssetName();
 
 	if (InMetaSound->IsDynamic())
 	{
 		return;
 	}
 
-	TOptional<FMetasoundGeneratorInitParams> InitParams = CreateInitParams(InMetaSound, *BuildParams);
+	TOptional<FMetasoundGeneratorInitParams> InitParams = CreateInitParams(InMetaSound, BuildParams);
 	if (!InitParams.IsSet())
 	{
 		return;
diff --git a/Engine/Plugins/Runtime/Metasound/Source/MetasoundEngine/Public/MetasoundOperatorCacheSubsystem.h b/Engine/Plugins/Runtime/Metasound/Source/MetasoundEngine/Public/MetasoundOperatorCacheSubsystem.h
--- a/Engine/Plugins/Runtime/Metasound/Source/MetasoundEngine/Public/MetasoundOperatorCacheSubsystem.h
+++ b/Engine/Plugins/Runtime/Metasound/Source/MetasoundEngine/Public/MetasoundOperatorCacheSubsystem.h
@@ -8,6 +8,7 @@
 #include "Misc/Guid.h"
 #include "HAL/CriticalSection.h"
 #include "Templates/UniquePtr.h"
+#include "Sound/SoundGenerator.h"
 #include "Subsystems/AudioEngineSubsystem.h"
 
 #include "MetasoundSource.h"
@@ -36,6 +37,6 @@ public:
 	void PrecacheMetaSound(UMetaSoundSource* InMetaSound, int32 InNumInstances = 1);
 
 private:
-	TUniquePtr<FSoundGeneratorInitParams> BuildParams;
+	FSoundGeneratorInitParams BuildParams;
 }; //UMetaSoundCacheSubsystem
 
