Check correct asset path when waiting for async metasound graph registration
#rnx

diff --git a/Engine/Plugins/Runtime/Metasound/Source/MetasoundFrontend/Private/MetasoundFrontendNodeRegistryPrivate.h b/Engine/Plugins/Runtime/Metasound/Source/MetasoundFrontend/Private/MetasoundFrontendNodeRegistryPrivate.h
--- a/Engine/Plugins/Runtime/Metasound/Source/MetasoundFrontend/Private/MetasoundFrontendNodeRegistryPrivate.h
+++ b/Engine/Plugins/Runtime/Metasound/Source/MetasoundFrontend/Private/MetasoundFrontendNodeRegistryPrivate.h
@@ -183,6 +183,7 @@ namespace Metasound
 			struct FActiveRegistrationTaskInfo
 			{
 				UE::Tasks::FTask Task;
+				FSoftObjectPath AssetPath;
 				TObjectPtr<UObject> OwningObject;
 			};
 
diff --git a/Engine/Plugins/Runtime/Metasound/Source/MetasoundFrontend/Private/MetasoundFrontendRegistries.cpp b/Engine/Plugins/Runtime/Metasound/Source/MetasoundFrontend/Private/MetasoundFrontendRegistries.cpp
--- a/Engine/Plugins/Runtime/Metasound/Source/MetasoundFrontend/Private/MetasoundFrontendRegistries.cpp
+++ b/Engine/Plugins/Runtime/Metasound/Source/MetasoundFrontend/Private/MetasoundFrontendRegistries.cpp
@@ -413,7 +413,7 @@ namespace Metasound
 				{
 					ObjectReferencer->AddObject(OwningObject);
 				}
-				ActiveRegistrationTasks.Add(RegistryKey, FActiveRegistrationTaskInfo{BuildAndRegisterTask, OwningObject});
+				ActiveRegistrationTasks.Add(RegistryKey, FActiveRegistrationTaskInfo{BuildAndRegisterTask, InAssetPath, OwningObject});
 			}
 			else
 			{
@@ -857,7 +857,7 @@ namespace Metasound
 					if (InAssetPath)
 					{
 						// Filter by asset path
-						if (*InAssetPath == FSoftObjectPath(FoundTask->OwningObject))
+						if (*InAssetPath == FoundTask->AssetPath)
 						{
 							ActiveRegistrationTask = FoundTask->Task;
 						}
