Update use of MetaSound FCreateOperatorParams to FBuilderOperatorParams for perf (Part 1)

diff --git a/Engine/Plugins/Runtime/AudioModulation/Source/AudioModulation/Private/MetaSoundNodes/GetModulatorValueNode.cpp b/Engine/Plugins/Runtime/AudioModulation/Source/AudioModulation/Private/MetaSoundNodes/GetModulatorValueNode.cpp
--- a/Engine/Plugins/Runtime/AudioModulation/Source/AudioModulation/Private/MetaSoundNodes/GetModulatorValueNode.cpp
+++ b/Engine/Plugins/Runtime/AudioModulation/Source/AudioModulation/Private/MetaSoundNodes/GetModulatorValueNode.cpp
@@ -61,17 +61,16 @@ namespace AudioModulation
 			return Metadata;
 		}
 
-		static TUniquePtr<Metasound::IOperator> CreateOperator(const Metasound::FCreateOperatorParams& InParams, TArray<TUniquePtr<Metasound::IOperatorBuildError>>& OutErrors)
+		static TUniquePtr<Metasound::IOperator> CreateOperator(const Metasound::FBuildOperatorParams& InParams, Metasound::FBuildResults& OutResults)
 		{
 			using namespace Metasound;
 
-			const FInputVertexInterface& InputInterface = InParams.Node.GetVertexInterface().GetInputInterface();
-			const FDataReferenceCollection& InputCollection = InParams.InputDataReferences;
+			const FInputVertexInterfaceData& InputData = InParams.InputData;
 
 			if (InParams.Environment.Contains<Audio::FDeviceId>(Metasound::Frontend::SourceInterface::Environment::DeviceID))
 			{
-				FSoundModulatorAssetReadRef ModulatorReadRef = InputCollection.GetDataReadReferenceOrConstruct<FSoundModulatorAsset>("Modulator");
-				FBoolReadRef NormalizedReadRef = InputCollection.GetDataReadReferenceOrConstructWithVertexDefault<bool>(InputInterface, "Normalized", InParams.OperatorSettings);
+				FSoundModulatorAssetReadRef ModulatorReadRef = InputData.GetOrConstructDataReadReference<FSoundModulatorAsset>("Modulator");
+				FBoolReadRef NormalizedReadRef = InputData.GetOrCreateDefaultDataReadReference<bool>("Normalized", InParams.OperatorSettings);
 
 				return MakeUnique<FGetModulatorValueNodeOperator>(InParams, ModulatorReadRef, NormalizedReadRef);
 			}
@@ -82,7 +81,7 @@ namespace AudioModulation
 			}
 		}
 
-		FGetModulatorValueNodeOperator(const Metasound::FCreateOperatorParams& InParams, const FSoundModulatorAssetReadRef& InModulator, const Metasound::FBoolReadRef& InNormalized)
+		FGetModulatorValueNodeOperator(const Metasound::FBuildOperatorParams& InParams, const FSoundModulatorAssetReadRef& InModulator, const Metasound::FBoolReadRef& InNormalized)
 			: DeviceId(InParams.Environment.GetValue<Audio::FDeviceId>(Metasound::Frontend::SourceInterface::Environment::DeviceID))
 			, Modulator(InModulator)
 			, Normalized(InNormalized)
diff --git a/Engine/Plugins/Runtime/AudioModulation/Source/AudioModulation/Private/MetaSoundNodes/MixModulatorsNode.cpp b/Engine/Plugins/Runtime/AudioModulation/Source/AudioModulation/Private/MetaSoundNodes/MixModulatorsNode.cpp
--- a/Engine/Plugins/Runtime/AudioModulation/Source/AudioModulation/Private/MetaSoundNodes/MixModulatorsNode.cpp
+++ b/Engine/Plugins/Runtime/AudioModulation/Source/AudioModulation/Private/MetaSoundNodes/MixModulatorsNode.cpp
@@ -64,19 +64,17 @@ namespace AudioModulation
 			return Metadata;
 		}
 
-		static TUniquePtr<Metasound::IOperator> CreateOperator(const Metasound::FCreateOperatorParams& InParams, TArray<TUniquePtr<Metasound::IOperatorBuildError>>& OutErrors)
+		static TUniquePtr<Metasound::IOperator> CreateOperator(const Metasound::FBuildOperatorParams& InParams, Metasound::FBuildResults& OutResults)
 		{
 			using namespace Metasound;
-
-			const FInputVertexInterface& InputInterface = InParams.Node.GetVertexInterface().GetInputInterface();
-			const FDataReferenceCollection& InputCollection = InParams.InputDataReferences;
+			const FInputVertexInterfaceData& InputData = InParams.InputData;
 				
 			if (InParams.Environment.Contains<Audio::FDeviceId>(Metasound::Frontend::SourceInterface::Environment::DeviceID))
 			{
-				FSoundModulatorAssetReadRef Modulator1ReadRef = InputCollection.GetDataReadReferenceOrConstruct<FSoundModulatorAsset>("In1");
-				FSoundModulatorAssetReadRef Modulator2ReadRef = InputCollection.GetDataReadReferenceOrConstruct<FSoundModulatorAsset>("In2");
-				FSoundModulationParameterAssetReadRef ParameterReadRef = InputCollection.GetDataReadReferenceOrConstruct<FSoundModulationParameterAsset>("MixParameter");
-				FBoolReadRef NormalizedReadRef = InputCollection.GetDataReadReferenceOrConstructWithVertexDefault<bool>(InputInterface, "Normalized", InParams.OperatorSettings);
+				FSoundModulatorAssetReadRef Modulator1ReadRef = InputData.GetOrConstructDataReadReference<FSoundModulatorAsset>("In1");
+				FSoundModulatorAssetReadRef Modulator2ReadRef = InputData.GetOrConstructDataReadReference<FSoundModulatorAsset>("In2");
+				FSoundModulationParameterAssetReadRef ParameterReadRef = InputData.GetOrConstructDataReadReference<FSoundModulationParameterAsset>("MixParameter");
+				FBoolReadRef NormalizedReadRef = InputData.GetOrCreateDefaultDataReadReference<bool>("Normalized", InParams.OperatorSettings);
 
 				return MakeUnique<FMixModulatorsNodeOperator>(InParams, Modulator1ReadRef, Modulator2ReadRef, ParameterReadRef, NormalizedReadRef);
 			}
@@ -88,7 +86,7 @@ namespace AudioModulation
 		}
 
 		FMixModulatorsNodeOperator(
-			const Metasound::FCreateOperatorParams& InParams,
+			const Metasound::FBuildOperatorParams& InParams,
 			const FSoundModulatorAssetReadRef& InModulator1,
 			const FSoundModulatorAssetReadRef& InModulator2,
 			const FSoundModulationParameterAssetReadRef& InParameter,
diff --git a/Engine/Plugins/Runtime/Metasound/Source/MetasoundEngine/Private/MetasoundAudioBusReaderNode.cpp b/Engine/Plugins/Runtime/Metasound/Source/MetasoundEngine/Private/MetasoundAudioBusReaderNode.cpp
--- a/Engine/Plugins/Runtime/Metasound/Source/MetasoundEngine/Private/MetasoundAudioBusReaderNode.cpp
+++ b/Engine/Plugins/Runtime/Metasound/Source/MetasoundEngine/Private/MetasoundAudioBusReaderNode.cpp
@@ -83,19 +83,19 @@ namespace Metasound
 			return Interface;
 		}
 
-		static TUniquePtr<IOperator> CreateOperator(const FCreateOperatorParams& InParams, FBuildErrorArray& OutErrors)
+		static TUniquePtr<IOperator> CreateOperator(const FBuildOperatorParams& InParams, FBuildResults& OutResults)
 		{
 			using namespace Frontend;
 			
 			using namespace AudioBusReaderNode; 
-			const FDataReferenceCollection& InputCollection = InParams.InputDataReferences;
+			const FInputVertexInterfaceData& InputData = InParams.InputData;
 
 			bool bHasEnvironmentVars = InParams.Environment.Contains<Audio::FDeviceId>(SourceInterface::Environment::DeviceID);
 			bHasEnvironmentVars &= InParams.Environment.Contains<int32>(SourceInterface::Environment::AudioMixerNumOutputFrames);
 			
 			if (bHasEnvironmentVars)
 			{
-				FAudioBusAssetReadRef AudioBusIn = InputCollection.GetDataReadReferenceOrConstruct<FAudioBusAsset>(METASOUND_GET_PARAM_NAME(InParamAudioBusInput));
+				FAudioBusAssetReadRef AudioBusIn = InputData.GetOrConstructDataReadReference<FAudioBusAsset>(METASOUND_GET_PARAM_NAME(InParamAudioBusInput));
 				return MakeUnique<TAudioBusReaderOperator<NumChannels>>(InParams, AudioBusIn);
 			}
 			else
@@ -106,7 +106,7 @@ namespace Metasound
 			}
 		}
 
-		TAudioBusReaderOperator(const FCreateOperatorParams& InParams, const FAudioBusAssetReadRef& InAudioBusAsset) : AudioBusAsset(InAudioBusAsset)
+		TAudioBusReaderOperator(const FBuildOperatorParams& InParams, const FAudioBusAssetReadRef& InAudioBusAsset) : AudioBusAsset(InAudioBusAsset)
 		{
 			for (int32 ChannelIndex = 0; ChannelIndex < NumChannels; ++ChannelIndex)
 			{
diff --git a/Engine/Plugins/Runtime/Metasound/Source/MetasoundEngine/Private/MetasoundAudioBusWriterNode.cpp b/Engine/Plugins/Runtime/Metasound/Source/MetasoundEngine/Private/MetasoundAudioBusWriterNode.cpp
--- a/Engine/Plugins/Runtime/Metasound/Source/MetasoundEngine/Private/MetasoundAudioBusWriterNode.cpp
+++ b/Engine/Plugins/Runtime/Metasound/Source/MetasoundEngine/Private/MetasoundAudioBusWriterNode.cpp
@@ -75,24 +75,24 @@ namespace Metasound
 			return Interface;
 		}
 
-		static TUniquePtr<IOperator> CreateOperator(const FCreateOperatorParams& InParams, FBuildErrorArray& OutErrors)
+		static TUniquePtr<IOperator> CreateOperator(const FBuildOperatorParams& InParams, FBuildResults& OutResults)
 		{
 			using namespace Frontend;
 			using namespace AudioBusWriterNode;
 
-			const FDataReferenceCollection& InputCollection = InParams.InputDataReferences;
+			const FInputVertexInterfaceData& InputData = InParams.InputData;
 
 			bool bHasEnvironmentVars = InParams.Environment.Contains<Audio::FDeviceId>(SourceInterface::Environment::DeviceID);
 			bHasEnvironmentVars &= InParams.Environment.Contains<int32>(SourceInterface::Environment::AudioMixerNumOutputFrames);
 
 			if (bHasEnvironmentVars)
 			{
-				FAudioBusAssetReadRef AudioBusIn = InputCollection.GetDataReadReferenceOrConstruct<FAudioBusAsset>(METASOUND_GET_PARAM_NAME(InParamAudioBusOutput));
+				FAudioBusAssetReadRef AudioBusIn = InputData.GetOrConstructDataReadReference<FAudioBusAsset>(METASOUND_GET_PARAM_NAME(InParamAudioBusOutput));
 
 				TArray<FAudioBufferReadRef> AudioInputs;
 				for (int32 ChannelIndex = 0; ChannelIndex < NumChannels; ++ChannelIndex)
 				{
-					AudioInputs.Add(InputCollection.GetDataReadReferenceOrConstruct<FAudioBuffer>(METASOUND_GET_PARAM_NAME_WITH_INDEX(InParamAudio, ChannelIndex), InParams.OperatorSettings));
+					AudioInputs.Add(InputData.GetOrConstructDataReadReference<FAudioBuffer>(METASOUND_GET_PARAM_NAME_WITH_INDEX(InParamAudio, ChannelIndex), InParams.OperatorSettings));
 				}
 
 				return MakeUnique<TAudioBusWriterOperator<NumChannels>>(InParams, MoveTemp(AudioBusIn), MoveTemp(AudioInputs));
@@ -105,7 +105,7 @@ namespace Metasound
 			}
 		}
 
-		TAudioBusWriterOperator(const FCreateOperatorParams& InParams, FAudioBusAssetReadRef InAudioBusAsset, TArray<FAudioBufferReadRef> InAudioInputs)
+		TAudioBusWriterOperator(const FBuildOperatorParams& InParams, FAudioBusAssetReadRef InAudioBusAsset, TArray<FAudioBufferReadRef> InAudioInputs)
 			: AudioBusAsset(MoveTemp(InAudioBusAsset))
 			, AudioInputs(MoveTemp(InAudioInputs))
 		{
diff --git a/Engine/Plugins/Runtime/Metasound/Source/MetasoundEngine/Private/MetasoundWaveInfoNode.cpp b/Engine/Plugins/Runtime/Metasound/Source/MetasoundEngine/Private/MetasoundWaveInfoNode.cpp
--- a/Engine/Plugins/Runtime/Metasound/Source/MetasoundEngine/Private/MetasoundWaveInfoNode.cpp
+++ b/Engine/Plugins/Runtime/Metasound/Source/MetasoundEngine/Private/MetasoundWaveInfoNode.cpp
@@ -46,7 +46,7 @@ namespace Metasound
 		// node interface
 		static const FNodeClassMetadata& GetNodeInfo();
 		static FVertexInterface DeclareVertexInterface();
-		static TUniquePtr<IOperator> CreateOperator(const FCreateOperatorParams& InParams, FBuildErrorArray& OutErrors);
+		static TUniquePtr<IOperator> CreateOperator(const FBuildOperatorParams& InParams, FBuildResults& OutResults);
 		virtual void BindInputs(FInputVertexInterfaceData& InOutVertexData) override;
 		virtual void BindOutputs(FOutputVertexInterfaceData& InOutVertexData) override;
 		virtual FDataReferenceCollection GetInputs() const override;
@@ -118,14 +118,13 @@ namespace Metasound
 	}
 
 
-	TUniquePtr<IOperator> FWaveInfoNodeOperator::CreateOperator(const FCreateOperatorParams& InParams, FBuildErrorArray& OutErrors)
+	TUniquePtr<IOperator> FWaveInfoNodeOperator::CreateOperator(const FBuildOperatorParams& InParams, FBuildResults& OutResults)
 	{
 		using namespace WaveInfoNodeParameterNames;
 		
-		const FDataReferenceCollection& InputDataRefs = InParams.InputDataReferences;
-
+		const FInputVertexInterfaceData& InputData = InParams.InputData;
 		// inputs
-		FWaveAssetReadRef WaveAssetIn = InputDataRefs.GetDataReadReferenceOrConstruct<FWaveAsset>(METASOUND_GET_PARAM_NAME(ParamWaveAsset));
+		FWaveAssetReadRef WaveAssetIn = InputData.GetOrConstructDataReadReference<FWaveAsset>(METASOUND_GET_PARAM_NAME(ParamWaveAsset));
 
 		return MakeUnique < FWaveInfoNodeOperator >(InParams.OperatorSettings, WaveAssetIn);
 	}
diff --git a/Engine/Plugins/Runtime/Metasound/Source/MetasoundEngine/Private/MetasoundWaveTableEnvelopeNode.cpp b/Engine/Plugins/Runtime/Metasound/Source/MetasoundEngine/Private/MetasoundWaveTableEnvelopeNode.cpp
--- a/Engine/Plugins/Runtime/Metasound/Source/MetasoundEngine/Private/MetasoundWaveTableEnvelopeNode.cpp
+++ b/Engine/Plugins/Runtime/Metasound/Source/MetasoundEngine/Private/MetasoundWaveTableEnvelopeNode.cpp
@@ -96,26 +96,25 @@ namespace Metasound
 			return Metadata;
 		}
 
-		static TUniquePtr<IOperator> CreateOperator(const FCreateOperatorParams& InParams, TArray<TUniquePtr<IOperatorBuildError>>& OutErrors)
+		static TUniquePtr<IOperator> CreateOperator(const FBuildOperatorParams& InParams, FBuildResults& OutResults)
 		{
 			using namespace WaveTable;
+			
+			const FInputVertexInterfaceData& InputData = InParams.InputData;
 
-			const FInputVertexInterface& InputInterface = InParams.Node.GetVertexInterface().GetInputInterface();
-			const FDataReferenceCollection& InputCollection = InParams.InputDataReferences;
-
-			FWaveTableReadRef InWaveTableReadRef = InputCollection.GetDataReadReferenceOrConstruct<FWaveTable>("WaveTable");
-			FTriggerReadRef InPlayReadRef = InputCollection.GetDataReadReferenceOrConstructWithVertexDefault<FTrigger>(InputInterface, "Play", InParams.OperatorSettings);
-			FTriggerReadRef InStopReadRef = InputCollection.GetDataReadReferenceOrConstructWithVertexDefault<FTrigger>(InputInterface, "Stop", InParams.OperatorSettings);
-			FTriggerReadRef InPauseReadRef = InputCollection.GetDataReadReferenceOrConstructWithVertexDefault<FTrigger>(InputInterface, "Pause", InParams.OperatorSettings);
-			FTimeReadRef InDurationReadRef = InputCollection.GetDataReadReferenceOrConstructWithVertexDefault<FTime>(InputInterface, "Duration", InParams.OperatorSettings);
-			FEnumWaveTableEnvelopeModeReadRef InModeReadRef = InputCollection.GetDataReadReferenceOrConstructWithVertexDefault<FEnumWaveTableEnvelopeMode>( InputInterface, "Mode", InParams.OperatorSettings);
-			FEnumWaveTableInterpModeReadRef InInterpReadRef = InputCollection.GetDataReadReferenceOrConstructWithVertexDefault<FEnumWaveTableInterpolationMode>(InputInterface, "Interpolation", InParams.OperatorSettings);
+			FWaveTableReadRef InWaveTableReadRef = InputData.GetOrConstructDataReadReference<FWaveTable>("WaveTable");
+			FTriggerReadRef InPlayReadRef = InputData.GetOrCreateDefaultDataReadReference<FTrigger>("Play", InParams.OperatorSettings);
+			FTriggerReadRef InStopReadRef = InputData.GetOrCreateDefaultDataReadReference<FTrigger>("Stop", InParams.OperatorSettings);
+			FTriggerReadRef InPauseReadRef = InputData.GetOrCreateDefaultDataReadReference<FTrigger>("Pause", InParams.OperatorSettings);
+			FTimeReadRef InDurationReadRef = InputData.GetOrCreateDefaultDataReadReference<FTime>("Duration", InParams.OperatorSettings);
+			FEnumWaveTableEnvelopeModeReadRef InModeReadRef = InputData.GetOrCreateDefaultDataReadReference<FEnumWaveTableEnvelopeMode>("Mode", InParams.OperatorSettings);
+			FEnumWaveTableInterpModeReadRef InInterpReadRef = InputData.GetOrCreateDefaultDataReadReference<FEnumWaveTableInterpolationMode>("Interpolation", InParams.OperatorSettings);
 
 			return MakeUnique<FMetasoundWaveTableEnvelopeNodeOperator>(InParams, InWaveTableReadRef, InPlayReadRef, InStopReadRef, InPauseReadRef, InDurationReadRef, InModeReadRef, InInterpReadRef);
 		}
 
 		FMetasoundWaveTableEnvelopeNodeOperator(
-			const FCreateOperatorParams& InParams,
+			const FBuildOperatorParams& InParams,
 			const FWaveTableReadRef& InWaveTableReadRef,
 			const FTriggerReadRef& InPlayReadRef,
 			const FTriggerReadRef& InStopReadRef,
diff --git a/Engine/Plugins/Runtime/Metasound/Source/MetasoundEngine/Private/MetasoundWaveTableEvaluateNode.cpp b/Engine/Plugins/Runtime/Metasound/Source/MetasoundEngine/Private/MetasoundWaveTableEvaluateNode.cpp
--- a/Engine/Plugins/Runtime/Metasound/Source/MetasoundEngine/Private/MetasoundWaveTableEvaluateNode.cpp
+++ b/Engine/Plugins/Runtime/Metasound/Source/MetasoundEngine/Private/MetasoundWaveTableEvaluateNode.cpp
@@ -77,22 +77,21 @@ namespace Metasound
 			return Metadata;
 		}
 
-		static TUniquePtr<IOperator> CreateOperator(const FCreateOperatorParams& InParams, TArray<TUniquePtr<IOperatorBuildError>>& OutErrors)
+		static TUniquePtr<IOperator> CreateOperator(const FBuildOperatorParams& InParams, FBuildResults& OutResults)
 		{
 			using namespace WaveTable;
 
-			const FInputVertexInterface& InputInterface = InParams.Node.GetVertexInterface().GetInputInterface();
-			const FDataReferenceCollection& InputCollection = InParams.InputDataReferences;
+			const FInputVertexInterfaceData& InputData = InParams.InputData;
 
-			FWaveTableReadRef InWaveTableReadRef = InputCollection.GetDataReadReferenceOrConstruct<FWaveTable>("WaveTable");
-			FFloatReadRef InInputReadRef = InputCollection.GetDataReadReferenceOrConstructWithVertexDefault<float>(InputInterface, "Input", InParams.OperatorSettings);
-			FEnumWaveTableInterpModeReadRef InInterpReadRef = InputCollection.GetDataReadReferenceOrConstructWithVertexDefault<FEnumWaveTableInterpolationMode>(InputInterface, "Interpolation", InParams.OperatorSettings);
+			FWaveTableReadRef InWaveTableReadRef = InputData.GetOrConstructDataReadReference<FWaveTable>("WaveTable");
+			FFloatReadRef InInputReadRef = InputData.GetOrCreateDefaultDataReadReference<float>("Input", InParams.OperatorSettings);
+			FEnumWaveTableInterpModeReadRef InInterpReadRef = InputData.GetOrCreateDefaultDataReadReference<FEnumWaveTableInterpolationMode>("Interpolation", InParams.OperatorSettings);
 
 			return MakeUnique<FMetasoundWaveTableEvaluateNodeOperator>(InParams, InWaveTableReadRef, InInputReadRef, InInterpReadRef);
 		}
 
 		FMetasoundWaveTableEvaluateNodeOperator(
-			const FCreateOperatorParams& InParams,
+			const FBuildOperatorParams& InParams,
 			const FWaveTableReadRef& InWaveTableReadRef,
 			const FFloatReadRef& InInputReadRef,
 			const FEnumWaveTableInterpModeReadRef& InInterpModeReadRef)
diff --git a/Engine/Plugins/Runtime/Metasound/Source/MetasoundEngine/Private/MetasoundWaveTableGetNode.cpp b/Engine/Plugins/Runtime/Metasound/Source/MetasoundEngine/Private/MetasoundWaveTableGetNode.cpp
--- a/Engine/Plugins/Runtime/Metasound/Source/MetasoundEngine/Private/MetasoundWaveTableGetNode.cpp
+++ b/Engine/Plugins/Runtime/Metasound/Source/MetasoundEngine/Private/MetasoundWaveTableGetNode.cpp
@@ -69,21 +69,20 @@ namespace Metasound
 			return Metadata;
 		}
 
-		static TUniquePtr<IOperator> CreateOperator(const FCreateOperatorParams& InParams, TArray<TUniquePtr<IOperatorBuildError>>& OutErrors)
+		static TUniquePtr<IOperator> CreateOperator(const FBuildOperatorParams& InParams, FBuildResults& OutResults)
 		{
 			using namespace Metasound;
+			
+			const FInputVertexInterfaceData& InputData = InParams.InputData;
 
-			const FInputVertexInterface& InputInterface = InParams.Node.GetVertexInterface().GetInputInterface();
-			const FDataReferenceCollection& InputCollection = InParams.InputDataReferences;
-
-			FWaveTableBankAssetReadRef WaveTableBankReadRef = InputCollection.GetDataReadReferenceOrConstruct<FWaveTableBankAsset>("WaveTableBank");
-			FFloatReadRef TableIndexReadRef = InputCollection.GetDataReadReferenceOrConstructWithVertexDefault<float>(InputInterface, "TableIndex", InParams.OperatorSettings);
+			FWaveTableBankAssetReadRef WaveTableBankReadRef = InputData.GetOrConstructDataReadReference<FWaveTableBankAsset>("WaveTableBank");
+			FFloatReadRef TableIndexReadRef = InputData.GetOrCreateDefaultDataReadReference<float>("TableIndex", InParams.OperatorSettings);
 
 			return MakeUnique<FMetasoundWaveTableGetNodeOperator>(InParams, WaveTableBankReadRef, TableIndexReadRef);
 		}
 
 		FMetasoundWaveTableGetNodeOperator(
-			const FCreateOperatorParams& InParams,
+			const FBuildOperatorParams& InParams,
 			const FWaveTableBankAssetReadRef& InWaveTableBankReadRef,
 			const FFloatReadRef& InTableIndexReadRef)
 			: WaveTableBankReadRef(InWaveTableBankReadRef)
diff --git a/Engine/Plugins/Runtime/Metasound/Source/MetasoundFrontend/Public/MetasoundArrayNodes.h b/Engine/Plugins/Runtime/Metasound/Source/MetasoundFrontend/Public/MetasoundArrayNodes.h
--- a/Engine/Plugins/Runtime/Metasound/Source/MetasoundFrontend/Public/MetasoundArrayNodes.h
+++ b/Engine/Plugins/Runtime/Metasound/Source/MetasoundFrontend/Public/MetasoundArrayNodes.h
@@ -115,15 +115,15 @@ namespace Metasound
 			return Metadata;
 		}
 
-		static TUniquePtr<IOperator> CreateOperator(const FCreateOperatorParams& InParams, TArray<TUniquePtr<IOperatorBuildError>>& OutErrors)
+		static TUniquePtr<IOperator> CreateOperator(const FBuildOperatorParams& InParams, FBuildResults& OutResults)
 		{
 			using namespace ArrayNodeVertexNames;
 			using namespace MetasoundArrayNodesPrivate;
 
-			const FInputVertexInterface& Inputs = InParams.Node.GetVertexInterface().GetInputInterface();
+			const FInputVertexInterfaceData& InputData = InParams.InputData;
 
 			// Get the input array or construct an empty one. 
-			FArrayDataReadReference Array = InParams.InputDataReferences.GetDataReadReferenceOrConstructWithVertexDefault<ArrayType>(Inputs, METASOUND_GET_PARAM_NAME(InputArray), InParams.OperatorSettings);
+			FArrayDataReadReference Array = InputData.GetOrCreateDefaultDataReadReference<ArrayType>(METASOUND_GET_PARAM_NAME(InputArray), InParams.OperatorSettings);
 
 			return MakeUnique<TArrayNumOperator>(Array);
 		}
@@ -247,21 +247,21 @@ namespace Metasound
 #endif
 		};
 
-		static TUniquePtr<IOperator> CreateOperator(const FCreateOperatorParams& InParams, TArray<TUniquePtr<IOperatorBuildError>>& OutErrors)
+		static TUniquePtr<IOperator> CreateOperator(const FBuildOperatorParams& InParams, FBuildResults& OutResults)
 		{
 			using namespace ArrayNodeVertexNames;
 			using namespace MetasoundArrayNodesPrivate;
 
-			const FInputVertexInterface& Inputs = InParams.Node.GetVertexInterface().GetInputInterface();
+			const FInputVertexInterfaceData& InputData = InParams.InputData;
 
 			// Input Trigger
-			TDataReadReference<FTrigger> Trigger = InParams.InputDataReferences.GetDataReadReferenceOrConstructWithVertexDefault<FTrigger>(Inputs, METASOUND_GET_PARAM_NAME(InputTriggerGet), InParams.OperatorSettings);
+			TDataReadReference<FTrigger> Trigger = InputData.GetOrCreateDefaultDataReadReference<FTrigger>(METASOUND_GET_PARAM_NAME(InputTriggerGet), InParams.OperatorSettings);
 			
 			// Input Array
-			FArrayDataReadReference Array = InParams.InputDataReferences.GetDataReadReferenceOrConstructWithVertexDefault<ArrayType>(Inputs, METASOUND_GET_PARAM_NAME(InputArray), InParams.OperatorSettings);
+			FArrayDataReadReference Array = InputData.GetOrCreateDefaultDataReadReference<ArrayType>(METASOUND_GET_PARAM_NAME(InputArray), InParams.OperatorSettings);
 
 			// Input Index
-			TDataReadReference<int32> Index = InParams.InputDataReferences.GetDataReadReferenceOrConstructWithVertexDefault<int32>(Inputs, METASOUND_GET_PARAM_NAME(InputIndex), InParams.OperatorSettings);
+			TDataReadReference<int32> Index = InputData.GetOrCreateDefaultDataReadReference<int32>(METASOUND_GET_PARAM_NAME(InputIndex), InParams.OperatorSettings);
 #if WITH_METASOUND_DEBUG_ENVIRONMENT
 			FString GraphName;
 			if (InParams.Environment.Contains<FString>(Frontend::SourceInterface::Environment::GraphName))
@@ -451,21 +451,21 @@ namespace Metasound
 #endif // WITH_METASOUND_DEBUG_ENVIRONMENT
 		};
 
-		static TUniquePtr<IOperator> CreateOperator(const FCreateOperatorParams& InParams, TArray<TUniquePtr<IOperatorBuildError>>& OutErrors)
+		static TUniquePtr<IOperator> CreateOperator(const FBuildOperatorParams& InParams, FBuildResults& OutResults)
 		{
 			using namespace ArrayNodeVertexNames;
 			using namespace MetasoundArrayNodesPrivate;
 
-			const FInputVertexInterface& Inputs = InParams.Node.GetVertexInterface().GetInputInterface();
+			const FInputVertexInterfaceData& InputData = InParams.InputData;
 			
-			TDataReadReference<FTrigger> Trigger = InParams.InputDataReferences.GetDataReadReferenceOrConstructWithVertexDefault<FTrigger>(Inputs, METASOUND_GET_PARAM_NAME(InputTriggerSet), InParams.OperatorSettings);
+			TDataReadReference<FTrigger> Trigger = InputData.GetOrCreateDefaultDataReadReference<FTrigger>(METASOUND_GET_PARAM_NAME(InputTriggerSet), InParams.OperatorSettings);
 
-			FArrayDataReadReference InitArray = InParams.InputDataReferences.GetDataReadReferenceOrConstructWithVertexDefault<ArrayType>(Inputs, InputInitialArrayName, InParams.OperatorSettings);
+			FArrayDataReadReference InitArray = InputData.GetOrCreateDefaultDataReadReference<ArrayType>(InputInitialArrayName, InParams.OperatorSettings);
 			FArrayDataWriteReference Array = TDataWriteReferenceFactory<ArrayType>::CreateExplicitArgs(InParams.OperatorSettings, *InitArray);
 
-			TDataReadReference<int32> Index = InParams.InputDataReferences.GetDataReadReferenceOrConstructWithVertexDefault<int32>(Inputs, METASOUND_GET_PARAM_NAME(InputIndex), InParams.OperatorSettings);
+			TDataReadReference<int32> Index = InputData.GetOrCreateDefaultDataReadReference<int32>(METASOUND_GET_PARAM_NAME(InputIndex), InParams.OperatorSettings);
 
-			TDataReadReference<ElementType> Value = InParams.InputDataReferences.GetDataReadReferenceOrConstructWithVertexDefault<ElementType>(Inputs, METASOUND_GET_PARAM_NAME(InputValue), InParams.OperatorSettings);
+			TDataReadReference<ElementType> Value = InputData.GetOrCreateDefaultDataReadReference<ElementType>(METASOUND_GET_PARAM_NAME(InputValue), InParams.OperatorSettings);
 #if WITH_METASOUND_DEBUG_ENVIRONMENT
 			FString GraphName;
 			if (InParams.Environment.Contains<FString>(Frontend::SourceInterface::Environment::GraphName))
@@ -630,17 +630,17 @@ namespace Metasound
 			return Metadata;
 		}
 
-		static TUniquePtr<IOperator> CreateOperator(const FCreateOperatorParams& InParams, TArray<TUniquePtr<IOperatorBuildError>>& OutErrors)
+		static TUniquePtr<IOperator> CreateOperator(const FBuildOperatorParams& InParams, FBuildResults& OutResults)
 		{
 			using namespace ArrayNodeVertexNames;
 			using namespace MetasoundArrayNodesPrivate;
 
-			const FInputVertexInterface& Inputs = InParams.Node.GetVertexInterface().GetInputInterface();
+			const FInputVertexInterfaceData& InputData = InParams.InputData;
 			
-			TDataReadReference<FTrigger> Trigger = InParams.InputDataReferences.GetDataReadReferenceOrConstructWithVertexDefault<FTrigger>(Inputs, METASOUND_GET_PARAM_NAME(InputTriggerGet), InParams.OperatorSettings);
+			TDataReadReference<FTrigger> Trigger = InputData.GetOrCreateDefaultDataReadReference<FTrigger>(METASOUND_GET_PARAM_NAME(InputTriggerGet), InParams.OperatorSettings);
 
-			FArrayDataReadReference LeftArray = InParams.InputDataReferences.GetDataReadReferenceOrConstructWithVertexDefault<ArrayType>(Inputs, METASOUND_GET_PARAM_NAME(InputLeftArray), InParams.OperatorSettings);
-			FArrayDataReadReference RightArray = InParams.InputDataReferences.GetDataReadReferenceOrConstructWithVertexDefault<ArrayType>(Inputs, METASOUND_GET_PARAM_NAME(InputRightArray), InParams.OperatorSettings);
+			FArrayDataReadReference LeftArray = InputData.GetOrCreateDefaultDataReadReference<ArrayType>(METASOUND_GET_PARAM_NAME(InputLeftArray), InParams.OperatorSettings);
+			FArrayDataReadReference RightArray = InputData.GetOrCreateDefaultDataReadReference<ArrayType>(METASOUND_GET_PARAM_NAME(InputRightArray), InParams.OperatorSettings);
 
 			return MakeUnique<TArrayConcatOperator>(Trigger, LeftArray, RightArray);
 		}
@@ -765,19 +765,19 @@ namespace Metasound
 			return Metadata;
 		}
 
-		static TUniquePtr<IOperator> CreateOperator(const FCreateOperatorParams& InParams, TArray<TUniquePtr<IOperatorBuildError>>& OutErrors)
+		static TUniquePtr<IOperator> CreateOperator(const FBuildOperatorParams& InParams, FBuildResults& OutResults)
 		{
 			using namespace ArrayNodeVertexNames;
 			using namespace MetasoundArrayNodesPrivate;
 
-			const FInputVertexInterface& Inputs = InParams.Node.GetVertexInterface().GetInputInterface();
+			const FInputVertexInterfaceData& InputData = InParams.InputData;
 			
-			TDataReadReference<FTrigger> Trigger = InParams.InputDataReferences.GetDataReadReferenceOrConstructWithVertexDefault<FTrigger>(Inputs, METASOUND_GET_PARAM_NAME(InputTriggerGet), InParams.OperatorSettings);
+			TDataReadReference<FTrigger> Trigger = InputData.GetOrCreateDefaultDataReadReference<FTrigger>(METASOUND_GET_PARAM_NAME(InputTriggerGet), InParams.OperatorSettings);
 
-			FArrayDataReadReference InArray = InParams.InputDataReferences.GetDataReadReferenceOrConstructWithVertexDefault<ArrayType>(Inputs, METASOUND_GET_PARAM_NAME(InputArray), InParams.OperatorSettings);
+			FArrayDataReadReference InArray = InputData.GetOrCreateDefaultDataReadReference<ArrayType>(METASOUND_GET_PARAM_NAME(InputArray), InParams.OperatorSettings);
 
-			TDataReadReference<int32> StartIndex = InParams.InputDataReferences.GetDataReadReferenceOrConstructWithVertexDefault<int32>(Inputs, METASOUND_GET_PARAM_NAME(InputStartIndex), InParams.OperatorSettings);
-			TDataReadReference<int32> EndIndex = InParams.InputDataReferences.GetDataReadReferenceOrConstructWithVertexDefault<int32>(Inputs, METASOUND_GET_PARAM_NAME(InputEndIndex), InParams.OperatorSettings);
+			TDataReadReference<int32> StartIndex = InputData.GetOrCreateDefaultDataReadReference<int32>(METASOUND_GET_PARAM_NAME(InputStartIndex), InParams.OperatorSettings);
+			TDataReadReference<int32> EndIndex = InputData.GetOrCreateDefaultDataReadReference<int32>(METASOUND_GET_PARAM_NAME(InputEndIndex), InParams.OperatorSettings);
 
 			TDataWriteReference<ArrayType> OutputArray = TDataWriteReferenceFactory<ArrayType>::CreateExplicitArgs(InParams.OperatorSettings);
 
diff --git a/Engine/Plugins/Runtime/Metasound/Source/MetasoundFrontend/Public/MetasoundArrayRandomNode.h b/Engine/Plugins/Runtime/Metasound/Source/MetasoundFrontend/Public/MetasoundArrayRandomNode.h
--- a/Engine/Plugins/Runtime/Metasound/Source/MetasoundFrontend/Public/MetasoundArrayRandomNode.h
+++ b/Engine/Plugins/Runtime/Metasound/Source/MetasoundFrontend/Public/MetasoundArrayRandomNode.h
@@ -144,26 +144,26 @@ namespace Metasound
 			return Metadata;
 		}
 
-		static TUniquePtr<IOperator> CreateOperator(const FCreateOperatorParams& InParams, TArray<TUniquePtr<IOperatorBuildError>>& OutErrors)
+		static TUniquePtr<IOperator> CreateOperator(const FBuildOperatorParams& InParams, FBuildResults& OutResults)
 		{
 			using namespace ArrayNodeRandomGetVertexNames;
 			using namespace MetasoundArrayNodesPrivate;
 
-			const FInputVertexInterface& Inputs = InParams.Node.GetVertexInterface().GetInputInterface();
+			const FInputVertexInterfaceData& InputData = InParams.InputData;
 
-			FTriggerReadRef InTriggerNext = InParams.InputDataReferences.GetDataReadReferenceOrConstructWithVertexDefault<FTrigger>(Inputs, METASOUND_GET_PARAM_NAME(InputTriggerNextValue), InParams.OperatorSettings);
-			FTriggerReadRef InTriggerReset = InParams.InputDataReferences.GetDataReadReferenceOrConstructWithVertexDefault<FTrigger>(Inputs, METASOUND_GET_PARAM_NAME(InputTriggerResetSeed), InParams.OperatorSettings);
-			FArrayDataReadReference InInputArray = InParams.InputDataReferences.GetDataReadReferenceOrConstructWithVertexDefault<ArrayType>(Inputs, METASOUND_GET_PARAM_NAME(InputRandomArray), InParams.OperatorSettings);
-			FArrayWeightReadReference InInputWeightsArray = InParams.InputDataReferences.GetDataReadReferenceOrConstructWithVertexDefault<WeightsArrayType>(Inputs, METASOUND_GET_PARAM_NAME(InputWeights), InParams.OperatorSettings);
-			FInt32ReadRef InSeedValue = InParams.InputDataReferences.GetDataReadReferenceOrConstructWithVertexDefault<int32>(Inputs, METASOUND_GET_PARAM_NAME(InputSeed), InParams.OperatorSettings);
-			FInt32ReadRef InNoRepeatOrder = InParams.InputDataReferences.GetDataReadReferenceOrConstructWithVertexDefault<int32>(Inputs, METASOUND_GET_PARAM_NAME(InputNoRepeatOrder), InParams.OperatorSettings);
-			FBoolReadRef bInEnableSharedState = InParams.InputDataReferences.GetDataReadReferenceOrConstructWithVertexDefault<bool>(Inputs, METASOUND_GET_PARAM_NAME(InputEnableSharedState), InParams.OperatorSettings);
+			FTriggerReadRef InTriggerNext = InputData.GetOrCreateDefaultDataReadReference<FTrigger>(METASOUND_GET_PARAM_NAME(InputTriggerNextValue), InParams.OperatorSettings);
+			FTriggerReadRef InTriggerReset = InputData.GetOrCreateDefaultDataReadReference<FTrigger>(METASOUND_GET_PARAM_NAME(InputTriggerResetSeed), InParams.OperatorSettings);
+			FArrayDataReadReference InInputArray = InputData.GetOrCreateDefaultDataReadReference<ArrayType>(METASOUND_GET_PARAM_NAME(InputRandomArray), InParams.OperatorSettings);
+			FArrayWeightReadReference InInputWeightsArray = InputData.GetOrCreateDefaultDataReadReference<WeightsArrayType>(METASOUND_GET_PARAM_NAME(InputWeights), InParams.OperatorSettings);
+			FInt32ReadRef InSeedValue = InputData.GetOrCreateDefaultDataReadReference<int32>(METASOUND_GET_PARAM_NAME(InputSeed), InParams.OperatorSettings);
+			FInt32ReadRef InNoRepeatOrder = InputData.GetOrCreateDefaultDataReadReference<int32>(METASOUND_GET_PARAM_NAME(InputNoRepeatOrder), InParams.OperatorSettings);
+			FBoolReadRef bInEnableSharedState = InputData.GetOrCreateDefaultDataReadReference<bool>(METASOUND_GET_PARAM_NAME(InputEnableSharedState), InParams.OperatorSettings);
 
 			return MakeUnique<TArrayRandomGetOperator<ArrayType>>(InParams, InTriggerNext, InTriggerReset, InInputArray, InInputWeightsArray, InSeedValue, InNoRepeatOrder, bInEnableSharedState);
 		}
 
 		TArrayRandomGetOperator(
-			const FCreateOperatorParams& InParams,
+			const FBuildOperatorParams& InParams,
 			const FTriggerReadRef& InTriggerNext,
 			const FTriggerReadRef& InTriggerReset,
 			const FArrayDataReadReference& InInputArray,
diff --git a/Engine/Plugins/Runtime/Metasound/Source/MetasoundFrontend/Public/MetasoundArrayShuffleNode.h b/Engine/Plugins/Runtime/Metasound/Source/MetasoundFrontend/Public/MetasoundArrayShuffleNode.h
--- a/Engine/Plugins/Runtime/Metasound/Source/MetasoundFrontend/Public/MetasoundArrayShuffleNode.h
+++ b/Engine/Plugins/Runtime/Metasound/Source/MetasoundFrontend/Public/MetasoundArrayShuffleNode.h
@@ -183,26 +183,26 @@ namespace Metasound
 			return Metadata;
 		}
 
-		static TUniquePtr<IOperator> CreateOperator(const FCreateOperatorParams& InParams, TArray<TUniquePtr<IOperatorBuildError>>& OutErrors)
+		static TUniquePtr<IOperator> CreateOperator(const FBuildOperatorParams& InParams, FBuildResults& OutResults)
 		{
 			using namespace ArrayNodeShuffleVertexNames;
 			using namespace MetasoundArrayNodesPrivate;
 
-			const FInputVertexInterface& Inputs = InParams.Node.GetVertexInterface().GetInputInterface();
+			const FInputVertexInterfaceData& InputData = InParams.InputData;
 
-			TDataReadReference<FTrigger> InTriggerNext = InParams.InputDataReferences.GetDataReadReferenceOrConstructWithVertexDefault<FTrigger>(Inputs, METASOUND_GET_PARAM_NAME(InputTriggerNext), InParams.OperatorSettings);
-			TDataReadReference<FTrigger> InTriggerShuffle = InParams.InputDataReferences.GetDataReadReferenceOrConstructWithVertexDefault<FTrigger>(Inputs, METASOUND_GET_PARAM_NAME(InputTriggerShuffle), InParams.OperatorSettings);
-			TDataReadReference<FTrigger> InTriggerReset = InParams.InputDataReferences.GetDataReadReferenceOrConstructWithVertexDefault<FTrigger>(Inputs, METASOUND_GET_PARAM_NAME(InputTriggerReset), InParams.OperatorSettings);
-			FArrayDataReadReference InInputArray = InParams.InputDataReferences.GetDataReadReferenceOrConstructWithVertexDefault<ArrayType>(Inputs, METASOUND_GET_PARAM_NAME(InputShuffleArray), InParams.OperatorSettings);
-			TDataReadReference<int32> InSeedValue = InParams.InputDataReferences.GetDataReadReferenceOrConstructWithVertexDefault<int32>(Inputs, METASOUND_GET_PARAM_NAME(InputShuffleSeed), InParams.OperatorSettings);
-			TDataReadReference<bool> bInAutoShuffle = InParams.InputDataReferences.GetDataReadReferenceOrConstructWithVertexDefault<bool>(Inputs, METASOUND_GET_PARAM_NAME(InputAutoShuffle), InParams.OperatorSettings);
-			TDataReadReference<bool> bInEnableSharedState = InParams.InputDataReferences.GetDataReadReferenceOrConstructWithVertexDefault<bool>(Inputs, METASOUND_GET_PARAM_NAME(InputShuffleEnableSharedState), InParams.OperatorSettings);
+			TDataReadReference<FTrigger> InTriggerNext = InputData.GetOrCreateDefaultDataReadReference<FTrigger>(METASOUND_GET_PARAM_NAME(InputTriggerNext), InParams.OperatorSettings);
+			TDataReadReference<FTrigger> InTriggerShuffle = InputData.GetOrCreateDefaultDataReadReference<FTrigger>(METASOUND_GET_PARAM_NAME(InputTriggerShuffle), InParams.OperatorSettings);
+			TDataReadReference<FTrigger> InTriggerReset = InputData.GetOrCreateDefaultDataReadReference<FTrigger>(METASOUND_GET_PARAM_NAME(InputTriggerReset), InParams.OperatorSettings);
+			FArrayDataReadReference InInputArray = InputData.GetOrCreateDefaultDataReadReference<ArrayType>(METASOUND_GET_PARAM_NAME(InputShuffleArray), InParams.OperatorSettings);
+			TDataReadReference<int32> InSeedValue = InputData.GetOrCreateDefaultDataReadReference<int32>(METASOUND_GET_PARAM_NAME(InputShuffleSeed), InParams.OperatorSettings);
+			TDataReadReference<bool> bInAutoShuffle = InputData.GetOrCreateDefaultDataReadReference<bool>(METASOUND_GET_PARAM_NAME(InputAutoShuffle), InParams.OperatorSettings);
+			TDataReadReference<bool> bInEnableSharedState = InputData.GetOrCreateDefaultDataReadReference<bool>(METASOUND_GET_PARAM_NAME(InputShuffleEnableSharedState), InParams.OperatorSettings);
 
 			return MakeUnique<TArrayShuffleOperator>(InParams, InTriggerNext, InTriggerShuffle, InTriggerReset, InInputArray, InSeedValue, bInAutoShuffle, bInEnableSharedState);
 		}
 
 		TArrayShuffleOperator(
-			const FCreateOperatorParams& InParams,
+			const FBuildOperatorParams& InParams,
 			const TDataReadReference<FTrigger>& InTriggerNext, 
 			const TDataReadReference<FTrigger>& InTriggerShuffle,
 			const TDataReadReference<FTrigger>& InTriggerReset,
diff --git a/Engine/Plugins/Runtime/Metasound/Source/MetasoundStandardNodes/Private/MetasoundADEnvelopeNode.cpp b/Engine/Plugins/Runtime/Metasound/Source/MetasoundStandardNodes/Private/MetasoundADEnvelopeNode.cpp
--- a/Engine/Plugins/Runtime/Metasound/Source/MetasoundStandardNodes/Private/MetasoundADEnvelopeNode.cpp
+++ b/Engine/Plugins/Runtime/Metasound/Source/MetasoundStandardNodes/Private/MetasoundADEnvelopeNode.cpp
@@ -285,24 +285,24 @@ namespace Metasound
 			return Metadata;
 		}
 
-		static TUniquePtr<IOperator> CreateOperator(const FCreateOperatorParams& InParams, TArray<TUniquePtr<IOperatorBuildError>>& OutErrors)
+		static TUniquePtr<IOperator> CreateOperator(const FBuildOperatorParams& InParams, FBuildResults& OutResults)
 		{
 			using namespace ADEnvelopeVertexNames;
 
-			const FInputVertexInterface& InputInterface = GetDefaultInterface().GetInputInterface();
+			const FInputVertexInterfaceData& InputData = InParams.InputData;
 
-			FTriggerReadRef TriggerIn = InParams.InputDataReferences.GetDataReadReferenceOrConstruct<FTrigger>(METASOUND_GET_PARAM_NAME(InputTrigger), InParams.OperatorSettings);
-			FTimeReadRef AttackTime = InParams.InputDataReferences.GetDataReadReferenceOrConstructWithVertexDefault<FTime>(InputInterface, METASOUND_GET_PARAM_NAME(InputAttackTime), InParams.OperatorSettings);
-			FTimeReadRef DecayTime = InParams.InputDataReferences.GetDataReadReferenceOrConstructWithVertexDefault<FTime>(InputInterface, METASOUND_GET_PARAM_NAME(InputDecayTime), InParams.OperatorSettings);
-			FFloatReadRef AttackCurveFactor = InParams.InputDataReferences.GetDataReadReferenceOrConstructWithVertexDefault<float>(InputInterface, METASOUND_GET_PARAM_NAME(InputAttackCurve), InParams.OperatorSettings);
-			FFloatReadRef DecayCurveFactor = InParams.InputDataReferences.GetDataReadReferenceOrConstructWithVertexDefault<float>(InputInterface, METASOUND_GET_PARAM_NAME(InputDecayCurve), InParams.OperatorSettings);
-			FBoolReadRef bLooping = InParams.InputDataReferences.GetDataReadReferenceOrConstructWithVertexDefault<bool>(InputInterface, METASOUND_GET_PARAM_NAME(InputLooping), InParams.OperatorSettings);
-			FBoolReadRef bHardReset = InParams.InputDataReferences.GetDataReadReferenceOrConstructWithVertexDefault<bool>(InputInterface, METASOUND_GET_PARAM_NAME(InputHardReset), InParams.OperatorSettings);
+			FTriggerReadRef TriggerIn = InputData.GetOrConstructDataReadReference<FTrigger>(METASOUND_GET_PARAM_NAME(InputTrigger), InParams.OperatorSettings);
+			FTimeReadRef AttackTime = InputData.GetOrCreateDefaultDataReadReference<FTime>(METASOUND_GET_PARAM_NAME(InputAttackTime), InParams.OperatorSettings);
+			FTimeReadRef DecayTime = InputData.GetOrCreateDefaultDataReadReference<FTime>(METASOUND_GET_PARAM_NAME(InputDecayTime), InParams.OperatorSettings);
+			FFloatReadRef AttackCurveFactor = InputData.GetOrCreateDefaultDataReadReference<float>(METASOUND_GET_PARAM_NAME(InputAttackCurve), InParams.OperatorSettings);
+			FFloatReadRef DecayCurveFactor = InputData.GetOrCreateDefaultDataReadReference<float>(METASOUND_GET_PARAM_NAME(InputDecayCurve), InParams.OperatorSettings);
+			FBoolReadRef bLooping = InputData.GetOrCreateDefaultDataReadReference<bool>(METASOUND_GET_PARAM_NAME(InputLooping), InParams.OperatorSettings);
+			FBoolReadRef bHardReset = InputData.GetOrCreateDefaultDataReadReference<bool>(METASOUND_GET_PARAM_NAME(InputHardReset), InParams.OperatorSettings);
 
 			return MakeUnique<TADEnvelopeNodeOperator<ValueType>>(InParams, TriggerIn, AttackTime, DecayTime, AttackCurveFactor, DecayCurveFactor, bLooping, bHardReset);
 		}
 
-		TADEnvelopeNodeOperator(const FCreateOperatorParams& InParams,
+		TADEnvelopeNodeOperator(const FBuildOperatorParams& InParams,
 			const FTriggerReadRef& InTriggerIn,
 			const FTimeReadRef& InAttackTime,
 			const FTimeReadRef& InDecayTime,
diff --git a/Engine/Plugins/Runtime/Metasound/Source/MetasoundStandardNodes/Private/MetasoundADSREnvelopeNode.cpp b/Engine/Plugins/Runtime/Metasound/Source/MetasoundStandardNodes/Private/MetasoundADSREnvelopeNode.cpp
--- a/Engine/Plugins/Runtime/Metasound/Source/MetasoundStandardNodes/Private/MetasoundADSREnvelopeNode.cpp
+++ b/Engine/Plugins/Runtime/Metasound/Source/MetasoundStandardNodes/Private/MetasoundADSREnvelopeNode.cpp
@@ -565,21 +565,21 @@ namespace Metasound
 		{
 			using namespace ADSREnvelopeVertexNames;
 
-			const FInputVertexInterface& InputInterface = GetDefaultInterface().GetInputInterface();
+			const FInputVertexInterfaceData& InputData = InParams.InputData;
 
 			FOperatorArgs Args 
 			{
 				InParams,
-				InParams.InputData.GetOrConstructDataReadReference<FTrigger>(METASOUND_GET_PARAM_NAME(InputAttackTrigger), InParams.OperatorSettings),
-				InParams.InputData.GetOrConstructDataReadReference<FTrigger>(METASOUND_GET_PARAM_NAME(InputReleaseTrigger), InParams.OperatorSettings),
-				InParams.InputData.GetOrCreateDefaultDataReadReference<FTime>(METASOUND_GET_PARAM_NAME(InputAttackTime), InParams.OperatorSettings),
-				InParams.InputData.GetOrCreateDefaultDataReadReference<FTime>(METASOUND_GET_PARAM_NAME(InputDecayTime), InParams.OperatorSettings),
-				InParams.InputData.GetOrCreateDefaultDataReadReference<float>(METASOUND_GET_PARAM_NAME(InputSustainLevel), InParams.OperatorSettings),
-				InParams.InputData.GetOrCreateDefaultDataReadReference<FTime>(METASOUND_GET_PARAM_NAME(InputReleaseTime), InParams.OperatorSettings),
-				InParams.InputData.GetOrCreateDefaultDataReadReference<float>(METASOUND_GET_PARAM_NAME(InputAttackCurve), InParams.OperatorSettings),
-				InParams.InputData.GetOrCreateDefaultDataReadReference<float>(METASOUND_GET_PARAM_NAME(InputDecayCurve), InParams.OperatorSettings),
-				InParams.InputData.GetOrCreateDefaultDataReadReference<float>(METASOUND_GET_PARAM_NAME(InputReleaseCurve), InParams.OperatorSettings),
-				InParams.InputData.GetOrCreateDefaultDataReadReference<bool>(METASOUND_GET_PARAM_NAME(InputHardReset), InParams.OperatorSettings)
+				InputData.GetOrConstructDataReadReference<FTrigger>(METASOUND_GET_PARAM_NAME(InputAttackTrigger), InParams.OperatorSettings),
+				InputData.GetOrConstructDataReadReference<FTrigger>(METASOUND_GET_PARAM_NAME(InputReleaseTrigger), InParams.OperatorSettings),
+				InputData.GetOrCreateDefaultDataReadReference<FTime>(METASOUND_GET_PARAM_NAME(InputAttackTime), InParams.OperatorSettings),
+				InputData.GetOrCreateDefaultDataReadReference<FTime>(METASOUND_GET_PARAM_NAME(InputDecayTime), InParams.OperatorSettings),
+				InputData.GetOrCreateDefaultDataReadReference<float>(METASOUND_GET_PARAM_NAME(InputSustainLevel), InParams.OperatorSettings),
+				InputData.GetOrCreateDefaultDataReadReference<FTime>(METASOUND_GET_PARAM_NAME(InputReleaseTime), InParams.OperatorSettings),
+				InputData.GetOrCreateDefaultDataReadReference<float>(METASOUND_GET_PARAM_NAME(InputAttackCurve), InParams.OperatorSettings),
+				InputData.GetOrCreateDefaultDataReadReference<float>(METASOUND_GET_PARAM_NAME(InputDecayCurve), InParams.OperatorSettings),
+				InputData.GetOrCreateDefaultDataReadReference<float>(METASOUND_GET_PARAM_NAME(InputReleaseCurve), InParams.OperatorSettings),
+				InputData.GetOrCreateDefaultDataReadReference<bool>(METASOUND_GET_PARAM_NAME(InputHardReset), InParams.OperatorSettings)
 			};
 
 			return MakeUnique<TADSREnvelopeNodeOperator<EnvelopeClass, ValueType>>(Args);
diff --git a/Engine/Plugins/Runtime/Metasound/Source/MetasoundStandardNodes/Private/MetasoundAbsoluteValueNodes.cpp b/Engine/Plugins/Runtime/Metasound/Source/MetasoundStandardNodes/Private/MetasoundAbsoluteValueNodes.cpp
--- a/Engine/Plugins/Runtime/Metasound/Source/MetasoundStandardNodes/Private/MetasoundAbsoluteValueNodes.cpp
+++ b/Engine/Plugins/Runtime/Metasound/Source/MetasoundStandardNodes/Private/MetasoundAbsoluteValueNodes.cpp
@@ -58,11 +58,11 @@ namespace Metasound
 				OutAbs = FMath::Abs(In);
 			}
 
-			static TDataReadReference<int32> CreateInRef(const FCreateOperatorParams& InParams)
+			static TDataReadReference<int32> CreateInRef(const FBuildOperatorParams& InParams)
 			{
 				using namespace AbsVertexNames;
-				const FInputVertexInterface& InputInterface = InParams.Node.GetVertexInterface().GetInputInterface();
-				return InParams.InputDataReferences.GetDataReadReferenceOrConstructWithVertexDefault<int32>(InputInterface, METASOUND_GET_PARAM_NAME(InputValue), InParams.OperatorSettings);
+				const FInputVertexInterfaceData& InputData = InParams.InputData;
+				return InputData.GetOrCreateDefaultDataReadReference<int32>(METASOUND_GET_PARAM_NAME(InputValue), InParams.OperatorSettings);
 			}
 		};
 
@@ -74,11 +74,11 @@ namespace Metasound
 				OutAbs = FMath::Abs(In);
 			}
 
-			static TDataReadReference<float> CreateInRef(const FCreateOperatorParams& InParams)
+			static TDataReadReference<float> CreateInRef(const FBuildOperatorParams& InParams)
 			{
 				using namespace AbsVertexNames;
-				const FInputVertexInterface& InputInterface = InParams.Node.GetVertexInterface().GetInputInterface();
-				return InParams.InputDataReferences.GetDataReadReferenceOrConstructWithVertexDefault<float>(InputInterface, METASOUND_GET_PARAM_NAME(InputValue), InParams.OperatorSettings);
+				const FInputVertexInterfaceData& InputData = InParams.InputData;
+				return InputData.GetOrCreateDefaultDataReadReference<float>(METASOUND_GET_PARAM_NAME(InputValue), InParams.OperatorSettings);
 			}
 
 			static bool IsAudioBuffer() { return false; }
@@ -92,11 +92,11 @@ namespace Metasound
 				OutAbs = FTime(FMath::Abs(In.GetSeconds()));
 			}
 
-			static TDataReadReference<FTime> CreateInRef(const FCreateOperatorParams& InParams)
+			static TDataReadReference<FTime> CreateInRef(const FBuildOperatorParams& InParams)
 			{
 				using namespace AbsVertexNames;
-				const FInputVertexInterface& InputInterface = InParams.Node.GetVertexInterface().GetInputInterface();
-				return InParams.InputDataReferences.GetDataReadReferenceOrConstructWithVertexDefault<FTime>(InputInterface, METASOUND_GET_PARAM_NAME(InputValue), InParams.OperatorSettings);
+				const FInputVertexInterfaceData& InputData = InParams.InputData;
+				return InputData.GetOrCreateDefaultDataReadReference<FTime>(METASOUND_GET_PARAM_NAME(InputValue), InParams.OperatorSettings);
 			}
 
 			static bool IsAudioBuffer() { return false; }
@@ -113,10 +113,11 @@ namespace Metasound
 				Audio::ArrayAbs(InView, OutAbsView);
 			}
 
-			static TDataReadReference<FAudioBuffer> CreateInRef(const FCreateOperatorParams& InParams)
+			static TDataReadReference<FAudioBuffer> CreateInRef(const FBuildOperatorParams& InParams)
 			{
 				using namespace AbsVertexNames;
-				return InParams.InputDataReferences.GetDataReadReferenceOrConstruct<FAudioBuffer>(METASOUND_GET_PARAM_NAME(InputValue), InParams.OperatorSettings);
+				const FInputVertexInterfaceData& InputData = InParams.InputData;
+				return InputData.GetOrConstructDataReadReference<FAudioBuffer>(METASOUND_GET_PARAM_NAME(InputValue), InParams.OperatorSettings);
 			}
 		};
 	}
@@ -159,7 +160,7 @@ namespace Metasound
 			return Metadata;
 		}
 
-		static TUniquePtr<IOperator> CreateOperator(const FCreateOperatorParams& InParams, TArray<TUniquePtr<IOperatorBuildError>>& OutErrors)
+		static TUniquePtr<IOperator> CreateOperator(const FBuildOperatorParams& InParams, FBuildResults& OutResults)
 		{
 			using namespace AbsVertexNames;
 			using namespace MetasoundAbsNodePrivate;
diff --git a/Engine/Plugins/Runtime/Metasound/Source/MetasoundStandardNodes/Private/MetasoundAdditiveSynthNode.cpp b/Engine/Plugins/Runtime/Metasound/Source/MetasoundStandardNodes/Private/MetasoundAdditiveSynthNode.cpp
--- a/Engine/Plugins/Runtime/Metasound/Source/MetasoundStandardNodes/Private/MetasoundAdditiveSynthNode.cpp
+++ b/Engine/Plugins/Runtime/Metasound/Source/MetasoundStandardNodes/Private/MetasoundAdditiveSynthNode.cpp
@@ -81,19 +81,17 @@ namespace Metasound
 			return Interface;
 		}
 
-		static TUniquePtr<IOperator> CreateOperator(const FCreateOperatorParams& InParams, FBuildErrorArray& OutErrors)
+		static TUniquePtr<IOperator> CreateOperator(const FBuildOperatorParams& InParams, FBuildResults& OutResults)
 		{
 			using namespace AdditiveSynthVertexNames;
 
-			const FDataReferenceCollection& InputCollection = InParams.InputDataReferences;
-			FVertexInterface Interface = GetVertexInterface();
-			const FInputVertexInterface& InputInterface = Interface.GetInputInterface();
+			const FInputVertexInterfaceData& InputData = InParams.InputData;
 
-			FFloatReadRef InputBaseFrequency = InputCollection.GetDataReadReferenceOrConstructWithVertexDefault<float>(InputInterface, METASOUND_GET_PARAM_NAME(BaseFrequency), InParams.OperatorSettings);
-			FInputFloatArrayReadRef InputHarmonicMultipliers = InputCollection.GetDataReadReferenceOrConstruct<FInputFloatArrayType>(METASOUND_GET_PARAM_NAME(HarmonicMultipliers));
-			FInputFloatArrayReadRef InputAmplitudes = InputCollection.GetDataReadReferenceOrConstruct<FInputFloatArrayType>(METASOUND_GET_PARAM_NAME(Amplitudes));
-			FInputFloatArrayReadRef InputPhases = InputCollection.GetDataReadReferenceOrConstruct<FInputFloatArrayType>(METASOUND_GET_PARAM_NAME(Phases));
-			FInputFloatArrayReadRef InputPanAmounts = InputCollection.GetDataReadReferenceOrConstruct<FInputFloatArrayType>(METASOUND_GET_PARAM_NAME(PanAmounts));
+			FFloatReadRef InputBaseFrequency = InputData.GetOrCreateDefaultDataReadReference<float>(METASOUND_GET_PARAM_NAME(BaseFrequency), InParams.OperatorSettings);
+			FInputFloatArrayReadRef InputHarmonicMultipliers = InputData.GetOrConstructDataReadReference<FInputFloatArrayType>(METASOUND_GET_PARAM_NAME(HarmonicMultipliers));
+			FInputFloatArrayReadRef InputAmplitudes = InputData.GetOrConstructDataReadReference<FInputFloatArrayType>(METASOUND_GET_PARAM_NAME(Amplitudes));
+			FInputFloatArrayReadRef InputPhases = InputData.GetOrConstructDataReadReference<FInputFloatArrayType>(METASOUND_GET_PARAM_NAME(Phases));
+			FInputFloatArrayReadRef InputPanAmounts = InputData.GetOrConstructDataReadReference<FInputFloatArrayType>(METASOUND_GET_PARAM_NAME(PanAmounts));
 
 			return MakeUnique<FAdditiveSynthOperator>(InParams.OperatorSettings, InputBaseFrequency, InputHarmonicMultipliers, InputAmplitudes, InputPhases, InputPanAmounts);
 		}
diff --git a/Engine/Plugins/Runtime/Metasound/Source/MetasoundStandardNodes/Private/MetasoundBPMToSecondsNode.cpp b/Engine/Plugins/Runtime/Metasound/Source/MetasoundStandardNodes/Private/MetasoundBPMToSecondsNode.cpp
--- a/Engine/Plugins/Runtime/Metasound/Source/MetasoundStandardNodes/Private/MetasoundBPMToSecondsNode.cpp
+++ b/Engine/Plugins/Runtime/Metasound/Source/MetasoundStandardNodes/Private/MetasoundBPMToSecondsNode.cpp
@@ -29,9 +29,9 @@ namespace Metasound
 
 		static const FNodeClassMetadata& GetNodeInfo();
 		static const FVertexInterface& GetVertexInterface();
-		static TUniquePtr<IOperator> CreateOperator(const FCreateOperatorParams& InParams, FBuildErrorArray& OutErrors);
+		static TUniquePtr<IOperator> CreateOperator(const FBuildOperatorParams& InParams, FBuildResults& OutResults);
 
-		FBPMToSecondsOperator(const FCreateOperatorParams& InParams, const FFloatReadRef& InBPM, const FFloatReadRef& InBeatMultiplier, const FFloatReadRef& InDivOfWholeNote);
+		FBPMToSecondsOperator(const FBuildOperatorParams& InParams, const FFloatReadRef& InBPM, const FFloatReadRef& InBeatMultiplier, const FFloatReadRef& InDivOfWholeNote);
 
 
 		virtual void BindInputs(FInputVertexInterfaceData& InOutVertexData) override;
@@ -64,7 +64,7 @@ namespace Metasound
 		static constexpr float MinDivOfWholeNote = 1.f;
 	};
 
-	FBPMToSecondsOperator::FBPMToSecondsOperator(const FCreateOperatorParams& InParams, const FFloatReadRef& InBPM, const FFloatReadRef& InBeatMultiplier, const FFloatReadRef& InDivOfWholeNote)
+	FBPMToSecondsOperator::FBPMToSecondsOperator(const FBuildOperatorParams& InParams, const FFloatReadRef& InBPM, const FFloatReadRef& InBeatMultiplier, const FFloatReadRef& InDivOfWholeNote)
 		: BPM(InBPM)
 		, BeatMultiplier(InBeatMultiplier)
 		, DivOfWholeNote(InDivOfWholeNote)
@@ -180,16 +180,15 @@ namespace Metasound
 		return Info;
 	}
 
-	TUniquePtr<IOperator> FBPMToSecondsOperator::CreateOperator(const FCreateOperatorParams& InParams, FBuildErrorArray& OutErrors)
+	TUniquePtr<IOperator> FBPMToSecondsOperator::CreateOperator(const FBuildOperatorParams& InParams, FBuildResults& OutResults)
 	{
 		using namespace BPMToSecondsVertexNames;
 
-		const FDataReferenceCollection& InputCollection = InParams.InputDataReferences;
-		const FInputVertexInterface& InputInterface = GetVertexInterface().GetInputInterface();
+		const FInputVertexInterfaceData& InputData = InParams.InputData;
 
-		FFloatReadRef InBPM = InputCollection.GetDataReadReferenceOrConstructWithVertexDefault<float>(InputInterface, METASOUND_GET_PARAM_NAME(InputBPM), InParams.OperatorSettings);
-		FFloatReadRef InBeatMultiplier = InputCollection.GetDataReadReferenceOrConstructWithVertexDefault<float>(InputInterface, METASOUND_GET_PARAM_NAME(InputBeatMultiplier), InParams.OperatorSettings);
-		FFloatReadRef InDivOfWholeNote = InputCollection.GetDataReadReferenceOrConstructWithVertexDefault<float>(InputInterface, METASOUND_GET_PARAM_NAME(InputDivOfWholeNote), InParams.OperatorSettings);
+		FFloatReadRef InBPM = InputData.GetOrCreateDefaultDataReadReference<float>(METASOUND_GET_PARAM_NAME(InputBPM), InParams.OperatorSettings);
+		FFloatReadRef InBeatMultiplier = InputData.GetOrCreateDefaultDataReadReference<float>(METASOUND_GET_PARAM_NAME(InputBeatMultiplier), InParams.OperatorSettings);
+		FFloatReadRef InDivOfWholeNote = InputData.GetOrCreateDefaultDataReadReference<float>(METASOUND_GET_PARAM_NAME(InputDivOfWholeNote), InParams.OperatorSettings);
 
 		return MakeUnique<FBPMToSecondsOperator>(InParams, InBPM, InBeatMultiplier, InDivOfWholeNote);
 	}
diff --git a/Engine/Plugins/Runtime/Metasound/Source/MetasoundStandardNodes/Private/MetasoundBandsplitterNode.cpp b/Engine/Plugins/Runtime/Metasound/Source/MetasoundStandardNodes/Private/MetasoundBandsplitterNode.cpp
--- a/Engine/Plugins/Runtime/Metasound/Source/MetasoundStandardNodes/Private/MetasoundBandsplitterNode.cpp
+++ b/Engine/Plugins/Runtime/Metasound/Source/MetasoundStandardNodes/Private/MetasoundBandsplitterNode.cpp
@@ -238,27 +238,26 @@ namespace Metasound
 			return DefaultInterface;
 		}
 
-		static TUniquePtr<IOperator> CreateOperator(const FCreateOperatorParams& InParams, FBuildErrorArray& OutErrors)
+		static TUniquePtr<IOperator> CreateOperator(const FBuildOperatorParams& InParams, FBuildResults& OutResults)
 		{
 			using namespace BandSplitterNode;
-
-			const FDataReferenceCollection& Inputs = InParams.InputDataReferences;
-			const FInputVertexInterface& InputInterface = InParams.Node.GetVertexInterface().GetInputInterface();
+			
+			const FInputVertexInterfaceData& InputData = InParams.InputData;
 
 			TArray<FAudioBufferReadRef> InputBuffers;
 			TArray<FFloatReadRef> InputCrossovers;
 
-			FEnumBandSplitterFilterOrderReadRef FilterOrderIn = Inputs.GetDataReadReferenceOrConstructWithVertexDefault<FEnumBandSplitterFilterOrder>(InputInterface, METASOUND_GET_PARAM_NAME(InputFilterOrder), InParams.OperatorSettings);
-			FBoolReadRef bPhaseCompensateIn = Inputs.GetDataReadReferenceOrConstructWithVertexDefault<bool>(InputInterface, METASOUND_GET_PARAM_NAME(InputPhaseCompensate), InParams.OperatorSettings);
+			FEnumBandSplitterFilterOrderReadRef FilterOrderIn = InputData.GetOrCreateDefaultDataReadReference<FEnumBandSplitterFilterOrder>(METASOUND_GET_PARAM_NAME(InputFilterOrder), InParams.OperatorSettings);
+			FBoolReadRef bPhaseCompensateIn = InputData.GetOrCreateDefaultDataReadReference<bool>(METASOUND_GET_PARAM_NAME(InputPhaseCompensate), InParams.OperatorSettings);
 
 			for (uint32 Chan = 0; Chan < NumChannels; Chan++)
 			{
-				InputBuffers.Add(Inputs.GetDataReadReferenceOrConstruct<FAudioBuffer>(GetAudioInputName(Chan), InParams.OperatorSettings));
+				InputBuffers.Add(InputData.GetOrConstructDataReadReference<FAudioBuffer>(GetAudioInputName(Chan), InParams.OperatorSettings));
 			}
 
 			for (uint32 Band = 0; Band < NumBands - 1; Band++)
 			{
-				InputCrossovers.Add(Inputs.GetDataReadReferenceOrConstructWithVertexDefault<float>(InputInterface, GetCrossoverInputName(Band), InParams.OperatorSettings));
+				InputCrossovers.Add(InputData.GetOrCreateDefaultDataReadReference<float>(GetCrossoverInputName(Band), InParams.OperatorSettings));
 			}
 
 			return MakeUnique<TBandSplitterOperator<NumBands, NumChannels>>(InParams.OperatorSettings, MoveTemp(InputBuffers), MoveTemp(InputCrossovers), FilterOrderIn, bPhaseCompensateIn);
diff --git a/Engine/Plugins/Runtime/Metasound/Source/MetasoundStandardNodes/Private/MetasoundBasicFilters.cpp b/Engine/Plugins/Runtime/Metasound/Source/MetasoundStandardNodes/Private/MetasoundBasicFilters.cpp
--- a/Engine/Plugins/Runtime/Metasound/Source/MetasoundStandardNodes/Private/MetasoundBasicFilters.cpp
+++ b/Engine/Plugins/Runtime/Metasound/Source/MetasoundStandardNodes/Private/MetasoundBasicFilters.cpp
@@ -95,7 +95,7 @@ namespace Metasound
 
 		static FVertexInterface DeclareVertexInterface();
 
-		static TUniquePtr<IOperator> CreateOperator(const FCreateOperatorParams& InParams, FBuildErrorArray& OutErrors);
+		static TUniquePtr<IOperator> CreateOperator(const FBuildOperatorParams& InParams, FBuildResults& OutResults);
 
 		virtual void BindInputs(FInputVertexInterfaceData& InOutVertexData) override
 		{
@@ -214,17 +214,14 @@ namespace Metasound
 		return Interface;
 	}
 
-	TUniquePtr<IOperator> FLadderFilterOperator::CreateOperator(const FCreateOperatorParams& InParams, FBuildErrorArray& OutErrors)
+	TUniquePtr<IOperator> FLadderFilterOperator::CreateOperator(const FBuildOperatorParams& InParams, FBuildResults& OutResults)
 	{
-
-		const FDataReferenceCollection& InputDataRefs = InParams.InputDataReferences;
-		FVertexInterface Interface = DeclareVertexInterface();
-		const FInputVertexInterface& InputInterface = Interface.GetInputInterface();
+		const FInputVertexInterfaceData& InputData = InParams.InputData;
 
 		// inputs
-		FAudioBufferReadRef AudioIn = InputDataRefs.GetDataReadReferenceOrConstruct<FAudioBuffer>(METASOUND_GET_PARAM_NAME(ParamAudioInput), InParams.OperatorSettings);
-		FFloatReadRef FrequencyIn = InputDataRefs.GetDataReadReferenceOrConstructWithVertexDefault<float>(InputInterface, METASOUND_GET_PARAM_NAME(ParamCutoffFrequency), InParams.OperatorSettings);
-		FFloatReadRef ResonanceIn = InputDataRefs.GetDataReadReferenceOrConstructWithVertexDefault<float>(InputInterface, METASOUND_GET_PARAM_NAME(ParamResonance), InParams.OperatorSettings);
+		FAudioBufferReadRef AudioIn = InputData.GetOrConstructDataReadReference<FAudioBuffer>(METASOUND_GET_PARAM_NAME(ParamAudioInput), InParams.OperatorSettings);
+		FFloatReadRef FrequencyIn = InputData.GetOrCreateDefaultDataReadReference<float>(METASOUND_GET_PARAM_NAME(ParamCutoffFrequency), InParams.OperatorSettings);
+		FFloatReadRef ResonanceIn = InputData.GetOrCreateDefaultDataReadReference<float>(METASOUND_GET_PARAM_NAME(ParamResonance), InParams.OperatorSettings);
 
 		return MakeUnique<FLadderFilterOperator>(
 			InParams.OperatorSettings
@@ -279,7 +276,7 @@ namespace Metasound
 
 		static FVertexInterface DeclareVertexInterface();
 
-		static TUniquePtr<IOperator> CreateOperator(const FCreateOperatorParams& InParams, FBuildErrorArray& OutErrors);
+		static TUniquePtr<IOperator> CreateOperator(const FBuildOperatorParams& InParams, FBuildResults& OutResults);
 
 		virtual void BindInputs(FInputVertexInterfaceData& InOutVertexData) override
 		{
@@ -422,18 +419,15 @@ namespace Metasound
 		return Interface;
 	}
 
-	TUniquePtr<IOperator> FStateVariableFilterOperator::CreateOperator(const FCreateOperatorParams& InParams, FBuildErrorArray& OutErrors)
+	TUniquePtr<IOperator> FStateVariableFilterOperator::CreateOperator(const FBuildOperatorParams& InParams, FBuildResults& OutResults)
 	{
-
-		const FDataReferenceCollection& InputDataRefs = InParams.InputDataReferences;
-		FVertexInterface Interface = DeclareVertexInterface();
-		const FInputVertexInterface& InputInterface = Interface.GetInputInterface();
+		const FInputVertexInterfaceData& InputData = InParams.InputData;
 
 		// inputs
-		FAudioBufferReadRef AudioIn = InputDataRefs.GetDataReadReferenceOrConstruct<FAudioBuffer>(METASOUND_GET_PARAM_NAME(ParamAudioInput), InParams.OperatorSettings);
-		FFloatReadRef FrequencyIn = InputDataRefs.GetDataReadReferenceOrConstructWithVertexDefault<float>(InputInterface, METASOUND_GET_PARAM_NAME(ParamCutoffFrequency), InParams.OperatorSettings);
-		FFloatReadRef ResonanceIn = InputDataRefs.GetDataReadReferenceOrConstructWithVertexDefault<float>(InputInterface, METASOUND_GET_PARAM_NAME(ParamResonance), InParams.OperatorSettings);
-		FFloatReadRef PassBandGainCompensationIn = InputDataRefs.GetDataReadReferenceOrConstructWithVertexDefault<float>(InputInterface, METASOUND_GET_PARAM_NAME(ParamBandStopControl), InParams.OperatorSettings);
+		FAudioBufferReadRef AudioIn = InputData.GetOrConstructDataReadReference<FAudioBuffer>(METASOUND_GET_PARAM_NAME(ParamAudioInput), InParams.OperatorSettings);
+		FFloatReadRef FrequencyIn = InputData.GetOrCreateDefaultDataReadReference<float>(METASOUND_GET_PARAM_NAME(ParamCutoffFrequency), InParams.OperatorSettings);
+		FFloatReadRef ResonanceIn = InputData.GetOrCreateDefaultDataReadReference<float>(METASOUND_GET_PARAM_NAME(ParamResonance), InParams.OperatorSettings);
+		FFloatReadRef PassBandGainCompensationIn = InputData.GetOrCreateDefaultDataReadReference<float>(METASOUND_GET_PARAM_NAME(ParamBandStopControl), InParams.OperatorSettings);
 
 		return MakeUnique<FStateVariableFilterOperator>(
 			InParams.OperatorSettings
@@ -507,7 +501,7 @@ namespace Metasound
 
 		static FVertexInterface DeclareVertexInterface();
 
-		static TUniquePtr<IOperator> CreateOperator(const FCreateOperatorParams& InParams, FBuildErrorArray& OutErrors);
+		static TUniquePtr<IOperator> CreateOperator(const FBuildOperatorParams& InParams, FBuildResults& OutResults);
 
 		virtual void BindInputs(FInputVertexInterfaceData& InOutVertexData) override
 		{
@@ -615,16 +609,13 @@ namespace Metasound
 		return Interface;
 	}
 
-	TUniquePtr<IOperator> FOnePoleLowPassFilterOperator::CreateOperator(const FCreateOperatorParams& InParams, FBuildErrorArray& OutErrors)
+	TUniquePtr<IOperator> FOnePoleLowPassFilterOperator::CreateOperator(const FBuildOperatorParams& InParams, FBuildResults& OutResults)
 	{
-
-		const FDataReferenceCollection& InputDataRefs = InParams.InputDataReferences;
-		FVertexInterface Interface = DeclareVertexInterface();
-		const FInputVertexInterface& InputInterface = Interface.GetInputInterface();
+		const FInputVertexInterfaceData& InputData = InParams.InputData;
 
 		// inputs
-		FAudioBufferReadRef AudioIn = InputDataRefs.GetDataReadReferenceOrConstruct<FAudioBuffer>(METASOUND_GET_PARAM_NAME(ParamAudioInput), InParams.OperatorSettings);
-		FFloatReadRef FrequencyIn = InputDataRefs.GetDataReadReferenceOrConstructWithVertexDefault<float>(InputInterface, METASOUND_GET_PARAM_NAME(ParamCutoffFrequency), InParams.OperatorSettings);
+		FAudioBufferReadRef AudioIn = InputData.GetOrConstructDataReadReference<FAudioBuffer>(METASOUND_GET_PARAM_NAME(ParamAudioInput), InParams.OperatorSettings);
+		FFloatReadRef FrequencyIn = InputData.GetOrCreateDefaultDataReadReference<float>(METASOUND_GET_PARAM_NAME(ParamCutoffFrequency), InParams.OperatorSettings);
 
 		return MakeUnique<FOnePoleLowPassFilterOperator>(
 			InParams.OperatorSettings
@@ -661,7 +652,7 @@ namespace Metasound
 
 		static FVertexInterface DeclareVertexInterface();
 
-		static TUniquePtr<IOperator> CreateOperator(const FCreateOperatorParams& InParams, FBuildErrorArray& OutErrors);
+		static TUniquePtr<IOperator> CreateOperator(const FBuildOperatorParams& InParams, FBuildResults& OutResults);
 
 		virtual void BindInputs(FInputVertexInterfaceData& InOutVertexData) override
 		{
@@ -776,16 +767,13 @@ namespace Metasound
 		return Interface;
 	}
 
-	TUniquePtr<IOperator> FOnePoleHighPassFilterOperator::CreateOperator(const FCreateOperatorParams& InParams, FBuildErrorArray& OutErrors)
+	TUniquePtr<IOperator> FOnePoleHighPassFilterOperator::CreateOperator(const FBuildOperatorParams& InParams, FBuildResults& OutResults)
 	{
-
-		const FDataReferenceCollection& InputDataRefs = InParams.InputDataReferences;
-		FVertexInterface Interface = DeclareVertexInterface();
-		const FInputVertexInterface& InputInterface = Interface.GetInputInterface();
+		const FInputVertexInterfaceData& InputData = InParams.InputData;
 
 		// inputs
-		FAudioBufferReadRef AudioIn = InputDataRefs.GetDataReadReferenceOrConstruct<FAudioBuffer>(METASOUND_GET_PARAM_NAME(ParamAudioInput), InParams.OperatorSettings);
-		FFloatReadRef FrequencyIn = InputDataRefs.GetDataReadReferenceOrConstructWithVertexDefault<float>(InputInterface, METASOUND_GET_PARAM_NAME(ParamCutoffFrequency), InParams.OperatorSettings);
+		FAudioBufferReadRef AudioIn = InputData.GetOrConstructDataReadReference<FAudioBuffer>(METASOUND_GET_PARAM_NAME(ParamAudioInput), InParams.OperatorSettings);
+		FFloatReadRef FrequencyIn = InputData.GetOrCreateDefaultDataReadReference<float>(METASOUND_GET_PARAM_NAME(ParamCutoffFrequency), InParams.OperatorSettings);
 
 		return MakeUnique<FOnePoleHighPassFilterOperator>(
 			InParams.OperatorSettings
@@ -822,7 +810,7 @@ namespace Metasound
 
 		static FVertexInterface DeclareVertexInterface();
 
-		static TUniquePtr<IOperator> CreateOperator(const FCreateOperatorParams& InParams, FBuildErrorArray& OutErrors);
+		static TUniquePtr<IOperator> CreateOperator(const FBuildOperatorParams& InParams, FBuildResults& OutResults);
 
 		virtual void BindInputs(FInputVertexInterfaceData& InOutVertexData) override
 		{
@@ -974,19 +962,16 @@ namespace Metasound
 		return Interface;
 	}
 
-	TUniquePtr<IOperator> FBiquadFilterOperator::CreateOperator(const FCreateOperatorParams& InParams, FBuildErrorArray& OutErrors)
+	TUniquePtr<IOperator> FBiquadFilterOperator::CreateOperator(const FBuildOperatorParams& InParams, FBuildResults& OutResults)
 	{
-
-		const FDataReferenceCollection& InputDataRefs = InParams.InputDataReferences;
-		FVertexInterface Interface = DeclareVertexInterface();
-		const FInputVertexInterface& InputInterface = Interface.GetInputInterface();
+		const FInputVertexInterfaceData& InputData = InParams.InputData;
 
 		// inputs
-		FAudioBufferReadRef AudioIn = InputDataRefs.GetDataReadReferenceOrConstruct<FAudioBuffer>(METASOUND_GET_PARAM_NAME(ParamAudioInput), InParams.OperatorSettings);
-		FFloatReadRef FrequencyIn = InputDataRefs.GetDataReadReferenceOrConstructWithVertexDefault<float>(InputInterface, METASOUND_GET_PARAM_NAME(ParamCutoffFrequency), InParams.OperatorSettings);
-		FFloatReadRef BandwidthIn = InputDataRefs.GetDataReadReferenceOrConstructWithVertexDefault<float>(InputInterface, METASOUND_GET_PARAM_NAME(ParamBandwidth), InParams.OperatorSettings);
-		FFloatReadRef FilterGainDbIn = InputDataRefs.GetDataReadReferenceOrConstructWithVertexDefault<float>(InputInterface, METASOUND_GET_PARAM_NAME(ParamGainDb), InParams.OperatorSettings);
-		FEnumBiQuadFilterReadRef FilterType = InputDataRefs.GetDataReadReferenceOrConstructWithVertexDefault<FEnumEBiquadFilterType>(InputInterface, METASOUND_GET_PARAM_NAME(ParamFilterType), InParams.OperatorSettings);
+		FAudioBufferReadRef AudioIn = InputData.GetOrConstructDataReadReference<FAudioBuffer>(METASOUND_GET_PARAM_NAME(ParamAudioInput), InParams.OperatorSettings);
+		FFloatReadRef FrequencyIn = InputData.GetOrCreateDefaultDataReadReference<float>(METASOUND_GET_PARAM_NAME(ParamCutoffFrequency), InParams.OperatorSettings);
+		FFloatReadRef BandwidthIn = InputData.GetOrCreateDefaultDataReadReference<float>(METASOUND_GET_PARAM_NAME(ParamBandwidth), InParams.OperatorSettings);
+		FFloatReadRef FilterGainDbIn = InputData.GetOrCreateDefaultDataReadReference<float>(METASOUND_GET_PARAM_NAME(ParamGainDb), InParams.OperatorSettings);
+		FEnumBiQuadFilterReadRef FilterType = InputData.GetOrCreateDefaultDataReadReference<FEnumEBiquadFilterType>(METASOUND_GET_PARAM_NAME(ParamFilterType), InParams.OperatorSettings);
 
 		return MakeUnique<FBiquadFilterOperator>(
 			InParams.OperatorSettings
diff --git a/Engine/Plugins/Runtime/Metasound/Source/MetasoundStandardNodes/Private/MetasoundBitcrusherNode.cpp b/Engine/Plugins/Runtime/Metasound/Source/MetasoundStandardNodes/Private/MetasoundBitcrusherNode.cpp
--- a/Engine/Plugins/Runtime/Metasound/Source/MetasoundStandardNodes/Private/MetasoundBitcrusherNode.cpp
+++ b/Engine/Plugins/Runtime/Metasound/Source/MetasoundStandardNodes/Private/MetasoundBitcrusherNode.cpp
@@ -33,7 +33,7 @@ namespace Metasound
 	{
 	public:
 
-		FBitcrusherOperator(const FCreateOperatorParams& InParams,
+		FBitcrusherOperator(const FBuildOperatorParams& InParams,
 			const FAudioBufferReadRef& InAudio,
 			const FFloatReadRef& InCrushedSampleRate,
 			const FFloatReadRef& InCrushedBitDepth)
@@ -125,15 +125,14 @@ namespace Metasound
 			return {};
 		}
 
-		static TUniquePtr<IOperator> CreateOperator(const FCreateOperatorParams& InParams, FBuildErrorArray& OutErrors)
+		static TUniquePtr<IOperator> CreateOperator(const FBuildOperatorParams& InParams, FBuildResults& OutResults)
 		{
 			using namespace BitcrusherVertexNames;
-			const FDataReferenceCollection& Inputs = InParams.InputDataReferences;
-			const FInputVertexInterface& InputInterface = DeclareVertexInterface().GetInputInterface();
+			const FInputVertexInterfaceData& InputData = InParams.InputData;
 
-			FAudioBufferReadRef AudioIn = Inputs.GetDataReadReferenceOrConstruct<FAudioBuffer>(METASOUND_GET_PARAM_NAME(InputAudio), InParams.OperatorSettings);
-			FFloatReadRef SampleRateIn = Inputs.GetDataReadReferenceOrConstructWithVertexDefault<float>(InputInterface, METASOUND_GET_PARAM_NAME(InputSampleRate), InParams.OperatorSettings);
-			FFloatReadRef BitDepthIn = Inputs.GetDataReadReferenceOrConstructWithVertexDefault<float>(InputInterface, METASOUND_GET_PARAM_NAME(InputBitDepth), InParams.OperatorSettings);
+			FAudioBufferReadRef AudioIn = InputData.GetOrConstructDataReadReference<FAudioBuffer>(METASOUND_GET_PARAM_NAME(InputAudio), InParams.OperatorSettings);
+			FFloatReadRef SampleRateIn = InputData.GetOrCreateDefaultDataReadReference<float>(METASOUND_GET_PARAM_NAME(InputSampleRate), InParams.OperatorSettings);
+			FFloatReadRef BitDepthIn = InputData.GetOrCreateDefaultDataReadReference<float>(METASOUND_GET_PARAM_NAME(InputBitDepth), InParams.OperatorSettings);
 
 			return MakeUnique<FBitcrusherOperator>(InParams, AudioIn, SampleRateIn, BitDepthIn);
 		}
diff --git a/Engine/Plugins/Runtime/Metasound/Source/MetasoundStandardNodes/Private/MetasoundClampNode.cpp b/Engine/Plugins/Runtime/Metasound/Source/MetasoundStandardNodes/Private/MetasoundClampNode.cpp
--- a/Engine/Plugins/Runtime/Metasound/Source/MetasoundStandardNodes/Private/MetasoundClampNode.cpp
+++ b/Engine/Plugins/Runtime/Metasound/Source/MetasoundStandardNodes/Private/MetasoundClampNode.cpp
@@ -57,25 +57,25 @@ namespace Metasound
 				OutClamped = FMath::Clamp(InValue, InMin, InMax);
 			}
 
-			static TDataReadReference<int32> CreateInValue(const FCreateOperatorParams& InParams)
+			static TDataReadReference<int32> CreateInValue(const FBuildOperatorParams& InParams)
 			{
 				using namespace ClampVertexNames;
-				const FInputVertexInterface& InputInterface = InParams.Node.GetVertexInterface().GetInputInterface();
-				return InParams.InputDataReferences.GetDataReadReferenceOrConstructWithVertexDefault<int32>(InputInterface, METASOUND_GET_PARAM_NAME(InputValue), InParams.OperatorSettings);
+				const FInputVertexInterfaceData& InputData = InParams.InputData;
+				return InputData.GetOrCreateDefaultDataReadReference<int32>(METASOUND_GET_PARAM_NAME(InputValue), InParams.OperatorSettings);
 			}
 
-			static TDataReadReference<int32> CreateInMin(const FCreateOperatorParams& InParams)
+			static TDataReadReference<int32> CreateInMin(const FBuildOperatorParams& InParams)
 			{
 				using namespace ClampVertexNames;
-				const FInputVertexInterface& InputInterface = InParams.Node.GetVertexInterface().GetInputInterface();
-				return InParams.InputDataReferences.GetDataReadReferenceOrConstructWithVertexDefault<int32>(InputInterface, METASOUND_GET_PARAM_NAME(InputMinValue), InParams.OperatorSettings);
+				const FInputVertexInterfaceData& InputData = InParams.InputData;
+				return InputData.GetOrCreateDefaultDataReadReference<int32>(METASOUND_GET_PARAM_NAME(InputMinValue), InParams.OperatorSettings);
 			}
 
-			static TDataReadReference<int32> CreateInMax(const FCreateOperatorParams& InParams)
+			static TDataReadReference<int32> CreateInMax(const FBuildOperatorParams& InParams)
 			{
 				using namespace ClampVertexNames;
-				const FInputVertexInterface& InputInterface = InParams.Node.GetVertexInterface().GetInputInterface();
-				return InParams.InputDataReferences.GetDataReadReferenceOrConstructWithVertexDefault<int32>(InputInterface, METASOUND_GET_PARAM_NAME(InputMaxValue), InParams.OperatorSettings);
+				const FInputVertexInterfaceData& InputData = InParams.InputData;
+				return InputData.GetOrCreateDefaultDataReadReference<int32>(METASOUND_GET_PARAM_NAME(InputMaxValue), InParams.OperatorSettings);
 			}
 		};
 
@@ -88,25 +88,25 @@ namespace Metasound
 				OutClamped = FMath::Clamp(InValue, InMin, InMax);
 			}
 
-			static TDataReadReference<float> CreateInValue(const FCreateOperatorParams& InParams)
+			static TDataReadReference<float> CreateInValue(const FBuildOperatorParams& InParams)
 			{
 				using namespace ClampVertexNames;
-				const FInputVertexInterface& InputInterface = InParams.Node.GetVertexInterface().GetInputInterface();
-				return InParams.InputDataReferences.GetDataReadReferenceOrConstructWithVertexDefault<float>(InputInterface, METASOUND_GET_PARAM_NAME(InputValue), InParams.OperatorSettings);
+				const FInputVertexInterfaceData& InputData = InParams.InputData;
+				return InputData.GetOrCreateDefaultDataReadReference<float>(METASOUND_GET_PARAM_NAME(InputValue), InParams.OperatorSettings);
 			}
 
-			static TDataReadReference<float> CreateInMin(const FCreateOperatorParams& InParams)
+			static TDataReadReference<float> CreateInMin(const FBuildOperatorParams& InParams)
 			{
 				using namespace ClampVertexNames;
-				const FInputVertexInterface& InputInterface = InParams.Node.GetVertexInterface().GetInputInterface();
-				return InParams.InputDataReferences.GetDataReadReferenceOrConstructWithVertexDefault<float>(InputInterface, METASOUND_GET_PARAM_NAME(InputMinValue), InParams.OperatorSettings);
+				const FInputVertexInterfaceData& InputData = InParams.InputData;
+				return InputData.GetOrCreateDefaultDataReadReference<float>(METASOUND_GET_PARAM_NAME(InputMinValue), InParams.OperatorSettings);
 			}
 
-			static TDataReadReference<float> CreateInMax(const FCreateOperatorParams& InParams)
+			static TDataReadReference<float> CreateInMax(const FBuildOperatorParams& InParams)
 			{
 				using namespace ClampVertexNames;
-				const FInputVertexInterface& InputInterface = InParams.Node.GetVertexInterface().GetInputInterface();
-				return InParams.InputDataReferences.GetDataReadReferenceOrConstructWithVertexDefault<float>(InputInterface, METASOUND_GET_PARAM_NAME(InputMaxValue), InParams.OperatorSettings);
+				const FInputVertexInterfaceData& InputData = InParams.InputData;
+				return InputData.GetOrCreateDefaultDataReadReference<float>(METASOUND_GET_PARAM_NAME(InputMaxValue), InParams.OperatorSettings);
 			}
 		};
 
@@ -126,22 +126,25 @@ namespace Metasound
 				}
 			}
 
-			static TDataReadReference<FAudioBuffer> CreateInValue(const FCreateOperatorParams& InParams)
+			static TDataReadReference<FAudioBuffer> CreateInValue(const FBuildOperatorParams& InParams)
 			{
 				using namespace ClampVertexNames;
-				return InParams.InputDataReferences.GetDataReadReferenceOrConstruct<FAudioBuffer>(METASOUND_GET_PARAM_NAME(InputValue), InParams.OperatorSettings);
+				const FInputVertexInterfaceData& InputData = InParams.InputData;
+				return InputData.GetOrConstructDataReadReference<FAudioBuffer>(METASOUND_GET_PARAM_NAME(InputValue), InParams.OperatorSettings);
 			}
 
-			static TDataReadReference<FAudioBuffer> CreateInMin(const FCreateOperatorParams& InParams)
+			static TDataReadReference<FAudioBuffer> CreateInMin(const FBuildOperatorParams& InParams)
 			{
 				using namespace ClampVertexNames;
-				return InParams.InputDataReferences.GetDataReadReferenceOrConstruct<FAudioBuffer>(METASOUND_GET_PARAM_NAME(InputMinValue), InParams.OperatorSettings);
+				const FInputVertexInterfaceData& InputData = InParams.InputData;
+				return InputData.GetOrConstructDataReadReference<FAudioBuffer>(METASOUND_GET_PARAM_NAME(InputMinValue), InParams.OperatorSettings);
 			}
 
-			static TDataReadReference<FAudioBuffer> CreateInMax(const FCreateOperatorParams& InParams)
+			static TDataReadReference<FAudioBuffer> CreateInMax(const FBuildOperatorParams& InParams)
 			{
 				using namespace ClampVertexNames;
-				return InParams.InputDataReferences.GetDataReadReferenceOrConstruct<FAudioBuffer>(METASOUND_GET_PARAM_NAME(InputMaxValue), InParams.OperatorSettings);
+				const FInputVertexInterfaceData& InputData = InParams.InputData;
+				return InputData.GetOrConstructDataReadReference<FAudioBuffer>(METASOUND_GET_PARAM_NAME(InputMaxValue), InParams.OperatorSettings);
 			}
 		};
 	}
@@ -186,7 +189,7 @@ namespace Metasound
 			return Metadata;
 		}
 
-		static TUniquePtr<IOperator> CreateOperator(const FCreateOperatorParams& InParams, TArray<TUniquePtr<IOperatorBuildError>>& OutErrors)
+		static TUniquePtr<IOperator> CreateOperator(const FBuildOperatorParams& InParams, FBuildResults& OutResults)
 		{
 			using namespace ClampVertexNames;
 			using namespace MetasoundClampNodePrivate;
diff --git a/Engine/Plugins/Runtime/Metasound/Source/MetasoundStandardNodes/Private/MetasoundConversionNodes.cpp b/Engine/Plugins/Runtime/Metasound/Source/MetasoundStandardNodes/Private/MetasoundConversionNodes.cpp
--- a/Engine/Plugins/Runtime/Metasound/Source/MetasoundStandardNodes/Private/MetasoundConversionNodes.cpp
+++ b/Engine/Plugins/Runtime/Metasound/Source/MetasoundStandardNodes/Private/MetasoundConversionNodes.cpp
@@ -34,12 +34,12 @@ namespace Metasound
 		template<>
 		struct TConversionNodeSpecialization<float, FTime>
 		{
-			static TDataReadReference<float> CreateInputRef(const FVertexInterface& Interface, const FCreateOperatorParams& InParams)
+			static TDataReadReference<float> CreateInputRef(const FVertexInterface& Interface, const FBuildOperatorParams& InParams)
 			{
 				using namespace ConversionNodeVertexNames;
-				const FDataReferenceCollection& InputCollection = InParams.InputDataReferences;
-				const FInputVertexInterface& InputInterface = Interface.GetInputInterface();
-				return InputCollection.GetDataReadReferenceOrConstructWithVertexDefault<float>(InputInterface, METASOUND_GET_PARAM_NAME(InputValue), InParams.OperatorSettings);
+				
+				const FInputVertexInterfaceData& InputData = InParams.InputData;
+				return InputData.GetOrCreateDefaultDataReadReference<float>(METASOUND_GET_PARAM_NAME(InputValue), InParams.OperatorSettings);
 			}
 
 			static void GetConvertedValue(float InValue, FTime& OutValue)
@@ -51,12 +51,12 @@ namespace Metasound
 		template<>
 		struct TConversionNodeSpecialization<float, FAudioBuffer>
 		{
-			static TDataReadReference<float> CreateInputRef(const FVertexInterface& Interface, const FCreateOperatorParams& InParams)
+			static TDataReadReference<float> CreateInputRef(const FVertexInterface& Interface, const FBuildOperatorParams& InParams)
 			{
 				using namespace ConversionNodeVertexNames;
-				const FDataReferenceCollection& InputCollection = InParams.InputDataReferences;
-				const FInputVertexInterface& InputInterface = Interface.GetInputInterface();
-				return InputCollection.GetDataReadReferenceOrConstructWithVertexDefault<float>(InputInterface, METASOUND_GET_PARAM_NAME(InputValue), InParams.OperatorSettings);
+				
+				const FInputVertexInterfaceData& InputData = InParams.InputData;
+				return InputData.GetOrCreateDefaultDataReadReference<float>(METASOUND_GET_PARAM_NAME(InputValue), InParams.OperatorSettings);
 			}
 
 			static void GetConvertedValue(float InValue, FAudioBuffer& OutValue)
@@ -73,12 +73,12 @@ namespace Metasound
 		template<>
 		struct TConversionNodeSpecialization<int32, FTime>
 		{
-			static TDataReadReference<int32> CreateInputRef(const FVertexInterface& Interface, const FCreateOperatorParams& InParams)
+			static TDataReadReference<int32> CreateInputRef(const FVertexInterface& Interface, const FBuildOperatorParams& InParams)
 			{
 				using namespace ConversionNodeVertexNames;
-				const FDataReferenceCollection& InputCollection = InParams.InputDataReferences;
-				const FInputVertexInterface& InputInterface = Interface.GetInputInterface();
-				return InputCollection.GetDataReadReferenceOrConstructWithVertexDefault<int32>(InputInterface, METASOUND_GET_PARAM_NAME(InputValue), InParams.OperatorSettings);
+				
+				const FInputVertexInterfaceData& InputData = InParams.InputData;
+				return InputData.GetOrCreateDefaultDataReadReference<int32>(METASOUND_GET_PARAM_NAME(InputValue), InParams.OperatorSettings);
 			}
 
 			static void GetConvertedValue(int32 InValue, FTime& OutValue)
@@ -90,12 +90,12 @@ namespace Metasound
 		template<>
 		struct TConversionNodeSpecialization<FTime, float>
 		{
-			static TDataReadReference<FTime> CreateInputRef(const FVertexInterface& Interface, const FCreateOperatorParams& InParams)
+			static TDataReadReference<FTime> CreateInputRef(const FVertexInterface& Interface, const FBuildOperatorParams& InParams)
 			{
 				using namespace ConversionNodeVertexNames;
-				const FDataReferenceCollection& InputCollection = InParams.InputDataReferences;
-				const FInputVertexInterface& InputInterface = Interface.GetInputInterface();
-				return InputCollection.GetDataReadReferenceOrConstructWithVertexDefault<FTime>(InputInterface, METASOUND_GET_PARAM_NAME(InputValue), InParams.OperatorSettings);
+				
+				const FInputVertexInterfaceData& InputData = InParams.InputData;
+				return InputData.GetOrCreateDefaultDataReadReference<FTime>(METASOUND_GET_PARAM_NAME(InputValue), InParams.OperatorSettings);
 			}
 
 			static void GetConvertedValue(const FTime& InValue, float& OutValue)
@@ -107,12 +107,12 @@ namespace Metasound
 		template<>
 		struct TConversionNodeSpecialization<FTime, int32>
 		{
-			static TDataReadReference<FTime> CreateInputRef(const FVertexInterface& Interface, const FCreateOperatorParams& InParams)
+			static TDataReadReference<FTime> CreateInputRef(const FVertexInterface& Interface, const FBuildOperatorParams& InParams)
 			{
 				using namespace ConversionNodeVertexNames;
-				const FDataReferenceCollection& InputCollection = InParams.InputDataReferences;
-				const FInputVertexInterface& InputInterface = Interface.GetInputInterface();
-				return InputCollection.GetDataReadReferenceOrConstructWithVertexDefault<FTime>(InputInterface, METASOUND_GET_PARAM_NAME(InputValue), InParams.OperatorSettings);
+				
+				const FInputVertexInterfaceData& InputData = InParams.InputData;
+				return InputData.GetOrCreateDefaultDataReadReference<FTime>(METASOUND_GET_PARAM_NAME(InputValue), InParams.OperatorSettings);
 			}
 
 			static void GetConvertedValue(const FTime& InValue, int32& OutValue)
@@ -124,11 +124,11 @@ namespace Metasound
 		template<>
 		struct TConversionNodeSpecialization<FAudioBuffer, float>
 		{
-			static TDataReadReference<FAudioBuffer> CreateInputRef(const FVertexInterface& Interface, const FCreateOperatorParams& InParams)
+			static TDataReadReference<FAudioBuffer> CreateInputRef(const FVertexInterface& Interface, const FBuildOperatorParams& InParams)
 			{
 				using namespace ConversionNodeVertexNames;
-				const FDataReferenceCollection& InputCollection = InParams.InputDataReferences;
-				return InputCollection.GetDataReadReferenceOrConstruct<FAudioBuffer>(METASOUND_GET_PARAM_NAME(InputValue), InParams.OperatorSettings);
+				const FInputVertexInterfaceData& InputData = InParams.InputData;
+				return InputData.GetOrConstructDataReadReference<FAudioBuffer>(METASOUND_GET_PARAM_NAME(InputValue), InParams.OperatorSettings);
 			}
 
 			static void GetConvertedValue(const FAudioBuffer& InValue, float& OutValue)
@@ -206,7 +206,7 @@ namespace Metasound
 			return Info;
 		}
 
-		static TUniquePtr<IOperator> CreateOperator(const FCreateOperatorParams& InParams, TArray<TUniquePtr<IOperatorBuildError>>& OutErrors)
+		static TUniquePtr<IOperator> CreateOperator(const FBuildOperatorParams& InParams, FBuildResults& OutResults)
 		{
 			using namespace MetasoundConversionNodePrivate;
 			TDataReadReference<FromType> InputValue = TConversionNodeSpecialization<FromType, ToType>::CreateInputRef(GetDefaultInterface(), InParams);
diff --git a/Engine/Plugins/Runtime/Metasound/Source/MetasoundStandardNodes/Private/MetasoundCrossfadeNode.cpp b/Engine/Plugins/Runtime/Metasound/Source/MetasoundStandardNodes/Private/MetasoundCrossfadeNode.cpp
--- a/Engine/Plugins/Runtime/Metasound/Source/MetasoundStandardNodes/Private/MetasoundCrossfadeNode.cpp
+++ b/Engine/Plugins/Runtime/Metasound/Source/MetasoundStandardNodes/Private/MetasoundCrossfadeNode.cpp
@@ -202,19 +202,18 @@ namespace Metasound
 			return Metadata;
 		}
 
-		static TUniquePtr<IOperator> CreateOperator(const FCreateOperatorParams& InParams, TArray<TUniquePtr<IOperatorBuildError>>& OutErrors)
+		static TUniquePtr<IOperator> CreateOperator(const FBuildOperatorParams& InParams, FBuildResults& OutResults)
 		{
 			using namespace CrossfadeVertexNames;
 
-			const FInputVertexInterface& InputInterface = InParams.Node.GetVertexInterface().GetInputInterface();
-			const FDataReferenceCollection& InputCollection = InParams.InputDataReferences;
+			const FInputVertexInterfaceData& InputData = InParams.InputData;
 
-			FFloatReadRef CrossfadeValue = InputCollection.GetDataReadReferenceOrConstructWithVertexDefault<float>(InputInterface, METASOUND_GET_PARAM_NAME(InputCrossfadeValue), InParams.OperatorSettings);
+			FFloatReadRef CrossfadeValue = InputData.GetOrCreateDefaultDataReadReference<float>(METASOUND_GET_PARAM_NAME(InputCrossfadeValue), InParams.OperatorSettings);
 
 			TArray<TDataReadReference<ValueType>> InputValues;
 			for (uint32 i = 0; i < NumInputs; ++i)
 			{
-				InputValues.Add(InputCollection.GetDataReadReferenceOrConstructWithVertexDefault<ValueType>(InputInterface, GetInputName(i), InParams.OperatorSettings));
+				InputValues.Add(InputData.GetOrCreateDefaultDataReadReference<ValueType>(GetInputName(i), InParams.OperatorSettings));
 			}
 
 			return MakeUnique<TCrossfadeOperator<ValueType, NumInputs>>(InParams.OperatorSettings, CrossfadeValue, MoveTemp(InputValues));
diff --git a/Engine/Plugins/Runtime/Metasound/Source/MetasoundStandardNodes/Private/MetasoundDebugLogNode.cpp b/Engine/Plugins/Runtime/Metasound/Source/MetasoundStandardNodes/Private/MetasoundDebugLogNode.cpp
--- a/Engine/Plugins/Runtime/Metasound/Source/MetasoundStandardNodes/Private/MetasoundDebugLogNode.cpp
+++ b/Engine/Plugins/Runtime/Metasound/Source/MetasoundStandardNodes/Private/MetasoundDebugLogNode.cpp
@@ -87,23 +87,22 @@ namespace Metasound
 				return Metadata;
 			}
 
-			static TUniquePtr<IOperator> CreateOperator(const FCreateOperatorParams& InParams, TArray<TUniquePtr<IOperatorBuildError>>& OutErrors)
+			static TUniquePtr<IOperator> CreateOperator(const FBuildOperatorParams& InParams, FBuildResults& OutResults)
 			{
 				using namespace PrintLogVertexNames;
 
-				const FInputVertexInterface& InputInterface = InParams.Node.GetVertexInterface().GetInputInterface();
-				const FDataReferenceCollection& InputCollection = InParams.InputDataReferences;
+				const FInputVertexInterfaceData& InputData = InParams.InputData;
 
-				FTriggerReadRef Trigger = InputCollection.GetDataReadReferenceOrConstruct<FTrigger>(METASOUND_GET_PARAM_NAME(InputTrigger), InParams.OperatorSettings);
+				FTriggerReadRef Trigger = InputData.GetOrConstructDataReadReference<FTrigger>(METASOUND_GET_PARAM_NAME(InputTrigger), InParams.OperatorSettings);
 
-				TDataReadReference<FString> Label = InputCollection.GetDataReadReferenceOrConstructWithVertexDefault<FString>(InputInterface, METASOUND_GET_PARAM_NAME(InputLabel), InParams.OperatorSettings);
-				TDataReadReference<PrintLogType> ValueToLog = InputCollection.GetDataReadReferenceOrConstructWithVertexDefault<PrintLogType>(InputInterface, METASOUND_GET_PARAM_NAME(InputValueToLog), InParams.OperatorSettings);
+				TDataReadReference<FString> Label = InputData.GetOrCreateDefaultDataReadReference<FString>(METASOUND_GET_PARAM_NAME(InputLabel), InParams.OperatorSettings);
+				TDataReadReference<PrintLogType> ValueToLog = InputData.GetOrCreateDefaultDataReadReference<PrintLogType>(METASOUND_GET_PARAM_NAME(InputValueToLog), InParams.OperatorSettings);
 
 				return MakeUnique<TPrintLogOperator<PrintLogType>>(InParams, Trigger, Label, ValueToLog);
 			}
 
 
-			TPrintLogOperator(const FCreateOperatorParams& InParams, TDataReadReference<FTrigger> InTrigger, TDataReadReference<FString> InLabelPrintLog, TDataReadReference<PrintLogType> InValueToLogPrintLog)
+			TPrintLogOperator(const FBuildOperatorParams& InParams, TDataReadReference<FTrigger> InTrigger, TDataReadReference<FString> InLabelPrintLog, TDataReadReference<PrintLogType> InValueToLogPrintLog)
 				: Trigger(InTrigger)
 				, Label(InLabelPrintLog)
 				, ValueToLog(InValueToLogPrintLog)
diff --git a/Engine/Plugins/Runtime/Metasound/Source/MetasoundStandardNodes/Private/MetasoundDebugTriggerOnNanNode.cpp b/Engine/Plugins/Runtime/Metasound/Source/MetasoundStandardNodes/Private/MetasoundDebugTriggerOnNanNode.cpp
--- a/Engine/Plugins/Runtime/Metasound/Source/MetasoundStandardNodes/Private/MetasoundDebugTriggerOnNanNode.cpp
+++ b/Engine/Plugins/Runtime/Metasound/Source/MetasoundStandardNodes/Private/MetasoundDebugTriggerOnNanNode.cpp
@@ -33,7 +33,7 @@ namespace Metasound
 	public:
 		static const FNodeClassMetadata& GetNodeInfo();
 		static const FVertexInterface& GetVertexInterface();
-		static TUniquePtr<IOperator> CreateOperator(const FCreateOperatorParams& InParams, FBuildErrorArray& OutErrors);
+		static TUniquePtr<IOperator> CreateOperator(const FBuildOperatorParams& InParams, FBuildResults& OutResults);
 
 		FTriggerOnNanOperator(const FOperatorSettings& InSettings, const FAudioBufferReadRef& InAudioInput, const FBoolReadRef& InTriggerOnce);
 
@@ -123,14 +123,14 @@ namespace Metasound
 	}
 
 
-	TUniquePtr<IOperator> FTriggerOnNanOperator::CreateOperator(const FCreateOperatorParams& InParams, FBuildErrorArray& OutErrors)
+	TUniquePtr<IOperator> FTriggerOnNanOperator::CreateOperator(const FBuildOperatorParams& InParams, FBuildResults& OutResults)
 	{
 		using namespace TriggerOnNanVertexNames;
 
-		const FInputVertexInterface& InputInterface = GetVertexInterface().GetInputInterface();
+		const FInputVertexInterfaceData& InputData = InParams.InputData;
 
-		FAudioBufferReadRef AudioInput = InParams.InputDataReferences.GetDataReadReferenceOrConstruct<FAudioBuffer>(METASOUND_GET_PARAM_NAME(AudioInput), InParams.OperatorSettings);
-		FBoolReadRef bTriggerOnce = InParams.InputDataReferences.GetDataReadReferenceOrConstructWithVertexDefault<bool>(InputInterface, METASOUND_GET_PARAM_NAME(TriggerOnce), InParams.OperatorSettings);
+		FAudioBufferReadRef AudioInput = InputData.GetOrConstructDataReadReference<FAudioBuffer>(METASOUND_GET_PARAM_NAME(AudioInput), InParams.OperatorSettings);
+		FBoolReadRef bTriggerOnce = InputData.GetOrCreateDefaultDataReadReference<bool>(METASOUND_GET_PARAM_NAME(TriggerOnce), InParams.OperatorSettings);
 
 		return MakeUnique<FTriggerOnNanOperator>(InParams.OperatorSettings, AudioInput, bTriggerOnce);
 	}
diff --git a/Engine/Plugins/Runtime/Metasound/Source/MetasoundStandardNodes/Private/MetasoundDecibelsToLinearGain.cpp b/Engine/Plugins/Runtime/Metasound/Source/MetasoundStandardNodes/Private/MetasoundDecibelsToLinearGain.cpp
--- a/Engine/Plugins/Runtime/Metasound/Source/MetasoundStandardNodes/Private/MetasoundDecibelsToLinearGain.cpp
+++ b/Engine/Plugins/Runtime/Metasound/Source/MetasoundStandardNodes/Private/MetasoundDecibelsToLinearGain.cpp
@@ -29,9 +29,9 @@ namespace Metasound
 
 		static const FNodeClassMetadata& GetNodeInfo();
 		static const FVertexInterface& GetVertexInterface();
-		static TUniquePtr<IOperator> CreateOperator(const FCreateOperatorParams& InParams, FBuildErrorArray& OutErrors);
+		static TUniquePtr<IOperator> CreateOperator(const FBuildOperatorParams& InParams, FBuildResults& OutResults);
 
-		FDecibelsToLinearGainOperator(const FCreateOperatorParams& InParams, const FFloatReadRef& InDecibelGain);
+		FDecibelsToLinearGainOperator(const FBuildOperatorParams& InParams, const FFloatReadRef& InDecibelGain);
 
 		virtual void BindInputs(FInputVertexInterfaceData& InOutVertexData) override;
 		virtual void BindOutputs(FOutputVertexInterfaceData& InOutVertexData) override;
@@ -48,7 +48,7 @@ namespace Metasound
 		FFloatWriteRef LinearGainOutput;
 	};
 
-	FDecibelsToLinearGainOperator::FDecibelsToLinearGainOperator(const FCreateOperatorParams& InParams, const FFloatReadRef& InDecibelGain)
+	FDecibelsToLinearGainOperator::FDecibelsToLinearGainOperator(const FBuildOperatorParams& InParams, const FFloatReadRef& InDecibelGain)
 		: DecibelGainInput(InDecibelGain)
 		, LinearGainOutput(FFloatWriteRef::CreateNew())
 	{
@@ -137,14 +137,13 @@ namespace Metasound
 		return Info;
 	}
 
-	TUniquePtr<IOperator> FDecibelsToLinearGainOperator::CreateOperator(const FCreateOperatorParams& InParams, FBuildErrorArray& OutErrors)
+	TUniquePtr<IOperator> FDecibelsToLinearGainOperator::CreateOperator(const FBuildOperatorParams& InParams, FBuildResults& OutResults)
 	{
 		using namespace DecibelsToLinearGainVertexNames;
+		
+		const FInputVertexInterfaceData& InputData = InParams.InputData;
 
-		const FDataReferenceCollection& InputCollection = InParams.InputDataReferences;
-		const FInputVertexInterface& InputInterface = GetVertexInterface().GetInputInterface();
-
-		FFloatReadRef InDecibelGain = InputCollection.GetDataReadReferenceOrConstructWithVertexDefault<float>(InputInterface, METASOUND_GET_PARAM_NAME(InputDecibelGain), InParams.OperatorSettings);
+		FFloatReadRef InDecibelGain = InputData.GetOrCreateDefaultDataReadReference<float>(METASOUND_GET_PARAM_NAME(InputDecibelGain), InParams.OperatorSettings);
 
 		return MakeUnique<FDecibelsToLinearGainOperator>(InParams, InDecibelGain);
 	}
diff --git a/Engine/Plugins/Runtime/Metasound/Source/MetasoundStandardNodes/Private/MetasoundDiffuserNode.cpp b/Engine/Plugins/Runtime/Metasound/Source/MetasoundStandardNodes/Private/MetasoundDiffuserNode.cpp
--- a/Engine/Plugins/Runtime/Metasound/Source/MetasoundStandardNodes/Private/MetasoundDiffuserNode.cpp
+++ b/Engine/Plugins/Runtime/Metasound/Source/MetasoundStandardNodes/Private/MetasoundDiffuserNode.cpp
@@ -75,16 +75,15 @@ namespace Metasound
 			return Interface;
 		}
 
-		static TUniquePtr<IOperator> CreateOperator(const FCreateOperatorParams& InParams, FBuildErrorArray& OutErrors)
+		static TUniquePtr<IOperator> CreateOperator(const FBuildOperatorParams& InParams, FBuildResults& OutResults)
 		{
 			using namespace DiffuserNode;
+			
+			const FInputVertexInterfaceData& InputData = InParams.InputData;
 
-			const FDataReferenceCollection& InputCollection = InParams.InputDataReferences;
-			const FInputVertexInterface& InputInterface = GetVertexInterface().GetInputInterface();
-
-			FAudioBufferReadRef AudioIn = InputCollection.GetDataReadReferenceOrConstruct<FAudioBuffer>(METASOUND_GET_PARAM_NAME(InputAudio), InParams.OperatorSettings);
-			FInt32ReadRef DepthIn = InputCollection.GetDataReadReferenceOrConstructWithVertexDefault<int32>(InputInterface, METASOUND_GET_PARAM_NAME(InputDiffusionDepth), InParams.OperatorSettings);
-			FFloatReadRef FeedbackIn = InputCollection.GetDataReadReferenceOrConstructWithVertexDefault<float>(InputInterface, METASOUND_GET_PARAM_NAME(InputFeedbackGain), InParams.OperatorSettings);
+			FAudioBufferReadRef AudioIn = InputData.GetOrConstructDataReadReference<FAudioBuffer>(METASOUND_GET_PARAM_NAME(InputAudio), InParams.OperatorSettings);
+			FInt32ReadRef DepthIn = InputData.GetOrCreateDefaultDataReadReference<int32>(METASOUND_GET_PARAM_NAME(InputDiffusionDepth), InParams.OperatorSettings);
+			FFloatReadRef FeedbackIn = InputData.GetOrCreateDefaultDataReadReference<float>(METASOUND_GET_PARAM_NAME(InputFeedbackGain), InParams.OperatorSettings);
 
 			return MakeUnique<FDiffuserOperator>(InParams.OperatorSettings, AudioIn, DepthIn, FeedbackIn);
 		}
diff --git a/Engine/Plugins/Runtime/Metasound/Source/MetasoundStandardNodes/Private/MetasoundDopplerPitchShiftNode.cpp b/Engine/Plugins/Runtime/Metasound/Source/MetasoundStandardNodes/Private/MetasoundDopplerPitchShiftNode.cpp
--- a/Engine/Plugins/Runtime/Metasound/Source/MetasoundStandardNodes/Private/MetasoundDopplerPitchShiftNode.cpp
+++ b/Engine/Plugins/Runtime/Metasound/Source/MetasoundStandardNodes/Private/MetasoundDopplerPitchShiftNode.cpp
@@ -28,9 +28,9 @@ namespace Metasound
 	public:
 		static const FNodeClassMetadata& GetNodeInfo();
 		static const FVertexInterface& GetVertexInterface();
-		static TUniquePtr<IOperator> CreateOperator(const FCreateOperatorParams& InParams, FBuildErrorArray& OutErrors);
+		static TUniquePtr<IOperator> CreateOperator(const FBuildOperatorParams& InParams, FBuildResults& OutResults);
 
-		FDopplerPitchShiftOperator(const FCreateOperatorParams& InParams, const FAudioBufferReadRef& InAudioInput, const FFloatReadRef& InPitchShift, const FFloatReadRef& InDelayLength);
+		FDopplerPitchShiftOperator(const FBuildOperatorParams& InParams, const FAudioBufferReadRef& InAudioInput, const FFloatReadRef& InPitchShift, const FFloatReadRef& InDelayLength);
 
 		virtual void BindInputs(FInputVertexInterfaceData& InOutVertexData) override;
 		virtual void BindOutputs(FOutputVertexInterfaceData& InOutVertexData) override;
@@ -48,7 +48,7 @@ namespace Metasound
 		Audio::FTapDelayPitchShifter DelayPitchShifter;
 	};
 
-	FDopplerPitchShiftOperator::FDopplerPitchShiftOperator(const FCreateOperatorParams& InParams,
+	FDopplerPitchShiftOperator::FDopplerPitchShiftOperator(const FBuildOperatorParams& InParams,
 		const FAudioBufferReadRef& InAudioInput,
 		const FFloatReadRef& InPitchShift,
 		const FFloatReadRef& InDelayLength)
@@ -145,16 +145,15 @@ namespace Metasound
 		return Info;
 	}
 
-	TUniquePtr<IOperator> FDopplerPitchShiftOperator::CreateOperator(const FCreateOperatorParams& InParams, FBuildErrorArray& OutErrors)
+	TUniquePtr<IOperator> FDopplerPitchShiftOperator::CreateOperator(const FBuildOperatorParams& InParams, FBuildResults& OutResults)
 	{
 		using namespace DopplerPitchShift;
 
-		const FDataReferenceCollection& InputCollection = InParams.InputDataReferences;
-		const FInputVertexInterface& InputInterface = GetVertexInterface().GetInputInterface();
+		const FInputVertexInterfaceData& InputData = InParams.InputData;
 
-		FAudioBufferReadRef AudioIn = InputCollection.GetDataReadReferenceOrConstruct<FAudioBuffer>(METASOUND_GET_PARAM_NAME(InParamAudioInput), InParams.OperatorSettings);
-		FFloatReadRef PitchShift = InputCollection.GetDataReadReferenceOrConstructWithVertexDefault<float>(InputInterface, METASOUND_GET_PARAM_NAME(InParamPitchShift), InParams.OperatorSettings);
-		FFloatReadRef DelayLength = InputCollection.GetDataReadReferenceOrConstructWithVertexDefault<float>(InputInterface, METASOUND_GET_PARAM_NAME(InParamDelayLength), InParams.OperatorSettings);
+		FAudioBufferReadRef AudioIn = InputData.GetOrConstructDataReadReference<FAudioBuffer>(METASOUND_GET_PARAM_NAME(InParamAudioInput), InParams.OperatorSettings);
+		FFloatReadRef PitchShift = InputData.GetOrCreateDefaultDataReadReference<float>(METASOUND_GET_PARAM_NAME(InParamPitchShift), InParams.OperatorSettings);
+		FFloatReadRef DelayLength = InputData.GetOrCreateDefaultDataReadReference<float>(METASOUND_GET_PARAM_NAME(InParamDelayLength), InParams.OperatorSettings);
 
 		return MakeUnique<FDopplerPitchShiftOperator>(InParams, AudioIn, PitchShift, DelayLength);
 	}
diff --git a/Engine/Plugins/Runtime/Metasound/Source/MetasoundStandardNodes/Private/MetasoundEnvelopeFollowerNode.cpp b/Engine/Plugins/Runtime/Metasound/Source/MetasoundStandardNodes/Private/MetasoundEnvelopeFollowerNode.cpp
--- a/Engine/Plugins/Runtime/Metasound/Source/MetasoundStandardNodes/Private/MetasoundEnvelopeFollowerNode.cpp
+++ b/Engine/Plugins/Runtime/Metasound/Source/MetasoundStandardNodes/Private/MetasoundEnvelopeFollowerNode.cpp
@@ -43,9 +43,9 @@ namespace Metasound
 
 		static const FNodeClassMetadata& GetNodeInfo();
 		static const FVertexInterface& GetVertexInterface();
-		static TUniquePtr<IOperator> CreateOperator(const FCreateOperatorParams& InParams, FBuildErrorArray& OutErrors);
+		static TUniquePtr<IOperator> CreateOperator(const FBuildOperatorParams& InParams, FBuildResults& OutResults);
 
-		FEnvelopeFollowerOperator(const FCreateOperatorParams& InOperatorSettings,
+		FEnvelopeFollowerOperator(const FBuildOperatorParams& InOperatorSettings,
 			const FBoolReadRef& InEnable,
 			const FAudioBufferReadRef& InAudioInput,
 			const FTimeReadRef& InAttackTime,
@@ -88,7 +88,7 @@ namespace Metasound
 		EEnvelopePeakMode PrevFollowMode = EEnvelopePeakMode::Peak;
 	};
 
-	FEnvelopeFollowerOperator::FEnvelopeFollowerOperator(const FCreateOperatorParams& InParams,
+	FEnvelopeFollowerOperator::FEnvelopeFollowerOperator(const FBuildOperatorParams& InParams,
 		const FBoolReadRef& InEnable,
 		const FAudioBufferReadRef& InAudioInput,
 		const FTimeReadRef& InAttackTime,
@@ -265,18 +265,17 @@ namespace Metasound
 		return Info;
 	}
 
-	TUniquePtr<IOperator> FEnvelopeFollowerOperator::CreateOperator(const FCreateOperatorParams& InParams, FBuildErrorArray& OutErrors)
+	TUniquePtr<IOperator> FEnvelopeFollowerOperator::CreateOperator(const FBuildOperatorParams& InParams, FBuildResults& OutResults)
 	{
 		using namespace EnvelopeFollowerVertexNames;
 
-		const FDataReferenceCollection& InputCollection = InParams.InputDataReferences;
-		const FInputVertexInterface& InputInterface = GetVertexInterface().GetInputInterface();
+		const FInputVertexInterfaceData& InputData = InParams.InputData;
 
-		FBoolReadRef EnableIn = InputCollection.GetDataReadReferenceOrConstructWithVertexDefault<bool>(InputInterface, METASOUND_GET_PARAM_NAME(InParamEnable), InParams.OperatorSettings);
-		FAudioBufferReadRef AudioIn = InputCollection.GetDataReadReferenceOrConstruct<FAudioBuffer>(METASOUND_GET_PARAM_NAME(InParamAudioInput), InParams.OperatorSettings);
-		FTimeReadRef AttackTime = InputCollection.GetDataReadReferenceOrConstructWithVertexDefault<FTime>(InputInterface, METASOUND_GET_PARAM_NAME(InParamAttackTime), InParams.OperatorSettings);
-		FTimeReadRef ReleaseTime = InputCollection.GetDataReadReferenceOrConstructWithVertexDefault<FTime>(InputInterface, METASOUND_GET_PARAM_NAME(InParamReleaseTime), InParams.OperatorSettings);
-		FEnvelopePeakModeReadRef EnvelopeModeIn = InputCollection.GetDataReadReferenceOrConstructWithVertexDefault<FEnumEnvelopePeakMode>(InputInterface, METASOUND_GET_PARAM_NAME(InParamFollowMode), InParams.OperatorSettings);
+		FBoolReadRef EnableIn = InputData.GetOrCreateDefaultDataReadReference<bool>(METASOUND_GET_PARAM_NAME(InParamEnable), InParams.OperatorSettings);
+		FAudioBufferReadRef AudioIn = InputData.GetOrConstructDataReadReference<FAudioBuffer>(METASOUND_GET_PARAM_NAME(InParamAudioInput), InParams.OperatorSettings);
+		FTimeReadRef AttackTime = InputData.GetOrCreateDefaultDataReadReference<FTime>(METASOUND_GET_PARAM_NAME(InParamAttackTime), InParams.OperatorSettings);
+		FTimeReadRef ReleaseTime = InputData.GetOrCreateDefaultDataReadReference<FTime>(METASOUND_GET_PARAM_NAME(InParamReleaseTime), InParams.OperatorSettings);
+		FEnvelopePeakModeReadRef EnvelopeModeIn = InputData.GetOrCreateDefaultDataReadReference<FEnumEnvelopePeakMode>(METASOUND_GET_PARAM_NAME(InParamFollowMode), InParams.OperatorSettings);
 
 		return MakeUnique<FEnvelopeFollowerOperator>(InParams, EnableIn, AudioIn, AttackTime, ReleaseTime, EnvelopeModeIn);
 	}
diff --git a/Engine/Plugins/Runtime/Metasound/Source/MetasoundStandardNodes/Private/MetasoundFilterQToBandwidthNode.cpp b/Engine/Plugins/Runtime/Metasound/Source/MetasoundStandardNodes/Private/MetasoundFilterQToBandwidthNode.cpp
--- a/Engine/Plugins/Runtime/Metasound/Source/MetasoundStandardNodes/Private/MetasoundFilterQToBandwidthNode.cpp
+++ b/Engine/Plugins/Runtime/Metasound/Source/MetasoundStandardNodes/Private/MetasoundFilterQToBandwidthNode.cpp
@@ -67,14 +67,13 @@ namespace Metasound
 			return Metadata;
 		}
 
-		static TUniquePtr<IOperator> CreateOperator(const FCreateOperatorParams& InParams, TArray<TUniquePtr<IOperatorBuildError>>& OutErrors)
+		static TUniquePtr<IOperator> CreateOperator(const FBuildOperatorParams& InParams, FBuildResults& OutResults)
 		{
 			using namespace QToBandwidth;
+			
+			const FInputVertexInterfaceData& InputData = InParams.InputData;
 
-			const FDataReferenceCollection& InputDataRefs = InParams.InputDataReferences;
-			const FInputVertexInterface& InputInterface = InParams.Node.GetVertexInterface().GetInputInterface();
-
-			FFloatReadRef FloatInput = InputDataRefs.GetDataReadReferenceOrConstructWithVertexDefault<float>(InputInterface, METASOUND_GET_PARAM_NAME(InputQ), InParams.OperatorSettings);
+			FFloatReadRef FloatInput = InputData.GetOrCreateDefaultDataReadReference<float>(METASOUND_GET_PARAM_NAME(InputQ), InParams.OperatorSettings);
 
 			return MakeUnique<FConvertQToBandwidthNodeOperator>(InParams.OperatorSettings, FloatInput);
 		}
diff --git a/Engine/Plugins/Runtime/Metasound/Source/MetasoundStandardNodes/Private/MetasoundFlangerNode.cpp b/Engine/Plugins/Runtime/Metasound/Source/MetasoundStandardNodes/Private/MetasoundFlangerNode.cpp
--- a/Engine/Plugins/Runtime/Metasound/Source/MetasoundStandardNodes/Private/MetasoundFlangerNode.cpp
+++ b/Engine/Plugins/Runtime/Metasound/Source/MetasoundStandardNodes/Private/MetasoundFlangerNode.cpp
@@ -77,23 +77,22 @@ namespace Metasound
 			return Interface;
 		}
 
-		static TUniquePtr<IOperator> CreateOperator(const FCreateOperatorParams& InParams, FBuildErrorArray& OutErrors)
+		static TUniquePtr<IOperator> CreateOperator(const FBuildOperatorParams& InParams, FBuildResults& OutResults)
 		{
 			using namespace FlangerVertexNames;
+			
+			const FInputVertexInterfaceData& InputData = InParams.InputData;
 
-			const FDataReferenceCollection& InputCollection = InParams.InputDataReferences;
-			const FInputVertexInterface& InputInterface = GetVertexInterface().GetInputInterface();
-
-			FAudioBufferReadRef AudioInput = InputCollection.GetDataReadReferenceOrConstruct<FAudioBuffer>(METASOUND_GET_PARAM_NAME(InputAudio), InParams.OperatorSettings);
-			FFloatReadRef ModulationRate = InParams.InputDataReferences.GetDataReadReferenceOrConstructWithVertexDefault<float>(InputInterface, METASOUND_GET_PARAM_NAME(InputModulationRate), InParams.OperatorSettings);
-			FFloatReadRef ModulationDepth = InParams.InputDataReferences.GetDataReadReferenceOrConstructWithVertexDefault<float>(InputInterface, METASOUND_GET_PARAM_NAME(InputModulationDepth), InParams.OperatorSettings);
-			FFloatReadRef CenterDelay = InParams.InputDataReferences.GetDataReadReferenceOrConstructWithVertexDefault<float>(InputInterface, METASOUND_GET_PARAM_NAME(InputCenterDelay), InParams.OperatorSettings);
-			FFloatReadRef MixLevel = InParams.InputDataReferences.GetDataReadReferenceOrConstructWithVertexDefault<float>(InputInterface, METASOUND_GET_PARAM_NAME(InputMixLevel), InParams.OperatorSettings);
+			FAudioBufferReadRef AudioInput = InputData.GetOrConstructDataReadReference<FAudioBuffer>(METASOUND_GET_PARAM_NAME(InputAudio), InParams.OperatorSettings);
+			FFloatReadRef ModulationRate = InputData.GetOrCreateDefaultDataReadReference<float>(METASOUND_GET_PARAM_NAME(InputModulationRate), InParams.OperatorSettings);
+			FFloatReadRef ModulationDepth = InputData.GetOrCreateDefaultDataReadReference<float>(METASOUND_GET_PARAM_NAME(InputModulationDepth), InParams.OperatorSettings);
+			FFloatReadRef CenterDelay = InputData.GetOrCreateDefaultDataReadReference<float>(METASOUND_GET_PARAM_NAME(InputCenterDelay), InParams.OperatorSettings);
+			FFloatReadRef MixLevel = InputData.GetOrCreateDefaultDataReadReference<float>(METASOUND_GET_PARAM_NAME(InputMixLevel), InParams.OperatorSettings);
 
 			return MakeUnique<FFlangerOperator>(InParams, AudioInput, ModulationRate, ModulationDepth, CenterDelay, MixLevel);
 		}
 
-		FFlangerOperator(const FCreateOperatorParams& InParams,
+		FFlangerOperator(const FBuildOperatorParams& InParams,
 			const FAudioBufferReadRef& InAudioInput,
 			const FFloatReadRef& InputModulationRate,
 			const FFloatReadRef& InputModulationDepth,
diff --git a/Engine/Plugins/Runtime/Metasound/Source/MetasoundStandardNodes/Private/MetasoundFreqMultiplierToSemitoneNode.cpp b/Engine/Plugins/Runtime/Metasound/Source/MetasoundStandardNodes/Private/MetasoundFreqMultiplierToSemitoneNode.cpp
--- a/Engine/Plugins/Runtime/Metasound/Source/MetasoundStandardNodes/Private/MetasoundFreqMultiplierToSemitoneNode.cpp
+++ b/Engine/Plugins/Runtime/Metasound/Source/MetasoundStandardNodes/Private/MetasoundFreqMultiplierToSemitoneNode.cpp
@@ -30,7 +30,7 @@ namespace Metasound
 
 		static const FNodeClassMetadata& GetNodeInfo();
 		static const FVertexInterface& GetVertexInterface();
-		static TUniquePtr<IOperator> CreateOperator(const FCreateOperatorParams& InParams, FBuildErrorArray& OutErrors);
+		static TUniquePtr<IOperator> CreateOperator(const FBuildOperatorParams& InParams, FBuildResults& OutResults);
 
 		FFrequencyMultiplierToSemitoneOperator(const FOperatorSettings& InSettings, const FFloatReadRef& InFrequencyMultiplier);
 
@@ -148,14 +148,13 @@ namespace Metasound
 		return Info;
 	}
 
-	TUniquePtr<IOperator> FFrequencyMultiplierToSemitoneOperator::CreateOperator(const FCreateOperatorParams& InParams, FBuildErrorArray& OutErrors)
+	TUniquePtr<IOperator> FFrequencyMultiplierToSemitoneOperator::CreateOperator(const FBuildOperatorParams& InParams, FBuildResults& OutResults)
 	{
 		using namespace FrequencyMultiplierToSemitoneVertexNames;
 
-		const FDataReferenceCollection& InputCollection = InParams.InputDataReferences;
-		const FInputVertexInterface& InputInterface = GetVertexInterface().GetInputInterface();
+		const FInputVertexInterfaceData& InputData = InParams.InputData;
 
-		FFloatReadRef InFrequencyMultiplier = InputCollection.GetDataReadReferenceOrConstructWithVertexDefault<float>(InputInterface, METASOUND_GET_PARAM_NAME(InputFrequencyMultiplier), InParams.OperatorSettings);
+		FFloatReadRef InFrequencyMultiplier = InputData.GetOrCreateDefaultDataReadReference<float>(METASOUND_GET_PARAM_NAME(InputFrequencyMultiplier), InParams.OperatorSettings);
 
 		return MakeUnique<FFrequencyMultiplierToSemitoneOperator>(InParams.OperatorSettings, InFrequencyMultiplier);
 	}
diff --git a/Engine/Plugins/Runtime/Metasound/Source/MetasoundStandardNodes/Private/MetasoundFreqToMIDINode.cpp b/Engine/Plugins/Runtime/Metasound/Source/MetasoundStandardNodes/Private/MetasoundFreqToMIDINode.cpp
--- a/Engine/Plugins/Runtime/Metasound/Source/MetasoundStandardNodes/Private/MetasoundFreqToMIDINode.cpp
+++ b/Engine/Plugins/Runtime/Metasound/Source/MetasoundStandardNodes/Private/MetasoundFreqToMIDINode.cpp
@@ -26,7 +26,7 @@ namespace Metasound
 
 		static const FNodeClassMetadata& GetNodeInfo();
 		static const FVertexInterface& GetVertexInterface();
-		static TUniquePtr<IOperator> CreateOperator(const FCreateOperatorParams& InParams, FBuildErrorArray& OutErrors);
+		static TUniquePtr<IOperator> CreateOperator(const FBuildOperatorParams& InParams, FBuildResults& OutResults);
 
 		FFreqToMidiOperator(const FOperatorSettings& InSettings, const FFloatReadRef& InMidiNote);
 
@@ -145,14 +145,13 @@ namespace Metasound
 		return Info;
 	}
 
-	TUniquePtr<IOperator> FFreqToMidiOperator::CreateOperator(const FCreateOperatorParams& InParams, FBuildErrorArray& OutErrors)
+	TUniquePtr<IOperator> FFreqToMidiOperator::CreateOperator(const FBuildOperatorParams& InParams, FBuildResults& OutResults)
 	{
 		using namespace FrequencyToMidiVertexNames;
 
-		const FDataReferenceCollection& InputCollection = InParams.InputDataReferences;
-		const FInputVertexInterface& InputInterface = GetVertexInterface().GetInputInterface();
+		const FInputVertexInterfaceData& InputData = InParams.InputData;
 
-		FFloatReadRef InFreq = InputCollection.GetDataReadReferenceOrConstructWithVertexDefault<float>(InputInterface, METASOUND_GET_PARAM_NAME(InputFreq), InParams.OperatorSettings);
+		FFloatReadRef InFreq = InputData.GetOrCreateDefaultDataReadReference<float>(METASOUND_GET_PARAM_NAME(InputFreq), InParams.OperatorSettings);
 
 		return MakeUnique<FFreqToMidiOperator>(InParams.OperatorSettings, InFreq);
 	}
diff --git a/Engine/Plugins/Runtime/Metasound/Source/MetasoundStandardNodes/Private/MetasoundGrainDelayNode.cpp b/Engine/Plugins/Runtime/Metasound/Source/MetasoundStandardNodes/Private/MetasoundGrainDelayNode.cpp
--- a/Engine/Plugins/Runtime/Metasound/Source/MetasoundStandardNodes/Private/MetasoundGrainDelayNode.cpp
+++ b/Engine/Plugins/Runtime/Metasound/Source/MetasoundStandardNodes/Private/MetasoundGrainDelayNode.cpp
@@ -66,7 +66,7 @@ namespace Metasound
 
 		static const FNodeClassMetadata& GetNodeInfo();
 		static const FVertexInterface& GetVertexInterface();
-		static TUniquePtr<IOperator> CreateOperator(const FCreateOperatorParams& InParams, FBuildErrorArray& OutErrors);
+		static TUniquePtr<IOperator> CreateOperator(const FBuildOperatorParams& InParams, FBuildResults& OutResults);
 		virtual void BindInputs(FInputVertexInterfaceData& InOutVertexData) override;
 		virtual void BindOutputs(FOutputVertexInterfaceData& InOutVertexData) override;
 		virtual FDataReferenceCollection GetInputs() const override;
@@ -188,24 +188,23 @@ namespace Metasound
 		return DefaultInterface;
 	}
 
-	TUniquePtr<IOperator> FGrainDelayOperator::CreateOperator(const FCreateOperatorParams& InParams, FBuildErrorArray& OutErrors)
+	TUniquePtr<IOperator> FGrainDelayOperator::CreateOperator(const FBuildOperatorParams& InParams, FBuildResults& OutResults)
 	{
 		using namespace GrainDelay;
 
-		const FDataReferenceCollection& InputCollection = InParams.InputDataReferences;
-		const FInputVertexInterface& InputInterface = GetVertexInterface().GetInputInterface();
+		const FInputVertexInterfaceData& InputData = InParams.InputData;
 		
-		FAudioBufferReadRef AudioInput = InputCollection.GetDataReadReferenceOrConstruct<FAudioBuffer>(METASOUND_GET_PARAM_NAME(InAudio), InParams.OperatorSettings);
-		FTriggerReadRef GrainSpawnTrigger = InputCollection.GetDataReadReferenceOrConstruct<FTrigger>(METASOUND_GET_PARAM_NAME(InGrainSpawnTrigger), InParams.OperatorSettings);
-		FFloatReadRef GrainDelay = InputCollection.GetDataReadReferenceOrConstructWithVertexDefault<float>(InputInterface, METASOUND_GET_PARAM_NAME(InGrainDelay), InParams.OperatorSettings);
-		FFloatReadRef GrainDelayDelta = InputCollection.GetDataReadReferenceOrConstructWithVertexDefault<float>(InputInterface, METASOUND_GET_PARAM_NAME(InGrainDelayDelta), InParams.OperatorSettings);
-		FFloatReadRef GrainDuration = InputCollection.GetDataReadReferenceOrConstructWithVertexDefault<float>(InputInterface, METASOUND_GET_PARAM_NAME(InGrainDuration), InParams.OperatorSettings);
-		FFloatReadRef GrainDurationDelta = InputCollection.GetDataReadReferenceOrConstructWithVertexDefault<float>(InputInterface, METASOUND_GET_PARAM_NAME(InGrainDurationDelta), InParams.OperatorSettings);
-		FFloatReadRef GrainPitchShift = InputCollection.GetDataReadReferenceOrConstructWithVertexDefault<float>(InputInterface, METASOUND_GET_PARAM_NAME(InGrainPitchShift), InParams.OperatorSettings);
-		FFloatReadRef GrainPitchShiftDelta = InputCollection.GetDataReadReferenceOrConstructWithVertexDefault<float>(InputInterface, METASOUND_GET_PARAM_NAME(InGrainPitchShiftDelta), InParams.OperatorSettings);
-		FGrainDelayEnvelopeReadRef GrainDelayEnvelope = InputCollection.GetDataReadReferenceOrConstructWithVertexDefault<FEnumGrainDelayEnvelope>(InputInterface, METASOUND_GET_PARAM_NAME(InGrainDelayEnvelope), InParams.OperatorSettings);
-		FInt32ReadRef GrainMaxCount = InputCollection.GetDataReadReferenceOrConstructWithVertexDefault<int32>(InputInterface, METASOUND_GET_PARAM_NAME(InGrainMaxCount), InParams.OperatorSettings);
-		FFloatReadRef GrainDelayFeedbackAmount = InputCollection.GetDataReadReferenceOrConstructWithVertexDefault<float>(InputInterface, METASOUND_GET_PARAM_NAME(InGrainDelayFeedbackAmount), InParams.OperatorSettings);
+		FAudioBufferReadRef AudioInput = InputData.GetOrConstructDataReadReference<FAudioBuffer>(METASOUND_GET_PARAM_NAME(InAudio), InParams.OperatorSettings);
+		FTriggerReadRef GrainSpawnTrigger = InputData.GetOrConstructDataReadReference<FTrigger>(METASOUND_GET_PARAM_NAME(InGrainSpawnTrigger), InParams.OperatorSettings);
+		FFloatReadRef GrainDelay = InputData.GetOrCreateDefaultDataReadReference<float>(METASOUND_GET_PARAM_NAME(InGrainDelay), InParams.OperatorSettings);
+		FFloatReadRef GrainDelayDelta = InputData.GetOrCreateDefaultDataReadReference<float>(METASOUND_GET_PARAM_NAME(InGrainDelayDelta), InParams.OperatorSettings);
+		FFloatReadRef GrainDuration = InputData.GetOrCreateDefaultDataReadReference<float>(METASOUND_GET_PARAM_NAME(InGrainDuration), InParams.OperatorSettings);
+		FFloatReadRef GrainDurationDelta = InputData.GetOrCreateDefaultDataReadReference<float>(METASOUND_GET_PARAM_NAME(InGrainDurationDelta), InParams.OperatorSettings);
+		FFloatReadRef GrainPitchShift = InputData.GetOrCreateDefaultDataReadReference<float>(METASOUND_GET_PARAM_NAME(InGrainPitchShift), InParams.OperatorSettings);
+		FFloatReadRef GrainPitchShiftDelta = InputData.GetOrCreateDefaultDataReadReference<float>(METASOUND_GET_PARAM_NAME(InGrainPitchShiftDelta), InParams.OperatorSettings);
+		FGrainDelayEnvelopeReadRef GrainDelayEnvelope = InputData.GetOrCreateDefaultDataReadReference<FEnumGrainDelayEnvelope>(METASOUND_GET_PARAM_NAME(InGrainDelayEnvelope), InParams.OperatorSettings);
+		FInt32ReadRef GrainMaxCount = InputData.GetOrCreateDefaultDataReadReference<int32>(METASOUND_GET_PARAM_NAME(InGrainMaxCount), InParams.OperatorSettings);
+		FFloatReadRef GrainDelayFeedbackAmount = InputData.GetOrCreateDefaultDataReadReference<float>(METASOUND_GET_PARAM_NAME(InGrainDelayFeedbackAmount), InParams.OperatorSettings);
 
 		return MakeUnique<FGrainDelayOperator>(InParams.OperatorSettings, 
 			AudioInput,
diff --git a/Engine/Plugins/Runtime/Metasound/Source/MetasoundStandardNodes/Private/MetasoundITDPannerNode.cpp b/Engine/Plugins/Runtime/Metasound/Source/MetasoundStandardNodes/Private/MetasoundITDPannerNode.cpp
--- a/Engine/Plugins/Runtime/Metasound/Source/MetasoundStandardNodes/Private/MetasoundITDPannerNode.cpp
+++ b/Engine/Plugins/Runtime/Metasound/Source/MetasoundStandardNodes/Private/MetasoundITDPannerNode.cpp
@@ -38,7 +38,7 @@ namespace Metasound
 
 		static const FNodeClassMetadata& GetNodeInfo();
 		static const FVertexInterface& GetVertexInterface();
-		static TUniquePtr<IOperator> CreateOperator(const FCreateOperatorParams& InParams, FBuildErrorArray& OutErrors);
+		static TUniquePtr<IOperator> CreateOperator(const FBuildOperatorParams& InParams, FBuildResults& OutResults);
 
 		FITDPannerOperator(const FOperatorSettings& InSettings,
 			const FAudioBufferReadRef& InAudioInput, 
@@ -299,17 +299,17 @@ namespace Metasound
 		return Info;
 	}
 
-	TUniquePtr<IOperator> FITDPannerOperator::CreateOperator(const FCreateOperatorParams& InParams, FBuildErrorArray& OutErrors)
+	TUniquePtr<IOperator> FITDPannerOperator::CreateOperator(const FBuildOperatorParams& InParams, FBuildResults& OutResults)
 	{
-		const FDataReferenceCollection& InputCollection = InParams.InputDataReferences;
-		const FInputVertexInterface& InputInterface = GetVertexInterface().GetInputInterface();
+		
+		const FInputVertexInterfaceData& InputData = InParams.InputData;
 
 		using namespace ITDPannerVertexNames;
 
-		FAudioBufferReadRef AudioIn = InputCollection.GetDataReadReferenceOrConstruct<FAudioBuffer>(METASOUND_GET_PARAM_NAME(InputAudio), InParams.OperatorSettings);
-		FFloatReadRef PanningAngle = InputCollection.GetDataReadReferenceOrConstructWithVertexDefault<float>(InputInterface, METASOUND_GET_PARAM_NAME(InputPanAngle), InParams.OperatorSettings);
-		FFloatReadRef DistanceFactor = InputCollection.GetDataReadReferenceOrConstructWithVertexDefault<float>(InputInterface, METASOUND_GET_PARAM_NAME(InputDistanceFactor), InParams.OperatorSettings);
-		FFloatReadRef HeadWidth = InputCollection.GetDataReadReferenceOrConstructWithVertexDefault<float>(InputInterface, METASOUND_GET_PARAM_NAME(InputHeadWidth), InParams.OperatorSettings);
+		FAudioBufferReadRef AudioIn = InputData.GetOrConstructDataReadReference<FAudioBuffer>(METASOUND_GET_PARAM_NAME(InputAudio), InParams.OperatorSettings);
+		FFloatReadRef PanningAngle = InputData.GetOrCreateDefaultDataReadReference<float>(METASOUND_GET_PARAM_NAME(InputPanAngle), InParams.OperatorSettings);
+		FFloatReadRef DistanceFactor = InputData.GetOrCreateDefaultDataReadReference<float>(METASOUND_GET_PARAM_NAME(InputDistanceFactor), InParams.OperatorSettings);
+		FFloatReadRef HeadWidth = InputData.GetOrCreateDefaultDataReadReference<float>(METASOUND_GET_PARAM_NAME(InputHeadWidth), InParams.OperatorSettings);
 
 		return MakeUnique<FITDPannerOperator>(InParams.OperatorSettings, AudioIn, PanningAngle, DistanceFactor, HeadWidth);
 	}
diff --git a/Engine/Plugins/Runtime/Metasound/Source/MetasoundStandardNodes/Private/MetasoundInterpToNode.cpp b/Engine/Plugins/Runtime/Metasound/Source/MetasoundStandardNodes/Private/MetasoundInterpToNode.cpp
--- a/Engine/Plugins/Runtime/Metasound/Source/MetasoundStandardNodes/Private/MetasoundInterpToNode.cpp
+++ b/Engine/Plugins/Runtime/Metasound/Source/MetasoundStandardNodes/Private/MetasoundInterpToNode.cpp
@@ -43,9 +43,9 @@ namespace Metasound
 
 		static const FNodeClassMetadata& GetNodeInfo();
 		static const FVertexInterface& GetVertexInterface();
-		static TUniquePtr<IOperator> CreateOperator(const FCreateOperatorParams& InParams, FBuildErrorArray& OutErrors);
+		static TUniquePtr<IOperator> CreateOperator(const FBuildOperatorParams& InParams, FBuildResults& OutResults);
 
-		FInterpToOperator(const FCreateOperatorParams& InSettings, const FFloatReadRef& InTargetValue, const FTimeReadRef& InInterpTime);
+		FInterpToOperator(const FBuildOperatorParams& InParams, const FFloatReadRef& InTargetValue, const FTimeReadRef& InInterpTime);
 
 		virtual void BindInputs(FInputVertexInterfaceData& InOutVertexData) override;
 		virtual void BindOutputs(FOutputVertexInterfaceData& InOutVertexData) override;
@@ -74,7 +74,7 @@ namespace Metasound
 		float PreviousTargetValue = 0.0f;
 	};
 
-	FInterpToOperator::FInterpToOperator(const FCreateOperatorParams& InParams, const FFloatReadRef& InTargetValue, const FTimeReadRef& InInterpTime)
+	FInterpToOperator::FInterpToOperator(const FBuildOperatorParams& InParams, const FFloatReadRef& InTargetValue, const FTimeReadRef& InInterpTime)
 		: TargetValue(InTargetValue)
 		, InterpTime(InInterpTime)
 		, ValueOutput(FFloatWriteRef::CreateNew(*TargetValue))
@@ -195,16 +195,14 @@ namespace Metasound
 	{
 	}
 
-	TUniquePtr<IOperator> FInterpToOperator::CreateOperator(const FCreateOperatorParams& InParams, FBuildErrorArray& OutErrors)
+	TUniquePtr<IOperator> FInterpToOperator::CreateOperator(const FBuildOperatorParams& InParams, FBuildResults& OutResults)
 	{
 		using namespace InterpToVertexNames; 
 
-		const FInterpToNode& InterpToNode = static_cast<const FInterpToNode&>(InParams.Node);
-		const FDataReferenceCollection& InputCollection = InParams.InputDataReferences;
-		const FInputVertexInterface& InputInterface = GetVertexInterface().GetInputInterface();
+		const FInputVertexInterfaceData& InputData = InParams.InputData;
 
-		FFloatReadRef TargetValue = InputCollection.GetDataReadReferenceOrConstructWithVertexDefault<float>(InputInterface, METASOUND_GET_PARAM_NAME(InParamTarget), InParams.OperatorSettings);
-		FTimeReadRef InterpTime = InputCollection.GetDataReadReferenceOrConstructWithVertexDefault<FTime>(InputInterface, METASOUND_GET_PARAM_NAME(InParamInterpTime), InParams.OperatorSettings);
+		FFloatReadRef TargetValue = InputData.GetOrCreateDefaultDataReadReference<float>(METASOUND_GET_PARAM_NAME(InParamTarget), InParams.OperatorSettings);
+		FTimeReadRef InterpTime = InputData.GetOrCreateDefaultDataReadReference<FTime>(METASOUND_GET_PARAM_NAME(InParamInterpTime), InParams.OperatorSettings);
 
 		return MakeUnique<FInterpToOperator>(InParams, TargetValue, InterpTime);
 	}
diff --git a/Engine/Plugins/Runtime/Metasound/Source/MetasoundStandardNodes/Private/MetasoundLimiterNode.cpp b/Engine/Plugins/Runtime/Metasound/Source/MetasoundStandardNodes/Private/MetasoundLimiterNode.cpp
--- a/Engine/Plugins/Runtime/Metasound/Source/MetasoundStandardNodes/Private/MetasoundLimiterNode.cpp
+++ b/Engine/Plugins/Runtime/Metasound/Source/MetasoundStandardNodes/Private/MetasoundLimiterNode.cpp
@@ -52,7 +52,7 @@ namespace Metasound
 		static constexpr float SoftKneeBandwitdh = 10.0f;
 		static constexpr float MaxInputGain = 100.0f;
 
-		FLimiterOperator(const FCreateOperatorParams& InParams,
+		FLimiterOperator(const FBuildOperatorParams& InParams,
 			const FAudioBufferReadRef& InAudio,
 			const FFloatReadRef& InGainDb,
 			const FFloatReadRef& InThresholdDb,
@@ -154,17 +154,17 @@ namespace Metasound
 			return {};
 		}
 
-		static TUniquePtr<IOperator> CreateOperator(const FCreateOperatorParams& InParams, FBuildErrorArray& OutErrors)
+		static TUniquePtr<IOperator> CreateOperator(const FBuildOperatorParams& InParams, FBuildResults& OutResults)
 		{
 			using namespace LimiterVertexNames;
-			const FDataReferenceCollection& Inputs = InParams.InputDataReferences;
-			const FInputVertexInterface& InputInterface = DeclareVertexInterface().GetInputInterface();
-
-			FAudioBufferReadRef AudioIn = Inputs.GetDataReadReferenceOrConstruct<FAudioBuffer>(METASOUND_GET_PARAM_NAME(InputAudio), InParams.OperatorSettings);
-			FFloatReadRef InGainDbIn = Inputs.GetDataReadReferenceOrConstructWithVertexDefault<float>(InputInterface, METASOUND_GET_PARAM_NAME(InputInGainDb), InParams.OperatorSettings);
-			FFloatReadRef ThresholdDbIn = Inputs.GetDataReadReferenceOrConstructWithVertexDefault<float>(InputInterface, METASOUND_GET_PARAM_NAME(InputThresholdDb), InParams.OperatorSettings);
-			FTimeReadRef ReleaseTimeIn = Inputs.GetDataReadReferenceOrConstructWithVertexDefault<FTime>(InputInterface, METASOUND_GET_PARAM_NAME(InputReleaseTime), InParams.OperatorSettings);
-			FKneeModeReadRef KneeModeIn = Inputs.GetDataReadReferenceOrConstructWithVertexDefault<FEnumKneeMode>(InputInterface, METASOUND_GET_PARAM_NAME(InputKneeMode), InParams.OperatorSettings);
+			
+			const FInputVertexInterfaceData& InputData = InParams.InputData;
+
+			FAudioBufferReadRef AudioIn = InputData.GetOrConstructDataReadReference<FAudioBuffer>(METASOUND_GET_PARAM_NAME(InputAudio), InParams.OperatorSettings);
+			FFloatReadRef InGainDbIn = InputData.GetOrCreateDefaultDataReadReference<float>(METASOUND_GET_PARAM_NAME(InputInGainDb), InParams.OperatorSettings);
+			FFloatReadRef ThresholdDbIn = InputData.GetOrCreateDefaultDataReadReference<float>(METASOUND_GET_PARAM_NAME(InputThresholdDb), InParams.OperatorSettings);
+			FTimeReadRef ReleaseTimeIn = InputData.GetOrCreateDefaultDataReadReference<FTime>(METASOUND_GET_PARAM_NAME(InputReleaseTime), InParams.OperatorSettings);
+			FKneeModeReadRef KneeModeIn = InputData.GetOrCreateDefaultDataReadReference<FEnumKneeMode>(METASOUND_GET_PARAM_NAME(InputKneeMode), InParams.OperatorSettings);
 
 			return MakeUnique<FLimiterOperator>(InParams, AudioIn, InGainDbIn, ThresholdDbIn, ReleaseTimeIn, KneeModeIn);
 		}
diff --git a/Engine/Plugins/Runtime/Metasound/Source/MetasoundStandardNodes/Private/MetasoundLinearGainToDecibels.cpp b/Engine/Plugins/Runtime/Metasound/Source/MetasoundStandardNodes/Private/MetasoundLinearGainToDecibels.cpp
--- a/Engine/Plugins/Runtime/Metasound/Source/MetasoundStandardNodes/Private/MetasoundLinearGainToDecibels.cpp
+++ b/Engine/Plugins/Runtime/Metasound/Source/MetasoundStandardNodes/Private/MetasoundLinearGainToDecibels.cpp
@@ -27,9 +27,9 @@ namespace Metasound
 
 		static const FNodeClassMetadata& GetNodeInfo();
 		static const FVertexInterface& GetVertexInterface();
-		static TUniquePtr<IOperator> CreateOperator(const FCreateOperatorParams& InParams, FBuildErrorArray& OutErrors);
+		static TUniquePtr<IOperator> CreateOperator(const FBuildOperatorParams& InParams, FBuildResults& OutResults);
 
-		FLinearGainToDecibelsOperator(const FCreateOperatorParams& InParams, const FFloatReadRef& InLinearGain);
+		FLinearGainToDecibelsOperator(const FBuildOperatorParams& InParams, const FFloatReadRef& InLinearGain);
 
 		virtual void BindInputs(FInputVertexInterfaceData& InOutVertexData) override;
 		virtual void BindOutputs(FOutputVertexInterfaceData& InOutVertexData) override;
@@ -46,7 +46,7 @@ namespace Metasound
 		FFloatWriteRef DecibelOutput;
 	};
 
-	FLinearGainToDecibelsOperator::FLinearGainToDecibelsOperator(const FCreateOperatorParams& InParams, const FFloatReadRef& InLinearGain)
+	FLinearGainToDecibelsOperator::FLinearGainToDecibelsOperator(const FBuildOperatorParams& InParams, const FFloatReadRef& InLinearGain)
 		: LinearGainInput(InLinearGain)
 		, DecibelOutput(FFloatWriteRef::CreateNew(0))
 	{
@@ -142,14 +142,13 @@ namespace Metasound
 		return Info;
 	}
 
-	TUniquePtr<IOperator> FLinearGainToDecibelsOperator::CreateOperator(const FCreateOperatorParams& InParams, FBuildErrorArray& OutErrors)
+	TUniquePtr<IOperator> FLinearGainToDecibelsOperator::CreateOperator(const FBuildOperatorParams& InParams, FBuildResults& OutResults)
 	{
 		using namespace LinearGainToDecibelsVertexNames;
+		
+		const FInputVertexInterfaceData& InputData = InParams.InputData;
 
-		const FDataReferenceCollection& InputCollection = InParams.InputDataReferences;
-		const FInputVertexInterface& InputInterface = GetVertexInterface().GetInputInterface();
-
-		FFloatReadRef InLinearGain = InputCollection.GetDataReadReferenceOrConstructWithVertexDefault<float>(InputInterface, METASOUND_GET_PARAM_NAME(InputLinearGain), InParams.OperatorSettings);
+		FFloatReadRef InLinearGain = InputData.GetOrCreateDefaultDataReadReference<float>(METASOUND_GET_PARAM_NAME(InputLinearGain), InParams.OperatorSettings);
 		
 		return MakeUnique<FLinearGainToDecibelsOperator>(InParams, InLinearGain);
 	}
diff --git a/Engine/Plugins/Runtime/Metasound/Source/MetasoundStandardNodes/Private/MetasoundLinearToLogFrequencyNode.cpp b/Engine/Plugins/Runtime/Metasound/Source/MetasoundStandardNodes/Private/MetasoundLinearToLogFrequencyNode.cpp
--- a/Engine/Plugins/Runtime/Metasound/Source/MetasoundStandardNodes/Private/MetasoundLinearToLogFrequencyNode.cpp
+++ b/Engine/Plugins/Runtime/Metasound/Source/MetasoundStandardNodes/Private/MetasoundLinearToLogFrequencyNode.cpp
@@ -76,18 +76,17 @@ namespace Metasound
 			return Interface;
 		}
 
-		static TUniquePtr<IOperator> CreateOperator(const FCreateOperatorParams& InParams, FBuildErrorArray& OutErrors)
+		static TUniquePtr<IOperator> CreateOperator(const FBuildOperatorParams& InParams, FBuildResults& OutResults)
 		{
 			using namespace LinearToLogFrequencyVertexNames;
-
-			const FDataReferenceCollection& InputCollection = InParams.InputDataReferences;
-			const FInputVertexInterface& InputInterface = GetVertexInterface().GetInputInterface();
-
-			TDataReadReference<float> InValue = InputCollection.GetDataReadReferenceOrConstructWithVertexDefault<float>(InputInterface, METASOUND_GET_PARAM_NAME(InputValue), InParams.OperatorSettings);
-			TDataReadReference<float> InDomainMin = InputCollection.GetDataReadReferenceOrConstructWithVertexDefault<float>(InputInterface, METASOUND_GET_PARAM_NAME(InputDomainMin), InParams.OperatorSettings);
-			TDataReadReference<float> InDomainMax = InputCollection.GetDataReadReferenceOrConstructWithVertexDefault<float>(InputInterface, METASOUND_GET_PARAM_NAME(InputDomainMax), InParams.OperatorSettings);
-			TDataReadReference<float> InRangeMin = InputCollection.GetDataReadReferenceOrConstructWithVertexDefault<float>(InputInterface, METASOUND_GET_PARAM_NAME(InputRangeMin), InParams.OperatorSettings);
-			TDataReadReference<float> InRangeMax = InputCollection.GetDataReadReferenceOrConstructWithVertexDefault<float>(InputInterface, METASOUND_GET_PARAM_NAME(InputRangeMax), InParams.OperatorSettings);
+			
+			const FInputVertexInterfaceData& InputData = InParams.InputData;
+
+			TDataReadReference<float> InValue = InputData.GetOrCreateDefaultDataReadReference<float>(METASOUND_GET_PARAM_NAME(InputValue), InParams.OperatorSettings);
+			TDataReadReference<float> InDomainMin = InputData.GetOrCreateDefaultDataReadReference<float>(METASOUND_GET_PARAM_NAME(InputDomainMin), InParams.OperatorSettings);
+			TDataReadReference<float> InDomainMax = InputData.GetOrCreateDefaultDataReadReference<float>(METASOUND_GET_PARAM_NAME(InputDomainMax), InParams.OperatorSettings);
+			TDataReadReference<float> InRangeMin = InputData.GetOrCreateDefaultDataReadReference<float>(METASOUND_GET_PARAM_NAME(InputRangeMin), InParams.OperatorSettings);
+			TDataReadReference<float> InRangeMax = InputData.GetOrCreateDefaultDataReadReference<float>(METASOUND_GET_PARAM_NAME(InputRangeMax), InParams.OperatorSettings);
 
 			return MakeUnique<FLinearToLogFrequencyOperator>(InValue, InDomainMin, InDomainMax, InRangeMin, InRangeMax);
 		}
diff --git a/Engine/Plugins/Runtime/Metasound/Source/MetasoundStandardNodes/Private/MetasoundLowFrequencyNoise.cpp b/Engine/Plugins/Runtime/Metasound/Source/MetasoundStandardNodes/Private/MetasoundLowFrequencyNoise.cpp
--- a/Engine/Plugins/Runtime/Metasound/Source/MetasoundStandardNodes/Private/MetasoundLowFrequencyNoise.cpp
+++ b/Engine/Plugins/Runtime/Metasound/Source/MetasoundStandardNodes/Private/MetasoundLowFrequencyNoise.cpp
@@ -155,32 +155,30 @@ namespace Metasound
 			FFloatReadRef MaxOutput;
 		};
 
-		static TUniquePtr<IOperator> CreateOperator(const FCreateOperatorParams& InParams, FBuildErrorArray& OutErrors)
+		static TUniquePtr<IOperator> CreateOperator(const FBuildOperatorParams& InParams, FBuildResults& OutResults)
 		{
-			//const FRandomLowFrequencyNoiseNode& Node = static_cast<const FRandomLowFrequencyNoiseNode&>(InParams.Node);
-			const FDataReferenceCollection& InputCol = InParams.InputDataReferences;
 			const FOperatorSettings& Settings = InParams.OperatorSettings;
-			const FInputVertexInterface& InputInterface = GetVertexInterface().GetInputInterface();
+			const FInputVertexInterfaceData& InputData = InParams.InputData;
 
 			using namespace LowFrequencyNoiseVertexNames;
 
 			FPinReadRefs Pins
 			{
-				  InputCol.GetDataReadReferenceOrConstructWithVertexDefault<float>(InputInterface, METASOUND_GET_PARAM_NAME(RatePin), Settings)
-				, InputCol.GetDataReadReferenceOrConstructWithVertexDefault<float>(InputInterface, METASOUND_GET_PARAM_NAME(RateJitterPin), Settings)
-				, InputCol.GetDataReadReferenceOrConstructWithVertexDefault<float>(InputInterface, METASOUND_GET_PARAM_NAME(StepLimitPin), Settings)
-				, InputCol.GetDataReadReferenceOrConstructWithVertexDefault<int32>(InputInterface, METASOUND_GET_PARAM_NAME(SeedPin), Settings)
-				, InputCol.GetDataReadReferenceOrConstructWithVertexDefault<FEnumLowFrequencyNoiseInterpType>(InputInterface, METASOUND_GET_PARAM_NAME(InterpTypePin), Settings)
-				, InputCol.GetDataReadReferenceOrConstructWithVertexDefault<FTrigger>(InputInterface, METASOUND_GET_PARAM_NAME(ResetSeedPin), Settings)
-				, InputCol.GetDataReadReferenceOrConstructWithVertexDefault<FTrigger>(InputInterface, METASOUND_GET_PARAM_NAME(SyncPin), Settings)
-				, InputCol.GetDataReadReferenceOrConstructWithVertexDefault<float>(InputInterface, METASOUND_GET_PARAM_NAME(MinOutputValuePin), Settings)
-				, InputCol.GetDataReadReferenceOrConstructWithVertexDefault<float>(InputInterface, METASOUND_GET_PARAM_NAME(MaxOutputValuePin), Settings)
+				  InputData.GetOrCreateDefaultDataReadReference<float>(METASOUND_GET_PARAM_NAME(RatePin), Settings)
+				, InputData.GetOrCreateDefaultDataReadReference<float>(METASOUND_GET_PARAM_NAME(RateJitterPin), Settings)
+				, InputData.GetOrCreateDefaultDataReadReference<float>(METASOUND_GET_PARAM_NAME(StepLimitPin), Settings)
+				, InputData.GetOrCreateDefaultDataReadReference<int32>(METASOUND_GET_PARAM_NAME(SeedPin), Settings)
+				, InputData.GetOrCreateDefaultDataReadReference<FEnumLowFrequencyNoiseInterpType>(METASOUND_GET_PARAM_NAME(InterpTypePin), Settings)
+				, InputData.GetOrCreateDefaultDataReadReference<FTrigger>(METASOUND_GET_PARAM_NAME(ResetSeedPin), Settings)
+				, InputData.GetOrCreateDefaultDataReadReference<FTrigger>(METASOUND_GET_PARAM_NAME(SyncPin), Settings)
+				, InputData.GetOrCreateDefaultDataReadReference<float>(METASOUND_GET_PARAM_NAME(MinOutputValuePin), Settings)
+				, InputData.GetOrCreateDefaultDataReadReference<float>(METASOUND_GET_PARAM_NAME(MaxOutputValuePin), Settings)
 			};
 
 			return MakeUnique<FLowFrequencyNoiseOperator>(InParams, MoveTemp(Pins));
 		}
 
-		FLowFrequencyNoiseOperator(const FCreateOperatorParams& InParams, FPinReadRefs&& InPins)
+		FLowFrequencyNoiseOperator(const FBuildOperatorParams& InParams, FPinReadRefs&& InPins)
 			: BlockRate{ InParams.OperatorSettings.GetActualBlockRate() }
 			, Pins{ MoveTemp(InPins) }
 			, ScaledOutput{ FFloatWriteRef::CreateNew(0.f) }
diff --git a/Engine/Plugins/Runtime/Metasound/Source/MetasoundStandardNodes/Private/MetasoundMIDIToFreqNode.cpp b/Engine/Plugins/Runtime/Metasound/Source/MetasoundStandardNodes/Private/MetasoundMIDIToFreqNode.cpp
--- a/Engine/Plugins/Runtime/Metasound/Source/MetasoundStandardNodes/Private/MetasoundMIDIToFreqNode.cpp
+++ b/Engine/Plugins/Runtime/Metasound/Source/MetasoundStandardNodes/Private/MetasoundMIDIToFreqNode.cpp
@@ -64,9 +64,9 @@ namespace Metasound
 
 		static const FNodeClassMetadata& GetNodeInfo();
 		static const FVertexInterface& GetVertexInterface();
-		static TUniquePtr<IOperator> CreateOperator(const FCreateOperatorParams& InParams, FBuildErrorArray& OutErrors);
+		static TUniquePtr<IOperator> CreateOperator(const FBuildOperatorParams& InParams, FBuildResults& OutResults);
 
-		TMidiToFreqOperator(const FCreateOperatorParams& InParams, const TDataReadReference<ValueType>& InMidiNote);
+		TMidiToFreqOperator(const FBuildOperatorParams& InParams, const TDataReadReference<ValueType>& InMidiNote);
 
 		virtual void BindInputs(FInputVertexInterfaceData& InOutVertexData) override;
 		virtual void BindOutputs(FOutputVertexInterfaceData& InOutVertexData) override;
@@ -87,7 +87,7 @@ namespace Metasound
 	};
 
 	template<typename ValueType>
-	TMidiToFreqOperator<ValueType>::TMidiToFreqOperator(const FCreateOperatorParams& InParams, const TDataReadReference<ValueType>& InMidiNote)
+	TMidiToFreqOperator<ValueType>::TMidiToFreqOperator(const FBuildOperatorParams& InParams, const TDataReadReference<ValueType>& InMidiNote)
 		: MidiNote(InMidiNote)
 		, FreqOutput(FFloatWriteRef::CreateNew(Audio::GetFrequencyFromMidi(*InMidiNote)))
 		, PrevMidiNote(*InMidiNote)
@@ -199,14 +199,12 @@ namespace Metasound
 	}
 
 	template<typename ValueType>
-	TUniquePtr<IOperator> TMidiToFreqOperator<ValueType>::CreateOperator(const FCreateOperatorParams& InParams, FBuildErrorArray& OutErrors)
+	TUniquePtr<IOperator> TMidiToFreqOperator<ValueType>::CreateOperator(const FBuildOperatorParams& InParams, FBuildResults& OutResults)
 	{
 		using namespace MidiToFrequencyVertexNames;
+		const FInputVertexInterfaceData& InputData = InParams.InputData;
 
-		const FDataReferenceCollection& InputCollection = InParams.InputDataReferences;
-		const FInputVertexInterface& InputInterface = GetVertexInterface().GetInputInterface();
-
-		TDataReadReference<ValueType> InMidiNote = InputCollection.GetDataReadReferenceOrConstructWithVertexDefault<ValueType>(InputInterface, METASOUND_GET_PARAM_NAME(InputMidi), InParams.OperatorSettings);
+		TDataReadReference<ValueType> InMidiNote = InputData.GetOrCreateDefaultDataReadReference<ValueType>(METASOUND_GET_PARAM_NAME(InputMidi), InParams.OperatorSettings);
 
 		return MakeUnique<TMidiToFreqOperator>(InParams, InMidiNote);
 	}
diff --git a/Engine/Plugins/Runtime/Metasound/Source/MetasoundStandardNodes/Private/MetasoundMaxNode.cpp b/Engine/Plugins/Runtime/Metasound/Source/MetasoundStandardNodes/Private/MetasoundMaxNode.cpp
--- a/Engine/Plugins/Runtime/Metasound/Source/MetasoundStandardNodes/Private/MetasoundMaxNode.cpp
+++ b/Engine/Plugins/Runtime/Metasound/Source/MetasoundStandardNodes/Private/MetasoundMaxNode.cpp
@@ -58,18 +58,18 @@ namespace Metasound
 				OutMax = FMath::Max(InA, InB);
 			}
 
-			static TDataReadReference<int32> CreateInRefA(const FCreateOperatorParams& InParams)
+			static TDataReadReference<int32> CreateInRefA(const FBuildOperatorParams& InParams)
 			{
 				using namespace MaxVertexNames;
-				const FInputVertexInterface& InputInterface = InParams.Node.GetVertexInterface().GetInputInterface();
-				return InParams.InputDataReferences.GetDataReadReferenceOrConstructWithVertexDefault<int32>(InputInterface, METASOUND_GET_PARAM_NAME(InputAValue), InParams.OperatorSettings);
+				const FInputVertexInterfaceData& InputData = InParams.InputData;
+				return InputData.GetOrCreateDefaultDataReadReference<int32>(METASOUND_GET_PARAM_NAME(InputAValue), InParams.OperatorSettings);
 			}
 
-			static TDataReadReference<int32> CreateInRefB(const FCreateOperatorParams& InParams)
+			static TDataReadReference<int32> CreateInRefB(const FBuildOperatorParams& InParams)
 			{
 				using namespace MaxVertexNames;
-				const FInputVertexInterface& InputInterface = InParams.Node.GetVertexInterface().GetInputInterface();
-				return InParams.InputDataReferences.GetDataReadReferenceOrConstructWithVertexDefault<int32>(InputInterface, METASOUND_GET_PARAM_NAME(InputBValue), InParams.OperatorSettings);
+				const FInputVertexInterfaceData& InputData = InParams.InputData;
+				return InputData.GetOrCreateDefaultDataReadReference<int32>(METASOUND_GET_PARAM_NAME(InputBValue), InParams.OperatorSettings);
 			}
 		};
 
@@ -82,18 +82,18 @@ namespace Metasound
 				OutMax = FMath::Max(InA, InB);
 			}
 
-			static TDataReadReference<float> CreateInRefA(const FCreateOperatorParams& InParams)
+			static TDataReadReference<float> CreateInRefA(const FBuildOperatorParams& InParams)
 			{
 				using namespace MaxVertexNames;
-				const FInputVertexInterface& InputInterface = InParams.Node.GetVertexInterface().GetInputInterface();
-				return InParams.InputDataReferences.GetDataReadReferenceOrConstructWithVertexDefault<float>(InputInterface, METASOUND_GET_PARAM_NAME(InputAValue), InParams.OperatorSettings);
+				const FInputVertexInterfaceData& InputData = InParams.InputData;
+				return InputData.GetOrCreateDefaultDataReadReference<float>(METASOUND_GET_PARAM_NAME(InputAValue), InParams.OperatorSettings);
 			}
 
-			static TDataReadReference<float> CreateInRefB(const FCreateOperatorParams& InParams)
+			static TDataReadReference<float> CreateInRefB(const FBuildOperatorParams& InParams)
 			{
 				using namespace MaxVertexNames;
-				const FInputVertexInterface& InputInterface = InParams.Node.GetVertexInterface().GetInputInterface();
-				return InParams.InputDataReferences.GetDataReadReferenceOrConstructWithVertexDefault<float>(InputInterface, METASOUND_GET_PARAM_NAME(InputBValue), InParams.OperatorSettings);
+				const FInputVertexInterfaceData& InputData = InParams.InputData;
+				return InputData.GetOrCreateDefaultDataReadReference<float>(METASOUND_GET_PARAM_NAME(InputBValue), InParams.OperatorSettings);
 			}
 
 			static bool IsAudioBuffer() { return false; }
@@ -114,16 +114,18 @@ namespace Metasound
 				}
 			}
 
-			static TDataReadReference<FAudioBuffer> CreateInRefA(const FCreateOperatorParams& InParams)
+			static TDataReadReference<FAudioBuffer> CreateInRefA(const FBuildOperatorParams& InParams)
 			{
 				using namespace MaxVertexNames;
-				return InParams.InputDataReferences.GetDataReadReferenceOrConstruct<FAudioBuffer>(METASOUND_GET_PARAM_NAME(InputAValue), InParams.OperatorSettings);
+				const FInputVertexInterfaceData& InputData = InParams.InputData;
+				return InputData.GetOrConstructDataReadReference<FAudioBuffer>(METASOUND_GET_PARAM_NAME(InputAValue), InParams.OperatorSettings);
 			}
 
-			static TDataReadReference<FAudioBuffer> CreateInRefB(const FCreateOperatorParams& InParams)
+			static TDataReadReference<FAudioBuffer> CreateInRefB(const FBuildOperatorParams& InParams)
 			{
-				using namespace MaxVertexNames;
-				return InParams.InputDataReferences.GetDataReadReferenceOrConstruct<FAudioBuffer>(METASOUND_GET_PARAM_NAME(InputBValue), InParams.OperatorSettings);
+				using namespace MaxVertexNames;				
+				const FInputVertexInterfaceData& InputData = InParams.InputData;
+				return InputData.GetOrConstructDataReadReference<FAudioBuffer>(METASOUND_GET_PARAM_NAME(InputBValue), InParams.OperatorSettings);
 			}
 		};
 	}
@@ -167,7 +169,7 @@ namespace Metasound
 			return Metadata;
 		}
 
-		static TUniquePtr<IOperator> CreateOperator(const FCreateOperatorParams& InParams, TArray<TUniquePtr<IOperatorBuildError>>& OutErrors)
+		static TUniquePtr<IOperator> CreateOperator(const FBuildOperatorParams& InParams, FBuildResults& OutResults)
 		{
 			using namespace MaxVertexNames;
 			using namespace MetasoundMaxNodePrivate;
diff --git a/Engine/Plugins/Runtime/Metasound/Source/MetasoundStandardNodes/Private/MetasoundMidSideNodes.cpp b/Engine/Plugins/Runtime/Metasound/Source/MetasoundStandardNodes/Private/MetasoundMidSideNodes.cpp
--- a/Engine/Plugins/Runtime/Metasound/Source/MetasoundStandardNodes/Private/MetasoundMidSideNodes.cpp
+++ b/Engine/Plugins/Runtime/Metasound/Source/MetasoundStandardNodes/Private/MetasoundMidSideNodes.cpp
@@ -34,7 +34,7 @@ namespace Metasound
 	{
 	public:
 
-		FMidSideEncodeOperator(const FCreateOperatorParams& InParams,
+		FMidSideEncodeOperator(const FBuildOperatorParams& InParams,
 			const FAudioBufferReadRef& InLeftAudioInput,
 			const FAudioBufferReadRef& InRightAudioInput,
 			const FFloatReadRef& InSpreadAmount,
@@ -138,16 +138,16 @@ namespace Metasound
 			return {};
 		}
 		
-		static TUniquePtr<IOperator> CreateOperator(const FCreateOperatorParams& InParams, FBuildErrorArray& OutErrors)
+		static TUniquePtr<IOperator> CreateOperator(const FBuildOperatorParams& InParams, FBuildResults& OutResults)
 		{
 			using namespace MidSideEncodeVertexNames;
-			const FDataReferenceCollection& Inputs = InParams.InputDataReferences;
-			const FInputVertexInterface& InputInterface = DeclareVertexInterface().GetInputInterface();
 			
-			FAudioBufferReadRef LeftAudioIn = Inputs.GetDataReadReferenceOrConstruct<FAudioBuffer>(METASOUND_GET_PARAM_NAME(InputAudioLeft), InParams.OperatorSettings);
-			FAudioBufferReadRef RightAudioIn = Inputs.GetDataReadReferenceOrConstruct<FAudioBuffer>(METASOUND_GET_PARAM_NAME(InputAudioRight), InParams.OperatorSettings);
-			FFloatReadRef SpreadAmountIn = Inputs.GetDataReadReferenceOrConstructWithVertexDefault<float>(InputInterface, METASOUND_GET_PARAM_NAME(InputSpreadAmount), InParams.OperatorSettings);
-			FBoolReadRef bEqualPowerIn = Inputs.GetDataReadReferenceOrConstructWithVertexDefault<bool>(InputInterface, METASOUND_GET_PARAM_NAME(InputEqualPower), InParams.OperatorSettings);
+			const FInputVertexInterfaceData& InputData = InParams.InputData;
+			
+			FAudioBufferReadRef LeftAudioIn = InputData.GetOrConstructDataReadReference<FAudioBuffer>(METASOUND_GET_PARAM_NAME(InputAudioLeft), InParams.OperatorSettings);
+			FAudioBufferReadRef RightAudioIn = InputData.GetOrConstructDataReadReference<FAudioBuffer>(METASOUND_GET_PARAM_NAME(InputAudioRight), InParams.OperatorSettings);
+			FFloatReadRef SpreadAmountIn = InputData.GetOrCreateDefaultDataReadReference<float>(METASOUND_GET_PARAM_NAME(InputSpreadAmount), InParams.OperatorSettings);
+			FBoolReadRef bEqualPowerIn = InputData.GetOrCreateDefaultDataReadReference<bool>(METASOUND_GET_PARAM_NAME(InputEqualPower), InParams.OperatorSettings);
 
 			return MakeUnique<FMidSideEncodeOperator>(InParams, LeftAudioIn, RightAudioIn, SpreadAmountIn, bEqualPowerIn);
 		}
@@ -282,7 +282,7 @@ namespace Metasound
 	{
 	public:
 
-		FMidSideDecodeOperator(const FCreateOperatorParams& InParams,
+		FMidSideDecodeOperator(const FBuildOperatorParams& InParams,
 			const FAudioBufferReadRef& InLeftAudioInput,
 			const FAudioBufferReadRef& InRightAudioInput,
 			const FFloatReadRef& InSpreadAmount,
@@ -392,16 +392,16 @@ namespace Metasound
 			return {};
 		}
 
-		static TUniquePtr<IOperator> CreateOperator(const FCreateOperatorParams& InParams, FBuildErrorArray& OutErrors)
+		static TUniquePtr<IOperator> CreateOperator(const FBuildOperatorParams& InParams, FBuildResults& OutResults)
 		{
 			using namespace MidSideDecodeVertexNames;
-			const FDataReferenceCollection& Inputs = InParams.InputDataReferences;
-			const FInputVertexInterface& InputInterface = DeclareVertexInterface().GetInputInterface();
+			
+			const FInputVertexInterfaceData& InputData = InParams.InputData;
 
-			FAudioBufferReadRef MidAudioIn = Inputs.GetDataReadReferenceOrConstruct<FAudioBuffer>(METASOUND_GET_PARAM_NAME(InputAudioMid), InParams.OperatorSettings);
-			FAudioBufferReadRef SideAudioIn = Inputs.GetDataReadReferenceOrConstruct<FAudioBuffer>(METASOUND_GET_PARAM_NAME(InputAudioSide), InParams.OperatorSettings);
-			FFloatReadRef SpreadAmountIn = Inputs.GetDataReadReferenceOrConstructWithVertexDefault<float>(InputInterface, METASOUND_GET_PARAM_NAME(InputSpreadAmount), InParams.OperatorSettings);
-			FBoolReadRef bEqualPowerIn = Inputs.GetDataReadReferenceOrConstructWithVertexDefault<bool>(InputInterface, METASOUND_GET_PARAM_NAME(InputEqualPower), InParams.OperatorSettings);
+			FAudioBufferReadRef MidAudioIn = InputData.GetOrConstructDataReadReference<FAudioBuffer>(METASOUND_GET_PARAM_NAME(InputAudioMid), InParams.OperatorSettings);
+			FAudioBufferReadRef SideAudioIn = InputData.GetOrConstructDataReadReference<FAudioBuffer>(METASOUND_GET_PARAM_NAME(InputAudioSide), InParams.OperatorSettings);
+			FFloatReadRef SpreadAmountIn = InputData.GetOrCreateDefaultDataReadReference<float>(METASOUND_GET_PARAM_NAME(InputSpreadAmount), InParams.OperatorSettings);
+			FBoolReadRef bEqualPowerIn = InputData.GetOrCreateDefaultDataReadReference<bool>(METASOUND_GET_PARAM_NAME(InputEqualPower), InParams.OperatorSettings);
 
 			return MakeUnique<FMidSideDecodeOperator>(InParams, MidAudioIn, SideAudioIn, SpreadAmountIn, bEqualPowerIn);
 		}
diff --git a/Engine/Plugins/Runtime/Metasound/Source/MetasoundStandardNodes/Private/MetasoundMidiNoteQuantizerNode.cpp b/Engine/Plugins/Runtime/Metasound/Source/MetasoundStandardNodes/Private/MetasoundMidiNoteQuantizerNode.cpp
--- a/Engine/Plugins/Runtime/Metasound/Source/MetasoundStandardNodes/Private/MetasoundMidiNoteQuantizerNode.cpp
+++ b/Engine/Plugins/Runtime/Metasound/Source/MetasoundStandardNodes/Private/MetasoundMidiNoteQuantizerNode.cpp
@@ -50,7 +50,7 @@ namespace Metasound
 
 		// ctor
 		FMidiNoteQuantizerOperator(
-			  const FCreateOperatorParams& InParams
+			  const FBuildOperatorParams& InParams
 			, const FFloatReadRef& InMidiNoteIn
 			, const FFloatReadRef& InRootNote
 			, const FArrayScaleDegreeReadRef& InScale
@@ -60,7 +60,7 @@ namespace Metasound
 		// node interface
 		static const FNodeClassMetadata& GetNodeInfo();
 		static FVertexInterface DeclareVertexInterface();
-		static TUniquePtr<IOperator> CreateOperator(const FCreateOperatorParams& InParams, FBuildErrorArray& OutErrors);
+		static TUniquePtr<IOperator> CreateOperator(const FBuildOperatorParams& InParams, FBuildResults& OutResults);
 
 		virtual void BindInputs(FInputVertexInterfaceData& InOutVertexData) override;
 		virtual void BindOutputs(FOutputVertexInterfaceData& InOutVertexData) override;
@@ -100,7 +100,7 @@ namespace Metasound
 
 	// ctor
 	FMidiNoteQuantizerOperator::FMidiNoteQuantizerOperator(
-		  const FCreateOperatorParams& InParams
+		  const FBuildOperatorParams& InParams
 		, const FFloatReadRef& InMidiNoteIn
 		, const FFloatReadRef& InRootNote
 		, const FArrayScaleDegreeReadRef& InScale
@@ -159,17 +159,15 @@ namespace Metasound
 	}
 
 
-	TUniquePtr<IOperator> FMidiNoteQuantizerOperator::CreateOperator(const FCreateOperatorParams& InParams, FBuildErrorArray& OutErrors)
+	TUniquePtr<IOperator> FMidiNoteQuantizerOperator::CreateOperator(const FBuildOperatorParams& InParams, FBuildResults& OutResults)
 	{
-		const FDataReferenceCollection& InputDataRefs = InParams.InputDataReferences;
-
-		const FInputVertexInterface& InputInterface = InParams.Node.GetVertexInterface().GetInputInterface();
+		const FInputVertexInterfaceData& InputData = InParams.InputData;
 
 		// inputs
-		FFloatReadRef MidiNoteIn = InputDataRefs.GetDataReadReferenceOrConstructWithVertexDefault<float>(InputInterface, METASOUND_GET_PARAM_NAME(ParamNoteIn), InParams.OperatorSettings);
-		FFloatReadRef RootNoteIn = InputDataRefs.GetDataReadReferenceOrConstructWithVertexDefault<float>(InputInterface, METASOUND_GET_PARAM_NAME(ParamRootNote), InParams.OperatorSettings);
-		FArrayScaleDegreeReadRef InScaleArray = InParams.InputDataReferences.GetDataReadReferenceOrConstruct<ScaleDegreeArrayType>(METASOUND_GET_PARAM_NAME(ParamScaleDegrees));
-		FFloatReadRef ScaleRangeIn = InputDataRefs.GetDataReadReferenceOrConstructWithVertexDefault<float>(InputInterface, METASOUND_GET_PARAM_NAME(ParamScaleRange), InParams.OperatorSettings);
+		FFloatReadRef MidiNoteIn = InputData.GetOrCreateDefaultDataReadReference<float>(METASOUND_GET_PARAM_NAME(ParamNoteIn), InParams.OperatorSettings);
+		FFloatReadRef RootNoteIn = InputData.GetOrCreateDefaultDataReadReference<float>(METASOUND_GET_PARAM_NAME(ParamRootNote), InParams.OperatorSettings);
+		FArrayScaleDegreeReadRef InScaleArray = InputData.GetOrConstructDataReadReference<ScaleDegreeArrayType>(METASOUND_GET_PARAM_NAME(ParamScaleDegrees));
+		FFloatReadRef ScaleRangeIn = InputData.GetOrCreateDefaultDataReadReference<float>(METASOUND_GET_PARAM_NAME(ParamScaleRange), InParams.OperatorSettings);
 
 		return MakeUnique <FMidiNoteQuantizerOperator>(
 			  InParams
diff --git a/Engine/Plugins/Runtime/Metasound/Source/MetasoundStandardNodes/Private/MetasoundMidiScaleToArrayNode.cpp b/Engine/Plugins/Runtime/Metasound/Source/MetasoundStandardNodes/Private/MetasoundMidiScaleToArrayNode.cpp
--- a/Engine/Plugins/Runtime/Metasound/Source/MetasoundStandardNodes/Private/MetasoundMidiScaleToArrayNode.cpp
+++ b/Engine/Plugins/Runtime/Metasound/Source/MetasoundStandardNodes/Private/MetasoundMidiScaleToArrayNode.cpp
@@ -89,7 +89,7 @@ namespace Metasound
 
 		// ctor
 		FMusicalScaleToNoteArrayOperator(
-			  const FCreateOperatorParams& InParams
+			  const FBuildOperatorParams& InParams
 			, const FEnumMusicalScaleReadRef& InScale
 			, const FBoolReadRef& InChordTonesOnly
 		);
@@ -97,7 +97,7 @@ namespace Metasound
 		// node interface
 		static const FNodeClassMetadata& GetNodeInfo();
 		static FVertexInterface DeclareVertexInterface();
-		static TUniquePtr<IOperator> CreateOperator(const FCreateOperatorParams& InParams, FBuildErrorArray& OutErrors);
+		static TUniquePtr<IOperator> CreateOperator(const FBuildOperatorParams& InParams, FBuildResults& OutResults);
 		virtual void BindInputs(FInputVertexInterfaceData& InOutVertexData) override;
 		virtual void BindOutputs(FOutputVertexInterfaceData& InOutVertexData) override;
 		virtual FDataReferenceCollection GetInputs() const override;
@@ -126,7 +126,7 @@ namespace Metasound
 
 	// ctor
 	FMusicalScaleToNoteArrayOperator::FMusicalScaleToNoteArrayOperator(
-		  const FCreateOperatorParams& InParams
+		  const FBuildOperatorParams& InParams
 		, const FEnumMusicalScaleReadRef& InScale
 		, const FBoolReadRef& InChordTonesOnly
 		)
@@ -181,14 +181,13 @@ namespace Metasound
 	}
 
 
-	TUniquePtr<IOperator> FMusicalScaleToNoteArrayOperator::CreateOperator(const FCreateOperatorParams& InParams, FBuildErrorArray& OutErrors)
+	TUniquePtr<IOperator> FMusicalScaleToNoteArrayOperator::CreateOperator(const FBuildOperatorParams& InParams, FBuildResults& OutResults)
 	{
-		const FDataReferenceCollection& InputDataRefs = InParams.InputDataReferences;
-		const FInputVertexInterface& InputInterface = InParams.Node.GetVertexInterface().GetInputInterface();
+		const FInputVertexInterfaceData& InputData = InParams.InputData;
 
 		// inputs
-		FEnumMusicalScaleReadRef Scale = InputDataRefs.GetDataReadReferenceOrConstructWithVertexDefault<FEnumEMusicalScale>(InputInterface, METASOUND_GET_PARAM_NAME(ParamScaleDegreesPreset), InParams.OperatorSettings);
-		FBoolReadRef ChordTonesOnly = InputDataRefs.GetDataReadReferenceOrConstructWithVertexDefault<bool>(InputInterface, METASOUND_GET_PARAM_NAME(ParamChordTonesOnly), InParams.OperatorSettings);
+		FEnumMusicalScaleReadRef Scale = InputData.GetOrCreateDefaultDataReadReference<FEnumEMusicalScale>(METASOUND_GET_PARAM_NAME(ParamScaleDegreesPreset), InParams.OperatorSettings);
+		FBoolReadRef ChordTonesOnly = InputData.GetOrCreateDefaultDataReadReference<bool>(METASOUND_GET_PARAM_NAME(ParamChordTonesOnly), InParams.OperatorSettings);
 
 		return MakeUnique <FMusicalScaleToNoteArrayOperator>(
 			  InParams
diff --git a/Engine/Plugins/Runtime/Metasound/Source/MetasoundStandardNodes/Private/MetasoundMinNode.cpp b/Engine/Plugins/Runtime/Metasound/Source/MetasoundStandardNodes/Private/MetasoundMinNode.cpp
--- a/Engine/Plugins/Runtime/Metasound/Source/MetasoundStandardNodes/Private/MetasoundMinNode.cpp
+++ b/Engine/Plugins/Runtime/Metasound/Source/MetasoundStandardNodes/Private/MetasoundMinNode.cpp
@@ -58,18 +58,18 @@ namespace Metasound
 				OutMin = FMath::Min(InA, InB);
 			}
 
-			static TDataReadReference<int32> CreateInRefA(const FCreateOperatorParams& InParams)
+			static TDataReadReference<int32> CreateInRefA(const FBuildOperatorParams& InParams)
 			{
 				using namespace MinVertexNames;
-				const FInputVertexInterface& InputInterface = InParams.Node.GetVertexInterface().GetInputInterface();
-				return InParams.InputDataReferences.GetDataReadReferenceOrConstructWithVertexDefault<int32>(InputInterface, METASOUND_GET_PARAM_NAME(InputAValue), InParams.OperatorSettings);
+				const FInputVertexInterfaceData& InputData = InParams.InputData;
+				return InputData.GetOrCreateDefaultDataReadReference<int32>(METASOUND_GET_PARAM_NAME(InputAValue), InParams.OperatorSettings);
 			}
 
-			static TDataReadReference<int32> CreateInRefB(const FCreateOperatorParams& InParams)
+			static TDataReadReference<int32> CreateInRefB(const FBuildOperatorParams& InParams)
 			{
 				using namespace MinVertexNames;
-				const FInputVertexInterface& InputInterface = InParams.Node.GetVertexInterface().GetInputInterface();
-				return InParams.InputDataReferences.GetDataReadReferenceOrConstructWithVertexDefault<int32>(InputInterface, METASOUND_GET_PARAM_NAME(InputBValue), InParams.OperatorSettings);
+				const FInputVertexInterfaceData& InputData = InParams.InputData;
+				return InputData.GetOrCreateDefaultDataReadReference<int32>(METASOUND_GET_PARAM_NAME(InputBValue), InParams.OperatorSettings);
 			}
 		};
 
@@ -82,18 +82,18 @@ namespace Metasound
 				OutMin = FMath::Min(InA, InB);
 			}
 
-			static TDataReadReference<float> CreateInRefA(const FCreateOperatorParams& InParams)
+			static TDataReadReference<float> CreateInRefA(const FBuildOperatorParams& InParams)
 			{
 				using namespace MinVertexNames;
-				const FInputVertexInterface& InputInterface = InParams.Node.GetVertexInterface().GetInputInterface();
-				return InParams.InputDataReferences.GetDataReadReferenceOrConstructWithVertexDefault<float>(InputInterface, METASOUND_GET_PARAM_NAME(InputAValue), InParams.OperatorSettings);
+				const FInputVertexInterfaceData& InputData = InParams.InputData;
+				return InputData.GetOrCreateDefaultDataReadReference<float>(METASOUND_GET_PARAM_NAME(InputAValue), InParams.OperatorSettings);
 			}
 
-			static TDataReadReference<float> CreateInRefB(const FCreateOperatorParams& InParams)
+			static TDataReadReference<float> CreateInRefB(const FBuildOperatorParams& InParams)
 			{
 				using namespace MinVertexNames;
-				const FInputVertexInterface& InputInterface = InParams.Node.GetVertexInterface().GetInputInterface();
-				return InParams.InputDataReferences.GetDataReadReferenceOrConstructWithVertexDefault<float>(InputInterface, METASOUND_GET_PARAM_NAME(InputBValue), InParams.OperatorSettings);
+				const FInputVertexInterfaceData& InputData = InParams.InputData;
+				return InputData.GetOrCreateDefaultDataReadReference<float>(METASOUND_GET_PARAM_NAME(InputBValue), InParams.OperatorSettings);
 			}
 		};
 
@@ -112,16 +112,18 @@ namespace Metasound
 				}
 			}
 
-			static TDataReadReference<FAudioBuffer> CreateInRefA(const FCreateOperatorParams& InParams)
+			static TDataReadReference<FAudioBuffer> CreateInRefA(const FBuildOperatorParams& InParams)
 			{
 				using namespace MinVertexNames;
-				return InParams.InputDataReferences.GetDataReadReferenceOrConstruct<FAudioBuffer>(METASOUND_GET_PARAM_NAME(InputAValue), InParams.OperatorSettings);
+				const FInputVertexInterfaceData& InputData = InParams.InputData;
+				return InputData.GetOrConstructDataReadReference<FAudioBuffer>(METASOUND_GET_PARAM_NAME(InputAValue), InParams.OperatorSettings);
 			}
 
-			static TDataReadReference<FAudioBuffer> CreateInRefB(const FCreateOperatorParams& InParams)
+			static TDataReadReference<FAudioBuffer> CreateInRefB(const FBuildOperatorParams& InParams)
 			{
 				using namespace MinVertexNames;
-				return InParams.InputDataReferences.GetDataReadReferenceOrConstruct<FAudioBuffer>(METASOUND_GET_PARAM_NAME(InputBValue), InParams.OperatorSettings);
+				const FInputVertexInterfaceData& InputData = InParams.InputData;
+				return InputData.GetOrConstructDataReadReference<FAudioBuffer>(METASOUND_GET_PARAM_NAME(InputBValue), InParams.OperatorSettings);
 			}
 		};
 
@@ -166,7 +168,7 @@ namespace Metasound
 			return Metadata;
 		}
 
-		static TUniquePtr<IOperator> CreateOperator(const FCreateOperatorParams& InParams, TArray<TUniquePtr<IOperatorBuildError>>& OutErrors)
+		static TUniquePtr<IOperator> CreateOperator(const FBuildOperatorParams& InParams, FBuildResults& OutResults)
 		{
 			TDataReadReference<ValueType> InputA = MetasoundMinNodePrivate::TMin<ValueType>::CreateInRefA(InParams);
 			TDataReadReference<ValueType> InputB = MetasoundMinNodePrivate::TMin<ValueType>::CreateInRefB(InParams);
diff --git a/Engine/Plugins/Runtime/Metasound/Source/MetasoundStandardNodes/Private/MetasoundMixerNode.cpp b/Engine/Plugins/Runtime/Metasound/Source/MetasoundStandardNodes/Private/MetasoundMixerNode.cpp
--- a/Engine/Plugins/Runtime/Metasound/Source/MetasoundStandardNodes/Private/MetasoundMixerNode.cpp
+++ b/Engine/Plugins/Runtime/Metasound/Source/MetasoundStandardNodes/Private/MetasoundMixerNode.cpp
@@ -25,7 +25,7 @@ namespace Metasound
 	{
 	public:
 		// ctor
-		TAudioMixerNodeOperator(const FCreateOperatorParams& InParams, const TArray<FAudioBufferReadRef>&& InInputBuffers, const TArray<FFloatReadRef>&& InGainValues)
+		TAudioMixerNodeOperator(const FBuildOperatorParams& InParams, const TArray<FAudioBufferReadRef>&& InInputBuffers, const TArray<FFloatReadRef>&& InGainValues)
 			: Gains(InGainValues)
 			, Inputs (InInputBuffers)
 		{
@@ -146,11 +146,10 @@ namespace Metasound
 			return Metadata;
 		}
 
-		static TUniquePtr<IOperator> CreateOperator(const FCreateOperatorParams& InParams, TArray<TUniquePtr<IOperatorBuildError>>& OutErrors)
+		static TUniquePtr<IOperator> CreateOperator(const FBuildOperatorParams& InParams, FBuildResults& OutResults)
 		{
-			const FInputVertexInterface& InputInterface = InParams.Node.GetVertexInterface().GetInputInterface();
-			const FDataReferenceCollection& InputCollection = InParams.InputDataReferences;
-
+			const FInputVertexInterfaceData& InputData = InParams.InputData;
+			
 			TArray<FAudioBufferReadRef> InputBuffers;
 			TArray<FFloatReadRef> InputGains;
 
@@ -158,10 +157,10 @@ namespace Metasound
 			{
 				for (uint32 Chan = 0; Chan < NumChannels; ++Chan)
 				{
-					InputBuffers.Add(InputCollection.GetDataReadReferenceOrConstruct<FAudioBuffer>(GetAudioInputName(i, Chan), InParams.OperatorSettings));
+					InputBuffers.Add(InputData.GetOrConstructDataReadReference<FAudioBuffer>(GetAudioInputName(i, Chan), InParams.OperatorSettings));
 				}
 
-				InputGains.Add(InputCollection.GetDataReadReferenceOrConstructWithVertexDefault<float>(InputInterface, GetGainInputName(i), InParams.OperatorSettings));
+				InputGains.Add(InputData.GetOrCreateDefaultDataReadReference<float>(GetGainInputName(i), InParams.OperatorSettings));
 			}
 
 			return MakeUnique<TAudioMixerNodeOperator<NumInputs, NumChannels>>(InParams, MoveTemp(InputBuffers), MoveTemp(InputGains));
diff --git a/Engine/Plugins/Runtime/Metasound/Source/MetasoundStandardNodes/Private/MetasoundNoiseGenerator.cpp b/Engine/Plugins/Runtime/Metasound/Source/MetasoundStandardNodes/Private/MetasoundNoiseGenerator.cpp
--- a/Engine/Plugins/Runtime/Metasound/Source/MetasoundStandardNodes/Private/MetasoundNoiseGenerator.cpp
+++ b/Engine/Plugins/Runtime/Metasound/Source/MetasoundStandardNodes/Private/MetasoundNoiseGenerator.cpp
@@ -46,7 +46,7 @@ namespace Metasound
 	public:
 		static constexpr int32 DefaultSeed = INDEX_NONE;
 
-		static TUniquePtr<IOperator> CreateOperator(const FCreateOperatorParams& InParams, FBuildErrorArray& OutErrors);
+		static TUniquePtr<IOperator> CreateOperator(const FBuildOperatorParams& InParams, FBuildResults& OutResults);
 		static const FNodeClassMetadata& GetNodeInfo();
 		static FVertexInterface DeclareVertexInterface();
 
@@ -238,20 +238,18 @@ namespace Metasound
 		: FNoiseNode(InInitData.InstanceName, InInitData.InstanceID, FNoiseOperator::DefaultSeed)
 	{}
 
-	TUniquePtr<IOperator> FNoiseOperator::CreateOperator(const FCreateOperatorParams& InParams, FBuildErrorArray& OutErrors)
+	TUniquePtr<IOperator> FNoiseOperator::CreateOperator(const FBuildOperatorParams& InParams, FBuildResults& OutResults)
 	{
 		using namespace NoiseGeneratorVertexNames; 
 
-		const FNoiseNode& Node = static_cast<const FNoiseNode&>(InParams.Node);
-		const FDataReferenceCollection& InputCol = InParams.InputDataReferences;
 		const FOperatorSettings& Settings = InParams.OperatorSettings;
-		const FInputVertexInterface& InputInterface = InParams.Node.GetVertexInterface().GetInputInterface();
+		const FInputVertexInterfaceData& InputData = InParams.InputData;
 
 		// Static property pin, only used for factory.
-		FEnumNoiseTypeReadRef Type = InputCol.GetDataReadReferenceOrConstructWithVertexDefault<FEnumNoiseType>(InputInterface, METASOUND_GET_PARAM_NAME(InputType), Settings);
+		FEnumNoiseTypeReadRef Type = InputData.GetOrCreateDefaultDataReadReference<FEnumNoiseType>(METASOUND_GET_PARAM_NAME(InputType), Settings);
 
 		// Seed.
-		FInt32ReadRef Seed = InputCol.GetDataReadReferenceOrConstructWithVertexDefault<int32>(InputInterface, METASOUND_GET_PARAM_NAME(InputSeed), Settings);
+		FInt32ReadRef Seed = InputData.GetOrCreateDefaultDataReadReference<int32>(METASOUND_GET_PARAM_NAME(InputSeed), Settings);
 
 		switch (*Type)
 		{
diff --git a/Engine/Plugins/Runtime/Metasound/Source/MetasoundStandardNodes/Private/MetasoundRingModNode.cpp b/Engine/Plugins/Runtime/Metasound/Source/MetasoundStandardNodes/Private/MetasoundRingModNode.cpp
--- a/Engine/Plugins/Runtime/Metasound/Source/MetasoundStandardNodes/Private/MetasoundRingModNode.cpp
+++ b/Engine/Plugins/Runtime/Metasound/Source/MetasoundStandardNodes/Private/MetasoundRingModNode.cpp
@@ -116,13 +116,13 @@ namespace Metasound
 			return {};
 		}
 		
-		static TUniquePtr<IOperator> CreateOperator(const FCreateOperatorParams& InParams, FBuildErrorArray& OutErrors)
+		static TUniquePtr<IOperator> CreateOperator(const FBuildOperatorParams& InParams, FBuildResults& OutResults)
 		{
 			using namespace RingModVertexNames;
-			const FDataReferenceCollection& Inputs = InParams.InputDataReferences;
-			
-			FAudioBufferReadRef CarrierAudioIn = Inputs.GetDataReadReferenceOrConstruct<FAudioBuffer>(METASOUND_GET_PARAM_NAME(InputAudioCarrier), InParams.OperatorSettings);
-			FAudioBufferReadRef ModulatorAudioIn = Inputs.GetDataReadReferenceOrConstruct<FAudioBuffer>(METASOUND_GET_PARAM_NAME(InputAudioModulator), InParams.OperatorSettings);
+			const FInputVertexInterfaceData& InputData = InParams.InputData;
+
+			FAudioBufferReadRef CarrierAudioIn = InputData.GetOrConstructDataReadReference<FAudioBuffer>(METASOUND_GET_PARAM_NAME(InputAudioCarrier), InParams.OperatorSettings);
+			FAudioBufferReadRef ModulatorAudioIn = InputData.GetOrConstructDataReadReference<FAudioBuffer>(METASOUND_GET_PARAM_NAME(InputAudioModulator), InParams.OperatorSettings);
 
 			return MakeUnique<FRingModOperator>(InParams.OperatorSettings, CarrierAudioIn, ModulatorAudioIn);
 		}
diff --git a/Engine/Plugins/Runtime/Metasound/Source/MetasoundStandardNodes/Private/MetasoundSampleAndHoldNode.cpp b/Engine/Plugins/Runtime/Metasound/Source/MetasoundStandardNodes/Private/MetasoundSampleAndHoldNode.cpp
--- a/Engine/Plugins/Runtime/Metasound/Source/MetasoundStandardNodes/Private/MetasoundSampleAndHoldNode.cpp
+++ b/Engine/Plugins/Runtime/Metasound/Source/MetasoundStandardNodes/Private/MetasoundSampleAndHoldNode.cpp
@@ -76,18 +76,18 @@ namespace Metasound
 			return Interface;
 		}
 		
-		static TUniquePtr<IOperator> CreateOperator(const FCreateOperatorParams& InParams, FBuildErrorArray& OutErrors)
+		static TUniquePtr<IOperator> CreateOperator(const FBuildOperatorParams& InParams, FBuildResults& OutResults)
 		{
 			using namespace SampleAndHoldVertexNames;
-			const FDataReferenceCollection& Inputs = InParams.InputDataReferences;
+			const FInputVertexInterfaceData& InputData = InParams.InputData;
 
-			FTriggerReadRef InputTriggerSampleAndHold = Inputs.GetDataReadReferenceOrConstruct<FTrigger>(METASOUND_GET_PARAM_NAME(InputTriggerSampleAndHold), InParams.OperatorSettings);
-			FAudioBufferReadRef InputAudio = Inputs.GetDataReadReferenceOrConstruct<FAudioBuffer>(METASOUND_GET_PARAM_NAME(InputAudio), InParams.OperatorSettings);
+			FTriggerReadRef InputTriggerSampleAndHold = InputData.GetOrConstructDataReadReference<FTrigger>(METASOUND_GET_PARAM_NAME(InputTriggerSampleAndHold), InParams.OperatorSettings);
+			FAudioBufferReadRef InputAudio = InputData.GetOrConstructDataReadReference<FAudioBuffer>(METASOUND_GET_PARAM_NAME(InputAudio), InParams.OperatorSettings);
 
 			return MakeUnique<FSampleAndHoldOperator>(InParams, InputTriggerSampleAndHold, InputAudio);
 		}
 
-		FSampleAndHoldOperator(const FCreateOperatorParams& InParams, const FTriggerReadRef& InTriggerSampleAndHold, const FAudioBufferReadRef& InAudioInput)
+		FSampleAndHoldOperator(const FBuildOperatorParams& InParams, const FTriggerReadRef& InTriggerSampleAndHold, const FAudioBufferReadRef& InAudioInput)
 			: AudioInput(InAudioInput)
 			, TriggerSampleAndHold(InTriggerSampleAndHold)
 			, AudioOutput(FAudioBufferWriteRef::CreateNew(InParams.OperatorSettings))
diff --git a/Engine/Plugins/Runtime/Metasound/Source/MetasoundStandardNodes/Private/MetasoundSemitoneToFreqMultiplierNode.cpp b/Engine/Plugins/Runtime/Metasound/Source/MetasoundStandardNodes/Private/MetasoundSemitoneToFreqMultiplierNode.cpp
--- a/Engine/Plugins/Runtime/Metasound/Source/MetasoundStandardNodes/Private/MetasoundSemitoneToFreqMultiplierNode.cpp
+++ b/Engine/Plugins/Runtime/Metasound/Source/MetasoundStandardNodes/Private/MetasoundSemitoneToFreqMultiplierNode.cpp
@@ -30,7 +30,7 @@ namespace Metasound
 
 		static const FNodeClassMetadata& GetNodeInfo();
 		static const FVertexInterface& GetVertexInterface();
-		static TUniquePtr<IOperator> CreateOperator(const FCreateOperatorParams& InParams, FBuildErrorArray& OutErrors);
+		static TUniquePtr<IOperator> CreateOperator(const FBuildOperatorParams& InParams, FBuildResults& OutResults);
 
 		FSemitoneToFrequencyMultiplierOperator(const FOperatorSettings& InSettings, const FFloatReadRef& InSemitone);
 
@@ -152,14 +152,13 @@ namespace Metasound
 		return Info;
 	}
 
-	TUniquePtr<IOperator> FSemitoneToFrequencyMultiplierOperator::CreateOperator(const FCreateOperatorParams& InParams, FBuildErrorArray& OutErrors)
+	TUniquePtr<IOperator> FSemitoneToFrequencyMultiplierOperator::CreateOperator(const FBuildOperatorParams& InParams, FBuildResults& OutResults)
 	{
 		using namespace SemitoneToFrequencyMultiplierVertexNames;
 
-		const FDataReferenceCollection& InputCollection = InParams.InputDataReferences;
-		const FInputVertexInterface& InputInterface = GetVertexInterface().GetInputInterface();
+		const FInputVertexInterfaceData& InputData = InParams.InputData;
 
-		FFloatReadRef InSemitone = InputCollection.GetDataReadReferenceOrConstructWithVertexDefault<float>(InputInterface, METASOUND_GET_PARAM_NAME(InputSemitone), InParams.OperatorSettings);
+		FFloatReadRef InSemitone = InputData.GetOrCreateDefaultDataReadReference<float>(METASOUND_GET_PARAM_NAME(InputSemitone), InParams.OperatorSettings);
 
 		return MakeUnique<FSemitoneToFrequencyMultiplierOperator>(InParams.OperatorSettings, InSemitone);
 	}
diff --git a/Engine/Plugins/Runtime/Metasound/Source/MetasoundStandardNodes/Private/MetasoundStereoDelayNode.cpp b/Engine/Plugins/Runtime/Metasound/Source/MetasoundStandardNodes/Private/MetasoundStereoDelayNode.cpp
--- a/Engine/Plugins/Runtime/Metasound/Source/MetasoundStandardNodes/Private/MetasoundStereoDelayNode.cpp
+++ b/Engine/Plugins/Runtime/Metasound/Source/MetasoundStandardNodes/Private/MetasoundStereoDelayNode.cpp
@@ -57,7 +57,7 @@ namespace Metasound
 
 		static const FNodeClassMetadata& GetNodeInfo();
 		static const FVertexInterface& GetVertexInterface();
-		static TUniquePtr<IOperator> CreateOperator(const FCreateOperatorParams& InParams, FBuildErrorArray& OutErrors);
+		static TUniquePtr<IOperator> CreateOperator(const FBuildOperatorParams& InParams, FBuildResults& OutResults);
 
 		FStereoDelayOperator(const FOperatorSettings& InSettings, 
 			const FAudioBufferReadRef& InLeftAudioInput, 
@@ -361,21 +361,19 @@ namespace Metasound
 		return Info;
 	}
 
-	TUniquePtr<IOperator> FStereoDelayOperator::CreateOperator(const FCreateOperatorParams& InParams, FBuildErrorArray& OutErrors)
+	TUniquePtr<IOperator> FStereoDelayOperator::CreateOperator(const FBuildOperatorParams& InParams, FBuildResults& OutResults)
 	{
 		using namespace StereoDelay; 
-
-		const FDataReferenceCollection& InputCollection = InParams.InputDataReferences;
-		const FInputVertexInterface& InputInterface = GetVertexInterface().GetInputInterface();
-
-		FAudioBufferReadRef LeftAudioIn = InputCollection.GetDataReadReferenceOrConstruct<FAudioBuffer>(METASOUND_GET_PARAM_NAME(InAudioInputLeft), InParams.OperatorSettings);
-		FAudioBufferReadRef RightAudioIn = InputCollection.GetDataReadReferenceOrConstruct<FAudioBuffer>(METASOUND_GET_PARAM_NAME(InAudioInputRight), InParams.OperatorSettings);
-		FStereoDelayModeReadRef StereoDelayMode = InputCollection.GetDataReadReferenceOrConstructWithVertexDefault<FEnumStereoDelayMode>(InputInterface, METASOUND_GET_PARAM_NAME(InDelayMode), InParams.OperatorSettings);
-		FTimeReadRef DelayTime = InputCollection.GetDataReadReferenceOrConstructWithVertexDefault<FTime>(InputInterface, METASOUND_GET_PARAM_NAME(InDelayTime), InParams.OperatorSettings);
-		FFloatReadRef DelayRatio = InputCollection.GetDataReadReferenceOrConstructWithVertexDefault<float>(InputInterface, METASOUND_GET_PARAM_NAME(InDelayRatio), InParams.OperatorSettings);
-		FFloatReadRef DryLevel = InputCollection.GetDataReadReferenceOrConstructWithVertexDefault<float>(InputInterface, METASOUND_GET_PARAM_NAME(InDryLevel), InParams.OperatorSettings);
-		FFloatReadRef WetLevel = InputCollection.GetDataReadReferenceOrConstructWithVertexDefault<float>(InputInterface, METASOUND_GET_PARAM_NAME(InWetLevel), InParams.OperatorSettings);
-		FFloatReadRef Feedback = InputCollection.GetDataReadReferenceOrConstructWithVertexDefault<float>(InputInterface, METASOUND_GET_PARAM_NAME(InFeedbackAmount), InParams.OperatorSettings);
+		const FInputVertexInterfaceData& InputData = InParams.InputData;
+
+		FAudioBufferReadRef LeftAudioIn = InputData.GetOrConstructDataReadReference<FAudioBuffer>(METASOUND_GET_PARAM_NAME(InAudioInputLeft), InParams.OperatorSettings);
+		FAudioBufferReadRef RightAudioIn = InputData.GetOrConstructDataReadReference<FAudioBuffer>(METASOUND_GET_PARAM_NAME(InAudioInputRight), InParams.OperatorSettings);
+		FStereoDelayModeReadRef StereoDelayMode = InputData.GetOrCreateDefaultDataReadReference<FEnumStereoDelayMode>(METASOUND_GET_PARAM_NAME(InDelayMode), InParams.OperatorSettings);
+		FTimeReadRef DelayTime = InputData.GetOrCreateDefaultDataReadReference<FTime>(METASOUND_GET_PARAM_NAME(InDelayTime), InParams.OperatorSettings);
+		FFloatReadRef DelayRatio = InputData.GetOrCreateDefaultDataReadReference<float>(METASOUND_GET_PARAM_NAME(InDelayRatio), InParams.OperatorSettings);
+		FFloatReadRef DryLevel = InputData.GetOrCreateDefaultDataReadReference<float>(METASOUND_GET_PARAM_NAME(InDryLevel), InParams.OperatorSettings);
+		FFloatReadRef WetLevel = InputData.GetOrCreateDefaultDataReadReference<float>(METASOUND_GET_PARAM_NAME(InWetLevel), InParams.OperatorSettings);
+		FFloatReadRef Feedback = InputData.GetOrCreateDefaultDataReadReference<float>(METASOUND_GET_PARAM_NAME(InFeedbackAmount), InParams.OperatorSettings);
 
 		return MakeUnique<FStereoDelayOperator>(InParams.OperatorSettings, LeftAudioIn, RightAudioIn, StereoDelayMode, DelayTime, DelayRatio, DryLevel, WetLevel, Feedback);
 	}
diff --git a/Engine/Plugins/Runtime/Metasound/Source/MetasoundStandardNodes/Private/MetasoundStereoPannerNode.cpp b/Engine/Plugins/Runtime/Metasound/Source/MetasoundStandardNodes/Private/MetasoundStereoPannerNode.cpp
--- a/Engine/Plugins/Runtime/Metasound/Source/MetasoundStandardNodes/Private/MetasoundStereoPannerNode.cpp
+++ b/Engine/Plugins/Runtime/Metasound/Source/MetasoundStandardNodes/Private/MetasoundStereoPannerNode.cpp
@@ -50,7 +50,7 @@ namespace Metasound
 
 		static const FNodeClassMetadata& GetNodeInfo();
 		static const FVertexInterface& GetVertexInterface();
-		static TUniquePtr<IOperator> CreateOperator(const FCreateOperatorParams& InParams, FBuildErrorArray& OutErrors);
+		static TUniquePtr<IOperator> CreateOperator(const FBuildOperatorParams& InParams, FBuildResults& OutResults);
 
 		FStereoPannerOperator(const FOperatorSettings& InSettings,
 			const FAudioBufferReadRef& InAudioInput, 
@@ -244,16 +244,15 @@ namespace Metasound
 		return Info;
 	}
 
-	TUniquePtr<IOperator> FStereoPannerOperator::CreateOperator(const FCreateOperatorParams& InParams, FBuildErrorArray& OutErrors)
+	TUniquePtr<IOperator> FStereoPannerOperator::CreateOperator(const FBuildOperatorParams& InParams, FBuildResults& OutResults)
 	{
-		const FDataReferenceCollection& InputCollection = InParams.InputDataReferences;
-		const FInputVertexInterface& InputInterface = GetVertexInterface().GetInputInterface();
+		const FInputVertexInterfaceData& InputData = InParams.InputData;
 
 		using namespace StereoPannerVertexNames;
 
-		FAudioBufferReadRef AudioIn = InputCollection.GetDataReadReferenceOrConstruct<FAudioBuffer>(METASOUND_GET_PARAM_NAME(InputAudio), InParams.OperatorSettings);
-		FFloatReadRef PanningAmount = InputCollection.GetDataReadReferenceOrConstructWithVertexDefault<float>(InputInterface, METASOUND_GET_PARAM_NAME(InputPanAmount), InParams.OperatorSettings);
-		FPanningLawReadRef PanningLaw = InputCollection.GetDataReadReferenceOrConstructWithVertexDefault<FEnumPanningLaw>(InputInterface, METASOUND_GET_PARAM_NAME(InputPanningLaw), InParams.OperatorSettings);
+		FAudioBufferReadRef AudioIn = InputData.GetOrConstructDataReadReference<FAudioBuffer>(METASOUND_GET_PARAM_NAME(InputAudio), InParams.OperatorSettings);
+		FFloatReadRef PanningAmount = InputData.GetOrCreateDefaultDataReadReference<float>(METASOUND_GET_PARAM_NAME(InputPanAmount), InParams.OperatorSettings);
+		FPanningLawReadRef PanningLaw = InputData.GetOrCreateDefaultDataReadReference<FEnumPanningLaw>(METASOUND_GET_PARAM_NAME(InputPanningLaw), InParams.OperatorSettings);
 
 		return MakeUnique<FStereoPannerOperator>(InParams.OperatorSettings, AudioIn, PanningAmount, PanningLaw);
 	}
diff --git a/Engine/Plugins/Runtime/Metasound/Source/MetasoundStandardNodes/Private/MetasoundSwitchAndRampNode.cpp b/Engine/Plugins/Runtime/Metasound/Source/MetasoundStandardNodes/Private/MetasoundSwitchAndRampNode.cpp
--- a/Engine/Plugins/Runtime/Metasound/Source/MetasoundStandardNodes/Private/MetasoundSwitchAndRampNode.cpp
+++ b/Engine/Plugins/Runtime/Metasound/Source/MetasoundStandardNodes/Private/MetasoundSwitchAndRampNode.cpp
@@ -72,16 +72,14 @@ namespace Metasound
 			return Interface;
 		}
 
-		static TUniquePtr<IOperator> CreateOperator(const FCreateOperatorParams& InParams, FBuildErrorArray& OutErrors)
+		static TUniquePtr<IOperator> CreateOperator(const FBuildOperatorParams& InParams, FBuildResults& OutResults)
 		{
 			using namespace SwitchAndRampVertexNames;
+			const FInputVertexInterfaceData& InputData = InParams.InputData;
 
-			const FDataReferenceCollection& InputCollection = InParams.InputDataReferences;
-			const FInputVertexInterface& InputInterface = GetVertexInterface().GetInputInterface();
-
-			FTriggerReadRef TriggerIn = InputCollection.GetDataReadReferenceOrConstruct<FTrigger>(METASOUND_GET_PARAM_NAME(InputTrigger), InParams.OperatorSettings);
-			FAudioBufferReadRef AudioInput = InputCollection.GetDataReadReferenceOrConstruct<FAudioBuffer>(METASOUND_GET_PARAM_NAME(InputAudio), InParams.OperatorSettings);
-			FTimeReadRef SmoothTime = InputCollection.GetDataReadReferenceOrConstructWithVertexDefault<FTime>(InputInterface, METASOUND_GET_PARAM_NAME(InputSmoothTime), InParams.OperatorSettings);
+			FTriggerReadRef TriggerIn = InputData.GetOrConstructDataReadReference<FTrigger>(METASOUND_GET_PARAM_NAME(InputTrigger), InParams.OperatorSettings);
+			FAudioBufferReadRef AudioInput = InputData.GetOrConstructDataReadReference<FAudioBuffer>(METASOUND_GET_PARAM_NAME(InputAudio), InParams.OperatorSettings);
+			FTimeReadRef SmoothTime = InputData.GetOrCreateDefaultDataReadReference<FTime>(METASOUND_GET_PARAM_NAME(InputSmoothTime), InParams.OperatorSettings);
 
 			return MakeUnique<FSwitchAndRampOperator>(InParams.OperatorSettings, TriggerIn, AudioInput, SmoothTime);
 		}
diff --git a/Engine/Plugins/Runtime/Metasound/Source/MetasoundStandardNodes/Private/MetasoundTriggerAnyNode.cpp b/Engine/Plugins/Runtime/Metasound/Source/MetasoundStandardNodes/Private/MetasoundTriggerAnyNode.cpp
--- a/Engine/Plugins/Runtime/Metasound/Source/MetasoundStandardNodes/Private/MetasoundTriggerAnyNode.cpp
+++ b/Engine/Plugins/Runtime/Metasound/Source/MetasoundStandardNodes/Private/MetasoundTriggerAnyNode.cpp
@@ -111,18 +111,17 @@ namespace Metasound
 			return Metadata;
 		}
 
-		static TUniquePtr<IOperator> CreateOperator(const FCreateOperatorParams& InParams, TArray<TUniquePtr<IOperatorBuildError>>& OutErrors)
+		static TUniquePtr<IOperator> CreateOperator(const FBuildOperatorParams& InParams, FBuildResults& OutResults)
 		{
 			using namespace TriggerAnyVertexNames;
 
-			const FInputVertexInterface& InputInterface = InParams.Node.GetVertexInterface().GetInputInterface();
-			const FDataReferenceCollection& InputCollection = InParams.InputDataReferences;
+			const FInputVertexInterfaceData& InputData = InParams.InputData;
 
 			TArray<FTriggerReadRef> InputTriggers;
 
 			for (uint32 i = 0; i < NumInputs; ++i)
 			{
-				InputTriggers.Add(InputCollection.GetDataReadReferenceOrConstruct<FTrigger>(GetInputTriggerName(i), InParams.OperatorSettings));
+				InputTriggers.Add(InputData.GetOrConstructDataReadReference<FTrigger>(GetInputTriggerName(i), InParams.OperatorSettings));
 			}
 
 			return MakeUnique<TTriggerAnyOperator<NumInputs>>(InParams.OperatorSettings, MoveTemp(InputTriggers));
diff --git a/Engine/Plugins/Runtime/Metasound/Source/MetasoundStandardNodes/Private/MetasoundTriggerCoinNode.cpp b/Engine/Plugins/Runtime/Metasound/Source/MetasoundStandardNodes/Private/MetasoundTriggerCoinNode.cpp
--- a/Engine/Plugins/Runtime/Metasound/Source/MetasoundStandardNodes/Private/MetasoundTriggerCoinNode.cpp
+++ b/Engine/Plugins/Runtime/Metasound/Source/MetasoundStandardNodes/Private/MetasoundTriggerCoinNode.cpp
@@ -35,7 +35,7 @@ namespace Metasound
 	public:
 		static const FNodeClassMetadata& GetNodeInfo();
 		static const FVertexInterface& GetVertexInterface();
-		static TUniquePtr<IOperator> CreateOperator(const FCreateOperatorParams& InParams, FBuildErrorArray& OutErrors);
+		static TUniquePtr<IOperator> CreateOperator(const FBuildOperatorParams& InParams, FBuildResults& OutResults);
 
 		FTriggerCoinOperator(const FOperatorSettings& InSettings, 
 			const FTriggerReadRef& InTrigger, 
@@ -180,16 +180,16 @@ namespace Metasound
 		RandomStream.Reset();
 	}
 
-	TUniquePtr<IOperator> FTriggerCoinOperator::CreateOperator(const FCreateOperatorParams& InParams, FBuildErrorArray& OutErrors)
+	TUniquePtr<IOperator> FTriggerCoinOperator::CreateOperator(const FBuildOperatorParams& InParams, FBuildResults& OutResults)
 	{
 		using namespace TriggerCoinVertexNames;
 
-		const FInputVertexInterface& InputInterface = GetVertexInterface().GetInputInterface();
+		const FInputVertexInterfaceData& InputData = InParams.InputData;
 
-		FTriggerReadRef TriggerIn = InParams.InputDataReferences.GetDataReadReferenceOrConstruct<FTrigger>(METASOUND_GET_PARAM_NAME(InputTrigger), InParams.OperatorSettings);
-		FTriggerReadRef TriggerResetIn = InParams.InputDataReferences.GetDataReadReferenceOrConstruct<FTrigger>(METASOUND_GET_PARAM_NAME(InputReset), InParams.OperatorSettings);
-		FInt32ReadRef SeedIn = InParams.InputDataReferences.GetDataReadReferenceOrConstructWithVertexDefault<int32>(InputInterface, METASOUND_GET_PARAM_NAME(InputSeed), InParams.OperatorSettings);
-		FFloatReadRef ProbabilityIn = InParams.InputDataReferences.GetDataReadReferenceOrConstructWithVertexDefault<float>(InputInterface, METASOUND_GET_PARAM_NAME(InputProbability), InParams.OperatorSettings);
+		FTriggerReadRef TriggerIn = InputData.GetOrConstructDataReadReference<FTrigger>(METASOUND_GET_PARAM_NAME(InputTrigger), InParams.OperatorSettings);
+		FTriggerReadRef TriggerResetIn = InputData.GetOrConstructDataReadReference<FTrigger>(METASOUND_GET_PARAM_NAME(InputReset), InParams.OperatorSettings);
+		FInt32ReadRef SeedIn = InputData.GetOrCreateDefaultDataReadReference<int32>(METASOUND_GET_PARAM_NAME(InputSeed), InParams.OperatorSettings);
+		FFloatReadRef ProbabilityIn = InputData.GetOrCreateDefaultDataReadReference<float>(METASOUND_GET_PARAM_NAME(InputProbability), InParams.OperatorSettings);
 
 		return MakeUnique<FTriggerCoinOperator>(InParams.OperatorSettings, TriggerIn, TriggerResetIn, SeedIn, ProbabilityIn);
 	}
diff --git a/Engine/Plugins/Runtime/Metasound/Source/MetasoundStandardNodes/Private/MetasoundTriggerControlNode.cpp b/Engine/Plugins/Runtime/Metasound/Source/MetasoundStandardNodes/Private/MetasoundTriggerControlNode.cpp
--- a/Engine/Plugins/Runtime/Metasound/Source/MetasoundStandardNodes/Private/MetasoundTriggerControlNode.cpp
+++ b/Engine/Plugins/Runtime/Metasound/Source/MetasoundStandardNodes/Private/MetasoundTriggerControlNode.cpp
@@ -36,7 +36,7 @@ namespace Metasound
 	public:
 		static const FNodeClassMetadata& GetNodeInfo();
 		static const FVertexInterface& GetVertexInterface();
-		static TUniquePtr<IOperator> CreateOperator(const FCreateOperatorParams& InParams, FBuildErrorArray& OutErrors);
+		static TUniquePtr<IOperator> CreateOperator(const FBuildOperatorParams& InParams, FBuildResults& OutResults);
 
 		FTriggerControlOperator(const FOperatorSettings& InSettings,
 			const FTriggerReadRef& InTriggerEnter,
@@ -179,17 +179,17 @@ namespace Metasound
 		bIsGateOpen = !(*bStartClosedInput);
 	}
 
-	TUniquePtr<IOperator> FTriggerControlOperator::CreateOperator(const FCreateOperatorParams& InParams, FBuildErrorArray& OutErrors)
+	TUniquePtr<IOperator> FTriggerControlOperator::CreateOperator(const FBuildOperatorParams& InParams, FBuildResults& OutResults)
 	{
 		using namespace TriggerControlVertexNames;
 
-		const FInputVertexInterface& InputInterface = GetVertexInterface().GetInputInterface();
+		const FInputVertexInterfaceData& InputData = InParams.InputData;
 
-		FTriggerReadRef TriggerEnterIn = InParams.InputDataReferences.GetDataReadReferenceOrConstruct<FTrigger>(METASOUND_GET_PARAM_NAME(InputEnter), InParams.OperatorSettings);
-		FTriggerReadRef TriggerOpenIn = InParams.InputDataReferences.GetDataReadReferenceOrConstruct<FTrigger>(METASOUND_GET_PARAM_NAME(InputOpen), InParams.OperatorSettings);
-		FTriggerReadRef TriggerCloseIn = InParams.InputDataReferences.GetDataReadReferenceOrConstruct<FTrigger>(METASOUND_GET_PARAM_NAME(InputClose), InParams.OperatorSettings);
-		FTriggerReadRef TriggerToggleIn = InParams.InputDataReferences.GetDataReadReferenceOrConstruct<FTrigger>(METASOUND_GET_PARAM_NAME(InputToggle), InParams.OperatorSettings);
-		FBoolReadRef bStartClosedIn = InParams.InputDataReferences.GetDataReadReferenceOrConstructWithVertexDefault<bool>(InputInterface, METASOUND_GET_PARAM_NAME(InputStartClosed), InParams.OperatorSettings);
+		FTriggerReadRef TriggerEnterIn = InputData.GetOrConstructDataReadReference<FTrigger>(METASOUND_GET_PARAM_NAME(InputEnter), InParams.OperatorSettings);
+		FTriggerReadRef TriggerOpenIn = InputData.GetOrConstructDataReadReference<FTrigger>(METASOUND_GET_PARAM_NAME(InputOpen), InParams.OperatorSettings);
+		FTriggerReadRef TriggerCloseIn = InputData.GetOrConstructDataReadReference<FTrigger>(METASOUND_GET_PARAM_NAME(InputClose), InParams.OperatorSettings);
+		FTriggerReadRef TriggerToggleIn = InputData.GetOrConstructDataReadReference<FTrigger>(METASOUND_GET_PARAM_NAME(InputToggle), InParams.OperatorSettings);
+		FBoolReadRef bStartClosedIn = InputData.GetOrCreateDefaultDataReadReference<bool>(METASOUND_GET_PARAM_NAME(InputStartClosed), InParams.OperatorSettings);
 
 		return MakeUnique<FTriggerControlOperator>(InParams.OperatorSettings, TriggerEnterIn, TriggerOpenIn, TriggerCloseIn, TriggerToggleIn, bStartClosedIn);
 	}
diff --git a/Engine/Plugins/Runtime/Metasound/Source/MetasoundStandardNodes/Private/MetasoundTriggerCounterNode.cpp b/Engine/Plugins/Runtime/Metasound/Source/MetasoundStandardNodes/Private/MetasoundTriggerCounterNode.cpp
--- a/Engine/Plugins/Runtime/Metasound/Source/MetasoundStandardNodes/Private/MetasoundTriggerCounterNode.cpp
+++ b/Engine/Plugins/Runtime/Metasound/Source/MetasoundStandardNodes/Private/MetasoundTriggerCounterNode.cpp
@@ -37,7 +37,7 @@ namespace Metasound
 		public:
 			static const FNodeClassMetadata& GetNodeInfo();
 			static const FVertexInterface& GetVertexInterface();
-			static TUniquePtr<IOperator> CreateOperator(const FCreateOperatorParams& InParams, FBuildErrorArray& OutErrors);
+			static TUniquePtr<IOperator> CreateOperator(const FBuildOperatorParams& InParams, FBuildResults& OutResults);
 
 			FTriggerCounterOperator(const FOperatorSettings& InSettings,
 				const FTriggerReadRef& InTriggerIn,
@@ -199,16 +199,16 @@ namespace Metasound
 		bIsFirstTrigger = true;
 	}
 
-	TUniquePtr<IOperator> FTriggerCounterOperator::CreateOperator(const FCreateOperatorParams& InParams, FBuildErrorArray& OutErrors)
+	TUniquePtr<IOperator> FTriggerCounterOperator::CreateOperator(const FBuildOperatorParams& InParams, FBuildResults& OutResults)
 	{
 		using namespace TriggerCounterVertexNames;
 
-		const FInputVertexInterface& InputInterface = GetVertexInterface().GetInputInterface();
-		FTriggerReadRef TriggerIn = InParams.InputDataReferences.GetDataReadReferenceOrConstruct<FTrigger>(METASOUND_GET_PARAM_NAME(InputInTrigger), InParams.OperatorSettings);
-		FTriggerReadRef TriggerReset = InParams.InputDataReferences.GetDataReadReferenceOrConstruct<FTrigger>(METASOUND_GET_PARAM_NAME(InputReset), InParams.OperatorSettings);
-		FFloatReadRef StartValue = InParams.InputDataReferences.GetDataReadReferenceOrConstructWithVertexDefault<float>(InputInterface, METASOUND_GET_PARAM_NAME(InputStartValue), InParams.OperatorSettings);
-		FFloatReadRef StepSize = InParams.InputDataReferences.GetDataReadReferenceOrConstructWithVertexDefault<float>(InputInterface, METASOUND_GET_PARAM_NAME(InputStepSize), InParams.OperatorSettings);
-		FInt32ReadRef AutoResetCount = InParams.InputDataReferences.GetDataReadReferenceOrConstructWithVertexDefault<int32>(InputInterface, METASOUND_GET_PARAM_NAME(InputAutoResetCount), InParams.OperatorSettings);
+		const FInputVertexInterfaceData& InputData = InParams.InputData;
+		FTriggerReadRef TriggerIn = InputData.GetOrConstructDataReadReference<FTrigger>(METASOUND_GET_PARAM_NAME(InputInTrigger), InParams.OperatorSettings);
+		FTriggerReadRef TriggerReset = InputData.GetOrConstructDataReadReference<FTrigger>(METASOUND_GET_PARAM_NAME(InputReset), InParams.OperatorSettings);
+		FFloatReadRef StartValue = InputData.GetOrCreateDefaultDataReadReference<float>(METASOUND_GET_PARAM_NAME(InputStartValue), InParams.OperatorSettings);
+		FFloatReadRef StepSize = InputData.GetOrCreateDefaultDataReadReference<float>(METASOUND_GET_PARAM_NAME(InputStepSize), InParams.OperatorSettings);
+		FInt32ReadRef AutoResetCount = InputData.GetOrCreateDefaultDataReadReference<int32>(METASOUND_GET_PARAM_NAME(InputAutoResetCount), InParams.OperatorSettings);
 
 		return MakeUnique<FTriggerCounterOperator>(InParams.OperatorSettings, TriggerIn, TriggerReset, StartValue, StepSize, AutoResetCount);
 	}
diff --git a/Engine/Plugins/Runtime/Metasound/Source/MetasoundStandardNodes/Private/MetasoundTriggerDelayNode.cpp b/Engine/Plugins/Runtime/Metasound/Source/MetasoundStandardNodes/Private/MetasoundTriggerDelayNode.cpp
--- a/Engine/Plugins/Runtime/Metasound/Source/MetasoundStandardNodes/Private/MetasoundTriggerDelayNode.cpp
+++ b/Engine/Plugins/Runtime/Metasound/Source/MetasoundStandardNodes/Private/MetasoundTriggerDelayNode.cpp
@@ -31,7 +31,7 @@ namespace Metasound
 		public:
 			static const FNodeClassMetadata& GetNodeInfo();
 			static const FVertexInterface& GetVertexInterface();
-			static TUniquePtr<IOperator> CreateOperator(const FCreateOperatorParams& InParams, FBuildErrorArray& OutErrors);
+			static TUniquePtr<IOperator> CreateOperator(const FBuildOperatorParams& InParams, FBuildResults& OutResults);
 
 			FTriggerDelayOperator(const FOperatorSettings& InSettings, const FTriggerReadRef& InTriggerReset, const FTriggerReadRef& InTriggerIn, const FTimeReadRef& InDelay);
 
@@ -123,16 +123,16 @@ namespace Metasound
 		TriggerOut->Reset();
 	}
 
-	TUniquePtr<IOperator> FTriggerDelayOperator::CreateOperator(const FCreateOperatorParams& InParams, FBuildErrorArray& OutErrors)
+	TUniquePtr<IOperator> FTriggerDelayOperator::CreateOperator(const FBuildOperatorParams& InParams, FBuildResults& OutResults)
 	{
 		using namespace TriggerDelayVertexNames;
 
-		const FInputVertexInterface& InputInterface = GetVertexInterface().GetInputInterface();
+		const FInputVertexInterfaceData& InputData = InParams.InputData;
 
-		FTimeReadRef Delay = InParams.InputDataReferences.GetDataReadReferenceOrConstructWithVertexDefault<FTime>(InputInterface, METASOUND_GET_PARAM_NAME(InputDelayTime), InParams.OperatorSettings);
+		FTimeReadRef Delay = InputData.GetOrCreateDefaultDataReadReference<FTime>(METASOUND_GET_PARAM_NAME(InputDelayTime), InParams.OperatorSettings);
 
-		FTriggerReadRef TriggerIn = InParams.InputDataReferences.GetDataReadReferenceOrConstruct<FTrigger>(METASOUND_GET_PARAM_NAME(InputInTriggerDelay), InParams.OperatorSettings);
-		FTriggerReadRef TriggerReset = InParams.InputDataReferences.GetDataReadReferenceOrConstruct<FTrigger>(METASOUND_GET_PARAM_NAME(InputResetDelay), InParams.OperatorSettings);
+		FTriggerReadRef TriggerIn = InputData.GetOrConstructDataReadReference<FTrigger>(METASOUND_GET_PARAM_NAME(InputInTriggerDelay), InParams.OperatorSettings);
+		FTriggerReadRef TriggerReset = InputData.GetOrConstructDataReadReference<FTrigger>(METASOUND_GET_PARAM_NAME(InputResetDelay), InParams.OperatorSettings);
 
 		return MakeUnique<FTriggerDelayOperator>(InParams.OperatorSettings, TriggerReset, TriggerIn, Delay);
 	}
diff --git a/Engine/Plugins/Runtime/Metasound/Source/MetasoundStandardNodes/Private/MetasoundTriggerOnThresholdNode.cpp b/Engine/Plugins/Runtime/Metasound/Source/MetasoundStandardNodes/Private/MetasoundTriggerOnThresholdNode.cpp
--- a/Engine/Plugins/Runtime/Metasound/Source/MetasoundStandardNodes/Private/MetasoundTriggerOnThresholdNode.cpp
+++ b/Engine/Plugins/Runtime/Metasound/Source/MetasoundStandardNodes/Private/MetasoundTriggerOnThresholdNode.cpp
@@ -104,15 +104,13 @@ namespace Metasound
 		}
 
 		// In CreateOperator, the only thing that varies between Audio and non-audio input is whether there is a default value for the Input. 
-		static TDataReadReference<ValueType> CreateInput(const FCreateOperatorParams& InParams)
+		static TDataReadReference<ValueType> CreateInput(const FBuildOperatorParams& InParams)
 		{
 			using namespace TriggerOnThresholdVertexNames;
-
-			const FDataReferenceCollection& InputCol = InParams.InputDataReferences;
 			const FOperatorSettings& Settings = InParams.OperatorSettings;
-			const FInputVertexInterface& InputInterface = InParams.Node.GetVertexInterface().GetInputInterface();
+			const FInputVertexInterfaceData& InputData = InParams.InputData;
 
-			return InputCol.GetDataReadReferenceOrConstructWithVertexDefault<ValueType>(InputInterface, METASOUND_GET_PARAM_NAME(InPin), Settings);
+			return InputData.GetOrCreateDefaultDataReadReference<ValueType>(METASOUND_GET_PARAM_NAME(InPin), Settings);
 		}
 
 	};
@@ -191,14 +189,14 @@ namespace Metasound
 			return Interface;
 		}
 
-		static FAudioBufferReadRef CreateInput(const FCreateOperatorParams& InParams)
+		static FAudioBufferReadRef CreateInput(const FBuildOperatorParams& InParams)
 		{
 			using namespace TriggerOnThresholdVertexNames;
 
-			const FDataReferenceCollection& InputCol = InParams.InputDataReferences;
+			const FInputVertexInterfaceData& InputData = InParams.InputData;
 			const FOperatorSettings& Settings = InParams.OperatorSettings;
 
-			return InputCol.GetDataReadReferenceOrConstruct<FAudioBuffer>(METASOUND_GET_PARAM_NAME(InPin), Settings);
+			return InputData.GetOrConstructDataReadReference<FAudioBuffer>(METASOUND_GET_PARAM_NAME(InPin), Settings);
 		}
 	};
 
@@ -300,17 +298,16 @@ namespace Metasound
 			return Info;
 		}
 
-		static TUniquePtr<IOperator> CreateOperator(const FCreateOperatorParams& InParams, FBuildErrorArray& OutErrors)
+		static TUniquePtr<IOperator> CreateOperator(const FBuildOperatorParams& InParams, FBuildResults& OutResults)
 		{
 			using namespace TriggerOnThresholdVertexNames;
 
-			const FDataReferenceCollection& InputCol = InParams.InputDataReferences;
 			const FOperatorSettings& Settings = InParams.OperatorSettings;
-			const FInputVertexInterface& InputInterface = InParams.Node.GetVertexInterface().GetInputInterface();
+			const FInputVertexInterfaceData& InputData = InParams.InputData;
 
-			FBufferTriggerTypeReadRef Type = InputCol.GetDataReadReferenceOrConstructWithVertexDefault<FEnumBufferTriggerType>(InputInterface, METASOUND_GET_PARAM_NAME(InTriggerType), Settings);
+			FBufferTriggerTypeReadRef Type = InputData.GetOrCreateDefaultDataReadReference<FEnumBufferTriggerType>(METASOUND_GET_PARAM_NAME(InTriggerType), Settings);
 			TDataReadReference<ValueType> Input = TTriggerOnThresholdHelper<ValueType>::CreateInput(InParams);
-			TDataReadReference<ThresholdType> Threshold = InputCol.GetDataReadReferenceOrConstructWithVertexDefault<ThresholdType>(InputInterface, METASOUND_GET_PARAM_NAME(InThresholdPin), Settings);
+			TDataReadReference<ThresholdType> Threshold = InputData.GetOrCreateDefaultDataReadReference<ThresholdType>(METASOUND_GET_PARAM_NAME(InThresholdPin), Settings);
 
 			return MakeUnique<TTriggerOnThresholdOperator<ValueType, ThresholdType>>(InParams.OperatorSettings, Input, Threshold, Type);
 		}		
diff --git a/Engine/Plugins/Runtime/Metasound/Source/MetasoundStandardNodes/Private/MetasoundTriggerOnValueChangeNode.cpp b/Engine/Plugins/Runtime/Metasound/Source/MetasoundStandardNodes/Private/MetasoundTriggerOnValueChangeNode.cpp
--- a/Engine/Plugins/Runtime/Metasound/Source/MetasoundStandardNodes/Private/MetasoundTriggerOnValueChangeNode.cpp
+++ b/Engine/Plugins/Runtime/Metasound/Source/MetasoundStandardNodes/Private/MetasoundTriggerOnValueChangeNode.cpp
@@ -47,7 +47,7 @@ namespace Metasound
 	public:
 		static const FNodeClassMetadata& GetNodeInfo();
 		static const FVertexInterface& GetVertexInterface();
-		static TUniquePtr<IOperator> CreateOperator(const FCreateOperatorParams& InParams, FBuildErrorArray& OutErrors);
+		static TUniquePtr<IOperator> CreateOperator(const FBuildOperatorParams& InParams, FBuildResults& OutResults);
 
 		TTriggerOnValueChangeOperator(const FOperatorSettings& InSettings, const TDataReadReference<ValueType>& InValue);
 
@@ -136,13 +136,13 @@ namespace Metasound
 	}
 
 	template <typename ValueType>
-	TUniquePtr<IOperator> TTriggerOnValueChangeOperator<ValueType>::CreateOperator(const FCreateOperatorParams& InParams, FBuildErrorArray& OutErrors)
+	TUniquePtr<IOperator> TTriggerOnValueChangeOperator<ValueType>::CreateOperator(const FBuildOperatorParams& InParams, FBuildResults& OutResults)
 	{
 		using namespace TriggerOnValueChangeVertexNames;
 
-		const FInputVertexInterface& InputInterface = GetVertexInterface().GetInputInterface();
+		const FInputVertexInterfaceData& InputData = InParams.InputData;
 
-		TDataReadReference<ValueType> ValueIn = InParams.InputDataReferences.GetDataReadReferenceOrConstructWithVertexDefault<ValueType>(InputInterface, METASOUND_GET_PARAM_NAME(InputValue), InParams.OperatorSettings);
+		TDataReadReference<ValueType> ValueIn = InputData.GetOrCreateDefaultDataReadReference<ValueType>(METASOUND_GET_PARAM_NAME(InputValue), InParams.OperatorSettings);
 
 		return MakeUnique<TTriggerOnValueChangeOperator>(InParams.OperatorSettings, ValueIn);
 	}
diff --git a/Engine/Plugins/Runtime/Metasound/Source/MetasoundStandardNodes/Private/MetasoundTriggerOnceNode.cpp b/Engine/Plugins/Runtime/Metasound/Source/MetasoundStandardNodes/Private/MetasoundTriggerOnceNode.cpp
--- a/Engine/Plugins/Runtime/Metasound/Source/MetasoundStandardNodes/Private/MetasoundTriggerOnceNode.cpp
+++ b/Engine/Plugins/Runtime/Metasound/Source/MetasoundStandardNodes/Private/MetasoundTriggerOnceNode.cpp
@@ -34,7 +34,7 @@ namespace Metasound
 	public:
 		static const FNodeClassMetadata& GetNodeInfo();
 		static const FVertexInterface& GetVertexInterface();
-		static TUniquePtr<IOperator> CreateOperator(const FCreateOperatorParams& InParams, FBuildErrorArray& OutErrors);
+		static TUniquePtr<IOperator> CreateOperator(const FBuildOperatorParams& InParams, FBuildResults& OutResults);
 
 		FTriggerOnceOperator(const FOperatorSettings& InSettings,
 			const FTriggerReadRef& InTriggerEnter,
@@ -144,15 +144,15 @@ namespace Metasound
 		bIsGateOpen = !(*bStartClosedInput);
 	}
 
-	TUniquePtr<IOperator> FTriggerOnceOperator::CreateOperator(const FCreateOperatorParams& InParams, FBuildErrorArray& OutErrors)
+	TUniquePtr<IOperator> FTriggerOnceOperator::CreateOperator(const FBuildOperatorParams& InParams, FBuildResults& OutResults)
 	{
 		using namespace TriggerOnceVertexNames;
 
-		const FInputVertexInterface& InputInterface = GetVertexInterface().GetInputInterface();
+		const FInputVertexInterfaceData& InputData = InParams.InputData;
 
-		FTriggerReadRef TriggerEnterIn = InParams.InputDataReferences.GetDataReadReferenceOrConstruct<FTrigger>(METASOUND_GET_PARAM_NAME(InputEnter), InParams.OperatorSettings);
-		FTriggerReadRef TriggerResetIn = InParams.InputDataReferences.GetDataReadReferenceOrConstruct<FTrigger>(METASOUND_GET_PARAM_NAME(InputReset), InParams.OperatorSettings);
-		FBoolReadRef bStartClosedIn = InParams.InputDataReferences.GetDataReadReferenceOrConstructWithVertexDefault<bool>(InputInterface, METASOUND_GET_PARAM_NAME(InputStartClosed), InParams.OperatorSettings);
+		FTriggerReadRef TriggerEnterIn = InputData.GetOrConstructDataReadReference<FTrigger>(METASOUND_GET_PARAM_NAME(InputEnter), InParams.OperatorSettings);
+		FTriggerReadRef TriggerResetIn = InputData.GetOrConstructDataReadReference<FTrigger>(METASOUND_GET_PARAM_NAME(InputReset), InParams.OperatorSettings);
+		FBoolReadRef bStartClosedIn = InputData.GetOrCreateDefaultDataReadReference<bool>(METASOUND_GET_PARAM_NAME(InputStartClosed), InParams.OperatorSettings);
 
 		return MakeUnique<FTriggerOnceOperator>(InParams.OperatorSettings, TriggerEnterIn, TriggerResetIn, bStartClosedIn);
 	}
diff --git a/Engine/Plugins/Runtime/Metasound/Source/MetasoundStandardNodes/Private/MetasoundTriggerPipeNode.cpp b/Engine/Plugins/Runtime/Metasound/Source/MetasoundStandardNodes/Private/MetasoundTriggerPipeNode.cpp
--- a/Engine/Plugins/Runtime/Metasound/Source/MetasoundStandardNodes/Private/MetasoundTriggerPipeNode.cpp
+++ b/Engine/Plugins/Runtime/Metasound/Source/MetasoundStandardNodes/Private/MetasoundTriggerPipeNode.cpp
@@ -34,7 +34,7 @@ namespace Metasound
 	public:
 		static const FNodeClassMetadata& GetNodeInfo();
 		static const FVertexInterface& GetVertexInterface();
-		static TUniquePtr<IOperator> CreateOperator(const FCreateOperatorParams& InParams, FBuildErrorArray& OutErrors);
+		static TUniquePtr<IOperator> CreateOperator(const FBuildOperatorParams& InParams, FBuildResults& OutResults);
 
 		FTriggerPipeOperator(const FOperatorSettings& InSettings, const FTriggerReadRef& InTriggerReset, const FTriggerReadRef& InTriggerIn, const FTimeReadRef& InDelay);
 
@@ -153,16 +153,14 @@ namespace Metasound
 		TriggerOut->Reset();
 	}
 
-	TUniquePtr<IOperator> FTriggerPipeOperator::CreateOperator(const FCreateOperatorParams& InParams, FBuildErrorArray& OutErrors)
+	TUniquePtr<IOperator> FTriggerPipeOperator::CreateOperator(const FBuildOperatorParams& InParams, FBuildResults& OutResults)
 	{
 		using namespace TriggerPipeVertexNames;
+		const FInputVertexInterfaceData& InputData = InParams.InputData;
 
-		const FDataReferenceCollection& InputCollection = InParams.InputDataReferences;
-		const FInputVertexInterface& InputInterface = GetVertexInterface().GetInputInterface();
-
-		FTriggerReadRef TriggerIn = InParams.InputDataReferences.GetDataReadReferenceOrConstruct<FTrigger>(METASOUND_GET_PARAM_NAME(InputInTrigger), InParams.OperatorSettings);
-		FTriggerReadRef TriggerReset = InParams.InputDataReferences.GetDataReadReferenceOrConstruct<FTrigger>(METASOUND_GET_PARAM_NAME(InputReset), InParams.OperatorSettings);
-		FTimeReadRef DelayTime = InputCollection.GetDataReadReferenceOrConstructWithVertexDefault<FTime>(InputInterface, METASOUND_GET_PARAM_NAME(InputDelayTime), InParams.OperatorSettings);
+		FTriggerReadRef TriggerIn = InputData.GetOrConstructDataReadReference<FTrigger>(METASOUND_GET_PARAM_NAME(InputInTrigger), InParams.OperatorSettings);
+		FTriggerReadRef TriggerReset = InputData.GetOrConstructDataReadReference<FTrigger>(METASOUND_GET_PARAM_NAME(InputReset), InParams.OperatorSettings);
+		FTimeReadRef DelayTime = InputData.GetOrCreateDefaultDataReadReference<FTime>(METASOUND_GET_PARAM_NAME(InputDelayTime), InParams.OperatorSettings);
 
 		return MakeUnique<FTriggerPipeOperator>(InParams.OperatorSettings, TriggerReset, TriggerIn, DelayTime);
 	}
diff --git a/Engine/Plugins/Runtime/Metasound/Source/MetasoundStandardNodes/Private/MetasoundTriggerRepeatNode.cpp b/Engine/Plugins/Runtime/Metasound/Source/MetasoundStandardNodes/Private/MetasoundTriggerRepeatNode.cpp
--- a/Engine/Plugins/Runtime/Metasound/Source/MetasoundStandardNodes/Private/MetasoundTriggerRepeatNode.cpp
+++ b/Engine/Plugins/Runtime/Metasound/Source/MetasoundStandardNodes/Private/MetasoundTriggerRepeatNode.cpp
@@ -38,7 +38,7 @@ namespace Metasound
 		public:
 			static const FNodeClassMetadata& GetNodeInfo();
 			static const FVertexInterface& GetVertexInterface();
-			static TUniquePtr<IOperator> CreateOperator(const FCreateOperatorParams& InParams, FBuildErrorArray& OutErrors);
+			static TUniquePtr<IOperator> CreateOperator(const FBuildOperatorParams& InParams, FBuildResults& OutResults);
 
 			FTriggerRepeatOperator(const FOperatorSettings& InSettings, const FTriggerReadRef& InTriggerEnable, const FTriggerReadRef& InTriggerDisable, 
 				const FTimeReadRef& InPeriod, const FInt32ReadRef& InNumRepeats);
@@ -178,17 +178,17 @@ namespace Metasound
 		CurrentRepeats = 0;
 	}
 
-	TUniquePtr<IOperator> FTriggerRepeatOperator::CreateOperator(const FCreateOperatorParams& InParams, FBuildErrorArray& OutErrors)
+	TUniquePtr<IOperator> FTriggerRepeatOperator::CreateOperator(const FBuildOperatorParams& InParams, FBuildResults& OutResults)
 	{
 		using namespace TriggerRepeatVertexNames;
 
 		const FTriggerRepeatNode& PeriodicTriggerNode = static_cast<const FTriggerRepeatNode&>(InParams.Node);
-		const FInputVertexInterface& InputInterface = GetVertexInterface().GetInputInterface();
+		const FInputVertexInterfaceData& InputData = InParams.InputData;
 
-		FTriggerReadRef TriggerEnable = InParams.InputDataReferences.GetDataReadReferenceOrConstruct<FTrigger>(METASOUND_GET_PARAM_NAME(InputStart), InParams.OperatorSettings);
-		FTriggerReadRef TriggerDisable = InParams.InputDataReferences.GetDataReadReferenceOrConstruct<FTrigger>(METASOUND_GET_PARAM_NAME(InputStop), InParams.OperatorSettings);
-		FTimeReadRef Period = InParams.InputDataReferences.GetDataReadReferenceOrConstructWithVertexDefault<FTime>(InputInterface, METASOUND_GET_PARAM_NAME(InputPeriod), InParams.OperatorSettings);
-		FInt32ReadRef NumRepeats = InParams.InputDataReferences.GetDataReadReferenceOrConstructWithVertexDefault<int32>(InputInterface, METASOUND_GET_PARAM_NAME(InputNumRepeats), InParams.OperatorSettings);
+		FTriggerReadRef TriggerEnable = InputData.GetOrConstructDataReadReference<FTrigger>(METASOUND_GET_PARAM_NAME(InputStart), InParams.OperatorSettings);
+		FTriggerReadRef TriggerDisable = InputData.GetOrConstructDataReadReference<FTrigger>(METASOUND_GET_PARAM_NAME(InputStop), InParams.OperatorSettings);
+		FTimeReadRef Period = InputData.GetOrCreateDefaultDataReadReference<FTime>(METASOUND_GET_PARAM_NAME(InputPeriod), InParams.OperatorSettings);
+		FInt32ReadRef NumRepeats = InputData.GetOrCreateDefaultDataReadReference<int32>(METASOUND_GET_PARAM_NAME(InputNumRepeats), InParams.OperatorSettings);
 
 		return MakeUnique<FTriggerRepeatOperator>(InParams.OperatorSettings, TriggerEnable, TriggerDisable, Period, NumRepeats);
 	}
diff --git a/Engine/Plugins/Runtime/Metasound/Source/MetasoundStandardNodes/Private/MetasoundTriggerRouteNode.cpp b/Engine/Plugins/Runtime/Metasound/Source/MetasoundStandardNodes/Private/MetasoundTriggerRouteNode.cpp
--- a/Engine/Plugins/Runtime/Metasound/Source/MetasoundStandardNodes/Private/MetasoundTriggerRouteNode.cpp
+++ b/Engine/Plugins/Runtime/Metasound/Source/MetasoundStandardNodes/Private/MetasoundTriggerRouteNode.cpp
@@ -144,26 +144,25 @@ namespace Metasound
 			return Metadata;
 		}
 
-		static TUniquePtr<IOperator> CreateOperator(const FCreateOperatorParams& InParams, TArray<TUniquePtr<IOperatorBuildError>>& OutErrors)
+		static TUniquePtr<IOperator> CreateOperator(const FBuildOperatorParams& InParams, FBuildResults& OutResults)
 		{
 			using namespace TriggerRouteVertexNames;
 
-			const FInputVertexInterface& InputInterface = InParams.Node.GetVertexInterface().GetInputInterface();
-			const FDataReferenceCollection& InputCollection = InParams.InputDataReferences;
+			const FInputVertexInterfaceData& InputData = InParams.InputData;
 
 			TArray<FTriggerReadRef> InputTriggers;
 			TArray<TDataReadReference<ValueType>> InputValues;
 
 			for (uint32 i = 0; i < NumInputs; ++i)
 			{
-				InputTriggers.Add(InputCollection.GetDataReadReferenceOrConstruct<FTrigger>(GetInputTriggerName(i), InParams.OperatorSettings));
-				InputValues.Add(InputCollection.GetDataReadReferenceOrConstructWithVertexDefault<ValueType>(InputInterface, GetInputValueName(i), InParams.OperatorSettings));
+				InputTriggers.Add(InputData.GetOrConstructDataReadReference<FTrigger>(GetInputTriggerName(i), InParams.OperatorSettings));
+				InputValues.Add(InputData.GetOrCreateDefaultDataReadReference<ValueType>(GetInputValueName(i), InParams.OperatorSettings));
 			}
 
 			return MakeUnique<TTriggerRouteOperator<ValueType, NumInputs>>(InParams, MoveTemp(InputTriggers), MoveTemp(InputValues));
 		}
 
-		TTriggerRouteOperator(const FCreateOperatorParams& InParams, TArray<FTriggerReadRef>&& InInputTriggers, TArray<TDataReadReference<ValueType>>&& InInputValues)
+		TTriggerRouteOperator(const FBuildOperatorParams& InParams, TArray<FTriggerReadRef>&& InInputTriggers, TArray<TDataReadReference<ValueType>>&& InInputValues)
 			: InputTriggers(MoveTemp(InInputTriggers))
 			, InputValues(MoveTemp(InInputValues))
 			, OutputTrigger(FTriggerWriteRef::CreateNew(InParams.OperatorSettings))
diff --git a/Engine/Plugins/Runtime/Metasound/Source/MetasoundStandardNodes/Private/MetasoundTriggerSelectNode.cpp b/Engine/Plugins/Runtime/Metasound/Source/MetasoundStandardNodes/Private/MetasoundTriggerSelectNode.cpp
--- a/Engine/Plugins/Runtime/Metasound/Source/MetasoundStandardNodes/Private/MetasoundTriggerSelectNode.cpp
+++ b/Engine/Plugins/Runtime/Metasound/Source/MetasoundStandardNodes/Private/MetasoundTriggerSelectNode.cpp
@@ -38,7 +38,7 @@ namespace Metasound
 	public:
 		static const FNodeClassMetadata& GetNodeInfo();
 		static const FVertexInterface& GetVertexInterface();
-		static TUniquePtr<IOperator> CreateOperator(const FCreateOperatorParams& InParams, FBuildErrorArray& OutErrors);
+		static TUniquePtr<IOperator> CreateOperator(const FBuildOperatorParams& InParams, FBuildResults& OutResults);
 
 		TTriggerSelectOperator(const FOperatorSettings& InSettings,
 			const FTriggerReadRef& InTrigger,
@@ -144,14 +144,14 @@ namespace Metasound
 	}
 
 	template<uint32 NumOutputs>
-	TUniquePtr<IOperator> TTriggerSelectOperator<NumOutputs>::CreateOperator(const FCreateOperatorParams& InParams, FBuildErrorArray& OutErrors)
+	TUniquePtr<IOperator> TTriggerSelectOperator<NumOutputs>::CreateOperator(const FBuildOperatorParams& InParams, FBuildResults& OutResults)
 	{
 		using namespace TriggerSelectVertexNames;
 
-		const FInputVertexInterface& InputInterface = GetVertexInterface().GetInputInterface();
+		const FInputVertexInterfaceData& InputData = InParams.InputData;
 
-		FTriggerReadRef TriggerIn = InParams.InputDataReferences.GetDataReadReferenceOrConstruct<FTrigger>(METASOUND_GET_PARAM_NAME(InputTrigger), InParams.OperatorSettings);
-		FInt32ReadRef InputIn = InParams.InputDataReferences.GetDataReadReferenceOrConstructWithVertexDefault<int32>(InputInterface, METASOUND_GET_PARAM_NAME(InputIndex), InParams.OperatorSettings);
+		FTriggerReadRef TriggerIn = InputData.GetOrConstructDataReadReference<FTrigger>(METASOUND_GET_PARAM_NAME(InputTrigger), InParams.OperatorSettings);
+		FInt32ReadRef InputIn = InputData.GetOrCreateDefaultDataReadReference<int32>(METASOUND_GET_PARAM_NAME(InputIndex), InParams.OperatorSettings);
 
 		return MakeUnique<TTriggerSelectOperator>(InParams.OperatorSettings, TriggerIn, InputIn);
 	}
diff --git a/Engine/Plugins/Runtime/Metasound/Source/MetasoundStandardNodes/Private/MetasoundTriggerSequenceNode.cpp b/Engine/Plugins/Runtime/Metasound/Source/MetasoundStandardNodes/Private/MetasoundTriggerSequenceNode.cpp
--- a/Engine/Plugins/Runtime/Metasound/Source/MetasoundStandardNodes/Private/MetasoundTriggerSequenceNode.cpp
+++ b/Engine/Plugins/Runtime/Metasound/Source/MetasoundStandardNodes/Private/MetasoundTriggerSequenceNode.cpp
@@ -40,7 +40,7 @@ namespace Metasound
 	public:
 		static const FNodeClassMetadata& GetNodeInfo();
 		static const FVertexInterface& GetVertexInterface();
-		static TUniquePtr<IOperator> CreateOperator(const FCreateOperatorParams& InParams, FBuildErrorArray& OutErrors);
+		static TUniquePtr<IOperator> CreateOperator(const FBuildOperatorParams& InParams, FBuildResults& OutResults);
 
 		TTriggerSequenceOperator(const FOperatorSettings& InSettings, 
 			const FTriggerReadRef& InTrigger, 
@@ -175,15 +175,15 @@ namespace Metasound
 	}
 
 	template<uint32 NumOutputs>
-	TUniquePtr<IOperator> TTriggerSequenceOperator<NumOutputs>::CreateOperator(const FCreateOperatorParams& InParams, FBuildErrorArray& OutErrors)
+	TUniquePtr<IOperator> TTriggerSequenceOperator<NumOutputs>::CreateOperator(const FBuildOperatorParams& InParams, FBuildResults& OutResults)
 	{
 		using namespace TriggerSequenceVertexNames;
 
-		const FInputVertexInterface& InputInterface = GetVertexInterface().GetInputInterface();
+		const FInputVertexInterfaceData& InputData = InParams.InputData;
 
-		FTriggerReadRef TriggerIn = InParams.InputDataReferences.GetDataReadReferenceOrConstruct<FTrigger>(METASOUND_GET_PARAM_NAME(InputTrigger), InParams.OperatorSettings);
-		FTriggerReadRef TriggerResetIn = InParams.InputDataReferences.GetDataReadReferenceOrConstruct<FTrigger>(METASOUND_GET_PARAM_NAME(InputReset), InParams.OperatorSettings);
-		FBoolReadRef bLoopIn = InParams.InputDataReferences.GetDataReadReferenceOrConstructWithVertexDefault<bool>(InputInterface, METASOUND_GET_PARAM_NAME(InputLoop), InParams.OperatorSettings);
+		FTriggerReadRef TriggerIn = InputData.GetOrConstructDataReadReference<FTrigger>(METASOUND_GET_PARAM_NAME(InputTrigger), InParams.OperatorSettings);
+		FTriggerReadRef TriggerResetIn = InputData.GetOrConstructDataReadReference<FTrigger>(METASOUND_GET_PARAM_NAME(InputReset), InParams.OperatorSettings);
+		FBoolReadRef bLoopIn = InputData.GetOrCreateDefaultDataReadReference<bool>(METASOUND_GET_PARAM_NAME(InputLoop), InParams.OperatorSettings);
 
 		return MakeUnique<TTriggerSequenceOperator>(InParams.OperatorSettings, TriggerIn, TriggerResetIn, bLoopIn);
 	}
diff --git a/Engine/Plugins/Runtime/Metasound/Source/MetasoundStandardNodes/Private/MetasoundTriggerToggleNode.cpp b/Engine/Plugins/Runtime/Metasound/Source/MetasoundStandardNodes/Private/MetasoundTriggerToggleNode.cpp
--- a/Engine/Plugins/Runtime/Metasound/Source/MetasoundStandardNodes/Private/MetasoundTriggerToggleNode.cpp
+++ b/Engine/Plugins/Runtime/Metasound/Source/MetasoundStandardNodes/Private/MetasoundTriggerToggleNode.cpp
@@ -35,7 +35,7 @@ namespace Metasound
 		public:
 			static const FNodeClassMetadata& GetNodeInfo();
 			static const FVertexInterface& GetVertexInterface();
-			static TUniquePtr<IOperator> CreateOperator(const FCreateOperatorParams& InParams, FBuildErrorArray& OutErrors);
+			static TUniquePtr<IOperator> CreateOperator(const FBuildOperatorParams& InParams, FBuildResults& OutResults);
 
 			FTriggerToggleOperator(const FOperatorSettings& InSettings, const FTriggerReadRef& InTriggerOn, const FTriggerReadRef& InTriggerOff, const FBoolReadRef& InInitValue);
 
@@ -129,14 +129,14 @@ namespace Metasound
 		*ValueOutput = *InitValue;
 	}
 
-	TUniquePtr<IOperator> FTriggerToggleOperator::CreateOperator(const FCreateOperatorParams& InParams, FBuildErrorArray& OutErrors)
+	TUniquePtr<IOperator> FTriggerToggleOperator::CreateOperator(const FBuildOperatorParams& InParams, FBuildResults& OutResults)
 	{
 		using namespace TriggerToggle;
 
-		const FInputVertexInterface& InputInterface = GetVertexInterface().GetInputInterface();
-		FTriggerReadRef InTriggerOn = InParams.InputDataReferences.GetDataReadReferenceOrConstruct<FTrigger>(METASOUND_GET_PARAM_NAME(InputOnTrigger), InParams.OperatorSettings);
-		FTriggerReadRef InTriggerOff = InParams.InputDataReferences.GetDataReadReferenceOrConstruct<FTrigger>(METASOUND_GET_PARAM_NAME(InputOffTrigger), InParams.OperatorSettings);
-		FBoolReadRef InInitValue = InParams.InputDataReferences.GetDataReadReferenceOrConstructWithVertexDefault<bool>(InputInterface, METASOUND_GET_PARAM_NAME(InputInit), InParams.OperatorSettings);
+		const FInputVertexInterfaceData& InputData = InParams.InputData;
+		FTriggerReadRef InTriggerOn = InputData.GetOrConstructDataReadReference<FTrigger>(METASOUND_GET_PARAM_NAME(InputOnTrigger), InParams.OperatorSettings);
+		FTriggerReadRef InTriggerOff = InputData.GetOrConstructDataReadReference<FTrigger>(METASOUND_GET_PARAM_NAME(InputOffTrigger), InParams.OperatorSettings);
+		FBoolReadRef InInitValue = InputData.GetOrCreateDefaultDataReadReference<bool>(METASOUND_GET_PARAM_NAME(InputInit), InParams.OperatorSettings);
 
 		return MakeUnique<FTriggerToggleOperator>(InParams.OperatorSettings, InTriggerOn, InTriggerOff, InInitValue);
 	}
diff --git a/Engine/Plugins/Runtime/Metasound/Source/MetasoundStandardNodes/Private/MetasoundWaveShaperNode.cpp b/Engine/Plugins/Runtime/Metasound/Source/MetasoundStandardNodes/Private/MetasoundWaveShaperNode.cpp
--- a/Engine/Plugins/Runtime/Metasound/Source/MetasoundStandardNodes/Private/MetasoundWaveShaperNode.cpp
+++ b/Engine/Plugins/Runtime/Metasound/Source/MetasoundStandardNodes/Private/MetasoundWaveShaperNode.cpp
@@ -45,7 +45,7 @@ namespace Metasound
 
 		static const FNodeClassMetadata& GetNodeInfo();
 		static const FVertexInterface& GetVertexInterface();
-		static TUniquePtr<IOperator> CreateOperator(const FCreateOperatorParams& InParams, FBuildErrorArray& OutErrors);
+		static TUniquePtr<IOperator> CreateOperator(const FBuildOperatorParams& InParams, FBuildResults& OutResults);
 
 		FWaveShaperOperator(const FOperatorSettings& InSettings, 
 			const FAudioBufferReadRef& InAudioInput, 
@@ -192,16 +192,15 @@ namespace Metasound
 		return Info;
 	}
 
-	TUniquePtr<IOperator> FWaveShaperOperator::CreateOperator(const FCreateOperatorParams& InParams, FBuildErrorArray& OutErrors)
+	TUniquePtr<IOperator> FWaveShaperOperator::CreateOperator(const FBuildOperatorParams& InParams, FBuildResults& OutResults)
 	{
-		const FDataReferenceCollection& InputCollection = InParams.InputDataReferences;
-		const FInputVertexInterface& InputInterface = GetVertexInterface().GetInputInterface();
-
-		FAudioBufferReadRef AudioIn = InputCollection.GetDataReadReferenceOrConstruct<FAudioBuffer>(METASOUND_GET_PARAM_NAME(InAudioInput), InParams.OperatorSettings);
-		FFloatReadRef InAmount = InputCollection.GetDataReadReferenceOrConstructWithVertexDefault<float>(InputInterface, METASOUND_GET_PARAM_NAME(InWaveShapeAmount), InParams.OperatorSettings);
-		FFloatReadRef InBias = InputCollection.GetDataReadReferenceOrConstructWithVertexDefault<float>(InputInterface, METASOUND_GET_PARAM_NAME(InWaveShapeBias), InParams.OperatorSettings);
-		FFloatReadRef OutGain = InputCollection.GetDataReadReferenceOrConstructWithVertexDefault<float>(InputInterface, METASOUND_GET_PARAM_NAME(InOutputGain), InParams.OperatorSettings);
-		FEnumWaveShaperReadRef InType = InputCollection.GetDataReadReferenceOrConstructWithVertexDefault<FEnumEWaveShaperType>(InputInterface, METASOUND_GET_PARAM_NAME(InType), InParams.OperatorSettings);
+		const FInputVertexInterfaceData& InputData = InParams.InputData;
+
+		FAudioBufferReadRef AudioIn = InputData.GetOrConstructDataReadReference<FAudioBuffer>(METASOUND_GET_PARAM_NAME(InAudioInput), InParams.OperatorSettings);
+		FFloatReadRef InAmount = InputData.GetOrCreateDefaultDataReadReference<float>(METASOUND_GET_PARAM_NAME(InWaveShapeAmount), InParams.OperatorSettings);
+		FFloatReadRef InBias = InputData.GetOrCreateDefaultDataReadReference<float>(METASOUND_GET_PARAM_NAME(InWaveShapeBias), InParams.OperatorSettings);
+		FFloatReadRef OutGain = InputData.GetOrCreateDefaultDataReadReference<float>(METASOUND_GET_PARAM_NAME(InOutputGain), InParams.OperatorSettings);
+		FEnumWaveShaperReadRef InType = InputData.GetOrCreateDefaultDataReadReference<FEnumEWaveShaperType>(METASOUND_GET_PARAM_NAME(InType), InParams.OperatorSettings);
 
 		return MakeUnique<FWaveShaperOperator>(InParams.OperatorSettings, AudioIn, InAmount, InBias, OutGain, InType);
 	}
diff --git a/Engine/Plugins/Runtime/Metasound/Source/MetasoundStandardNodes/Public/MetasoundTriggerAccumulatorNode.h b/Engine/Plugins/Runtime/Metasound/Source/MetasoundStandardNodes/Public/MetasoundTriggerAccumulatorNode.h
--- a/Engine/Plugins/Runtime/Metasound/Source/MetasoundStandardNodes/Public/MetasoundTriggerAccumulatorNode.h
+++ b/Engine/Plugins/Runtime/Metasound/Source/MetasoundStandardNodes/Public/MetasoundTriggerAccumulatorNode.h
@@ -72,19 +72,18 @@ namespace Metasound
 			return Metadata;
 		}
 
-		static TUniquePtr<IOperator> CreateOperator(const FCreateOperatorParams& InParams, TArray<TUniquePtr<IOperatorBuildError>>& OutErrors)
+		static TUniquePtr<IOperator> CreateOperator(const FBuildOperatorParams& InParams, FBuildResults& OutResults)
 		{
 			using namespace TriggerAccumulatorVertexNames;
 
-			const FInputVertexInterface& InputInterface = InParams.Node.GetVertexInterface().GetInputInterface();
-			const FDataReferenceCollection& InputCollection = InParams.InputDataReferences;
-
-			FBoolReadRef bInAutoReset = InputCollection.GetDataReadReferenceOrConstructWithVertexDefault<bool>(InputInterface, METASOUND_GET_PARAM_NAME(InputAutoReset), InParams.OperatorSettings);
+			const FInputVertexInterfaceData& InputData = InParams.InputData;
+			
+			FBoolReadRef bInAutoReset = InputData.GetOrCreateDefaultDataReadReference<bool>(METASOUND_GET_PARAM_NAME(InputAutoReset), InParams.OperatorSettings);
 			TArray<FTriggerReadRef> InputTriggers;
 
 			for (uint32 i = 0; i < NumInputs; ++i)
 			{
-				InputTriggers.Add(InputCollection.GetDataReadReferenceOrConstruct<FTrigger>(METASOUND_GET_PARAM_NAME_WITH_INDEX(InputTrigger, i), InParams.OperatorSettings));
+				InputTriggers.Add(InputData.GetOrConstructDataReadReference<FTrigger>(METASOUND_GET_PARAM_NAME_WITH_INDEX(InputTrigger, i), InParams.OperatorSettings));
 			}
 
 			return MakeUnique<TTriggerAccumulatorOperator<NumInputs>>(InParams.OperatorSettings, bInAutoReset, MoveTemp(InputTriggers));
diff --git a/Engine/Plugins/Runtime/Metasound/Source/MetasoundStandardNodes/Public/MetasoundTriggerCompareNode.h b/Engine/Plugins/Runtime/Metasound/Source/MetasoundStandardNodes/Public/MetasoundTriggerCompareNode.h
--- a/Engine/Plugins/Runtime/Metasound/Source/MetasoundStandardNodes/Public/MetasoundTriggerCompareNode.h
+++ b/Engine/Plugins/Runtime/Metasound/Source/MetasoundStandardNodes/Public/MetasoundTriggerCompareNode.h
@@ -92,17 +92,16 @@ namespace Metasound
 			return Metadata;
 		}
 
-		static TUniquePtr<IOperator> CreateOperator(const FCreateOperatorParams& InParams, TArray<TUniquePtr<IOperatorBuildError>>& OutErrors)
+		static TUniquePtr<IOperator> CreateOperator(const FBuildOperatorParams& InParams, FBuildResults& OutResults)
 		{
 			using namespace TriggerCompareVertexNames;
+			
+			const FInputVertexInterfaceData& InputData = InParams.InputData;
 
-			const FDataReferenceCollection& InputCollection = InParams.InputDataReferences;
-			const FInputVertexInterface& InputInterface = GetDefaultInterface().GetInputInterface();
-
-			FTriggerReadRef InOnTriggerCompare = InputCollection.GetDataReadReferenceOrConstruct<FTrigger>(METASOUND_GET_PARAM_NAME(InputCompare), InParams.OperatorSettings);
-			TDataReadReference<ValueType> InValueA = InputCollection.GetDataReadReferenceOrConstructWithVertexDefault<ValueType>(InputInterface, METASOUND_GET_PARAM_NAME(InputParamA), InParams.OperatorSettings);
-			TDataReadReference<ValueType> InValueB = InputCollection.GetDataReadReferenceOrConstructWithVertexDefault<ValueType>(InputInterface, METASOUND_GET_PARAM_NAME(InputParamB), InParams.OperatorSettings);
-			FEnumTriggerComparisonTypeReadRef InComparison = InputCollection.GetDataReadReferenceOrConstructWithVertexDefault<FEnumTriggerComparisonType>(InputInterface, METASOUND_GET_PARAM_NAME(InputCompareType), InParams.OperatorSettings);
+			FTriggerReadRef InOnTriggerCompare = InputData.GetOrConstructDataReadReference<FTrigger>(METASOUND_GET_PARAM_NAME(InputCompare), InParams.OperatorSettings);
+			TDataReadReference<ValueType> InValueA = InputData.GetOrCreateDefaultDataReadReference<ValueType>(METASOUND_GET_PARAM_NAME(InputParamA), InParams.OperatorSettings);
+			TDataReadReference<ValueType> InValueB = InputData.GetOrCreateDefaultDataReadReference<ValueType>(METASOUND_GET_PARAM_NAME(InputParamB), InParams.OperatorSettings);
+			FEnumTriggerComparisonTypeReadRef InComparison = InputData.GetOrCreateDefaultDataReadReference<FEnumTriggerComparisonType>(METASOUND_GET_PARAM_NAME(InputCompareType), InParams.OperatorSettings);
 
 			return MakeUnique<TTriggerCompareNodeOperator<ValueType>>(InParams.OperatorSettings, InOnTriggerCompare, InValueA, InValueB, InComparison);
 
diff --git a/Engine/Plugins/Runtime/Metasound/Source/MetasoundStandardNodes/Public/MetasoundValueNode.h b/Engine/Plugins/Runtime/Metasound/Source/MetasoundStandardNodes/Public/MetasoundValueNode.h
--- a/Engine/Plugins/Runtime/Metasound/Source/MetasoundStandardNodes/Public/MetasoundValueNode.h
+++ b/Engine/Plugins/Runtime/Metasound/Source/MetasoundStandardNodes/Public/MetasoundValueNode.h
@@ -74,17 +74,16 @@ namespace Metasound
 			return Metadata;
 		}
 
-		static TUniquePtr<IOperator> CreateOperator(const FCreateOperatorParams& InParams, TArray<TUniquePtr<IOperatorBuildError>>& OutErrors)
+		static TUniquePtr<IOperator> CreateOperator(const FBuildOperatorParams& InParams, FBuildResults& OutResults)
 		{
 			using namespace ValueVertexNames;
 
-			const FInputVertexInterface& InputInterface = InParams.Node.GetVertexInterface().GetInputInterface();
-			const FDataReferenceCollection& InputCollection = InParams.InputDataReferences;
-
-			FTriggerReadRef SetTrigger = InputCollection.GetDataReadReferenceOrConstruct<FTrigger>(METASOUND_GET_PARAM_NAME(InputSetTrigger), InParams.OperatorSettings);
-			FTriggerReadRef ResetTrigger = InputCollection.GetDataReadReferenceOrConstruct<FTrigger>(METASOUND_GET_PARAM_NAME(InputResetTrigger), InParams.OperatorSettings);
-			TDataReadReference<ValueType> InitValue = InputCollection.GetDataReadReferenceOrConstructWithVertexDefault<ValueType>(InputInterface, METASOUND_GET_PARAM_NAME(InputInitValue), InParams.OperatorSettings);
-			TDataReadReference<ValueType> TargetValue = InputCollection.GetDataReadReferenceOrConstructWithVertexDefault<ValueType>(InputInterface, METASOUND_GET_PARAM_NAME(InputTargetValue), InParams.OperatorSettings);
+			const FInputVertexInterfaceData& InputData = InParams.InputData;
+			
+			FTriggerReadRef SetTrigger = InputData.GetOrConstructDataReadReference<FTrigger>(METASOUND_GET_PARAM_NAME(InputSetTrigger), InParams.OperatorSettings);
+			FTriggerReadRef ResetTrigger = InputData.GetOrConstructDataReadReference<FTrigger>(METASOUND_GET_PARAM_NAME(InputResetTrigger), InParams.OperatorSettings);
+			TDataReadReference<ValueType> InitValue = InputData.GetOrCreateDefaultDataReadReference<ValueType>(METASOUND_GET_PARAM_NAME(InputInitValue), InParams.OperatorSettings);
+			TDataReadReference<ValueType> TargetValue = InputData.GetOrCreateDefaultDataReadReference<ValueType>(METASOUND_GET_PARAM_NAME(InputTargetValue), InParams.OperatorSettings);
 
 			return MakeUnique<TValueOperator<ValueType>>(InParams.OperatorSettings, SetTrigger, ResetTrigger, InitValue, TargetValue);
 		}
