Update use of MetaSound FCreateOperatorParams to FBuilderOperatorParams for perf (Part 2)

diff --git a/Engine/Plugins/Runtime/Metasound/Source/MetasoundEngine/Private/MetasoundWaveTableOscillatorNode.cpp b/Engine/Plugins/Runtime/Metasound/Source/MetasoundEngine/Private/MetasoundWaveTableOscillatorNode.cpp
--- a/Engine/Plugins/Runtime/Metasound/Source/MetasoundEngine/Private/MetasoundWaveTableOscillatorNode.cpp
+++ b/Engine/Plugins/Runtime/Metasound/Source/MetasoundEngine/Private/MetasoundWaveTableOscillatorNode.cpp
@@ -79,30 +79,29 @@ namespace Metasound
 			return Metadata;
 		}
 
-		static TUniquePtr<IOperator> CreateOperator(const FCreateOperatorParams& InParams, TArray<TUniquePtr<IOperatorBuildError>>& OutErrors)
+		static TUniquePtr<IOperator> CreateOperator(const FBuildOperatorParams& InParams, FBuildResults& OutResults)
 		{
 			using namespace WaveTable;
+			
+			const FInputVertexInterfaceData& InputData = InParams.InputData;
 
-			const FInputVertexInterface& InputInterface = InParams.Node.GetVertexInterface().GetInputInterface();
-			const FDataReferenceCollection& InputCollection = InParams.InputDataReferences;
-
-			FWaveTableReadRef InWaveTableReadRef = InputCollection.GetDataReadReferenceOrConstruct<FWaveTable>("WaveTable");
-			FTriggerReadRef InPlayReadRef = InputCollection.GetDataReadReferenceOrConstructWithVertexDefault<FTrigger>(InputInterface, "Play", InParams.OperatorSettings);
-			FTriggerReadRef InStopReadRef = InputCollection.GetDataReadReferenceOrConstructWithVertexDefault<FTrigger>(InputInterface, "Stop", InParams.OperatorSettings);
-			FTriggerReadRef InSyncReadRef = InputCollection.GetDataReadReferenceOrConstructWithVertexDefault<FTrigger>(InputInterface, "Sync", InParams.OperatorSettings);
-			FFloatReadRef InFreqReadRef = InputCollection.GetDataReadReferenceOrConstructWithVertexDefault<float>(InputInterface, "Freq", InParams.OperatorSettings);
+			FWaveTableReadRef InWaveTableReadRef = InputData.GetOrConstructDataReadReference<FWaveTable>("WaveTable");
+			FTriggerReadRef InPlayReadRef = InputData.GetOrCreateDefaultDataReadReference<FTrigger>("Play", InParams.OperatorSettings);
+			FTriggerReadRef InStopReadRef = InputData.GetOrCreateDefaultDataReadReference<FTrigger>("Stop", InParams.OperatorSettings);
+			FTriggerReadRef InSyncReadRef = InputData.GetOrCreateDefaultDataReadReference<FTrigger>("Sync", InParams.OperatorSettings);
+			FFloatReadRef InFreqReadRef = InputData.GetOrCreateDefaultDataReadReference<float>("Freq", InParams.OperatorSettings);
 
 			TOptional<FAudioBufferReadRef> InPhaseModReadRef;
-			if (InputCollection.ContainsDataReadReference<FAudioBuffer>("PhaseMod"))
+			if (const FAnyDataReference* DataRef = InputData.FindDataReference("PhaseMod"))
 			{
-				InPhaseModReadRef = InputCollection.GetDataReadReference<FAudioBuffer>("PhaseMod");
+				InPhaseModReadRef = DataRef->GetDataReadReference<FAudioBuffer>();
 			}
 
 			return MakeUnique<FMetasoundWaveTableOscillatorNodeOperator>(InParams, InWaveTableReadRef, InPlayReadRef, InStopReadRef, InSyncReadRef, InFreqReadRef, MoveTemp(InPhaseModReadRef));
 		}
 
 		FMetasoundWaveTableOscillatorNodeOperator(
-			const FCreateOperatorParams& InParams,
+			const FBuildOperatorParams& InParams,
 			const FWaveTableReadRef& InWaveTableReadRef,
 			const FTriggerReadRef& InPlayReadRef,
 			const FTriggerReadRef& InStopReadRef,
diff --git a/Engine/Plugins/Runtime/Metasound/Source/MetasoundEngine/Private/MetasoundWaveTablePlayerNode.cpp b/Engine/Plugins/Runtime/Metasound/Source/MetasoundEngine/Private/MetasoundWaveTablePlayerNode.cpp
--- a/Engine/Plugins/Runtime/Metasound/Source/MetasoundEngine/Private/MetasoundWaveTablePlayerNode.cpp
+++ b/Engine/Plugins/Runtime/Metasound/Source/MetasoundEngine/Private/MetasoundWaveTablePlayerNode.cpp
@@ -104,32 +104,32 @@ namespace Metasound
 			return Metadata;
 		}
 
-		static TUniquePtr<IOperator> CreateOperator(const FCreateOperatorParams& InParams, TArray<TUniquePtr<IOperatorBuildError>>& OutErrors)
+		static TUniquePtr<IOperator> CreateOperator(const FBuildOperatorParams& InParams, FBuildResults& OutResults)
 		{
 			using namespace WaveTable;
 
-			const FInputVertexInterface& InputInterface = InParams.Node.GetVertexInterface().GetInputInterface();
-			const FDataReferenceCollection& InputCollection = InParams.InputDataReferences;
+			
+			const FInputVertexInterfaceData& InputData = InParams.InputData;
 
-			FWaveTableBankAssetReadRef InWaveTableBankReadRef = InputCollection.GetDataReadReferenceOrConstruct<FWaveTableBankAsset>("WaveTableBank");
-			FInt32ReadRef InIndexReadRef = InputCollection.GetDataReadReferenceOrConstruct<int32>("Index");
-			FTriggerReadRef InPlayReadRef = InputCollection.GetDataReadReferenceOrConstructWithVertexDefault<FTrigger>(InputInterface, "Play", InParams.OperatorSettings);
-			FTriggerReadRef InStopReadRef = InputCollection.GetDataReadReferenceOrConstructWithVertexDefault<FTrigger>(InputInterface, "Stop", InParams.OperatorSettings);
-			FTriggerReadRef InSyncReadRef = InputCollection.GetDataReadReferenceOrConstructWithVertexDefault<FTrigger>(InputInterface, "Sync", InParams.OperatorSettings);
-			FFloatReadRef InPitchShiftReadRef = InputCollection.GetDataReadReferenceOrConstructWithVertexDefault<float>(InputInterface, "PitchShift", InParams.OperatorSettings);
-			FBoolReadRef InLoopReadRef = InputCollection.GetDataReadReferenceOrConstructWithVertexDefault<bool>(InputInterface, "Loop", InParams.OperatorSettings);
+			FWaveTableBankAssetReadRef InWaveTableBankReadRef = InputData.GetOrConstructDataReadReference<FWaveTableBankAsset>("WaveTableBank");
+			FInt32ReadRef InIndexReadRef = InputData.GetOrConstructDataReadReference<int32>("Index");
+			FTriggerReadRef InPlayReadRef = InputData.GetOrCreateDefaultDataReadReference<FTrigger>("Play", InParams.OperatorSettings);
+			FTriggerReadRef InStopReadRef = InputData.GetOrCreateDefaultDataReadReference<FTrigger>("Stop", InParams.OperatorSettings);
+			FTriggerReadRef InSyncReadRef = InputData.GetOrCreateDefaultDataReadReference<FTrigger>("Sync", InParams.OperatorSettings);
+			FFloatReadRef InPitchShiftReadRef = InputData.GetOrCreateDefaultDataReadReference<float>("PitchShift", InParams.OperatorSettings);
+			FBoolReadRef InLoopReadRef = InputData.GetOrCreateDefaultDataReadReference<bool>("Loop", InParams.OperatorSettings);
 
 			TOptional<FAudioBufferReadRef> InPhaseModReadRef;
-			if (InputCollection.ContainsDataReadReference<FAudioBuffer>("PhaseMod"))
+			if (const FAnyDataReference* DataRef = InputData.FindDataReference("PhaseMod"))
 			{
-				InPhaseModReadRef = InputCollection.GetDataReadReference<FAudioBuffer>("PhaseMod");
+				InPhaseModReadRef = InputData.GetDataReadReference<FAudioBuffer>("PhaseMod");
 			}
 
 			return MakeUnique<FMetaSoundWaveTablePlayerNodeOperator>(InParams, InWaveTableBankReadRef, InIndexReadRef, InPlayReadRef, InStopReadRef, InSyncReadRef, InPitchShiftReadRef, InLoopReadRef, MoveTemp(InPhaseModReadRef));
 		}
 
 		FMetaSoundWaveTablePlayerNodeOperator(
-			const FCreateOperatorParams& InParams,
+			const FBuildOperatorParams& InParams,
 			const FWaveTableBankAssetReadRef& InWaveTableBankReadRef,
 			const FInt32ReadRef& InIndexReadRef,
 			const FTriggerReadRef& InPlayReadRef,
diff --git a/Engine/Plugins/Runtime/Metasound/Source/MetasoundStandardNodes/Private/MetasoundCompressorNode.cpp b/Engine/Plugins/Runtime/Metasound/Source/MetasoundStandardNodes/Private/MetasoundCompressorNode.cpp
--- a/Engine/Plugins/Runtime/Metasound/Source/MetasoundStandardNodes/Private/MetasoundCompressorNode.cpp
+++ b/Engine/Plugins/Runtime/Metasound/Source/MetasoundStandardNodes/Private/MetasoundCompressorNode.cpp
@@ -223,27 +223,27 @@ namespace Metasound
 			return {};
 		}
 
-		static TUniquePtr<IOperator> CreateOperator(const FCreateOperatorParams& InParams, FBuildErrorArray& OutErrors)
+		static TUniquePtr<IOperator> CreateOperator(const FBuildOperatorParams& InParams, FBuildResults& OutResults)
 		{
 			using namespace CompressorVertexNames;
-			const FDataReferenceCollection& Inputs = InParams.InputDataReferences;
-			const FInputVertexInterface& InputInterface = DeclareVertexInterface().GetInputInterface();
-
-			FBoolReadRef bIsBypassedIn = Inputs.GetDataReadReferenceOrConstructWithVertexDefault<bool>(InputInterface, METASOUND_GET_PARAM_NAME(InputIsBypassed), InParams.OperatorSettings);
-			FAudioBufferReadRef AudioIn = Inputs.GetDataReadReferenceOrConstruct<FAudioBuffer>(METASOUND_GET_PARAM_NAME(InputAudio), InParams.OperatorSettings);
-			FFloatReadRef RatioIn = Inputs.GetDataReadReferenceOrConstructWithVertexDefault<float>(InputInterface, METASOUND_GET_PARAM_NAME(InputRatio), InParams.OperatorSettings);
-			FFloatReadRef ThresholdDbIn = Inputs.GetDataReadReferenceOrConstructWithVertexDefault<float>(InputInterface, METASOUND_GET_PARAM_NAME(InputThreshold), InParams.OperatorSettings);
-			FTimeReadRef AttackTimeIn = Inputs.GetDataReadReferenceOrConstructWithVertexDefault<FTime>(InputInterface, METASOUND_GET_PARAM_NAME(InputAttackTime), InParams.OperatorSettings);
-			FTimeReadRef ReleaseTimeIn = Inputs.GetDataReadReferenceOrConstructWithVertexDefault<FTime>(InputInterface, METASOUND_GET_PARAM_NAME(InputReleaseTime), InParams.OperatorSettings);
-			FTimeReadRef LookaheadTimeIn = Inputs.GetDataReadReferenceOrConstructWithVertexDefault<FTime>(InputInterface, METASOUND_GET_PARAM_NAME(InputLookaheadTime), InParams.OperatorSettings);
-			FFloatReadRef KneeIn = Inputs.GetDataReadReferenceOrConstructWithVertexDefault<float>(InputInterface, METASOUND_GET_PARAM_NAME(InputKnee), InParams.OperatorSettings);
-			FAudioBufferReadRef SidechainIn = Inputs.GetDataReadReferenceOrConstruct<FAudioBuffer>(METASOUND_GET_PARAM_NAME(InputSidechain), InParams.OperatorSettings);
-			FEnvelopePeakModeReadRef EnvelopeModeIn = Inputs.GetDataReadReferenceOrConstructWithVertexDefault<FEnumEnvelopePeakMode>(InputInterface, METASOUND_GET_PARAM_NAME(InputEnvelopeMode), InParams.OperatorSettings);
-			FBoolReadRef bIsAnalogIn = Inputs.GetDataReadReferenceOrConstructWithVertexDefault<bool>(InputInterface, METASOUND_GET_PARAM_NAME(InputIsAnalog), InParams.OperatorSettings);
-			FBoolReadRef bIsUpwardsIn = Inputs.GetDataReadReferenceOrConstructWithVertexDefault<bool>(InputInterface, METASOUND_GET_PARAM_NAME(InputIsUpwards), InParams.OperatorSettings);
-			FFloatReadRef WetDryMixIn = Inputs.GetDataReadReferenceOrConstructWithVertexDefault<float>(InputInterface, METASOUND_GET_PARAM_NAME(InputWetDryMix), InParams.OperatorSettings);
-
-			bool bIsSidechainConnected = Inputs.ContainsDataReadReference<FAudioBuffer>(METASOUND_GET_PARAM_NAME(InputSidechain));
+			
+			const FInputVertexInterfaceData& InputData = InParams.InputData;
+
+			FBoolReadRef bIsBypassedIn = InputData.GetOrCreateDefaultDataReadReference<bool>(METASOUND_GET_PARAM_NAME(InputIsBypassed), InParams.OperatorSettings);
+			FAudioBufferReadRef AudioIn = InputData.GetOrConstructDataReadReference<FAudioBuffer>(METASOUND_GET_PARAM_NAME(InputAudio), InParams.OperatorSettings);
+			FFloatReadRef RatioIn = InputData.GetOrCreateDefaultDataReadReference<float>(METASOUND_GET_PARAM_NAME(InputRatio), InParams.OperatorSettings);
+			FFloatReadRef ThresholdDbIn = InputData.GetOrCreateDefaultDataReadReference<float>(METASOUND_GET_PARAM_NAME(InputThreshold), InParams.OperatorSettings);
+			FTimeReadRef AttackTimeIn = InputData.GetOrCreateDefaultDataReadReference<FTime>(METASOUND_GET_PARAM_NAME(InputAttackTime), InParams.OperatorSettings);
+			FTimeReadRef ReleaseTimeIn = InputData.GetOrCreateDefaultDataReadReference<FTime>(METASOUND_GET_PARAM_NAME(InputReleaseTime), InParams.OperatorSettings);
+			FTimeReadRef LookaheadTimeIn = InputData.GetOrCreateDefaultDataReadReference<FTime>(METASOUND_GET_PARAM_NAME(InputLookaheadTime), InParams.OperatorSettings);
+			FFloatReadRef KneeIn = InputData.GetOrCreateDefaultDataReadReference<float>(METASOUND_GET_PARAM_NAME(InputKnee), InParams.OperatorSettings);
+			FAudioBufferReadRef SidechainIn = InputData.GetOrConstructDataReadReference<FAudioBuffer>(METASOUND_GET_PARAM_NAME(InputSidechain), InParams.OperatorSettings);
+			FEnvelopePeakModeReadRef EnvelopeModeIn = InputData.GetOrCreateDefaultDataReadReference<FEnumEnvelopePeakMode>(METASOUND_GET_PARAM_NAME(InputEnvelopeMode), InParams.OperatorSettings);
+			FBoolReadRef bIsAnalogIn = InputData.GetOrCreateDefaultDataReadReference<bool>(METASOUND_GET_PARAM_NAME(InputIsAnalog), InParams.OperatorSettings);
+			FBoolReadRef bIsUpwardsIn = InputData.GetOrCreateDefaultDataReadReference<bool>(METASOUND_GET_PARAM_NAME(InputIsUpwards), InParams.OperatorSettings);
+			FFloatReadRef WetDryMixIn = InputData.GetOrCreateDefaultDataReadReference<float>(METASOUND_GET_PARAM_NAME(InputWetDryMix), InParams.OperatorSettings);
+
+			bool bIsSidechainConnected = InputData.IsVertexBound(METASOUND_GET_PARAM_NAME(InputSidechain));
 
 			return MakeUnique<FCompressorOperator>(InParams.OperatorSettings, bIsBypassedIn, AudioIn, RatioIn, ThresholdDbIn, AttackTimeIn, ReleaseTimeIn, LookaheadTimeIn, KneeIn, bIsSidechainConnected, SidechainIn, EnvelopeModeIn, bIsAnalogIn, bIsUpwardsIn, WetDryMixIn);
 		}
diff --git a/Engine/Plugins/Runtime/Metasound/Source/MetasoundStandardNodes/Private/MetasoundDynamicFilterNode.cpp b/Engine/Plugins/Runtime/Metasound/Source/MetasoundStandardNodes/Private/MetasoundDynamicFilterNode.cpp
--- a/Engine/Plugins/Runtime/Metasound/Source/MetasoundStandardNodes/Private/MetasoundDynamicFilterNode.cpp
+++ b/Engine/Plugins/Runtime/Metasound/Source/MetasoundStandardNodes/Private/MetasoundDynamicFilterNode.cpp
@@ -200,29 +200,29 @@ namespace Metasound
 			return Interface;
 		}
 
-		static TUniquePtr<IOperator> CreateOperator(const FCreateOperatorParams& InParams, FBuildErrorArray& OutErrors)
+		static TUniquePtr<IOperator> CreateOperator(const FBuildOperatorParams& InParams, FBuildResults& OutResults)
 		{
 			using namespace DynamicFilterNode;
 
-			const FDataReferenceCollection& Inputs = InParams.InputDataReferences;
-			const FInputVertexInterface& InputInterface = DeclareVertexInterface().GetInputInterface();
-
-			FAudioBufferReadRef AudioIn = Inputs.GetDataReadReferenceOrConstruct<FAudioBuffer>(METASOUND_GET_PARAM_NAME(InputAudio), InParams.OperatorSettings);
-			FAudioBufferReadRef SidechainIn = Inputs.GetDataReadReferenceOrConstruct<FAudioBuffer>(METASOUND_GET_PARAM_NAME(InputSidechain), InParams.OperatorSettings);
-			FEnumEDynamicFilterReadRef FilterTypeIn = Inputs.GetDataReadReferenceOrConstructWithVertexDefault<FEnumEDynamicFilterType>(InputInterface, METASOUND_GET_PARAM_NAME(InputFilterType), InParams.OperatorSettings);
-			FFloatReadRef FrequencyIn = Inputs.GetDataReadReferenceOrConstructWithVertexDefault<float>(InputInterface, METASOUND_GET_PARAM_NAME(InputFrequency), InParams.OperatorSettings);
-			FFloatReadRef QIn = Inputs.GetDataReadReferenceOrConstructWithVertexDefault<float>(InputInterface, METASOUND_GET_PARAM_NAME(InputQ), InParams.OperatorSettings);
-			FFloatReadRef ThresholdIn = Inputs.GetDataReadReferenceOrConstructWithVertexDefault<float>(InputInterface, METASOUND_GET_PARAM_NAME(InputThreshold), InParams.OperatorSettings);
-			FFloatReadRef RatioIn = Inputs.GetDataReadReferenceOrConstructWithVertexDefault<float>(InputInterface, METASOUND_GET_PARAM_NAME(InputRatio), InParams.OperatorSettings);
-			FFloatReadRef KneeIn = Inputs.GetDataReadReferenceOrConstructWithVertexDefault<float>(InputInterface, METASOUND_GET_PARAM_NAME(InputKnee), InParams.OperatorSettings);
-			FFloatReadRef RangeIn = Inputs.GetDataReadReferenceOrConstructWithVertexDefault<float>(InputInterface, METASOUND_GET_PARAM_NAME(InputRange), InParams.OperatorSettings);
-			FFloatReadRef GainIn = Inputs.GetDataReadReferenceOrConstructWithVertexDefault<float>(InputInterface, METASOUND_GET_PARAM_NAME(InputGain), InParams.OperatorSettings);
-			FTimeReadRef AttackTimeIn = Inputs.GetDataReadReferenceOrConstructWithVertexDefault<FTime>(InputInterface, METASOUND_GET_PARAM_NAME(InputAttackTime), InParams.OperatorSettings);
-			FTimeReadRef ReleaseTimeIn = Inputs.GetDataReadReferenceOrConstructWithVertexDefault<FTime>(InputInterface, METASOUND_GET_PARAM_NAME(InputReleaseTime), InParams.OperatorSettings);
-			FEnvelopePeakModeReadRef EnvelopeModeIn = Inputs.GetDataReadReferenceOrConstructWithVertexDefault<FEnumEnvelopePeakMode>(InputInterface, METASOUND_GET_PARAM_NAME(InputEnvelopeMode), InParams.OperatorSettings);
-			FBoolReadRef bAnalogIn = Inputs.GetDataReadReferenceOrConstructWithVertexDefault<bool>(InputInterface, METASOUND_GET_PARAM_NAME(InputAnalogMode), InParams.OperatorSettings);
-
-			bool bIsSidechainConnected = Inputs.ContainsDataReadReference<FAudioBuffer>(METASOUND_GET_PARAM_NAME(InputSidechain));
+			
+			const FInputVertexInterfaceData& InputData = InParams.InputData;
+
+			FAudioBufferReadRef AudioIn = InputData.GetOrConstructDataReadReference<FAudioBuffer>(METASOUND_GET_PARAM_NAME(InputAudio), InParams.OperatorSettings);
+			FAudioBufferReadRef SidechainIn = InputData.GetOrConstructDataReadReference<FAudioBuffer>(METASOUND_GET_PARAM_NAME(InputSidechain), InParams.OperatorSettings);
+			FEnumEDynamicFilterReadRef FilterTypeIn = InputData.GetOrCreateDefaultDataReadReference<FEnumEDynamicFilterType>(METASOUND_GET_PARAM_NAME(InputFilterType), InParams.OperatorSettings);
+			FFloatReadRef FrequencyIn = InputData.GetOrCreateDefaultDataReadReference<float>(METASOUND_GET_PARAM_NAME(InputFrequency), InParams.OperatorSettings);
+			FFloatReadRef QIn = InputData.GetOrCreateDefaultDataReadReference<float>(METASOUND_GET_PARAM_NAME(InputQ), InParams.OperatorSettings);
+			FFloatReadRef ThresholdIn = InputData.GetOrCreateDefaultDataReadReference<float>(METASOUND_GET_PARAM_NAME(InputThreshold), InParams.OperatorSettings);
+			FFloatReadRef RatioIn = InputData.GetOrCreateDefaultDataReadReference<float>(METASOUND_GET_PARAM_NAME(InputRatio), InParams.OperatorSettings);
+			FFloatReadRef KneeIn = InputData.GetOrCreateDefaultDataReadReference<float>(METASOUND_GET_PARAM_NAME(InputKnee), InParams.OperatorSettings);
+			FFloatReadRef RangeIn = InputData.GetOrCreateDefaultDataReadReference<float>(METASOUND_GET_PARAM_NAME(InputRange), InParams.OperatorSettings);
+			FFloatReadRef GainIn = InputData.GetOrCreateDefaultDataReadReference<float>(METASOUND_GET_PARAM_NAME(InputGain), InParams.OperatorSettings);
+			FTimeReadRef AttackTimeIn = InputData.GetOrCreateDefaultDataReadReference<FTime>(METASOUND_GET_PARAM_NAME(InputAttackTime), InParams.OperatorSettings);
+			FTimeReadRef ReleaseTimeIn = InputData.GetOrCreateDefaultDataReadReference<FTime>(METASOUND_GET_PARAM_NAME(InputReleaseTime), InParams.OperatorSettings);
+			FEnvelopePeakModeReadRef EnvelopeModeIn = InputData.GetOrCreateDefaultDataReadReference<FEnumEnvelopePeakMode>(METASOUND_GET_PARAM_NAME(InputEnvelopeMode), InParams.OperatorSettings);
+			FBoolReadRef bAnalogIn = InputData.GetOrCreateDefaultDataReadReference<bool>(METASOUND_GET_PARAM_NAME(InputAnalogMode), InParams.OperatorSettings);
+
+			bool bIsSidechainConnected = InputData.IsVertexBound(METASOUND_GET_PARAM_NAME(InputSidechain));
 
 			return MakeUnique<FDynamicFilterOperator>(InParams.OperatorSettings, AudioIn, SidechainIn, FilterTypeIn, FrequencyIn, QIn, ThresholdIn, RatioIn, KneeIn, RangeIn, GainIn, AttackTimeIn, ReleaseTimeIn, EnvelopeModeIn, bAnalogIn, bIsSidechainConnected);
 		}
diff --git a/Engine/Plugins/Runtime/Metasound/Source/MetasoundStandardNodes/Private/MetasoundMapRangeNode.cpp b/Engine/Plugins/Runtime/Metasound/Source/MetasoundStandardNodes/Private/MetasoundMapRangeNode.cpp
--- a/Engine/Plugins/Runtime/Metasound/Source/MetasoundStandardNodes/Private/MetasoundMapRangeNode.cpp
+++ b/Engine/Plugins/Runtime/Metasound/Source/MetasoundStandardNodes/Private/MetasoundMapRangeNode.cpp
@@ -63,13 +63,13 @@ namespace Metasound
 				return DefaultInterface;
 			}
 
-			FIntRange(const FOperatorSettings& OperatorSettings, const FInputVertexInterface& InputInterface, const FDataReferenceCollection& InputCollection)
-				: Value(InputCollection.GetDataReadReferenceOrConstructWithVertexDefault<int32>(InputInterface, METASOUND_GET_PARAM_NAME(InputValueName), OperatorSettings))
-				, InRangeA(InputCollection.GetDataReadReferenceOrConstructWithVertexDefault<int32>(InputInterface, METASOUND_GET_PARAM_NAME(InputInRangeAName), OperatorSettings))
-				, InRangeB(InputCollection.GetDataReadReferenceOrConstructWithVertexDefault<int32>(InputInterface, METASOUND_GET_PARAM_NAME(InputInRangeBName), OperatorSettings))
-				, OutRangeA(InputCollection.GetDataReadReferenceOrConstructWithVertexDefault<int32>(InputInterface, METASOUND_GET_PARAM_NAME(InputOutRangeAName), OperatorSettings))
-				, OutRangeB(InputCollection.GetDataReadReferenceOrConstructWithVertexDefault<int32>(InputInterface, METASOUND_GET_PARAM_NAME(InputOutRangeBName), OperatorSettings))
-				, bClamped(InputCollection.GetDataReadReferenceOrConstructWithVertexDefault<bool>(InputInterface, METASOUND_GET_PARAM_NAME(InputClampedName), OperatorSettings))
+			FIntRange(const FOperatorSettings& OperatorSettings, const FInputVertexInterfaceData& InInputData)
+				: Value(InInputData.GetOrCreateDefaultDataReadReference<int32>(METASOUND_GET_PARAM_NAME(InputValueName), OperatorSettings))
+				, InRangeA(InInputData.GetOrCreateDefaultDataReadReference<int32>(METASOUND_GET_PARAM_NAME(InputInRangeAName), OperatorSettings))
+				, InRangeB(InInputData.GetOrCreateDefaultDataReadReference<int32>(METASOUND_GET_PARAM_NAME(InputInRangeBName), OperatorSettings))
+				, OutRangeA(InInputData.GetOrCreateDefaultDataReadReference<int32>(METASOUND_GET_PARAM_NAME(InputOutRangeAName), OperatorSettings))
+				, OutRangeB(InInputData.GetOrCreateDefaultDataReadReference<int32>(METASOUND_GET_PARAM_NAME(InputOutRangeBName), OperatorSettings))
+				, bClamped(InInputData.GetOrCreateDefaultDataReadReference<bool>(METASOUND_GET_PARAM_NAME(InputClampedName), OperatorSettings))
 				, OutputValue(TDataWriteReferenceFactory<int32>::CreateAny(OperatorSettings))
 			{
 			}
@@ -142,13 +142,13 @@ namespace Metasound
 				return DefaultInterface;
 			}
 
-			FFloatRange(const FOperatorSettings& OperatorSettings, const FInputVertexInterface& InputInterface, const FDataReferenceCollection& InputCollection)
-				: Value(InputCollection.GetDataReadReferenceOrConstructWithVertexDefault<float>(InputInterface,		METASOUND_GET_PARAM_NAME(InputValueName),		OperatorSettings))
-				, InRangeA(InputCollection.GetDataReadReferenceOrConstructWithVertexDefault<float>(InputInterface,	METASOUND_GET_PARAM_NAME(InputInRangeAName),	OperatorSettings))
-				, InRangeB(InputCollection.GetDataReadReferenceOrConstructWithVertexDefault<float>(InputInterface,	METASOUND_GET_PARAM_NAME(InputInRangeBName),	OperatorSettings))
-				, OutRangeA(InputCollection.GetDataReadReferenceOrConstructWithVertexDefault<float>(InputInterface, METASOUND_GET_PARAM_NAME(InputOutRangeAName),	OperatorSettings))
-				, OutRangeB(InputCollection.GetDataReadReferenceOrConstructWithVertexDefault<float>(InputInterface, METASOUND_GET_PARAM_NAME(InputOutRangeBName),	OperatorSettings))
-				, bClamped(InputCollection.GetDataReadReferenceOrConstructWithVertexDefault<bool>(InputInterface,	METASOUND_GET_PARAM_NAME(InputClampedName),		OperatorSettings))
+			FFloatRange(const FOperatorSettings& OperatorSettings, const FInputVertexInterfaceData& InInputData)
+				: Value(InInputData.GetOrCreateDefaultDataReadReference<float>(METASOUND_GET_PARAM_NAME(InputValueName),		OperatorSettings))
+				, InRangeA(InInputData.GetOrCreateDefaultDataReadReference<float>(METASOUND_GET_PARAM_NAME(InputInRangeAName),	OperatorSettings))
+				, InRangeB(InInputData.GetOrCreateDefaultDataReadReference<float>(METASOUND_GET_PARAM_NAME(InputInRangeBName),	OperatorSettings))
+				, OutRangeA(InInputData.GetOrCreateDefaultDataReadReference<float>(METASOUND_GET_PARAM_NAME(InputOutRangeAName),	OperatorSettings))
+				, OutRangeB(InInputData.GetOrCreateDefaultDataReadReference<float>(METASOUND_GET_PARAM_NAME(InputOutRangeBName),	OperatorSettings))
+				, bClamped(InInputData.GetOrCreateDefaultDataReadReference<bool>(METASOUND_GET_PARAM_NAME(InputClampedName),		OperatorSettings))
 				, OutputValue(TDataWriteReferenceFactory<float>::CreateAny(OperatorSettings))
 			{
 			}
@@ -223,13 +223,13 @@ namespace Metasound
 				return DefaultInterface;
 			}
 
-			FAudioRange(const FOperatorSettings& OperatorSettings, const FInputVertexInterface& InputInterface, const FDataReferenceCollection& InputCollection)
-				: Value(InputCollection.GetDataReadReferenceOrConstructWithVertexDefault<FAudioBuffer>(InputInterface, METASOUND_GET_PARAM_NAME(InputValueName), OperatorSettings))
-				, InRangeA(InputCollection.GetDataReadReferenceOrConstructWithVertexDefault<float>(InputInterface, METASOUND_GET_PARAM_NAME(InputInRangeAName), OperatorSettings))
-				, InRangeB(InputCollection.GetDataReadReferenceOrConstructWithVertexDefault<float>(InputInterface, METASOUND_GET_PARAM_NAME(InputInRangeBName), OperatorSettings))
-				, OutRangeA(InputCollection.GetDataReadReferenceOrConstructWithVertexDefault<float>(InputInterface, METASOUND_GET_PARAM_NAME(InputOutRangeAName), OperatorSettings))
-				, OutRangeB(InputCollection.GetDataReadReferenceOrConstructWithVertexDefault<float>(InputInterface, METASOUND_GET_PARAM_NAME(InputOutRangeBName), OperatorSettings))
-				, bClamped(InputCollection.GetDataReadReferenceOrConstructWithVertexDefault<bool>(InputInterface, METASOUND_GET_PARAM_NAME(InputClampedName), OperatorSettings))
+			FAudioRange(const FOperatorSettings& OperatorSettings, const FInputVertexInterfaceData& InInputData)
+				: Value(InInputData.GetOrCreateDefaultDataReadReference<FAudioBuffer>(METASOUND_GET_PARAM_NAME(InputValueName), OperatorSettings))
+				, InRangeA(InInputData.GetOrCreateDefaultDataReadReference<float>(METASOUND_GET_PARAM_NAME(InputInRangeAName), OperatorSettings))
+				, InRangeB(InInputData.GetOrCreateDefaultDataReadReference<float>(METASOUND_GET_PARAM_NAME(InputInRangeBName), OperatorSettings))
+				, OutRangeA(InInputData.GetOrCreateDefaultDataReadReference<float>(METASOUND_GET_PARAM_NAME(InputOutRangeAName), OperatorSettings))
+				, OutRangeB(InInputData.GetOrCreateDefaultDataReadReference<float>(METASOUND_GET_PARAM_NAME(InputOutRangeBName), OperatorSettings))
+				, bClamped(InInputData.GetOrCreateDefaultDataReadReference<bool>(METASOUND_GET_PARAM_NAME(InputClampedName), OperatorSettings))
 				, OutputValue(TDataWriteReferenceFactory<FAudioBuffer>::CreateAny(OperatorSettings))
 			{
 			}
@@ -313,12 +313,9 @@ namespace Metasound
 			return Metadata;
 		}
 
-		static TUniquePtr<IOperator> CreateOperator(const FCreateOperatorParams& InParams, TArray<TUniquePtr<IOperatorBuildError>>& OutErrors)
+		static TUniquePtr<IOperator> CreateOperator(const FBuildOperatorParams& InParams, FBuildResults& OutResults)
 		{
-			const FInputVertexInterface& InputInterface = InParams.Node.GetVertexInterface().GetInputInterface();
-			const FDataReferenceCollection& InputCollection = InParams.InputDataReferences;
-
-			FMappingClass MappingObject(InParams.OperatorSettings, InputInterface, InputCollection);
+			FMappingClass MappingObject(InParams.OperatorSettings, InParams.InputData);
 
 			return MakeUnique<TMapRangeOperator<FMappingClass>>(MappingObject);
 		}
diff --git a/Engine/Plugins/Runtime/Metasound/Source/MetasoundStandardNodes/Private/MetasoundMathNodes.cpp b/Engine/Plugins/Runtime/Metasound/Source/MetasoundStandardNodes/Private/MetasoundMathNodes.cpp
--- a/Engine/Plugins/Runtime/Metasound/Source/MetasoundStandardNodes/Private/MetasoundMathNodes.cpp
+++ b/Engine/Plugins/Runtime/Metasound/Source/MetasoundStandardNodes/Private/MetasoundMathNodes.cpp
@@ -245,34 +245,25 @@ namespace Metasound
 				TMathOpClass::Calculate(InstanceData, PrimaryOperandRef, AdditionalOperandRefs, ValueRef);
 			}
 
-			static TUniquePtr<IOperator> CreateOperator(const FCreateOperatorParams& InParams, FBuildErrorArray& OutErrors)
+			static TUniquePtr<IOperator> CreateOperator(const FBuildOperatorParams& InParams, FBuildResults& OutResults)
 			{
 				using namespace MathOpNames;
 
 				const TMathOpNode& MathOpNode = static_cast<const TMathOpNode&>(InParams.Node);
 
+				const FInputVertexInterfaceData& InputData = InParams.InputData;
 				const FInputVertexInterface& InputInterface = MathOpNode.GetVertexInterface().GetInputInterface();
 
-				FLiteral DefaultValue = FLiteral::CreateInvalid();
-				if (InputInterface.Contains(METASOUND_GET_PARAM_NAME(PrimaryOperand)))
-				{
-					DefaultValue = InputInterface[METASOUND_GET_PARAM_NAME(PrimaryOperand)].GetDefaultLiteral();
-				}
-				TDataClassReadRef PrimaryOperand = InParams.InputDataReferences.GetDataReadReferenceOrConstruct<TDataClass>(METASOUND_GET_PARAM_NAME(PrimaryOperand), TMathOpClass::GetDefault(InParams.OperatorSettings, DefaultValue));
+				TDataClassReadRef PrimaryOperand = InputData.GetOrCreateDefaultDataReadReference<TDataClass>(METASOUND_GET_PARAM_NAME(PrimaryOperand), InParams.OperatorSettings);
 
 				// TODO: Support dynamic number of inputs
-				const FVertexName& OpName = METASOUND_GET_PARAM_NAME(AdditionalOperands);
-				if (InputInterface.Contains(OpName))
-				{
-					DefaultValue = InputInterface[OpName].GetDefaultLiteral();
-				}
-				TOperandDataClassReadRef Op1 = InParams.InputDataReferences.GetDataReadReferenceOrConstruct<TOperandDataClass>(OpName, TMathOpClass::GetDefaultOp(InParams.OperatorSettings, DefaultValue));
+				TOperandDataClassReadRef Op1 = InputData.GetOrCreateDefaultDataReadReference<TOperandDataClass>(METASOUND_GET_PARAM_NAME(AdditionalOperands), InParams.OperatorSettings);
 				TArray<TOperandDataClassReadRef> AdditionalOperandRefs = { Op1 };
 
 				return MakeUnique<TMathOperator>(InParams, PrimaryOperand, AdditionalOperandRefs);
 			}
 
-			TMathOperator(const FCreateOperatorParams& InParams, const TDataClassReadRef& InPrimaryOperand, const TArray<TOperandDataClassReadRef>& InAdditionalOperands)
+			TMathOperator(const FBuildOperatorParams& InParams, const TDataClassReadRef& InPrimaryOperand, const TArray<TOperandDataClassReadRef>& InAdditionalOperands)
 				: PrimaryOperandRef(InPrimaryOperand)
 				, AdditionalOperandRefs(InAdditionalOperands)
 				, ValueRef(TDataWriteReferenceFactory<TDataClass>::CreateAny(InParams.OperatorSettings))
@@ -364,16 +355,6 @@ namespace Metasound
 			return Interface;
 		}
 
-		static TDataClass GetDefault(const FOperatorSettings& InSettings, const FLiteral& InVertexDefault)
-		{
-			return static_cast<TDataClass>(0);
-		}
-
-		static TOperandDataClass GetDefaultOp(const FOperatorSettings& InSettings, const FLiteral& InVertexDefault)
-		{
-			return static_cast<TOperandDataClass>(0);
-		}
-
 		static void Calculate(TMathOpAdd<TDataClass, TOperandDataClass>& InInstanceData, const TDataClassReadRef& InPrimaryOperand, const TArray<TOperandDataClassReadRef>& InAdditionalOperands, TDataClassWriteRef& OutResult)
 		{
 			if (!ensure(!InAdditionalOperands.IsEmpty()))
@@ -433,16 +414,6 @@ namespace Metasound
 			return Interface;
 		}
 
-		static TDataClass GetDefault(const FOperatorSettings& InSettings, const FLiteral& InVertexDefault)
-		{
-			return static_cast<TDataClass>(0);
-		}
-
-		static TOperandDataClass GetDefaultOp(const FOperatorSettings& InSettings, const FLiteral& InVertexDefault)
-		{
-			return static_cast<TOperandDataClass>(0);
-		}
-
 		static void Calculate(TMathOpSubtract<TDataClass, TOperandDataClass>& InInstanceData, const TDataClassReadRef& InPrimaryOperand, const TArray<TOperandDataClassReadRef>& InAdditionalOperands, TDataClassWriteRef& OutResult)
 		{
 			if (!ensure(!InAdditionalOperands.IsEmpty()))
@@ -502,16 +473,6 @@ namespace Metasound
 			return Interface;
 		}
 
-		static TDataClass GetDefault(const FOperatorSettings& InSettings, const FLiteral& InVertexDefault)
-		{
-			return static_cast<TDataClass>(1);
-		}
-
-		static TOperandDataClass GetDefaultOp(const FOperatorSettings& InSettings, const FLiteral& InVertexDefault)
-		{
-			return static_cast<TOperandDataClass>(1);
-		}
-
 		static void Calculate(TMathOpMultiply<TDataClass, TOperandDataClass>& InInstanceData, const TDataClassReadRef& InPrimaryOperand, const TArray<TOperandDataClassReadRef>& InAdditionalOperands, TDataClassWriteRef& OutResult)
 		{
 			if (!ensure(!InAdditionalOperands.IsEmpty()))
@@ -571,16 +532,6 @@ namespace Metasound
 			return Interface;
 		}
 
-		static TDataClass GetDefault(const FOperatorSettings& InSettings, const FLiteral& InVertexDefault)
-		{
-			return static_cast<TDataClass>(1);
-		}
-
-		static TOperandDataClass GetDefaultOp(const FOperatorSettings& InSettings, const FLiteral& InVertexDefault)
-		{
-			return static_cast<TOperandDataClass>(1);
-		}
-
 		static void Calculate(TMathOpDivide& InInstanceData, const TDataClassReadRef& InPrimaryOperand, const TArray<TOperandDataClassReadRef>& InAdditionalOperands, TDataClassWriteRef& OutResult)
 		{
 			if (!ensure(!InAdditionalOperands.IsEmpty()))
@@ -647,16 +598,6 @@ namespace Metasound
 			return Interface;
 		}
 
-		static TDataClass GetDefault(const FOperatorSettings& InSettings, const FLiteral& InVertexDefault)
-		{
-			return static_cast<TDataClass>(1);
-		}
-
-		static TOperandDataClass GetDefaultOp(const FOperatorSettings& InSettings, const FLiteral& InVertexDefault)
-		{
-			return static_cast<TOperandDataClass>(1);
-		}
-
 		static void Calculate(TMathOpModulo& InInstanceData, const TDataClassReadRef& InPrimaryOperand, const TArray<TOperandDataClassReadRef>& InAdditionalOperands, TDataClassWriteRef& OutResult)
 		{
 			if (!ensure(!InAdditionalOperands.IsEmpty()))
@@ -725,16 +666,6 @@ namespace Metasound
 			return DefaultInterface;
 		}
 
-		static TDataClass GetDefault(const FOperatorSettings& InSettings, const FLiteral& InVertexDefault)
-		{
-			return static_cast<TDataClass>(1);
-		}
-
-		static TOperandDataClass GetDefaultOp(const FOperatorSettings& InSettings, const FLiteral& InVertexDefault)
-		{
-			return static_cast<TOperandDataClass>(1);
-		}
-
 		static void Calculate(TMathOpPower& InInstanceData, const TDataClassReadRef& InPrimaryOperand, const TArray<TOperandDataClassReadRef>& InAdditionalOperands, TDataClassWriteRef& OutResult)
 		{
 			if (!ensure(!InAdditionalOperands.IsEmpty()))
@@ -797,16 +728,6 @@ namespace Metasound
 			return DefaultInterface;
 		}
 
-		static TDataClass GetDefault(const FOperatorSettings& InSettings, const FLiteral& InVertexDefault)
-		{
-			return static_cast<TDataClass>(1);
-		}
-
-		static TOperandDataClass GetDefaultOp(const FOperatorSettings& InSettings, const FLiteral& InVertexDefault)
-		{
-			return static_cast<TOperandDataClass>(1);
-		}
-
 		static void Calculate(TMathOpLogarithm& InInstanceData, const TDataClassReadRef& InPrimaryOperand, const TArray<TOperandDataClassReadRef>& InAdditionalOperands, TDataClassWriteRef& OutResult)
 		{
 			if (!ensure(!InAdditionalOperands.IsEmpty()))
@@ -860,16 +781,6 @@ namespace Metasound
 			return Interface;
 		}
 
-		static FAudioBuffer GetDefault(const FOperatorSettings& InSettings, const FLiteral& InVertexDefault)
-		{
-			return FAudioBuffer(InSettings.GetNumFramesPerBlock());
-		}
-
-		static FAudioBuffer GetDefaultOp(const FOperatorSettings& InSettings, const FLiteral& InVertexDefault)
-		{
-			return GetDefault(InSettings, InVertexDefault);
-		}
-
 		static void Calculate(TMathOpAdd<FAudioBuffer>& InInstanceData, FAudioBufferReadRef& InPrimaryOperand, const TArray<FAudioBufferReadRef>& InAdditionalOperands, FAudioBufferWriteRef& OutResult)
 		{
 			if (!ensure(!InAdditionalOperands.IsEmpty()))
@@ -954,16 +865,6 @@ namespace Metasound
 			return Interface;
 		}
 
-		static FTime GetDefault(const FOperatorSettings& InSettings, const FLiteral& InVertexDefault)
-		{
-			return FTime(InVertexDefault.Value.Get<float>());
-		}
-
-		static FTime GetDefaultOp(const FOperatorSettings& InSettings, const FLiteral& InVertexDefault)
-		{
-			return GetDefault(InSettings, InVertexDefault);
-		}
-
 		static void Calculate(TMathOpAdd& InInstanceData, const FTimeReadRef& InPrimaryOperand, const TArray<FTimeReadRef>& InAdditionalOperands, FTimeWriteRef& OutResult)
 		{
 			if (!ensure(!InAdditionalOperands.IsEmpty()))
@@ -1012,16 +913,6 @@ namespace Metasound
 			return Interface;
 		}
 
-		static FAudioBuffer GetDefault(const FOperatorSettings& InSettings, const FLiteral& InVertexDefault)
-		{
-			return FAudioBuffer(InSettings.GetNumFramesPerBlock());
-		}
-
-		static float GetDefaultOp(const FOperatorSettings& InSettings, const FLiteral& InVertexDefault)
-		{
-			return InVertexDefault.Value.Get<float>();
-		}
-
 		static void Calculate(TMathOpAdd& InInstanceData, const FAudioBufferReadRef& InPrimaryOperand, const TArray<FFloatReadRef>& InAdditionalOperands, FAudioBufferWriteRef& OutResult)
 		{
 			FMemory::Memcpy(OutResult->GetData(), InPrimaryOperand->GetData(), sizeof(float) * OutResult->Num());
@@ -1068,16 +959,6 @@ namespace Metasound
 			return Interface;
 		}
 
-		static FAudioBuffer GetDefault(const FOperatorSettings& InSettings, const FLiteral& InVertexDefault)
-		{
-			return FAudioBuffer(InSettings.GetNumFramesPerBlock());
-		}
-
-		static FAudioBuffer GetDefaultOp(const FOperatorSettings& InSettings, const FLiteral& InVertexDefault)
-		{
-			return GetDefault(InSettings, InVertexDefault);
-		}
-
 		static void Calculate(TMathOpSubtract& InInstanceData, FAudioBufferReadRef& InPrimaryOperand, const TArray<FAudioBufferReadRef>& InAdditionalOperands, FAudioBufferWriteRef& OutResult)
 		{
 			if (!ensure(!InAdditionalOperands.IsEmpty()))
@@ -1153,16 +1034,6 @@ namespace Metasound
 			return Interface;
 		}
 
-		static FTime GetDefault(const FOperatorSettings& InSettings, const FLiteral& InVertexDefault)
-		{
-			return FTime(InVertexDefault.Value.Get<float>());
-		}
-
-		static FTime GetDefaultOp(const FOperatorSettings& InSettings, const FLiteral& InVertexDefault)
-		{
-			return GetDefault(InSettings, InVertexDefault);
-		}
-
 		static void Calculate(TMathOpSubtract& InInstanceData, const FTimeReadRef& InPrimaryOperand, const TArray<FTimeReadRef>& InAdditionalOperands, FTimeWriteRef& OutResult)
 		{
 			if (!ensure(!InAdditionalOperands.IsEmpty()))
@@ -1211,16 +1082,6 @@ namespace Metasound
 			return Interface;
 		}
 
-		static FAudioBuffer GetDefault(const FOperatorSettings& InSettings, const FLiteral& InVertexDefault)
-		{
-			return FAudioBuffer(InSettings.GetNumFramesPerBlock());
-		}
-
-		static FAudioBuffer GetDefaultOp(const FOperatorSettings& InSettings, const FLiteral& InVertexDefault)
-		{
-			return GetDefault(InSettings, InVertexDefault);
-		}
-
 		static void Calculate(TMathOpMultiply& InInstanceData, const FAudioBufferReadRef& InPrimaryOperand, const TArray<FAudioBufferReadRef>& InAdditionalOperands, FAudioBufferWriteRef& OutResult)
 		{
 			if (!ensure(!InAdditionalOperands.IsEmpty()))
@@ -1293,16 +1154,6 @@ namespace Metasound
 			return Interface;
 		}
 
-		static FAudioBuffer GetDefault(const FOperatorSettings& InSettings, const FLiteral& InVertexDefault)
-		{
-			return FAudioBuffer(InSettings.GetNumFramesPerBlock());
-		}
-
-		static float GetDefaultOp(const FOperatorSettings& InSettings, const FLiteral& InVertexDefault)
-		{
-			return InVertexDefault.Value.Get<float>();
-		}
-
 		static void Calculate(TMathOpMultiply& InInstanceData, const FAudioBufferReadRef& InPrimaryOperand, const TArray<FFloatReadRef>& InAdditionalOperands, FAudioBufferWriteRef& OutResult)
 		{
 			constexpr float MaxGain = TNumericLimits<float>::Max() / 1024.f;
@@ -1374,16 +1225,6 @@ namespace Metasound
 			return Interface;
 		}
 
-		static FTime GetDefault(const FOperatorSettings& InSettings, const FLiteral& InVertexDefault)
-		{
-			return FTime(InVertexDefault.Value.Get<float>());
-		}
-
-		static float GetDefaultOp(const FOperatorSettings& InSettings, const FLiteral& InVertexDefault)
-		{
-			return 1.0f;
-		}
-
 		static void Calculate(TMathOpMultiply& InInstanceData, const FTimeReadRef& InPrimaryOperand, const TArray<FFloatReadRef>& InAdditionalOperands, FTimeWriteRef& OutResult)
 		{
 			*OutResult = *InPrimaryOperand;
@@ -1440,16 +1281,6 @@ namespace Metasound
 			return Interface;
 		}
 
-		static FTime GetDefault(const FOperatorSettings& InSettings, const FLiteral& InVertexDefault)
-		{
-			return FTime(InVertexDefault.Value.Get<float>());
-		}
-
-		static float GetDefaultOp(const FOperatorSettings& InSettings, const FLiteral& InVertexDefault)
-		{
-			return 1.0f;
-		}
-
 		static void Calculate(TMathOpDivide& InInstanceData, const FTimeReadRef& InPrimaryOperand, const TArray<FFloatReadRef>& InAdditionalOperands, FTimeWriteRef& OutResult)
 		{
 			*OutResult = *InPrimaryOperand;
diff --git a/Engine/Plugins/Runtime/Metasound/Source/MetasoundStandardNodes/Private/MetasoundWaveWriterNode.cpp b/Engine/Plugins/Runtime/Metasound/Source/MetasoundStandardNodes/Private/MetasoundWaveWriterNode.cpp
--- a/Engine/Plugins/Runtime/Metasound/Source/MetasoundStandardNodes/Private/MetasoundWaveWriterNode.cpp
+++ b/Engine/Plugins/Runtime/Metasound/Source/MetasoundStandardNodes/Private/MetasoundWaveWriterNode.cpp
@@ -358,7 +358,7 @@ namespace Metasound
 			return Metadata;
 		}
 
-		static TUniquePtr<IOperator> CreateOperator(const FCreateOperatorParams& InParams, FBuildErrorArray& OutErrors);
+		static TUniquePtr<IOperator> CreateOperator(const FBuildOperatorParams& InParams, FBuildResults& OutResults);
 
 		void Execute()
 		{
@@ -508,22 +508,24 @@ namespace Metasound
 	};
 
 	template<int32 NumInputChannels>
-	TUniquePtr<Metasound::IOperator> TWaveWriterOperator<NumInputChannels>::CreateOperator(const FCreateOperatorParams& InParams, FBuildErrorArray& OutErrors)
+	TUniquePtr<Metasound::IOperator> TWaveWriterOperator<NumInputChannels>::CreateOperator(const FBuildOperatorParams& InParams, FBuildResults& OutResults)
 	{
 		using namespace WaveWriterOperatorPrivate;
 		using namespace WaveWriterVertexNames;
 
-		const FDataReferenceCollection& InputCol = InParams.InputDataReferences;
 		const FOperatorSettings& Settings = InParams.OperatorSettings;
-		const FInputVertexInterface& InputInterface = DeclareVertexInterface().GetInputInterface();
+		const FInputVertexInterfaceData& InputData = InParams.InputData;
 
 		int32 NumConnectedAudioPins = 0;
 		TArray<FAudioBufferReadRef> InputBuffers;
 		for (int32 i = 0; i < NumInputChannels; ++i)
 		{
 			const FVertexName PinName = GetAudioInputName(i);
-			NumConnectedAudioPins += (int32)InputCol.ContainsDataReadReference<FAudioBuffer>(PinName);
-			InputBuffers.Add(InputCol.GetDataReadReferenceOrConstruct<FAudioBuffer>(PinName, InParams.OperatorSettings));
+			if (InputData.IsVertexBound(PinName))
+			{
+				NumConnectedAudioPins++;
+			}
+			InputBuffers.Add(InputData.GetOrConstructDataReadReference<FAudioBuffer>(PinName, InParams.OperatorSettings));
 		}
 		
 		// Only create a real operator if there's some connected pins.
@@ -532,9 +534,9 @@ namespace Metasound
 			return MakeUnique<TWaveWriterOperator>(
 				Settings,
 				MoveTemp(InputBuffers),
-				InputCol.GetDataReadReferenceOrConstructWithVertexDefault<bool>(InputInterface, METASOUND_GET_PARAM_NAME(InEnabledPin), Settings),
+				InputData.GetOrCreateDefaultDataReadReference<bool>(METASOUND_GET_PARAM_NAME(InEnabledPin), Settings),
 				GetNameCache(),
-				InputCol.GetDataReadReferenceOrConstructWithVertexDefault<FString>(InputInterface, METASOUND_GET_PARAM_NAME(InFilenamePrefixPin), Settings)
+				InputData.GetOrCreateDefaultDataReadReference<FString>(METASOUND_GET_PARAM_NAME(InFilenamePrefixPin), Settings)
 			);
 		}
 
diff --git a/Engine/Plugins/Runtime/Metasound/Source/MetasoundStandardNodes/Public/MetasoundRandomNode.h b/Engine/Plugins/Runtime/Metasound/Source/MetasoundStandardNodes/Public/MetasoundRandomNode.h
--- a/Engine/Plugins/Runtime/Metasound/Source/MetasoundStandardNodes/Public/MetasoundRandomNode.h
+++ b/Engine/Plugins/Runtime/Metasound/Source/MetasoundStandardNodes/Public/MetasoundRandomNode.h
@@ -76,18 +76,16 @@ namespace Metasound
 			return InStream.RandRange(InMin, InMax);
 		}
 
-		static TDataReadReference<int32> CreateMinValueRef(const FDataReferenceCollection& InputCollection, const FInputVertexInterface& InputInterface, const FOperatorSettings& InOperatorSettings)
+		static TDataReadReference<int32> CreateMinValueRef(const FInputVertexInterfaceData& InputData, const FOperatorSettings& InOperatorSettings)
 		{
 			using namespace RandomNodeNames; 
-
-			return InputCollection.GetDataReadReferenceOrConstructWithVertexDefault<int32>(InputInterface, METASOUND_GET_PARAM_NAME(InputMin), InOperatorSettings);
+			return InputData.GetOrCreateDefaultDataReadReference<int32>(METASOUND_GET_PARAM_NAME(InputMin), InOperatorSettings);
 		}
 
-		static TDataReadReference<int32> CreateMaxValueRef(const FDataReferenceCollection& InputCollection, const FInputVertexInterface& InputInterface, const FOperatorSettings& InOperatorSettings)
+		static TDataReadReference<int32> CreateMaxValueRef(const FInputVertexInterfaceData& InputData, const FOperatorSettings& InOperatorSettings)
 		{
 			using namespace RandomNodeNames; 
-
-			return InputCollection.GetDataReadReferenceOrConstructWithVertexDefault<int32>(InputInterface, METASOUND_GET_PARAM_NAME(InputMax), InOperatorSettings);
+			return InputData.GetOrCreateDefaultDataReadReference<int32>(METASOUND_GET_PARAM_NAME(InputMax), InOperatorSettings);
 		}
 	};
 
@@ -129,18 +127,16 @@ namespace Metasound
 			return InStream.FRandRange(InMin, InMax);
 		}
 
-		static TDataReadReference<float> CreateMinValueRef(const FDataReferenceCollection& InputCollection, const FInputVertexInterface& InputInterface, const FOperatorSettings& InOperatorSettings)
+		static TDataReadReference<float> CreateMinValueRef(const FInputVertexInterfaceData& InputData, const FOperatorSettings& InOperatorSettings)
 		{
 			using namespace RandomNodeNames;
-
-			return InputCollection.GetDataReadReferenceOrConstructWithVertexDefault<float>(InputInterface, METASOUND_GET_PARAM_NAME(InputMin), InOperatorSettings);
+			return InputData.GetOrCreateDefaultDataReadReference<float>(METASOUND_GET_PARAM_NAME(InputMin), InOperatorSettings);
 		}
 
-		static TDataReadReference<float> CreateMaxValueRef(const FDataReferenceCollection& InputCollection, const FInputVertexInterface& InputInterface, const FOperatorSettings& InOperatorSettings)
+		static TDataReadReference<float> CreateMaxValueRef(const FInputVertexInterfaceData& InputData, const FOperatorSettings& InOperatorSettings)
 		{
 			using namespace RandomNodeNames;
-			
-			return InputCollection.GetDataReadReferenceOrConstructWithVertexDefault<float>(InputInterface, METASOUND_GET_PARAM_NAME(InputMax), InOperatorSettings);
+			return InputData.GetOrCreateDefaultDataReadReference<float>(METASOUND_GET_PARAM_NAME(InputMax), InOperatorSettings);
 		}
 	};
 
@@ -182,12 +178,12 @@ namespace Metasound
 			return (bool)InStream.RandRange(0, 1);
 		}
 
-		static TDataReadReference<bool> CreateMinValueRef(const FDataReferenceCollection&, const FInputVertexInterface&, const FOperatorSettings&)
+		static TDataReadReference<bool> CreateMinValueRef(const FInputVertexInterfaceData&, const FOperatorSettings&)
 		{
 			return TDataReadReference<bool>::CreateNew();
 		}
 
-		static TDataReadReference<bool> CreateMaxValueRef(const FDataReferenceCollection&, const FInputVertexInterface&, const FOperatorSettings&)
+		static TDataReadReference<bool> CreateMaxValueRef(const FInputVertexInterfaceData&, const FOperatorSettings&)
 		{
 			return TDataReadReference<bool>::CreateNew();
 		}
@@ -233,18 +229,16 @@ namespace Metasound
 			return FTime(Seconds);
 		}
 
-		static TDataReadReference<FTime> CreateMinValueRef(const FDataReferenceCollection& InputCollection, const FInputVertexInterface& InputInterface, const FOperatorSettings& InOperatorSettings)
+		static TDataReadReference<FTime> CreateMinValueRef(const FInputVertexInterfaceData& InputData, const FOperatorSettings& InOperatorSettings)
 		{
 			using namespace RandomNodeNames; 
-
-			return InputCollection.GetDataReadReferenceOrConstructWithVertexDefault<FTime>(InputInterface, METASOUND_GET_PARAM_NAME(InputMin), InOperatorSettings);
+			return InputData.GetOrCreateDefaultDataReadReference<FTime>(METASOUND_GET_PARAM_NAME(InputMin), InOperatorSettings);
 		}
 
-		static TDataReadReference<FTime> CreateMaxValueRef(const FDataReferenceCollection& InputCollection, const FInputVertexInterface& InputInterface, const FOperatorSettings& InOperatorSettings)
+		static TDataReadReference<FTime> CreateMaxValueRef(const FInputVertexInterfaceData& InputData, const FOperatorSettings& InOperatorSettings)
 		{
 			using namespace RandomNodeNames;
-
-			return InputCollection.GetDataReadReferenceOrConstructWithVertexDefault<FTime>(InputInterface, METASOUND_GET_PARAM_NAME(InputMax), InOperatorSettings);
+			return InputData.GetOrCreateDefaultDataReadReference<FTime>(METASOUND_GET_PARAM_NAME(InputMax), InOperatorSettings);
 		}
 	};
 
@@ -320,21 +314,21 @@ namespace Metasound
 			return Info;
 		}
 
-		static TUniquePtr<IOperator> CreateOperator(const FCreateOperatorParams& InParams, TArray<TUniquePtr<IOperatorBuildError>>& OutErrors)
+		static TUniquePtr<IOperator> CreateOperator(const FBuildOperatorParams& InParams, FBuildResults& OutResults)
 		{
- 			const FInputVertexInterface& InputInterface = InParams.Node.GetVertexInterface().GetInputInterface();
- 			const FDataReferenceCollection& InputCollection = InParams.InputDataReferences;
+ 			const FInputVertexInterfaceData& InputData = InParams.InputData;
+ 			
 			using namespace RandomNodeNames;
 
- 			FTriggerReadRef InNextTrigger = InputCollection.GetDataReadReferenceOrConstruct<FTrigger>(METASOUND_GET_PARAM_NAME(InputNextTrigger), InParams.OperatorSettings);
-			FTriggerReadRef InResetTrigger = InputCollection.GetDataReadReferenceOrConstruct<FTrigger>(METASOUND_GET_PARAM_NAME(InputResetTrigger), InParams.OperatorSettings);
-			FInt32ReadRef InSeedValue = InputCollection.GetDataReadReferenceOrConstructWithVertexDefault<int32>(InputInterface, METASOUND_GET_PARAM_NAME(InputSeed), InParams.OperatorSettings);
+ 			FTriggerReadRef InNextTrigger = InputData.GetOrConstructDataReadReference<FTrigger>(METASOUND_GET_PARAM_NAME(InputNextTrigger), InParams.OperatorSettings);
+			FTriggerReadRef InResetTrigger = InputData.GetOrConstructDataReadReference<FTrigger>(METASOUND_GET_PARAM_NAME(InputResetTrigger), InParams.OperatorSettings);
+			FInt32ReadRef InSeedValue = InputData.GetOrCreateDefaultDataReadReference<int32>(METASOUND_GET_PARAM_NAME(InputSeed), InParams.OperatorSettings);
 
 			// note: random bool does not have a range
 			if (TRandomNodeSpecialization<ValueType>::HasRange())
 			{
-				TDataReadReference<ValueType> InMinValue = TRandomNodeSpecialization<ValueType>::CreateMinValueRef(InputCollection, InputInterface, InParams.OperatorSettings);
-				TDataReadReference<ValueType> InMaxValue = TRandomNodeSpecialization<ValueType>::CreateMaxValueRef(InputCollection, InputInterface, InParams.OperatorSettings);
+				TDataReadReference<ValueType> InMinValue = TRandomNodeSpecialization<ValueType>::CreateMinValueRef(InputData, InParams.OperatorSettings);
+				TDataReadReference<ValueType> InMaxValue = TRandomNodeSpecialization<ValueType>::CreateMaxValueRef(InputData, InParams.OperatorSettings);
 
 				return MakeUnique<TRandomNodeOperator<ValueType>>(InParams.OperatorSettings, InNextTrigger, InResetTrigger, InSeedValue, InMinValue, InMaxValue);
 			}
