[Backout] - CL28807780
Original CL Desc
-----------------------------------------------------------------
Mixer patch fast connect

New patch connections to the mixer are established before the next mix.
Additional adjustments to bus reader and writer buffering to ensure enough data is in their patches.

Differences between the MetaSound and mixer sample rates are still not accounted for in the bus reader/writer nodes or elsewhere.

diff --git a/Engine/Plugins/Runtime/Metasound/Source/MetasoundEngine/Private/MetasoundAudioBusWriterNode.cpp b/Engine/Plugins/Runtime/Metasound/Source/MetasoundEngine/Private/MetasoundAudioBusWriterNode.cpp
--- a/Engine/Plugins/Runtime/Metasound/Source/MetasoundEngine/Private/MetasoundAudioBusWriterNode.cpp
+++ b/Engine/Plugins/Runtime/Metasound/Source/MetasoundEngine/Private/MetasoundAudioBusWriterNode.cpp
@@ -115,10 +115,10 @@ namespace Metasound
 			}
 		}
 
-		TAudioBusWriterOperator(const FBuildOperatorParams& InParams, FAudioBusAssetReadRef InAudioBusAsset, TArray<FAudioBufferReadRef> InAudioInputs, FString InGraphName)
-			: AudioBusAsset(MoveTemp(InAudioBusAsset))
-			, AudioInputs(MoveTemp(InAudioInputs))
-			, GraphName(MoveTemp(InGraphName))
+		TAudioBusWriterOperator(const FCreateOperatorParams& InParams, FAudioBusAssetReadRef InAudioBusAsset, TArray<FAudioBufferReadRef> InAudioInputs, FString InGraphName) :
+			AudioBusAsset(MoveTemp(InAudioBusAsset)),
+			AudioInputs(MoveTemp(InAudioInputs)),
+			GraphName(MoveTemp(InGraphName))
 		{
 			Reset(InParams);
 		}
