Created a new method for generating locally unique GUID-compatible IDs for metasound. Changed over several NewGuid calls to use the new system in MetasoundFrontendGraph.

diff --git a/Engine/Plugins/Runtime/Metasound/Source/MetasoundFrontend/Private/MetasoundFrontendDocumentIdGenerator.cpp b/Engine/Plugins/Runtime/Metasound/Source/MetasoundFrontend/Private/MetasoundFrontendDocumentIdGenerator.cpp
--- a/Engine/Plugins/Runtime/Metasound/Source/MetasoundFrontend/Private/MetasoundFrontendDocumentIdGenerator.cpp
+++ b/Engine/Plugins/Runtime/Metasound/Source/MetasoundFrontend/Private/MetasoundFrontendDocumentIdGenerator.cpp
@@ -6,6 +6,11 @@ namespace Metasound
 {
 	namespace Frontend
 	{
+		namespace
+		{
+			std::atomic<uint64> GlobalAtomicMetasoundIdCounter = 1; // First ID will be 1. This is because we will make this into an FGuid which cannot be zero.
+		}
+
 		int32 MetaSoundEnableCookDeterministicIDGeneration = 1;
 		FAutoConsoleVariableRef CVarMetaSoundEnableCookDeterministicIDGeneration(
 			TEXT("au.MetaSound.EnableCookDeterministicIDGeneration"),
@@ -124,5 +129,11 @@ namespace Metasound
 
 			return FGuid(HashUint32[0], HashUint32[1], HashUint32[2], HashUint32[3]);
 		}
+
+		FGuid CreateLocallyUniqueId()
+		{
+			uint64 NextId = GlobalAtomicMetasoundIdCounter.fetch_add(1, std::memory_order_relaxed);
+			return FGuid(0, 0, NextId >> 32, NextId & 0xFFFFFFFF);
+		}
 	}
 }
diff --git a/Engine/Plugins/Runtime/Metasound/Source/MetasoundFrontend/Private/MetasoundFrontendGraph.cpp b/Engine/Plugins/Runtime/Metasound/Source/MetasoundFrontend/Private/MetasoundFrontendGraph.cpp
--- a/Engine/Plugins/Runtime/Metasound/Source/MetasoundFrontend/Private/MetasoundFrontendGraph.cpp
+++ b/Engine/Plugins/Runtime/Metasound/Source/MetasoundFrontend/Private/MetasoundFrontendGraph.cpp
@@ -9,6 +9,7 @@
 #include "CoreMinimal.h"
 #include "MetasoundFrontend.h"
 #include "MetasoundFrontendDataTypeRegistry.h"
+#include "MetasoundFrontendDocumentIdGenerator.h"
 #include "MetasoundFrontendNodeTemplateRegistry.h"
 #include "MetasoundFrontendProxyDataCache.h"
 #include "MetasoundFrontendRegistries.h"
@@ -552,7 +553,7 @@ namespace Metasound
 		for (FDefaultInputByIDMap::ElementType& Pair : InGraphContext.DefaultInputs)
 		{
 			FDefaultLiteralData& LiteralData = Pair.Value;
-			const FGuid LiteralNodeID = FGuid::NewGuid();
+			const FGuid LiteralNodeID = Frontend::CreateLocallyUniqueId();
 
 			// 1. Construct and add the default variable to the graph
 			{
@@ -613,7 +614,7 @@ namespace Metasound
 					if (bRequiresDefault)
 					{
 						InitParams.Literal = Literal.Value.ToLiteral(Vertex.TypeName, &InContext.DataTypeRegistry, &InContext.ProxyDataCache);
-						InitParams.InstanceID = FGuid::NewGuid();
+						InitParams.InstanceID = Frontend::CreateLocallyUniqueId();
 						InitParams.NodeName = "Literal";
 						TypeName = Vertex.TypeName;
 
@@ -743,7 +744,7 @@ namespace Metasound
 
 		FBuildGraphContext BuildGraphContext
 		{
-			MakeUnique<FFrontendGraph>(GraphName, FGuid::NewGuid()),
+			MakeUnique<FFrontendGraph>(GraphName, Frontend::CreateLocallyUniqueId()),
 			InGraphClass,
 			InContext
 		};
diff --git a/Engine/Plugins/Runtime/Metasound/Source/MetasoundFrontend/Public/MetasoundFrontendDocumentIdGenerator.h b/Engine/Plugins/Runtime/Metasound/Source/MetasoundFrontend/Public/MetasoundFrontendDocumentIdGenerator.h
--- a/Engine/Plugins/Runtime/Metasound/Source/MetasoundFrontend/Public/MetasoundFrontendDocumentIdGenerator.h
+++ b/Engine/Plugins/Runtime/Metasound/Source/MetasoundFrontend/Public/MetasoundFrontendDocumentIdGenerator.h
@@ -88,5 +88,7 @@ namespace Metasound
 
 			FGuid CreateNamespacedIDFromString(const FGuid NamespaceGuid, const FString& StringToHash) const;
 		};
+
+		METASOUNDFRONTEND_API FGuid CreateLocallyUniqueId();
 	}
 }
