Move the FWaveWriter class out of local definition in MetasoundWaveWriterNode.cpp and into a place where it can be used by other systems.

diff --git a/Engine/Plugins/Runtime/Metasound/Source/MetasoundStandardNodes/Private/MetasoundWaveWriterNode.cpp b/Engine/Plugins/Runtime/Metasound/Source/MetasoundStandardNodes/Private/MetasoundWaveWriterNode.cpp
--- a/Engine/Plugins/Runtime/Metasound/Source/MetasoundStandardNodes/Private/MetasoundWaveWriterNode.cpp
+++ b/Engine/Plugins/Runtime/Metasound/Source/MetasoundStandardNodes/Private/MetasoundWaveWriterNode.cpp
@@ -1,5 +1,6 @@
 // Copyright Epic Games, Inc. All Rights Reserved.
 
+#include "Audio/SimpleWaveWriter.h"
 #include "HAL/FileManager.h"
 #include "MetasoundBuildError.h"
 #include "MetasoundEnumRegistrationMacro.h"
@@ -24,104 +25,6 @@ namespace Metasound
 		METASOUND_PARAM(InFilenamePrefixPin, "Filename Prefix", "Filename Prefix of file you are writing.");
 	}
 
-	// Incremental wave writer class.
-	class FWaveWriter
-	{
-		// Local definition so we don't depend on platform includes.
-		enum EFormatType { IEEE_FLOAT = 0x3 }; // WAVE_FORMAT_IEEE_FLOAT
-		struct FWaveFormatEx
-		{
-			uint16	FormatTag;
-			uint16	NumChannels;
-			uint32	NumSamplesPerSec;
-			uint32	AverageBytesPerSec;
-			uint16	BlockAlign;
-			uint16	NumBitsPerSample;
-			uint16	Size;
-		};
-		
-	public:
-		FWaveWriter(TUniquePtr<FArchive>&& InOutputStream, int32 InSampleRate, int32 InNumChannels, bool bInUpdateHeaderAfterEveryWrite)
-			: OutputStream{ MoveTemp(InOutputStream) }
-			, bUpdateHeaderAfterEveryWrite{ bInUpdateHeaderAfterEveryWrite }
-		{
-			WriteHeader(InSampleRate, InNumChannels);
-		}
-
-		~FWaveWriter()
-		{
-			UpdateHeader();
-		}
-
-		void Write(TArrayView<const float> InBuffer)
-		{
-			OutputStream->Serialize((void*)InBuffer.GetData(), InBuffer.GetTypeSize()*InBuffer.Num());
-			
-			if (bUpdateHeaderAfterEveryWrite)
-			{
-				UpdateHeader();
-			}
-		}
-
-	private:
-		void UpdateHeader()
-		{
-			// RIFF/fmt/data. (bytes per chunk)
-			static const int32 HeaderSize = sizeof(FWaveFormatEx) + sizeof(int32) + sizeof(int32) + sizeof(int32) + sizeof(int32) + sizeof(int32);
-
-			int32 WritePos = OutputStream->Tell();
-
-			// update data chunk size
-			OutputStream->Seek(DataSizePos);
-			int32 DataSize = WritePos - DataSizePos - 4;
-			*OutputStream << DataSize;
-
-			// update top riff size
-			OutputStream->Seek(RiffSizePos);
-			int32 RiffSize = HeaderSize + DataSize - 4;
-			*OutputStream << RiffSize;
-
-			OutputStream->Seek(WritePos);
-		}
-
-		TUniquePtr<FArchive> OutputStream;
-		int32 RiffSizePos = 0;
-		int32 DataSizePos = 0;
-		bool bUpdateHeaderAfterEveryWrite = false;
-
-		void WriteHeader(int32 InSampleRate, int32 InNumChannels)
-		{	
-			FWaveFormatEx Fmt = { 0 };
-			Fmt.NumChannels = InNumChannels;
-			Fmt.NumSamplesPerSec = InSampleRate;
-			Fmt.NumBitsPerSample = sizeof(float) * 8;
-			Fmt.BlockAlign = (Fmt.NumBitsPerSample * InNumChannels) / 8;
-			Fmt.AverageBytesPerSec = Fmt.BlockAlign * InSampleRate;
-			Fmt.FormatTag = EFormatType::IEEE_FLOAT;// WAVE_FORMAT_IEEE_FLOAT;
-		
-			int32 ID = 'FFIR';
-			*OutputStream << ID;
-			RiffSizePos = OutputStream->Tell();
-			int32 RiffChunkSize = 0;
-			*OutputStream << RiffChunkSize;
-
-			ID = 'EVAW';
-			*OutputStream << ID;
-
-			ID = ' tmf';
-			*OutputStream << ID;
-			int32 FmtSize = sizeof(Fmt);
-			*OutputStream << FmtSize;
-			OutputStream->Serialize((void*)&Fmt, FmtSize);
-
-			ID = 'atad';
-			*OutputStream << ID;
-			DataSizePos = OutputStream->Tell();
-			int32 DataChunkSize = 0;
-			*OutputStream << DataChunkSize;
-		}
-	};
-
 	class FFileWriteError : public FBuildErrorBase
 	{
 	public:
@@ -483,7 +386,7 @@ namespace Metasound
 				TUniquePtr<FArchive> Stream{ IFileManager::Get().CreateFileWriter(*Filename, IO_WRITE) };
 				if (Stream.IsValid())
 				{
-					Writer = MakeUnique<FWaveWriter>(MoveTemp(Stream), SampleRate, NumInputChannels, true);
+					Writer = MakeUnique<Audio::FSimpleWaveWriter>(MoveTemp(Stream), SampleRate, NumInputChannels, true);
 				}
 			}
 		}
@@ -500,7 +403,7 @@ namespace Metasound
 		TArray<const float*> AudioInputBufferPtrs;
 		TArray<float> InterleaveBuffer;
 		FBoolReadRef Enabled;
-		TUniquePtr<FWaveWriter> Writer;
+		TUniquePtr<Audio::FSimpleWaveWriter> Writer;
 		TSharedPtr<FNumberedFileCache, ESPMode::ThreadSafe> NumberedFileCacheSP;
 		FStringReadRef FileNamePrefix;
 		float SampleRate = 0.f;
diff --git a/Engine/Source/Runtime/Core/Private/Audio/SimpleWaveWriter.cpp b/Engine/Source/Runtime/Core/Private/Audio/SimpleWaveWriter.cpp
new file mode 100644
--- /dev/null
+++ b/Engine/Source/Runtime/Core/Private/Audio/SimpleWaveWriter.cpp
@@ -0,0 +1,79 @@
+// Copyright Epic Games, Inc. All Rights Reserved.
+#include "Audio/SimpleWaveWriter.h"
+
+namespace Audio
+{
+	FSimpleWaveWriter::FSimpleWaveWriter(TUniquePtr<FArchive>&& InOutputStream, int32 InSampleRate, int32 InNumChannels, bool bInUpdateHeaderAfterEveryWrite)
+		: OutputStream{ MoveTemp(InOutputStream) }
+		, bUpdateHeaderAfterEveryWrite{ bInUpdateHeaderAfterEveryWrite }
+	{
+		WriteHeader(InSampleRate, InNumChannels);
+	}
+
+	FSimpleWaveWriter::~FSimpleWaveWriter()
+	{
+		UpdateHeader();
+	}
+
+	void FSimpleWaveWriter::Write(TArrayView<const float> InBuffer)
+	{
+		OutputStream->Serialize((void*)InBuffer.GetData(), InBuffer.GetTypeSize() * InBuffer.Num());
+
+		if (bUpdateHeaderAfterEveryWrite)
+		{
+			UpdateHeader();
+		}
+	}
+
+	void FSimpleWaveWriter::UpdateHeader()
+	{
+		// RIFF/fmt/data. (bytes per chunk)
+		static const int32 HeaderSize = sizeof(FWaveFormatEx) + sizeof(int32) + sizeof(int32) + sizeof(int32) + sizeof(int32) + sizeof(int32);
+
+		int32 WritePos = (uint32)OutputStream->Tell();
+
+		// update data chunk size
+		OutputStream->Seek(DataSizePos);
+		int32 DataSize = WritePos - DataSizePos - 4;
+		*OutputStream << DataSize;
+
+		// update top riff size
+		OutputStream->Seek(RiffSizePos);
+		int32 RiffSize = HeaderSize + DataSize - 4;
+		*OutputStream << RiffSize;
+
+		OutputStream->Seek(WritePos);
+	}
+
+	void FSimpleWaveWriter::WriteHeader(int32 InSampleRate, int32 InNumChannels)
+	{
+		FWaveFormatEx Fmt = { 0 };
+		Fmt.NumChannels = (uint16)InNumChannels;
+		Fmt.NumSamplesPerSec = InSampleRate;
+		Fmt.NumBitsPerSample = sizeof(float) * 8;
+		Fmt.BlockAlign = (uint16)((Fmt.NumBitsPerSample * InNumChannels) / 8);
+		Fmt.AverageBytesPerSec = Fmt.BlockAlign * InSampleRate;
+		Fmt.FormatTag = EFormatType::IEEE_FLOAT;// WAVE_FORMAT_IEEE_FLOAT;
+
+		int32 ID = 'FFIR';
+		*OutputStream << ID;
+		RiffSizePos = (int32)OutputStream->Tell();
+		int32 RiffChunkSize = 0;
+		*OutputStream << RiffChunkSize;
+
+		ID = 'EVAW';
+		*OutputStream << ID;
+
+		ID = ' tmf';
+		*OutputStream << ID;
+		int32 FmtSize = sizeof(Fmt);
+		*OutputStream << FmtSize;
+		OutputStream->Serialize((void*)&Fmt, FmtSize);
+
+		ID = 'atad';
+		*OutputStream << ID;
+		DataSizePos = (int32)OutputStream->Tell();
+		int32 DataChunkSize = 0;
+		*OutputStream << DataChunkSize;
+	}
+}
\ No newline at end of file
diff --git a/Engine/Source/Runtime/Core/Public/Audio/SimpleWaveWriter.h b/Engine/Source/Runtime/Core/Public/Audio/SimpleWaveWriter.h
new file mode 100644
--- /dev/null
+++ b/Engine/Source/Runtime/Core/Public/Audio/SimpleWaveWriter.h
@@ -0,0 +1,37 @@
+// Copyright Epic Games, Inc. All Rights Reserved.
+#pragma once
+#include "CoreMinimal.h"
+
+namespace Audio
+{
+	class FSimpleWaveWriter
+	{
+		// Local definition so we don't depend on platform includes.
+		enum EFormatType { IEEE_FLOAT = 0x3 }; // WAVE_FORMAT_IEEE_FLOAT
+		struct FWaveFormatEx
+		{
+			uint16	FormatTag;
+			uint16	NumChannels;
+			uint32	NumSamplesPerSec;
+			uint32	AverageBytesPerSec;
+			uint16	BlockAlign;
+			uint16	NumBitsPerSample;
+			uint16	Size;
+		};
+
+	public:
+		CORE_API FSimpleWaveWriter(TUniquePtr<FArchive>&& InOutputStream, int32 InSampleRate, int32 InNumChannels, bool bInUpdateHeaderAfterEveryWrite);
+		CORE_API ~FSimpleWaveWriter();
+
+		CORE_API void Write(TArrayView<const float> InBuffer);
+
+	private:
+		void WriteHeader(int32 InSampleRate, int32 InNumChannels);
+		void UpdateHeader();
+
+		TUniquePtr<FArchive> OutputStream;
+		int32 RiffSizePos = 0;
+		int32 DataSizePos = 0;
+		bool bUpdateHeaderAfterEveryWrite = false;
+	};
+}
\ No newline at end of file
