Remove unnecessary FString format in MetaSound CreateNodeInitData

diff --git a/Engine/Plugins/Runtime/Metasound/Source/MetasoundFrontend/Private/MetasoundFrontendGraph.cpp b/Engine/Plugins/Runtime/Metasound/Source/MetasoundFrontend/Private/MetasoundFrontendGraph.cpp
--- a/Engine/Plugins/Runtime/Metasound/Source/MetasoundFrontend/Private/MetasoundFrontendGraph.cpp
+++ b/Engine/Plugins/Runtime/Metasound/Source/MetasoundFrontend/Private/MetasoundFrontendGraph.cpp
@@ -26,7 +26,7 @@ namespace Metasound
 		{
 			FNodeInitData InitData;
 
-			InitData.InstanceName = *FString::Format(TEXT("{0}_{1}"), { *InNode.Name.ToString(), *InNode.GetID().ToString() });
+			InitData.InstanceName = InNode.Name;
 			InitData.InstanceID = InNode.GetID();
 
 			return InitData;
