Unregister graph on async task to prevent other MetaSounds being unable to find the node between unregistration and reregistration

diff --git a/Engine/Plugins/Runtime/Metasound/Source/MetasoundFrontend/Private/MetasoundAssetBase.cpp b/Engine/Plugins/Runtime/Metasound/Source/MetasoundFrontend/Private/MetasoundAssetBase.cpp
--- a/Engine/Plugins/Runtime/Metasound/Source/MetasoundFrontend/Private/MetasoundAssetBase.cpp
+++ b/Engine/Plugins/Runtime/Metasound/Source/MetasoundFrontend/Private/MetasoundAssetBase.cpp
@@ -265,8 +265,8 @@ void FMetasoundAssetBase::RegisterGraphWithFrontend(Metasound::Frontend::FMetaSo
 	// Graph registration must only happen on one thread to avoid race conditions on graph registration.
 	checkf(IsInGameThread(), TEXT("MetaSound %s graph can only be registered on the GameThread"), *GetOwningAssetName());
 
-
 	METASOUND_TRACE_CPUPROFILER_EVENT_SCOPE(MetaSoundAssetBase::RegisterGraphWithFrontend);
+	METASOUND_TRACE_CPUPROFILER_EVENT_SCOPE_TEXT(*FString::Printf(TEXT("MetaSoundAssetBase::RegisterGraphWithFrontend asset %s"), *this->GetOwningAssetName()));
 	if (!InRegistrationOptions.bForceReregister)
 	{
 		if (IsRegistered())
@@ -322,7 +322,9 @@ void FMetasoundAssetBase::RegisterGraphWithFrontend(Metasound::Frontend::FMetaSo
 
 	{
 		TScriptInterface<IMetaSoundDocumentInterface> RegistryDocInterface = AssetBasePrivate::BuildRegistryDocument(Owner);
-		UnregisterGraphWithFrontend();
+		// Clear before reregistration
+		RegistryKey = FNodeRegistryKey();
+		RegisteredGraphAssetPath = FSoftObjectPath();
 		RegistryKey = CacheRuntimeData(RegistryDocInterface);
 	}
 
diff --git a/Engine/Plugins/Runtime/Metasound/Source/MetasoundFrontend/Private/MetasoundFrontendNodeRegistryPrivate.h b/Engine/Plugins/Runtime/Metasound/Source/MetasoundFrontend/Private/MetasoundFrontendNodeRegistryPrivate.h
--- a/Engine/Plugins/Runtime/Metasound/Source/MetasoundFrontend/Private/MetasoundFrontendNodeRegistryPrivate.h
+++ b/Engine/Plugins/Runtime/Metasound/Source/MetasoundFrontend/Private/MetasoundFrontendNodeRegistryPrivate.h
@@ -108,6 +108,7 @@ namespace Metasound
 
 			virtual bool UnregisterNode(const FNodeRegistryKey& InKey) override;
 			virtual bool IsNodeRegistered(const FNodeRegistryKey& InKey) const override;
+			virtual bool IsGraphRegistered(const FNodeRegistryKey& InKey, const FSoftObjectPath& InAssetPath) const override;
 			virtual bool IsNodeNative(const FNodeRegistryKey& InKey) const override;
 
 			virtual bool RegisterConversionNode(const FConverterNodeRegistryKey& InNodeKey, const FConverterNodeInfo& InNodeInfo) override;
diff --git a/Engine/Plugins/Runtime/Metasound/Source/MetasoundFrontend/Private/MetasoundFrontendRegistries.cpp b/Engine/Plugins/Runtime/Metasound/Source/MetasoundFrontend/Private/MetasoundFrontendRegistries.cpp
--- a/Engine/Plugins/Runtime/Metasound/Source/MetasoundFrontend/Private/MetasoundFrontendRegistries.cpp
+++ b/Engine/Plugins/Runtime/Metasound/Source/MetasoundFrontend/Private/MetasoundFrontendRegistries.cpp
@@ -201,6 +201,7 @@ namespace Metasound
 			void BuildAndRegisterGraphFromDocument(const FMetasoundFrontendDocument& InDocument, const FProxyDataCache& InProxyDataCache, const FNodeClassInfo& InNodeClassInfo)
 			{
 				METASOUND_TRACE_CPUPROFILER_EVENT_SCOPE(Metasound::FRegistryContainerImpl::BuildAndRegisterGraphFromDocument);
+				METASOUND_TRACE_CPUPROFILER_EVENT_SCOPE_TEXT(*FString::Printf(TEXT("Metasound::FRegistryContainerImpl::BuildAndRegisterGraphFromDocument asset %s"), *InNodeClassInfo.AssetPath.ToString()));
 
 				TUniquePtr<FFrontendGraph> FrontendGraph = FFrontendGraphBuilder::CreateGraph(InDocument, InProxyDataCache, InNodeClassInfo.AssetPath.ToString());
 				if (!FrontendGraph.IsValid())
@@ -366,6 +367,7 @@ namespace Metasound
 			const FMetasoundFrontendDocument& Document = InDocumentInterface->GetConstDocument();
 			FNodeRegistryKey RegistryKey = FNodeRegistryKey(Document.RootGraph);
 			FNodeClassInfo NodeClassInfo(Document.RootGraph, InAssetPath);
+			METASOUND_TRACE_CPUPROFILER_EVENT_SCOPE_TEXT(*FString::Printf(TEXT("FRegistryContainerImpl::RegisterGraph key:%s, asset %s"), *RegistryKey.ToString(), *InAssetPath.ToString()));
 
 			// Proxies are created synchronously to avoid creating proxies in async tasks. Proxies
 			// are created from UObjects which need to be protected from GC and non-GT access.
@@ -386,13 +388,18 @@ namespace Metasound
 					UE_SOURCE_LOCATION,
 					[RegistryKey, NodeClassInfo=NodeClassInfo, DocumentInterface=InDocumentInterface, ProxyDataCache=MoveTemp(ProxyDataCache)]()
 					{
+						FRegistryContainerImpl& Registry = FRegistryContainerImpl::Get();
+						// Unregister the graph before reregistering
+						if (Registry.IsGraphRegistered(RegistryKey, NodeClassInfo.AssetPath))
+						{
+							Registry.UnregisterGraph(RegistryKey, NodeClassInfo.AssetPath);
+						}
+
 						// Build and add the graph to the node registry
 						MetasoundFrontendRegistryPrivate::BuildAndRegisterGraphFromDocument(DocumentInterface->GetConstDocument(), ProxyDataCache, NodeClassInfo);
 
 						// cleanup async task tracking
 						{
-							FRegistryContainerImpl& Registry = FRegistryContainerImpl::Get();
-
 							FScopeLock LockActiveReg(&Registry.ActiveRegistrationTasksCriticalSection);
 							if (Registry.ObjectReferencer)
 							{
@@ -490,7 +497,7 @@ namespace Metasound
 				FNodeRegistryTransaction::FTimeType Timestamp = FPlatformTime::Cycles64();
 
 				Key = FNodeRegistryKey(Entry->GetClassInfo());
-
+				METASOUND_TRACE_CPUPROFILER_EVENT_SCOPE_TEXT(*FString::Printf(TEXT("FRegistryContainerImpl::RegisterNode with key %s"), *Key.ToString()))
 				{
 					FScopeLock Lock(&RegistryMapsCriticalSection);
 
@@ -555,9 +562,10 @@ namespace Metasound
 		{
 			METASOUND_LLM_SCOPE;
 			
-			check(IsInGameThread());
 			if (InKey.IsValid())
 			{
+				METASOUND_TRACE_CPUPROFILER_EVENT_SCOPE_TEXT(*FString::Printf(TEXT("FRegistryContainerImpl::UnregisterNode key %s"), *InKey.ToString()))
+
 				FScopeLock Lock(&RegistryMapsCriticalSection);
 				if (const TSharedRef<INodeRegistryEntry, ESPMode::ThreadSafe>* EntryPtr = RegisteredNodes.Find(InKey))
 				{
@@ -639,6 +647,15 @@ namespace Metasound
 			}
 		}
 
+		bool FRegistryContainerImpl::IsGraphRegistered(const FNodeRegistryKey& InKey, const FSoftObjectPath& InAssetPath) const
+		{
+			FScopeLock Lock(&RegistryMapsCriticalSection);
+
+			FGraphRegistryKey GraphRegistryKey{ InKey, InAssetPath };
+
+			return RegisteredGraphs.Contains(GraphRegistryKey);
+		}
+
 		bool FRegistryContainerImpl::IsNodeNative(const FNodeRegistryKey& InKey) const
 		{
 			if (const INodeRegistryEntry* Entry = FindNodeEntry(InKey))
@@ -846,6 +863,7 @@ namespace Metasound
 			if (AsyncRegistrationPipe.IsInContext())
 			{
 				// It is not safe to wait for an async registration task from within the async registration pipe because it will result in a deadlock. 
+				UE_LOG(LogMetaSound, Verbose, TEXT("Async registration pipe is already in context for registering key %s. Task will not be waited for."), *InRegistryKey.ToString());
 				return;
 			}
 
diff --git a/Engine/Plugins/Runtime/Metasound/Source/MetasoundFrontend/Public/MetasoundFrontendRegistries.h b/Engine/Plugins/Runtime/Metasound/Source/MetasoundFrontend/Public/MetasoundFrontendRegistries.h
--- a/Engine/Plugins/Runtime/Metasound/Source/MetasoundFrontend/Public/MetasoundFrontendRegistries.h
+++ b/Engine/Plugins/Runtime/Metasound/Source/MetasoundFrontend/Public/MetasoundFrontendRegistries.h
@@ -320,6 +320,9 @@ public:
 	/** Returns true if the provided registry key corresponds to a valid registered node. */
 	virtual bool IsNodeRegistered(const FNodeRegistryKey& InKey) const = 0;
 
+	/** Returns true if the provided registry key and asset path corresponds to a valid registered graph. */
+	virtual bool IsGraphRegistered(const FNodeRegistryKey& InKey, const FSoftObjectPath& InAssetPath) const = 0;
+
 	/** Returns true if the provided registry key corresponds to a valid registered node that is natively defined. */
 	virtual bool IsNodeNative(const FNodeRegistryKey& InKey) const = 0;
 
