Fix crash caused by race condition in metasound node registration
#rnx

diff --git a/Engine/Plugins/Runtime/Metasound/Source/MetasoundFrontend/Private/MetasoundFrontendDocumentBuilder.cpp b/Engine/Plugins/Runtime/Metasound/Source/MetasoundFrontend/Private/MetasoundFrontendDocumentBuilder.cpp
--- a/Engine/Plugins/Runtime/Metasound/Source/MetasoundFrontend/Private/MetasoundFrontendDocumentBuilder.cpp
+++ b/Engine/Plugins/Runtime/Metasound/Source/MetasoundFrontend/Private/MetasoundFrontendDocumentBuilder.cpp
@@ -806,12 +806,12 @@ bool FMetaSoundFrontendDocumentBuilder::AddEdgesByNodeClassInterfaceBindings(con
 	using namespace Metasound;
 	using namespace Metasound::Frontend;
 
-	const TSet<FMetasoundFrontendVersion>* FromInterfaceVersions = FindNodeClassInterfaces(InFromNodeID);
-	const TSet<FMetasoundFrontendVersion>* ToInterfaceVersions = FindNodeClassInterfaces(InToNodeID);
-	if (FromInterfaceVersions && ToInterfaceVersions)
+	TSet<FMetasoundFrontendVersion> FromInterfaceVersions;
+	TSet<FMetasoundFrontendVersion> ToInterfaceVersions;
+	if (FindNodeClassInterfaces(InFromNodeID, FromInterfaceVersions) && FindNodeClassInterfaces(InToNodeID, ToInterfaceVersions))
 	{
 		TSet<FNamedEdge> NamedEdges;
-		if (DocumentBuilderPrivate::TryGetInterfaceBoundEdges(InFromNodeID, *FromInterfaceVersions, InToNodeID, *ToInterfaceVersions, NamedEdges))
+		if (DocumentBuilderPrivate::TryGetInterfaceBoundEdges(InFromNodeID, FromInterfaceVersions, InToNodeID, ToInterfaceVersions, NamedEdges))
 		{
 			return AddNamedEdges(NamedEdges, nullptr, bReplaceExistingConnections);
 		}
@@ -829,15 +829,16 @@ bool FMetaSoundFrontendDocumentBuilder::AddEdgesFromMatchingInterfaceNodeOutputs
 
 	OutEdgesCreated.Reset();
 
-	const TSet<FMetasoundFrontendVersion>* NodeInterfaces = FindNodeClassInterfaces(InNodeID);
-	if (!NodeInterfaces)
+	TSet<FMetasoundFrontendVersion> NodeInterfaces;
+	if (!FindNodeClassInterfaces(InNodeID, NodeInterfaces))
 	{
+		// Did not find any node interfaces
 		return false;
 	}
 
 	const IDocumentGraphNodeCache& NodeCache = DocumentCache->GetNodeCache();
 	const IDocumentGraphInterfaceCache& InterfaceCache = DocumentCache->GetInterfaceCache();
-	const TSet<FMetasoundFrontendVersion> CommonInterfaces = NodeInterfaces->Intersect(GetDocument().Interfaces);
+	const TSet<FMetasoundFrontendVersion> CommonInterfaces = NodeInterfaces.Intersect(GetDocument().Interfaces);
 
 	TSet<FNamedEdge> EdgesToMake;
 	for (const FMetasoundFrontendVersion& Version : CommonInterfaces)
@@ -872,15 +873,16 @@ bool FMetaSoundFrontendDocumentBuilder::AddEdgesFromMatchingInterfaceNodeInputsT
 
 	OutEdgesCreated.Reset();
 
-	const TSet<FMetasoundFrontendVersion>* NodeInterfaces = FindNodeClassInterfaces(InNodeID);
-	if (!NodeInterfaces)
+	TSet<FMetasoundFrontendVersion> NodeInterfaces;
+	if (!FindNodeClassInterfaces(InNodeID, NodeInterfaces))
 	{
+		// Did not find any node interfaces
 		return false;
 	}
 
 	const IDocumentGraphNodeCache& NodeCache = DocumentCache->GetNodeCache();
 	const IDocumentGraphInterfaceCache& InterfaceCache = DocumentCache->GetInterfaceCache();
-	const TSet<FMetasoundFrontendVersion> CommonInterfaces = NodeInterfaces->Intersect(GetDocument().Interfaces);
+	const TSet<FMetasoundFrontendVersion> CommonInterfaces = NodeInterfaces.Intersect(GetDocument().Interfaces);
 
 	TSet<FNamedEdge> EdgesToMake;
 	for (const FMetasoundFrontendVersion& Version : CommonInterfaces)
@@ -1519,7 +1521,7 @@ const FMetasoundFrontendNode* FMetaSoundFrontendDocumentBuilder::FindNode(const
 	return NodeCache.FindNode(InNodeID);
 }
 
-const TSet<FMetasoundFrontendVersion>* FMetaSoundFrontendDocumentBuilder::FindNodeClassInterfaces(const FGuid& InNodeID) const
+bool FMetaSoundFrontendDocumentBuilder::FindNodeClassInterfaces(const FGuid& InNodeID, TSet<FMetasoundFrontendVersion>& OutInterfaces) const
 {
 	using namespace Metasound;
 	using namespace Metasound::Frontend;
@@ -1530,13 +1532,12 @@ const TSet<FMetasoundFrontendVersion>* FMetaSoundFrontendDocumentBuilder::FindNo
 	{
 		if (const FMetasoundFrontendClass* NodeClass = DocumentCache->FindDependency(Node->ClassID))
 		{
-			// 1. May be a serialized asset, so first check with asset manager.
 			const FNodeRegistryKey NodeClassRegistryKey = FNodeRegistryKey(NodeClass->Metadata);
-			return FMetasoundFrontendRegistryContainer::Get()->FindImplementedInterfacesFromRegistered(NodeClassRegistryKey);
+			return FMetasoundFrontendRegistryContainer::Get()->FindImplementedInterfacesFromRegistered(NodeClassRegistryKey, OutInterfaces);
 		}
 	}
 
-	return nullptr;
+	return false;
 }
 
 const FMetasoundFrontendVertex* FMetaSoundFrontendDocumentBuilder::FindNodeInput(const FGuid& InNodeID, const FGuid& InVertexID) const
@@ -2318,12 +2319,13 @@ bool FMetaSoundFrontendDocumentBuilder::RemoveEdgesByNodeClassInterfaceBindings(
 	using namespace Metasound;
 	using namespace Metasound::Frontend;
 
-	const TSet<FMetasoundFrontendVersion>* FromInterfaceVersions = FindNodeClassInterfaces(InFromNodeID);
-	const TSet<FMetasoundFrontendVersion>* ToInterfaceVersions = FindNodeClassInterfaces(InToNodeID);
-	if (FromInterfaceVersions && ToInterfaceVersions)
+	TSet<FMetasoundFrontendVersion> FromInterfaceVersions;
+	TSet<FMetasoundFrontendVersion> ToInterfaceVersions;
+
+	if (FindNodeClassInterfaces(InFromNodeID, FromInterfaceVersions) && FindNodeClassInterfaces(InToNodeID, ToInterfaceVersions))
 	{
 		TSet<FNamedEdge> NamedEdges;
-		if (DocumentBuilderPrivate::TryGetInterfaceBoundEdges(InFromNodeID, *FromInterfaceVersions, InToNodeID, *ToInterfaceVersions, NamedEdges))
+		if (DocumentBuilderPrivate::TryGetInterfaceBoundEdges(InFromNodeID, FromInterfaceVersions, InToNodeID, ToInterfaceVersions, NamedEdges))
 		{
 			return RemoveNamedEdges(NamedEdges);
 		}
diff --git a/Engine/Plugins/Runtime/Metasound/Source/MetasoundFrontend/Private/MetasoundFrontendNodeRegistryPrivate.h b/Engine/Plugins/Runtime/Metasound/Source/MetasoundFrontend/Private/MetasoundFrontendNodeRegistryPrivate.h
--- a/Engine/Plugins/Runtime/Metasound/Source/MetasoundFrontend/Private/MetasoundFrontendNodeRegistryPrivate.h
+++ b/Engine/Plugins/Runtime/Metasound/Source/MetasoundFrontend/Private/MetasoundFrontendNodeRegistryPrivate.h
@@ -117,11 +117,14 @@ namespace Metasound
 
 			// Find Frontend class data.
 			virtual bool FindFrontendClassFromRegistered(const FNodeRegistryKey& InKey, FMetasoundFrontendClass& OutClass) override;
+
+			UE_DEPRECATED(5.4, "This implementation of FindImplementedInterfacesFromRegistered(...). Please use bool FindImplementedInterfacesFromRegistered(const Metasound::Frontend::FNodeRegistryKey& InKey, TSet<FMetasoundFrontendVersion>& OutInterfacVersion)")
 			virtual const TSet<FMetasoundFrontendVersion>* FindImplementedInterfacesFromRegistered(const Metasound::Frontend::FNodeRegistryKey& InKey) const override;
 
+			virtual bool FindImplementedInterfacesFromRegistered(const Metasound::Frontend::FNodeRegistryKey& InKey, TSet<FMetasoundFrontendVersion>& OutInterfaceVersions) const override;
+			
 			UE_DEPRECATED(5.4, "Use FindFrontendClassFromRegistered instead")
 			virtual bool FindNodeClassInfoFromRegistered(const Metasound::Frontend::FNodeRegistryKey& InKey, FNodeClassInfo& OutInfo) override;
-
 			virtual bool FindInputNodeRegistryKeyForDataType(const FName& InDataTypeName, const EMetasoundFrontendVertexAccessType InAccessType, FNodeRegistryKey& OutKey) override;
 			virtual bool FindVariableNodeRegistryKeyForDataType(const FName& InDataTypeName, FNodeRegistryKey& OutKey) override;
 			virtual bool FindOutputNodeRegistryKeyForDataType(const FName& InDataTypeName, const EMetasoundFrontendVertexAccessType InAccessType, FNodeRegistryKey& OutKey) override;
@@ -153,7 +156,12 @@ namespace Metasound
 			void WaitForAsyncRegistrationInternal(const FNodeRegistryKey& InRegistryKey, const FSoftObjectPath* InAssetPath) const;
 			void RegisterGraphInternal(const FNodeRegistryKey& InKey, const FSoftObjectPath& InAssetPath, TSharedPtr<const FGraph> InGraph);
 
-			const INodeRegistryEntry* FindNodeEntry(const FNodeRegistryKey& InKey) const;
+			// Access a node entry safely. Node entries can be added/removed asynchronously. Functions passed to this method will be
+			// executed in a manner where access to the node registry entry is safe from threading issues. 
+			//
+			// @returns true if a node registry entry was found and the function executed. False if the entry was not 
+			//          found and the function not executed. 
+			bool AccessNodeEntryThreadSafe(const FNodeRegistryKey& InKey, TFunctionRef<void(const INodeRegistryEntry&)> InFunc) const;
 
 			const INodeRegistryTemplateEntry* FindNodeTemplateEntry(const FNodeRegistryKey& InKey) const;
 
diff --git a/Engine/Plugins/Runtime/Metasound/Source/MetasoundFrontend/Private/MetasoundFrontendRegistries.cpp b/Engine/Plugins/Runtime/Metasound/Source/MetasoundFrontend/Private/MetasoundFrontendRegistries.cpp
--- a/Engine/Plugins/Runtime/Metasound/Source/MetasoundFrontend/Private/MetasoundFrontendRegistries.cpp
+++ b/Engine/Plugins/Runtime/Metasound/Source/MetasoundFrontend/Private/MetasoundFrontendRegistries.cpp
@@ -289,52 +289,68 @@ namespace Metasound
 
 		TUniquePtr<Metasound::INode> FRegistryContainerImpl::CreateNode(const FNodeRegistryKey& InKey, const Metasound::FNodeInitData& InInitData) const
 		{
-			if (const INodeRegistryEntry* Entry = FindNodeEntry(InKey))
+			TUniquePtr<INode> Node;
+			auto CreateNode = [&Node, &InInitData](const INodeRegistryEntry& Entry) 
+			{ 
+				Node = Entry.CreateNode(InInitData); 
+			};
+
+			if (!AccessNodeEntryThreadSafe(InKey, CreateNode))
 			{
-				return Entry->CreateNode(InInitData);
+				// Creation of external nodes can rely on assets being unavailable due to errors in loading order, asset(s)
+				// missing, etc. 
+				UE_LOG(LogMetaSound, Error, TEXT("Could not find node [RegistryKey:%s]"), *InKey.ToString());
 			}
 
-			// Because creation of external nodes can rely on assets being unavailable due to errors in loading order, asset(s)
-			// missing, etc. only log error and don't throw ensure to avoid blocking start-up if assets are missing. All other
-			// CreateNode calls are natively managed and thus better suited to throw ensures.
-			UE_LOG(LogMetaSound, Error, TEXT("Could not find node [RegistryKey:%s]"), *InKey.ToString());
-			return nullptr;
+			return MoveTemp(Node);
 		}
 
 		TUniquePtr<Metasound::INode> FRegistryContainerImpl::CreateNode(const FNodeRegistryKey& InKey, FDefaultLiteralNodeConstructorParams&& InParams) const
 		{
-			const INodeRegistryEntry* Entry = FindNodeEntry(InKey);
+			TUniquePtr<INode> Node;
+			auto CreateNode = [&Node, &InParams](const INodeRegistryEntry& Entry) 
+			{ 
+				Node = Entry.CreateNode(MoveTemp(InParams)); 
+			};
 
-			if (ensureAlwaysMsgf(nullptr != Entry, TEXT("Could not find node [RegistryKey:%s]"), *InKey.ToString()))
+			if (!AccessNodeEntryThreadSafe(InKey, CreateNode))
 			{
-				return Entry->CreateNode(MoveTemp(InParams));
+				UE_LOG(LogMetaSound, Error, TEXT("Could not find node [RegistryKey:%s]"), *InKey.ToString());
 			}
 
-			return nullptr;
+			return MoveTemp(Node);
 		}
 
 		TUniquePtr<Metasound::INode> FRegistryContainerImpl::CreateNode(const FNodeRegistryKey& InKey, FDefaultNamedVertexNodeConstructorParams&& InParams) const
 		{
-			const INodeRegistryEntry* Entry = FindNodeEntry(InKey);
+			TUniquePtr<INode> Node;
+			auto CreateNode = [&Node, &InParams](const INodeRegistryEntry& Entry) 
+			{ 
+				Node = Entry.CreateNode(MoveTemp(InParams)); 
+			};
 
-			if (ensureAlwaysMsgf(nullptr != Entry, TEXT("Could not find node [RegistryKey:%s]"), *InKey.ToString()))
+			if (!AccessNodeEntryThreadSafe(InKey, CreateNode))
 			{
-				return Entry->CreateNode(MoveTemp(InParams));
+				UE_LOG(LogMetaSound, Error, TEXT("Could not find node [RegistryKey:%s]"), *InKey.ToString());
 			}
 
-			return nullptr;
+			return MoveTemp(Node);
 		}
 
 		TUniquePtr<Metasound::INode> FRegistryContainerImpl::CreateNode(const FNodeRegistryKey& InKey, FDefaultNamedVertexWithLiteralNodeConstructorParams&& InParams) const
 		{
-			const INodeRegistryEntry* Entry = FindNodeEntry(InKey);
+			TUniquePtr<INode> Node;
+			auto CreateNode = [&Node, &InParams](const INodeRegistryEntry& Entry) 
+			{ 
+				Node = Entry.CreateNode(MoveTemp(InParams)); 
+			};
 
-			if (ensureAlwaysMsgf(nullptr != Entry, TEXT("Could not find node [RegistryKey:%s]"), *InKey.ToString()))
+			if (!AccessNodeEntryThreadSafe(InKey, CreateNode))
 			{
-				return Entry->CreateNode(MoveTemp(InParams));
+				UE_LOG(LogMetaSound, Error, TEXT("Could not find node [RegistryKey:%s]"), *InKey.ToString());
 			}
 
-			return nullptr;
+			return MoveTemp(Node);
 		}
 
 		TArray<::Metasound::Frontend::FConverterNodeInfo> FRegistryContainerImpl::GetPossibleConverterNodes(const FName& FromDataType, const FName& ToDataType)
@@ -658,24 +674,33 @@ namespace Metasound
 
 		bool FRegistryContainerImpl::IsNodeNative(const FNodeRegistryKey& InKey) const
 		{
-			if (const INodeRegistryEntry* Entry = FindNodeEntry(InKey))
+			bool bIsNative = false;
+			auto SetIsNative = [&bIsNative](const INodeRegistryEntry& Entry) 
+			{ 
+				bIsNative = Entry.IsNative();
+			};
+
+			if (AccessNodeEntryThreadSafe(InKey, SetIsNative))
 			{
-				return Entry->IsNative();
+				return bIsNative;
 			}
 
 			if (const INodeRegistryTemplateEntry* TemplateEntry = FindNodeTemplateEntry(InKey))
 			{
 				return true;
 			}
-
 			return false;
 		}
 
 		bool FRegistryContainerImpl::FindFrontendClassFromRegistered(const FNodeRegistryKey& InKey, FMetasoundFrontendClass& OutClass)
 		{
-			if (const INodeRegistryEntry* Entry = FindNodeEntry(InKey))
+			auto SetFrontendClass = [&OutClass](const INodeRegistryEntry& Entry)
+			{
+				OutClass = Entry.GetFrontendClass();
+			};
+
+			if (AccessNodeEntryThreadSafe(InKey, SetFrontendClass))
 			{
-				OutClass = Entry->GetFrontendClass();
 				return true;
 			}
 
@@ -690,19 +715,52 @@ namespace Metasound
 
 		const TSet<FMetasoundFrontendVersion>* FRegistryContainerImpl::FindImplementedInterfacesFromRegistered(const Metasound::Frontend::FNodeRegistryKey& InKey) const
 		{
-			if (const INodeRegistryEntry* Entry = FindNodeEntry(InKey))
+			static bool bHasWarningBeenIssued = false;
+			if (!bHasWarningBeenIssued)
 			{
-				return Entry->GetImplementedInterfaces();
+				// This function is known to be thread-unsafe and should not longer be used.
+				// This implementation exists to support deprecated usage. 
+				UE_LOG(LogMetaSound, Warning, TEXT("Accessing non-thread-safe implementation of FindImplementedInterfacesFromRegistered(...) is known to cause crashes. Please update your code to use the non-deprecated version of this function with the same name"));
+				bHasWarningBeenIssued = true;
 			}
 
-			return nullptr;
+			const TSet<FMetasoundFrontendVersion>* Interfaces = nullptr;
+			AccessNodeEntryThreadSafe(InKey, 
+				[&Interfaces](const INodeRegistryEntry& Entry)
+				{
+					Interfaces = Entry.GetImplementedInterfaces();
+				}
+			);
+
+			return Interfaces;
+		}
+
+		bool FRegistryContainerImpl::FindImplementedInterfacesFromRegistered(const Metasound::Frontend::FNodeRegistryKey& InKey, TSet<FMetasoundFrontendVersion>& OutInterfaceVersions) const 
+		{
+			bool bDidCopy = false;
+
+			auto CopyImplementedInterfaces = [&OutInterfaceVersions, &bDidCopy](const INodeRegistryEntry& Entry)
+			{
+				if (const TSet<FMetasoundFrontendVersion>* Interfaces = Entry.GetImplementedInterfaces())
+				{
+					OutInterfaceVersions = *Interfaces;
+					bDidCopy = true;
+				}
+			};
+
+			AccessNodeEntryThreadSafe(InKey, CopyImplementedInterfaces);
+
+			return bDidCopy;
 		}
 
 		bool FRegistryContainerImpl::FindNodeClassInfoFromRegistered(const Metasound::Frontend::FNodeRegistryKey& InKey, FNodeClassInfo& OutInfo)
 		{
-			if (const INodeRegistryEntry* Entry = FindNodeEntry(InKey))
+			auto CopyClassInfo = [&OutInfo](const INodeRegistryEntry& Entry)
+			{
+				OutInfo = Entry.GetClassInfo();
+			};
+			if (AccessNodeEntryThreadSafe(InKey, CopyClassInfo))
 			{
-				OutInfo = Entry->GetClassInfo();
 				return true;
 			}
 
@@ -815,30 +873,28 @@ namespace Metasound
 			}
 		}
 
-		const INodeRegistryEntry* FRegistryContainerImpl::FindNodeEntry(const FNodeRegistryKey& InKey) const
+		bool FRegistryContainerImpl::AccessNodeEntryThreadSafe(const FNodeRegistryKey& InKey, TFunctionRef<void(const INodeRegistryEntry&)> InFunc) const
 		{
-			auto TryFindNodeEntry = [this, &InKey]() -> const INodeRegistryEntry*
+			auto TryAccessNodeEntry = [this, &InKey, &InFunc]() -> bool
 			{
-				// This scope lock protects against race conditions manipulating the `RegisteredNodes` map, but it does not
-				// protect against the individual INodeRegistryEntry from being removed after this function returns. Generally
-				// this is not an issue as the acess to the INodeRegistryEntry pointer happens on the same thread as calls to 
-				// remove node registry entries. 
 				FScopeLock Lock(&RegistryMapsCriticalSection);
 				if (const TSharedRef<INodeRegistryEntry, ESPMode::ThreadSafe>* Entry = RegisteredNodes.Find(InKey))
 				{
-					return &Entry->Get();
+					InFunc(*(*Entry));
+					return true;
 				}
-				return nullptr;
+				return false;
 			};
 
-			if (const INodeRegistryEntry* Entry = TryFindNodeEntry())
+			if (TryAccessNodeEntry())
 			{
-				return Entry;
+				return true;
 			}
 			else
 			{
+				// Wait for any async registration tasks related to the registry key. 
 				WaitForAsyncRegistrationInternal(InKey, nullptr /* InAssetPath */);
-				return TryFindNodeEntry();
+				return TryAccessNodeEntry();
 			}
 		}
 
diff --git a/Engine/Plugins/Runtime/Metasound/Source/MetasoundFrontend/Public/MetasoundFrontendDocumentBuilder.h b/Engine/Plugins/Runtime/Metasound/Source/MetasoundFrontend/Public/MetasoundFrontendDocumentBuilder.h
--- a/Engine/Plugins/Runtime/Metasound/Source/MetasoundFrontend/Public/MetasoundFrontendDocumentBuilder.h
+++ b/Engine/Plugins/Runtime/Metasound/Source/MetasoundFrontend/Public/MetasoundFrontendDocumentBuilder.h
@@ -292,7 +292,7 @@ private:
 
 	FMetasoundFrontendNode* AddNodeInternal(const FMetasoundFrontendClassMetadata& InClassMetadata, Metasound::Frontend::FFinalizeNodeFunctionRef FinalizeNode, FGuid InNodeID = FGuid::NewGuid());
 
-	const TSet<FMetasoundFrontendVersion>* FindNodeClassInterfaces(const FGuid& InNodeID) const;
+	bool FindNodeClassInterfaces(const FGuid& InNodeID, TSet<FMetasoundFrontendVersion>& OutInterfaces) const;
 
 	void InitCacheInternal(bool bPrimeCache = false);
 
diff --git a/Engine/Plugins/Runtime/Metasound/Source/MetasoundFrontend/Public/MetasoundFrontendRegistries.h b/Engine/Plugins/Runtime/Metasound/Source/MetasoundFrontend/Public/MetasoundFrontendRegistries.h
--- a/Engine/Plugins/Runtime/Metasound/Source/MetasoundFrontend/Public/MetasoundFrontendRegistries.h
+++ b/Engine/Plugins/Runtime/Metasound/Source/MetasoundFrontend/Public/MetasoundFrontendRegistries.h
@@ -331,11 +331,14 @@ public:
 
 	// Query for MetaSound Frontend document objects.
 	virtual bool FindFrontendClassFromRegistered(const Metasound::Frontend::FNodeRegistryKey& InKey, FMetasoundFrontendClass& OutClass) = 0;
+
+	UE_DEPRECATED(5.4, "This implementation of FindImplementedInterfacesFromRegistered(...). Please use bool FindImplementedInterfacesFromRegistered(const Metasound::Frontend::FNodeRegistryKey& InKey, TSet<FMetasoundFrontendVersion>& OutInterfacVersion)")
 	virtual const TSet<FMetasoundFrontendVersion>* FindImplementedInterfacesFromRegistered(const Metasound::Frontend::FNodeRegistryKey& InKey) const = 0;
 
+	virtual bool FindImplementedInterfacesFromRegistered(const Metasound::Frontend::FNodeRegistryKey& InKey, TSet<FMetasoundFrontendVersion>& OutInterfaceVersions) const = 0;
+	
 	UE_DEPRECATED(5.4, "Use FindFrontendClassFromRegistered instead")
 	virtual bool FindNodeClassInfoFromRegistered(const Metasound::Frontend::FNodeRegistryKey& InKey, FNodeClassInfo& OutInfo) = 0;
-
 	virtual bool FindInputNodeRegistryKeyForDataType(const FName& InDataTypeName, const EMetasoundFrontendVertexAccessType InAccessType, FNodeRegistryKey& OutKey) = 0;
 	virtual bool FindVariableNodeRegistryKeyForDataType(const FName& InDataTypeName, FNodeRegistryKey& OutKey) = 0;
 	virtual bool FindOutputNodeRegistryKeyForDataType(const FName& InDataTypeName, const EMetasoundFrontendVertexAccessType InAccessType, FNodeRegistryKey& OutKey) = 0;
