[Patchwork] Live Update Support

diff --git a/Engine/Plugins/Runtime/Metasound/Source/MetasoundEngine/Private/MetasoundGeneratorHandle.cpp b/Engine/Plugins/Runtime/Metasound/Source/MetasoundEngine/Private/MetasoundGeneratorHandle.cpp
--- a/Engine/Plugins/Runtime/Metasound/Source/MetasoundEngine/Private/MetasoundGeneratorHandle.cpp
+++ b/Engine/Plugins/Runtime/Metasound/Source/MetasoundEngine/Private/MetasoundGeneratorHandle.cpp
@@ -524,6 +524,11 @@ namespace Metasound
 			{
 				PinnedThis->SendParametersToGenerator();
 				PinnedThis->FixUpOutputWatchers();
+
+				if (PinnedThis->OnGeneratorIOUpdated.IsBound())
+				{
+					PinnedThis->OnGeneratorIOUpdated.Execute();
+				}
 			}
 		});
 	}
@@ -703,5 +708,10 @@ bool UMetasoundGeneratorHandle::InitGeneratorHandle(TWeakObjectPtr<UAudioCompone
 		OnGeneratorsGraphChanged.Broadcast();
 	});
 
+	GeneratorHandle->OnGeneratorIOUpdated.BindLambda([this]()
+	{
+		OnIOUpdated.Broadcast();
+	});
+
 	return true;
 }
diff --git a/Engine/Plugins/Runtime/Metasound/Source/MetasoundEngine/Public/MetasoundGeneratorHandle.h b/Engine/Plugins/Runtime/Metasound/Source/MetasoundEngine/Public/MetasoundGeneratorHandle.h
--- a/Engine/Plugins/Runtime/Metasound/Source/MetasoundEngine/Public/MetasoundGeneratorHandle.h
+++ b/Engine/Plugins/Runtime/Metasound/Source/MetasoundEngine/Public/MetasoundGeneratorHandle.h
@@ -83,7 +83,12 @@ namespace Metasound
 		 */
 		FOnGeneratorGraphSet OnGraphUpdated;
 	
-		DECLARE_DELEGATE_OneParam(FOnGeneratorVertexInterfaceUpdated, const FVertexInterfaceData&);
+		DECLARE_DELEGATE(FOnGeneratorIOUpdated);
+
+		/**
+		 * Fires on the game thread when the generator's graph inputs or outputs change
+		 */
+		FOnGeneratorIOUpdated OnGeneratorIOUpdated;
 
 		/**
 		 * Update the current parameter state on this handle and enqueue the changes on the generator.
@@ -321,6 +326,9 @@ public:
 	FDelegateHandle AddGraphSetCallback(UMetasoundGeneratorHandle::FOnSetGraph::FDelegate&& Delegate);
 	bool RemoveGraphSetCallback(const FDelegateHandle& Handle);
 
+	DECLARE_MULTICAST_DELEGATE(FOnIOUpdated)
+	FOnIOUpdated OnIOUpdated;
+
 	/**
 	 * Watch an output value.
 	 *
