Metasound static analysis fix
#rnx

diff --git a/Engine/Plugins/Runtime/Metasound/Source/MetasoundFrontend/Private/MetasoundFrontendRegistryTransaction.h b/Engine/Plugins/Runtime/Metasound/Source/MetasoundFrontend/Private/MetasoundFrontendRegistryTransaction.h
--- a/Engine/Plugins/Runtime/Metasound/Source/MetasoundFrontend/Private/MetasoundFrontendRegistryTransaction.h
+++ b/Engine/Plugins/Runtime/Metasound/Source/MetasoundFrontend/Private/MetasoundFrontendRegistryTransaction.h
@@ -73,6 +73,7 @@ namespace Metasound
 					FTransactionReaderHandle NewReaderHandle = CurrentReaderHandle;
 					CurrentReaderHandle++;
 					ReaderPositions.Add(NewReaderHandle, Position);
+					return NewReaderHandle;
 				}
 			}
 
