Fix missing ;

diff --git a/Engine/Plugins/Runtime/Metasound/Source/MetasoundEngineTest/Private/MetasoundDynamicGeneratorTests.cpp b/Engine/Plugins/Runtime/Metasound/Source/MetasoundEngineTest/Private/MetasoundDynamicGeneratorTests.cpp
--- a/Engine/Plugins/Runtime/Metasound/Source/MetasoundEngineTest/Private/MetasoundDynamicGeneratorTests.cpp
+++ b/Engine/Plugins/Runtime/Metasound/Source/MetasoundEngineTest/Private/MetasoundDynamicGeneratorTests.cpp
@@ -178,7 +178,7 @@ namespace Metasound::Test::Generator::Dynamic
 			const FAnyDataReference* InputRef = LatestInterfaceData.GetInputs().FindDataReference(InputName);
 			UTEST_NOT_NULL("Vertex data contains input", InputRef);
 			const float* Value = InputRef->GetValue<float>();
-			UTEST_NOT_NULL("Value exists", Value)
+			UTEST_NOT_NULL("Value exists", Value);
 			UTEST_EQUAL("Value is default", *Value, DefaultValue);
 		}
 
diff --git a/Engine/Source/Runtime/Core/Public/Misc/AutomationTest.h b/Engine/Source/Runtime/Core/Public/Misc/AutomationTest.h
--- a/Engine/Source/Runtime/Core/Public/Misc/AutomationTest.h
+++ b/Engine/Source/Runtime/Core/Public/Misc/AutomationTest.h
@@ -4001,7 +4001,7 @@ public: \
 	{\
 		return false;\
 	}\
-	CA_ASSUME(Pointer)
+	CA_ASSUME(Pointer);
 
 #define UTEST_NOT_NULL_EXPR(Pointer)\
 	if (!TestNotNull(TEXT(#Pointer), Pointer))\
