[Metasound] Add native delegate option for WatchOutput in FMetasoundGeneratorHandle and related classes.

- Updated StressTest_MetasoundOutputs to optionally use the new native delegates.
- In local testing the native delegates were ~40% faster on average than dynamic delegates, when measuring FMetasoundGeneratorHandle::UpdateOutputWatchers.

#tests EngineTest - PIE, Standalone

#virtualized

diff --git a/Engine/Plugins/Runtime/Metasound/Source/MetasoundEngine/Private/MetasoundGeneratorHandle.cpp b/Engine/Plugins/Runtime/Metasound/Source/MetasoundEngine/Private/MetasoundGeneratorHandle.cpp
--- a/Engine/Plugins/Runtime/Metasound/Source/MetasoundEngine/Private/MetasoundGeneratorHandle.cpp
+++ b/Engine/Plugins/Runtime/Metasound/Source/MetasoundEngine/Private/MetasoundGeneratorHandle.cpp
@@ -134,82 +134,16 @@
 		const FName AnalyzerName,
 		const FName AnalyzerOutputName)
 	{
-		METASOUND_LLM_SCOPE;
-		METASOUND_TRACE_CPUPROFILER_EVENT_SCOPE(FMetasoundGeneratorHandle::WatchOutput);
-
-		check(IsInGameThread());
-
-		if (!IsValid())
-		{
-			return false;
-		}
-
-		// Make the analyzer address.
-		Frontend::FAnalyzerAddress AnalyzerAddress;
-		AnalyzerAddress.InstanceID = GetAudioComponentId();
-		AnalyzerAddress.OutputName = OutputName;
-		AnalyzerAddress.AnalyzerName = AnalyzerName;
-		AnalyzerAddress.AnalyzerMemberName = AnalyzerOutputName;
-		AnalyzerAddress.AnalyzerInstanceID = FGuid::NewGuid();
-		
-		// Find the output node and get the data type/node id from that
-		{
-			const TWeakObjectPtr<UMetaSoundSource> Source = GetMetaSoundSource();
-
-			if (!Source.IsValid())
-			{
-				UE_LOG(LogMetaSound, Warning, TEXT("Couldn't find the MetaSound Source"));
-				return false;
-			}
-		
-			// Find the node id and type name
-			const FMetasoundFrontendClassOutput* OutputPtr =
-				Source->GetConstDocument().RootGraph.Interface.Outputs.FindByPredicate(
-					[&AnalyzerAddress](const FMetasoundFrontendClassOutput& Output)
-					{
-						return Output.Name == AnalyzerAddress.OutputName;
-					});
-
-			if (nullptr == OutputPtr)
-			{
-				return false;
-			}
-			
-			AnalyzerAddress.NodeID = OutputPtr->NodeID;
-			AnalyzerAddress.DataType = OutputPtr->TypeName;
-		}
-		
-		// If no analyzer name was provided, try to find a passthrough analyzer
-		if (AnalyzerAddress.AnalyzerName.IsNone())
-		{
-			if (!PassthroughAnalyzers.Contains(AnalyzerAddress.DataType))
-			{
-				return false;
-			}
-
-			AnalyzerAddress.AnalyzerName = PassthroughAnalyzers[AnalyzerAddress.DataType].AnalyzerName;
-			AnalyzerAddress.AnalyzerMemberName = PassthroughAnalyzers[AnalyzerAddress.DataType].OutputName;
-		}
-
-		// Check to see if the analyzer exists
-		{
-			using namespace Metasound::Frontend;
-			const IVertexAnalyzerFactory* Factory =
-				IVertexAnalyzerRegistry::Get().FindAnalyzerFactory(AnalyzerAddress.AnalyzerName);
-			
-			if (nullptr == Factory)
-			{
-				return false;
-			}
-		}
-
-		// Create the watcher
-		CreateOutputWatcher(AnalyzerAddress, OnOutputValueChanged);
-
-		// Update the generator's analyzers if necessary
-		FixUpOutputWatchers();
-
-		return true;
+		return WatchOutputInternal(OutputName, FWatchOutputUnifiedDelegate(OnOutputValueChanged), AnalyzerName, AnalyzerOutputName);
 	}
 
+	bool FMetasoundGeneratorHandle::WatchOutput(
+		const FName OutputName,
+		const FOnMetasoundOutputValueChangedNative& OnOutputValueChanged,
+		const FName AnalyzerName,
+		const FName AnalyzerOutputName)
+	{
+		return WatchOutputInternal(OutputName, FWatchOutputUnifiedDelegate(OnOutputValueChanged), AnalyzerName, AnalyzerOutputName);
+	}
+
 	void FMetasoundGeneratorHandle::UpdateOutputWatchers()
@@ -396,6 +330,90 @@
 		}
 	}
 
+	bool FMetasoundGeneratorHandle::WatchOutputInternal(
+		const FName OutputName,
+		const FWatchOutputUnifiedDelegate& OnOutputValueChanged,
+		const FName AnalyzerName,
+		const FName AnalyzerOutputName)
+	{
+		METASOUND_LLM_SCOPE;
+		METASOUND_TRACE_CPUPROFILER_EVENT_SCOPE(FMetasoundGeneratorHandle::WatchOutputInternal);
+
+		check(IsInGameThread());
+
+		if (!IsValid())
+		{
+			return false;
+		}
+
+		// Make the analyzer address.
+		Frontend::FAnalyzerAddress AnalyzerAddress;
+		AnalyzerAddress.InstanceID = GetAudioComponentId();
+		AnalyzerAddress.OutputName = OutputName;
+		AnalyzerAddress.AnalyzerName = AnalyzerName;
+		AnalyzerAddress.AnalyzerMemberName = AnalyzerOutputName;
+		AnalyzerAddress.AnalyzerInstanceID = FGuid::NewGuid();
+
+		// Find the output node and get the data type/node id from that
+		{
+			const TWeakObjectPtr<UMetaSoundSource> Source = GetMetaSoundSource();
+
+			if (!Source.IsValid())
+			{
+				UE_LOG(LogMetaSound, Warning, TEXT("Couldn't find the MetaSound Source"));
+				return false;
+			}
+
+			// Find the node id and type name
+			const FMetasoundFrontendClassOutput* OutputPtr =
+				Source->GetConstDocument().RootGraph.Interface.Outputs.FindByPredicate(
+					[&AnalyzerAddress](const FMetasoundFrontendClassOutput& Output)
+					{
+						return Output.Name == AnalyzerAddress.OutputName;
+					});
+
+			if (nullptr == OutputPtr)
+			{
+				return false;
+			}
+
+			AnalyzerAddress.NodeID = OutputPtr->NodeID;
+			AnalyzerAddress.DataType = OutputPtr->TypeName;
+		}
+
+		// If no analyzer name was provided, try to find a passthrough analyzer
+		if (AnalyzerAddress.AnalyzerName.IsNone())
+		{
+			if (!PassthroughAnalyzers.Contains(AnalyzerAddress.DataType))
+			{
+				return false;
+			}
+
+			AnalyzerAddress.AnalyzerName = PassthroughAnalyzers[AnalyzerAddress.DataType].AnalyzerName;
+			AnalyzerAddress.AnalyzerMemberName = PassthroughAnalyzers[AnalyzerAddress.DataType].OutputName;
+		}
+
+		// Check to see if the analyzer exists
+		{
+			using namespace Metasound::Frontend;
+			const IVertexAnalyzerFactory* Factory =
+				IVertexAnalyzerRegistry::Get().FindAnalyzerFactory(AnalyzerAddress.AnalyzerName);
+
+			if (nullptr == Factory)
+			{
+				return false;
+			}
+		}
+
+		// Create the watcher
+		CreateOutputWatcher(AnalyzerAddress, OnOutputValueChanged);
+
+		// Update the generator's analyzers if necessary
+		FixUpOutputWatchers();
+
+		return true;
+	}
+
 	void FMetasoundGeneratorHandle::FixUpOutputWatchers()
 	{
 		METASOUND_LLM_SCOPE;
@@ -421,7 +439,7 @@
 	
 	void FMetasoundGeneratorHandle::CreateOutputWatcher(
 		const Frontend::FAnalyzerAddress& AnalyzerAddress,
-		const FOnMetasoundOutputValueChanged& OnOutputValueChanged)
+		const FWatchOutputUnifiedDelegate& OnOutputValueChanged)
 	{
 		METASOUND_LLM_SCOPE;
 		METASOUND_TRACE_CPUPROFILER_EVENT_SCOPE(FMetasoundGeneratorHandle::CreateOutputWatcher);
@@ -438,7 +456,7 @@
 
 		if (FOutputWatcher* Watcher = OutputWatchers.Find(WatcherKey))
 		{
-			Watcher->OnOutputValueChanged.AddUnique(OnOutputValueChanged);
+			Watcher->OnOutputValueChanged.Add(OnOutputValueChanged);
 		}
 		// Otherwise add a new watcher
 		else
@@ -641,6 +659,23 @@
 	return GeneratorHandle->WatchOutput(OutputName, OnOutputValueChanged, AnalyzerName, AnalyzerOutputName);
 }
 
+bool UMetasoundGeneratorHandle::WatchOutput(
+	const FName OutputName,
+	const FOnMetasoundOutputValueChangedNative& OnOutputValueChanged,
+	const FName AnalyzerName,
+	const FName AnalyzerOutputName)
+{
+	METASOUND_LLM_SCOPE;
+	METASOUND_TRACE_CPUPROFILER_EVENT_SCOPE(UMetasoundGeneratorHandle::WatchOutput);
+
+	if (!IsValid())
+	{
+		return false;
+	}
+
+	return GeneratorHandle->WatchOutput(OutputName, OnOutputValueChanged, AnalyzerName, AnalyzerOutputName);
+}
+
 void UMetasoundGeneratorHandle::RegisterPassthroughAnalyzerForType(
 	const FName TypeName,
 	const FName AnalyzerName,
diff --git a/Engine/Plugins/Runtime/Metasound/Source/MetasoundEngine/Private/MetasoundOutputSubsystem.cpp b/Engine/Plugins/Runtime/Metasound/Source/MetasoundEngine/Private/MetasoundOutputSubsystem.cpp
--- a/Engine/Plugins/Runtime/Metasound/Source/MetasoundEngine/Private/MetasoundOutputSubsystem.cpp
+++ b/Engine/Plugins/Runtime/Metasound/Source/MetasoundEngine/Private/MetasoundOutputSubsystem.cpp
@@ -16,35 +16,28 @@ bool UMetaSoundOutputSubsystem::WatchOutput(
 	const FName AnalyzerName,
 	const FName AnalyzerOutputName)
 {
-	METASOUND_TRACE_CPUPROFILER_EVENT_SCOPE(UMetasoundOutputSubsystem::WatchOutput);
+	METASOUND_TRACE_CPUPROFILER_EVENT_SCOPE(UMetasoundOutputSubsystem::WatchOutput_Dynamic);
 
-	if (nullptr == AudioComponent)
+	UMetasoundGeneratorHandle* Handle = GetOrCreateGeneratorHandle(AudioComponent);
+
+	if (nullptr == Handle)
 	{
 		return false;
 	}
 
-	UMetasoundGeneratorHandle* Handle = nullptr;
+	return Handle->WatchOutput(OutputName, OnOutputValueChanged, AnalyzerName, AnalyzerOutputName);
+}
 
-	// Try to find an existing handle
-	const uint64 AudioComponentId = AudioComponent->GetAudioComponentID();
-	if (const TObjectPtr<UMetasoundGeneratorHandle>* FoundHandle = TrackedGenerators.FindByPredicate(
-		[AudioComponentId](const TObjectPtr<UMetasoundGeneratorHandle> ExistingHandle)
-		{
-			return ExistingHandle->IsValid() && ExistingHandle->GetAudioComponentId() == AudioComponentId;
-		}))
-	{
-		Handle = *FoundHandle;
-	}
-	// Create a new one
-	else
-	{
-		Handle = UMetasoundGeneratorHandle::CreateMetaSoundGeneratorHandle(AudioComponent);
-		
-		if(nullptr != Handle && Handle->IsValid())
-		{
-			TrackedGenerators.Add(Handle);
-		}
-	}
+bool UMetaSoundOutputSubsystem::WatchOutput(
+	UAudioComponent* AudioComponent,
+	const FName OutputName,
+	const FOnMetasoundOutputValueChangedNative& OnOutputValueChanged,
+	const FName AnalyzerName,
+	const FName AnalyzerOutputName)
+{
+	METASOUND_TRACE_CPUPROFILER_EVENT_SCOPE(UMetasoundOutputSubsystem::WatchOutput_Native);
+
+	UMetasoundGeneratorHandle* Handle = GetOrCreateGeneratorHandle(AudioComponent);
 
 	if (nullptr == Handle)
 	{
@@ -80,3 +73,36 @@ TStatId UMetaSoundOutputSubsystem::GetStatId() const
 {
 	RETURN_QUICK_DECLARE_CYCLE_STAT(UMetasoundGeneratorAccessSubsystem, STATGROUP_Tickables);
 }
+
+UMetasoundGeneratorHandle* UMetaSoundOutputSubsystem::GetOrCreateGeneratorHandle(UAudioComponent* AudioComponent)
+{
+	if (nullptr == AudioComponent)
+	{
+		return nullptr;
+	}
+
+	UMetasoundGeneratorHandle* Handle = nullptr;
+
+	// Try to find an existing handle
+	const uint64 AudioComponentId = AudioComponent->GetAudioComponentID();
+	if (const TObjectPtr<UMetasoundGeneratorHandle>* FoundHandle = TrackedGenerators.FindByPredicate(
+		[AudioComponentId](const TObjectPtr<UMetasoundGeneratorHandle> ExistingHandle)
+		{
+			return ExistingHandle->IsValid() && ExistingHandle->GetAudioComponentId() == AudioComponentId;
+		}))
+	{
+		Handle = *FoundHandle;
+	}
+	// Create a new one
+	else
+	{
+		Handle = UMetasoundGeneratorHandle::CreateMetaSoundGeneratorHandle(AudioComponent);
+
+		if (nullptr != Handle && Handle->IsValid())
+		{
+			TrackedGenerators.Add(Handle);
+		}
+	}
+
+	return Handle;
+}
\ No newline at end of file
diff --git a/Engine/Plugins/Runtime/Metasound/Source/MetasoundEngine/Public/MetasoundGeneratorHandle.h b/Engine/Plugins/Runtime/Metasound/Source/MetasoundEngine/Public/MetasoundGeneratorHandle.h
--- a/Engine/Plugins/Runtime/Metasound/Source/MetasoundEngine/Public/MetasoundGeneratorHandle.h
+++ b/Engine/Plugins/Runtime/Metasound/Source/MetasoundEngine/Public/MetasoundGeneratorHandle.h
@@ -12,6 +12,7 @@
 
 #include "MetasoundGeneratorHandle.generated.h"
 
+struct FMetaSoundOutput;
 class UAudioComponent;
 class UMetaSoundSource;
 
@@ -21,9 +22,11 @@ namespace Metasound
 }
 
 DECLARE_DYNAMIC_DELEGATE_TwoParams(FOnMetasoundOutputValueChanged, FName, OutputName, const FMetaSoundOutput&, Output);
-
+DECLARE_DELEGATE_TwoParams(FOnMetasoundOutputValueChangedNative, FName, const FMetaSoundOutput&);
 
 DECLARE_DYNAMIC_MULTICAST_DELEGATE_TwoParams(FOnMetasoundOutputValueChangedMulticast, FName, Name, const FMetaSoundOutput&, Output);
+DECLARE_MULTICAST_DELEGATE_TwoParams(FOnMetasoundOutputValueChangedNativeMulticast, FName, const FMetaSoundOutput&);
+
 
 namespace Metasound
 {
@@ -113,6 +116,12 @@ namespace Metasound
 			FName AnalyzerName = NAME_None,
 			FName AnalyzerOutputName = NAME_None);
 
+		bool WatchOutput(
+			FName OutputName,
+			const FOnMetasoundOutputValueChangedNative& OnOutputValueChanged,
+			FName AnalyzerName = NAME_None,
+			FName AnalyzerOutputName = NAME_None);
+
 		/**
 		 * Update any watched outputs
 		 */
@@ -152,11 +161,27 @@ namespace Metasound
 
 		void SendParametersToGenerator() const;
 		
+		struct FWatchOutputUnifiedDelegate
+		{
+			FOnMetasoundOutputValueChanged WatchDelegate;
+			FOnMetasoundOutputValueChangedNative NativeWatchDelegate;
+
+			FWatchOutputUnifiedDelegate() {}
+			FWatchOutputUnifiedDelegate(const FOnMetasoundOutputValueChanged& Delegate) : WatchDelegate(Delegate) {}
+			FWatchOutputUnifiedDelegate(const FOnMetasoundOutputValueChangedNative& Delegate) : NativeWatchDelegate(Delegate) {}
+		};
+
+		bool WatchOutputInternal(
+			FName OutputName,
+			const FWatchOutputUnifiedDelegate& OnOutputValueChanged,
+			FName AnalyzerName = NAME_None,
+			FName AnalyzerOutputName = NAME_None);
+
 		void FixUpOutputWatchers();
 
 		void CreateOutputWatcher(
 			const Frontend::FAnalyzerAddress& AnalyzerAddress,
-			const FOnMetasoundOutputValueChanged& OnOutputValueChanged);
+			const FWatchOutputUnifiedDelegate& OnOutputValueChanged);
 		
 		TWeakObjectPtr<UAudioComponent> AudioComponent;
 		const uint64 AudioComponentId;
@@ -220,20 +245,45 @@ namespace Metasound
 				GetTypeHash(Key.AnalyzerMemberName));
 		}
 
+		struct FWatchOutputUnifiedMulticastDelegate
+		{
+			FOnMetasoundOutputValueChangedMulticast WatchDelegates;
+			FOnMetasoundOutputValueChangedNativeMulticast NativeWatchDelegates;
+
+			void Add(const FWatchOutputUnifiedDelegate& Delegate)
+			{
+				if (Delegate.WatchDelegate.IsBound())
+				{
+					WatchDelegates.AddUnique(Delegate.WatchDelegate);
+				}
+
+				if (Delegate.NativeWatchDelegate.IsBound())
+				{
+					NativeWatchDelegates.Add(Delegate.NativeWatchDelegate);
+				}
+			}
+
+			void Broadcast(FName OutputName, const FMetaSoundOutput& Output) const
+			{
+				WatchDelegates.Broadcast(OutputName, Output);
+				NativeWatchDelegates.Broadcast(OutputName, Output);
+			}
+		};
+
 		/**
 		 * Info about an output being watched by one or more listeners
 		 */
 		struct FOutputWatcher
 		{
 			Frontend::FAnalyzerAddress AnalyzerAddress;
-			FOnMetasoundOutputValueChangedMulticast OnOutputValueChanged;
+			FWatchOutputUnifiedMulticastDelegate OnOutputValueChanged;
 
 			FOutputWatcher(
 				const Frontend::FAnalyzerAddress& InAnalyzerAddress,
-				const FOnMetasoundOutputValueChanged& InOnOutputValueChanged)
+				const FWatchOutputUnifiedDelegate& InOnOutputValueChanged)
 					: AnalyzerAddress(InAnalyzerAddress)
 			{
-				OnOutputValueChanged.AddUnique(InOnOutputValueChanged);
+				OnOutputValueChanged.Add(InOnOutputValueChanged);
 			}
 		};
 		
@@ -345,6 +395,12 @@ public:
 		FName AnalyzerName = NAME_None,
 		FName AnalyzerOutputName = NAME_None);
 
+	bool WatchOutput(
+		FName OutputName,
+		const FOnMetasoundOutputValueChangedNative& OnOutputValueChanged,
+		FName AnalyzerName = NAME_None,
+		FName AnalyzerOutputName = NAME_None);
+
 	/**
 	 * Map a type name to a passthrough analyzer name to use as a default for UMetasoundOutputSubsystem::WatchOutput()
 	 *
diff --git a/Engine/Plugins/Runtime/Metasound/Source/MetasoundEngine/Public/MetasoundOutputSubsystem.h b/Engine/Plugins/Runtime/Metasound/Source/MetasoundEngine/Public/MetasoundOutputSubsystem.h
--- a/Engine/Plugins/Runtime/Metasound/Source/MetasoundEngine/Public/MetasoundOutputSubsystem.h
+++ b/Engine/Plugins/Runtime/Metasound/Source/MetasoundEngine/Public/MetasoundOutputSubsystem.h
@@ -8,6 +8,7 @@
 #include "MetasoundOutputSubsystem.generated.h"
 
 class UAudioComponent;
+class UMetasoundGeneratorHandle;
 
 /**
  * Provides access to a playing Metasound generator's outputs
@@ -37,6 +38,13 @@ public:
 		const FOnMetasoundOutputValueChanged& OnOutputValueChanged,
 		FName AnalyzerName = NAME_None,
 		FName AnalyzerOutputName = NAME_None);
+
+	bool WatchOutput(
+		UAudioComponent* AudioComponent,
+		FName OutputName,
+		const FOnMetasoundOutputValueChangedNative& OnOutputValueChanged,
+		FName AnalyzerName = NAME_None,
+		FName AnalyzerOutputName = NAME_None);
 	
 	/** Begin UTickableWorldSubsystem */
 	virtual bool IsTickable() const override;
@@ -45,6 +53,8 @@ public:
 	/** End UTickableWorldSubsystem */
 
 private:
+	UMetasoundGeneratorHandle* GetOrCreateGeneratorHandle(UAudioComponent* AudioComponent);
+
 	UPROPERTY(Transient)
 	TArray<TObjectPtr<UMetasoundGeneratorHandle>> TrackedGenerators;
 };
