WaveTable Bank Evaluate Node

diff --git a/Engine/Plugins/Runtime/Metasound/Source/MetasoundEngine/Private/MetasoundWaveTableBankEvaluateNode.cpp b/Engine/Plugins/Runtime/Metasound/Source/MetasoundEngine/Private/MetasoundWaveTableBankEvaluateNode.cpp
new file mode 100644
--- /dev/null
+++ b/Engine/Plugins/Runtime/Metasound/Source/MetasoundEngine/Private/MetasoundWaveTableBankEvaluateNode.cpp
@@ -0,0 +1,270 @@
+// Copyright Epic Games, Inc. All Rights Reserved.
+
+#include "Interfaces/MetasoundFrontendSourceInterface.h"
+#include "MetasoundDataFactory.h"
+#include "MetasoundEngineNodesNames.h"
+#include "MetasoundEnumRegistrationMacro.h"
+#include "MetasoundExecutableOperator.h"
+#include "MetasoundFacade.h"
+#include "MetasoundNodeRegistrationMacro.h"
+#include "MetasoundPrimitives.h"
+#include "MetasoundParamHelper.h"
+#include "MetasoundStandardNodesCategories.h"
+#include "MetasoundWaveTable.h"
+#include "WaveTable.h"
+#include "WaveTableSampler.h"
+
+#define LOCTEXT_NAMESPACE "MetasoundStandardNodes"
+
+namespace Metasound
+{
+	namespace WaveTableBankEvaluateNode
+	{
+		METASOUND_PARAM(WaveTableBankParam, "WaveTableBank", "The WaveTableBank to evaluate");
+		METASOUND_PARAM(InputParam, "Input", "[0, 1] the X input with which to evaluate the wavetable");
+		METASOUND_PARAM(IndexParam, "Index", "The table index to interpolate and evaluate the result of (wraps over number of entries)");
+		METASOUND_PARAM(OutParam, "Output", "The linearly mixed value of the provided WaveTableBank's applicable entries");
+	} // WaveTableBankEvaluateNode
+	
+	class FMetasoundWaveTableBankEvaluateNodeOperator : public TExecutableOperator<FMetasoundWaveTableBankEvaluateNodeOperator>
+	{
+	public:
+		static const FVertexInterface& GetDefaultInterface()
+		{
+			using namespace WaveTable;
+			using namespace WaveTableBankEvaluateNode;
+
+			static const FVertexInterface DefaultInterface(
+				FInputVertexInterface(
+					TInputDataVertex<FWaveTableBankAsset>(METASOUND_GET_PARAM_NAME_AND_METADATA(WaveTableBankParam)),
+					TInputDataVertex<float>(METASOUND_GET_PARAM_NAME_AND_METADATA(InputParam), 0.0f),
+					TInputDataVertex<float>(METASOUND_GET_PARAM_NAME_AND_METADATA(IndexParam), 0.0f),
+					TInputDataVertex<FEnumWaveTableInterpolationMode>("Interpolation", FDataVertexMetadata
+					{
+						LOCTEXT("MetasoundWaveTableBankEvaluateNode_InterpDescription", "How interpolation occurs between WaveTable values."),
+						LOCTEXT("MetasoundWaveTableBankEvaluateNode_Interp", "Interpolation"),
+						true /* bIsAdvancedDisplay */
+					}, static_cast<int32>(FWaveTableSampler::EInterpolationMode::Linear))
+				),
+				FOutputVertexInterface(
+					TOutputDataVertex<float>(METASOUND_GET_PARAM_NAME_AND_METADATA(OutParam))
+				)
+			);
+
+			return DefaultInterface;
+		}
+
+		static const FNodeClassMetadata& GetNodeInfo()
+		{
+			auto CreateNodeClassMetadata = []() -> FNodeClassMetadata
+			{
+				FNodeClassMetadata Metadata
+				{
+					{ EngineNodes::Namespace, "WaveTableBankEvaluate", "" },
+					1, // Major Version
+					0, // Minor Version
+					LOCTEXT("MetasoundWaveTableBankEvaluateNode_Name", "Evaluate WaveTableBank"),
+					LOCTEXT("MetasoundWaveTableBankEvaluateNode_Description",
+						"Evaluates a WaveTableBank's given entires for a given input value, linearly interpolating inline between using the provided index. More performant "
+						"than using 'WaveTableGet' and using the 'WaveTableEvaluate' nodes as no resulting WaveTable is generated any time the input float index is changed."),
+					PluginAuthor,
+					PluginNodeMissingPrompt,
+					GetDefaultInterface(),
+					{ NodeCategories::WaveTables },
+					{ NodeCategories::Envelopes, METASOUND_LOCTEXT("WaveTableEvaluateCurveKeyword", "Curve") },
+					{ }
+				};
+
+				return Metadata;
+			};
+
+			static const FNodeClassMetadata Metadata = CreateNodeClassMetadata();
+			return Metadata;
+		}
+
+		static TUniquePtr<IOperator> CreateOperator(const FBuildOperatorParams& InParams, FBuildResults& OutResults)
+		{
+			using namespace WaveTable;
+
+			const FInputVertexInterfaceData& InputData = InParams.InputData;
+
+			FWaveTableBankAssetReadRef InWaveTableReadRef = InputData.GetOrConstructDataReadReference<FWaveTableBankAsset>("WaveTableBank");
+			FFloatReadRef InInputReadRef = InputData.GetOrCreateDefaultDataReadReference<float>("Input", InParams.OperatorSettings);
+			FFloatReadRef InIndexReadRef = InputData.GetOrCreateDefaultDataReadReference<float>("Index", InParams.OperatorSettings);
+			FEnumWaveTableInterpModeReadRef InInterpReadRef = InputData.GetOrCreateDefaultDataReadReference<FEnumWaveTableInterpolationMode>("Interpolation", InParams.OperatorSettings);
+
+			return MakeUnique<FMetasoundWaveTableBankEvaluateNodeOperator>(InParams, InWaveTableReadRef, InInputReadRef, InIndexReadRef, InInterpReadRef);
+		}
+
+		FMetasoundWaveTableBankEvaluateNodeOperator(
+			const FBuildOperatorParams& InParams,
+			const FWaveTableBankAssetReadRef& InWaveTableBankReadRef,
+			const FFloatReadRef& InInputReadRef,
+			const FFloatReadRef& InIndexReadRef,
+			const FEnumWaveTableInterpModeReadRef& InInterpModeReadRef)
+			: WaveTableBankReadRef(InWaveTableBankReadRef)
+			, InputReadRef(InInputReadRef)
+			, IndexReadRef(InIndexReadRef)
+			, InterpModeReadRef(InInterpModeReadRef)
+			, OutWriteRef(TDataWriteReferenceFactory<float>::CreateAny(InParams.OperatorSettings))
+		{
+			Reset(InParams);
+		}
+
+		virtual ~FMetasoundWaveTableBankEvaluateNodeOperator() = default;
+
+		virtual void BindInputs(FInputVertexInterfaceData& InOutVertexData) override
+		{
+			using namespace WaveTableBankEvaluateNode;
+			InOutVertexData.BindReadVertex(METASOUND_GET_PARAM_NAME(WaveTableBankParam), WaveTableBankReadRef);
+			InOutVertexData.BindReadVertex(METASOUND_GET_PARAM_NAME(InputParam), InputReadRef);
+			InOutVertexData.BindReadVertex(METASOUND_GET_PARAM_NAME(IndexParam), IndexReadRef);
+			InOutVertexData.BindReadVertex("Interpolation", InterpModeReadRef);
+		}
+
+		virtual void BindOutputs(FOutputVertexInterfaceData& InOutVertexData) override
+		{
+			using namespace WaveTableBankEvaluateNode;
+			InOutVertexData.BindReadVertex(METASOUND_GET_PARAM_NAME(OutParam), OutWriteRef);
+		}
+		
+		virtual FDataReferenceCollection GetInputs() const override
+		{
+			// This should never be called. Bind(...) is called instead. This method
+			// exists as a stop-gap until the API can be deprecated and removed. 
+			checkNoEntry();
+
+			FDataReferenceCollection InputDataReferences;
+			return InputDataReferences;
+		}
+
+		virtual FDataReferenceCollection GetOutputs() const override
+		{
+			// This should never be called. Bind(...) is called instead. This method
+			// exists as a stop-gap until the API can be deprecated and removed. 
+			checkNoEntry();
+
+			FDataReferenceCollection OutputDataReferences;
+			return OutputDataReferences;
+		}
+
+		void Execute()
+		{
+			using namespace WaveTable;
+
+			const FWaveTableBankAsset& WaveTableBankAsset = *WaveTableBankReadRef;
+			FWaveTableBankAssetProxyPtr Proxy = WaveTableBankAsset.GetProxy();
+			float NewIndex = 0.0f;
+			if (!ResolveNextComputeIndex(Proxy, NewIndex))
+			{
+				return;
+			}
+
+			const TArray<FWaveTableData>& WaveTables = WaveTableBankAsset->GetWaveTableData();
+			const int32 IndexFloor = FMath::FloorToInt32(NewIndex) % WaveTables.Num();
+			const int32 IndexCeil = FMath::CeilToInt32(NewIndex) % WaveTables.Num();
+
+			const FWaveTableData& IndexFloorTable = WaveTables[IndexFloor];
+
+			float IndexFloorValue = 0.0f;
+			constexpr FWaveTableSampler::ESingleSampleMode SampleMode =  FWaveTableSampler::ESingleSampleMode::Hold;
+			const float Input = FMath::Clamp(*InputReadRef, 0.0f, 1.0f);
+
+			Sampler.Reset();
+			Sampler.SetPhase(Input);
+			Sampler.Process(IndexFloorTable, IndexFloorValue, SampleMode);
+
+			if (IndexFloor != IndexCeil)
+			{
+				const FWaveTableData& IndexCeilTable = WaveTables[IndexCeil];
+				float IndexCeilValue = 0.0f;
+				Sampler.Reset();
+				Sampler.SetPhase(Input);
+				Sampler.Process(IndexCeilTable, IndexCeilValue, SampleMode);
+
+				const float Fractional = FMath::Frac(NewIndex);
+				LastValue = (IndexFloorValue * (1.0f - Fractional)) + (IndexCeilValue * Fractional);
+			}
+			else
+			{
+				LastValue = IndexFloorValue;
+			}
+
+			*OutWriteRef = LastValue;
+		}
+
+		void Reset(const IOperator::FResetParams& InParams)
+		{
+			using namespace WaveTable;
+			FWaveTableSampler::FSettings Settings;
+			Settings.Freq = 0.0f; // Sampler phase is manually progressed via this node
+			Sampler = FWaveTableSampler(MoveTemp(Settings));
+
+			*OutWriteRef = 0.f;
+
+			LastInterpolationMode = WaveTable::FWaveTableSampler::EInterpolationMode::COUNT;
+			LastIndex = -1.0f;
+			LastValue = 0.0f;
+		}
+
+	private:
+		// Returns true if new computation is required, setting OutIndex to index to compute.
+		// Returns false if new computation isn't required, resetting output & cached data accordingly.
+		bool ResolveNextComputeIndex(const FWaveTableBankAssetProxyPtr& Proxy, float& OutIndex)
+		{
+			using namespace WaveTable;
+
+			if (!Proxy.IsValid())
+			{
+				LastValue = 0.0f;
+				*OutWriteRef = LastValue;
+				return false;
+			}
+
+			const float Index = *IndexReadRef;
+			const FWaveTableSampler::EInterpolationMode NewInterpolationMode = *InterpModeReadRef;
+			if (LastInterpolationMode == NewInterpolationMode)
+			{
+				if (FMath::IsNearlyEqual(Index, LastIndex))
+				{
+					*OutWriteRef = LastValue;
+					return false;
+				}
+			}
+
+			LastInterpolationMode = NewInterpolationMode;
+			LastIndex = Index;
+			Sampler.SetInterpolationMode(NewInterpolationMode);
+			OutIndex = FMath::Abs(Index); // Avoids fractional, interpolative flip at zero crossing
+			return true;
+		}
+
+		FWaveTableBankAssetReadRef WaveTableBankReadRef;
+		FFloatReadRef InputReadRef;
+		FFloatReadRef IndexReadRef;
+		FEnumWaveTableInterpModeReadRef InterpModeReadRef;
+
+		WaveTable::FWaveTableSampler Sampler;
+
+		FFloatWriteRef OutWriteRef;
+
+		WaveTable::FWaveTableSampler::EInterpolationMode LastInterpolationMode = WaveTable::FWaveTableSampler::EInterpolationMode::COUNT;
+		float LastIndex = -1.0f;
+		float LastValue = 0.0f;
+	};
+
+	class FMetasoundWaveTableBankEvaluateNode : public FNodeFacade
+	{
+	public:
+		FMetasoundWaveTableBankEvaluateNode(const FNodeInitData& InInitData)
+			: FNodeFacade(InInitData.InstanceName, InInitData.InstanceID, TFacadeOperatorClass<FMetasoundWaveTableBankEvaluateNodeOperator>())
+		{
+		}
+
+		virtual ~FMetasoundWaveTableBankEvaluateNode() = default;
+	};
+
+	METASOUND_REGISTER_NODE(FMetasoundWaveTableBankEvaluateNode)
+} // namespace Metasound
+
+#undef LOCTEXT_NAMESPACE // MetasoundStandardNodes
diff --git a/Engine/Plugins/Runtime/Metasound/Source/MetasoundEngine/Private/MetasoundWaveTableEvaluateNode.cpp b/Engine/Plugins/Runtime/Metasound/Source/MetasoundEngine/Private/MetasoundWaveTableEvaluateNode.cpp
--- a/Engine/Plugins/Runtime/Metasound/Source/MetasoundEngine/Private/MetasoundWaveTableEvaluateNode.cpp
+++ b/Engine/Plugins/Runtime/Metasound/Source/MetasoundEngine/Private/MetasoundWaveTableEvaluateNode.cpp
@@ -105,15 +105,6 @@ namespace Metasound
 
 		virtual ~FMetasoundWaveTableEvaluateNodeOperator() = default;
 
-		virtual void Bind(FVertexInterfaceData& InVertexData) const override
-		{
-			using namespace WaveTableEvaluateNode;
-			
-			FInputVertexInterfaceData& Inputs = InVertexData.GetInputs();
-
-			FOutputVertexInterfaceData& Outputs = InVertexData.GetOutputs();
-		}
-
 		virtual void BindInputs(FInputVertexInterfaceData& InOutVertexData) override
 		{
 			using namespace WaveTableEvaluateNode;
diff --git a/Engine/Plugins/Runtime/Metasound/Source/MetasoundEngine/Private/MetasoundWaveTableGetNode.cpp b/Engine/Plugins/Runtime/Metasound/Source/MetasoundEngine/Private/MetasoundWaveTableGetNode.cpp
--- a/Engine/Plugins/Runtime/Metasound/Source/MetasoundEngine/Private/MetasoundWaveTableGetNode.cpp
+++ b/Engine/Plugins/Runtime/Metasound/Source/MetasoundEngine/Private/MetasoundWaveTableGetNode.cpp
@@ -124,14 +124,6 @@ namespace Metasound
 		}
 
 	private:
-		FORCEINLINE static void WrapIndex(int32 InMax, float& InOutIndex)
-		{
-			InOutIndex = FMath::Abs(InOutIndex); // Avoids remainder offset flip at 0 crossing
-			const int32 WrapIndex = FMath::TruncToInt32(InOutIndex) % InMax;
-			const float Remainder = FMath::Frac(InOutIndex);
-			InOutIndex = WrapIndex + Remainder;
-		};
-
 		void ComputeGainIndexData(const TArray<FWaveTableData>& WaveTables, EWaveTableSamplingMode InSampleMode, ::WaveTable::FWaveTable& OutputWaveTable)
 		{
 			checkf(!WaveTables.IsEmpty(), TEXT("ComputGainIndexData must have WaveTable to operate on"));
@@ -210,7 +202,7 @@ namespace Metasound
 			}
 
 			GainIndexData.TableIndex = *TableIndexReadRef;
-			WrapIndex(WaveTables.Num(), GainIndexData.TableIndex);
+			FWaveTable::WrapIndexSmooth(WaveTables.Num(), GainIndexData.TableIndex);
 
 			const bool bIsLastIndex = FMath::IsNearlyEqual(GainIndexData.LastTableIndex, GainIndexData.TableIndex);
 			const bool bIsLastProxy = LastTableId == Proxy->GetObjectId();
diff --git a/Engine/Plugins/Runtime/WaveTable/Source/WaveTable/Private/WaveTable.cpp b/Engine/Plugins/Runtime/WaveTable/Source/WaveTable/Private/WaveTable.cpp
--- a/Engine/Plugins/Runtime/WaveTable/Source/WaveTable/Private/WaveTable.cpp
+++ b/Engine/Plugins/Runtime/WaveTable/Source/WaveTable/Private/WaveTable.cpp
@@ -406,41 +406,48 @@ bool FWaveTableData::GetDataView(TArrayView<int16>& OutData)
 	return false;
 }
 
-float FWaveTableData::GetFinalValue() const
+float FWaveTableData::GetLastValue() const
 {
-	if (FMath::IsNearlyEqual(FinalValue, WaveTable::InvalidWaveTableValue))
+	if (Data.IsEmpty())
+	{
+		return 0.0f;
+	}
+
+	switch (BitDepth)
 	{
-		if (Data.IsEmpty())
+		case EWaveTableBitDepth::IEEE_Float:
 		{
-			return 0.0f;
+			TArrayView<const float> TableView;
+			ensureAlways(GetDataView(TableView));
+			return TableView.Last();
 		}
+		break;
 
-		switch (BitDepth)
+		case EWaveTableBitDepth::PCM_16:
 		{
-			case EWaveTableBitDepth::IEEE_Float:
-			{
-				TArrayView<const float> TableView;
-				ensureAlways(GetDataView(TableView));
-				return TableView.Last();
-			}
-			break;
-
-			case EWaveTableBitDepth::PCM_16:
-			{
-				TArrayView<const int16> TableView;
-				ensureAlways(GetDataView(TableView));
-				constexpr float ConversionValue = 1.0f / static_cast<float>(TNumericLimits<int16>::Max());
-				return TableView.Last() * ConversionValue;
-			}
-			break;
+			TArrayView<const int16> TableView;
+			ensureAlways(GetDataView(TableView));
+			constexpr float ConversionValue = 1.0f / static_cast<float>(TNumericLimits<int16>::Max());
+			return TableView.Last() * ConversionValue;
+		}
+		break;
 
-			default:
-			{
-				static_assert(static_cast<int32>(EWaveTableBitDepth::COUNT) == 2, "Possible missing switch case coverage for 'EWaveTableBitDepth'");
-				checkNoEntry();
-			}
-			break;
+		default:
+		{
+			static_assert(static_cast<int32>(EWaveTableBitDepth::COUNT) == 2, "Possible missing switch case coverage for 'EWaveTableBitDepth'");
+			checkNoEntry();
 		}
+		break;
+	}
+
+	return 0.0f;
+}
+
+float FWaveTableData::GetFinalValue() const
+{
+	if (FMath::IsNearlyEqual(FinalValue, WaveTable::InvalidWaveTableValue))
+	{
+		return GetLastValue();
 	}
 
 	return FinalValue;
diff --git a/Engine/Plugins/Runtime/WaveTable/Source/WaveTable/Private/WaveTableSampler.cpp b/Engine/Plugins/Runtime/WaveTable/Source/WaveTable/Private/WaveTableSampler.cpp
--- a/Engine/Plugins/Runtime/WaveTable/Source/WaveTable/Private/WaveTableSampler.cpp
+++ b/Engine/Plugins/Runtime/WaveTable/Source/WaveTable/Private/WaveTableSampler.cpp
@@ -267,75 +267,10 @@
 		{
 			return 0.f;
 		}
-		
+
 		TArrayView<float> OutSamplesView { &OutSample, 1 };
-		float Index = Process(InTableView, { }, { }, { }, OutSamplesView);
-
-		// Sampler functions by default cyclically and interpolates against initial value so
-		// early out as following index check and potential interp re-evaluation isn't necessary.
-		if (InMode == ESingleSampleMode::Loop)
-		{
-			return Index;
-		}
-
-		// If interpolating between last value and next, check
-		// to see if interpolation needs to be re-evaluated.
-		// This check is expensive in the context of buffer
-		// eval and thus only supported in "single-sample" mode.
-		const bool bIsFinalInterp = Index > InTableView.Num() - 1;
-		const bool bHasCompleted = Index == 0 && (Settings.Phase > 0.0f || LastIndex > 0);
-		if (bIsFinalInterp || bHasCompleted)
-		{
-			float ModeValue = 0.f;
-			
-			switch (InMode)
-			{
-				case ESingleSampleMode::Zero:
-				case ESingleSampleMode::Unit:
-				{
-					ModeValue = static_cast<float>(InMode);
-				}
-				break;
-
-				case ESingleSampleMode::Hold:
-				{
-					ModeValue = InTableView.FinalValue;
-				}
-				break;
-				
-				case ESingleSampleMode::Loop:
-				{
-					checkNoEntry(); // Should be handled above
-				}
-				break;
-
-				default:
-				break;
-			}
-
-			if (bIsFinalInterp)
-			{
-				TArray<float> FinalInterp = { InTableView.SampleView.Last(), ModeValue };
-				OutSample = FMath::Frac(Index);
-				Interpolate(FinalInterp, OutSamplesView, Settings.InterpolationMode);
-			}
-			else
-			{
-				OutSample = ModeValue;
-			}
-
-			if (!FMath::IsNearlyEqual(Settings.Amplitude, 1.0f))
-			{
-				OutSample *= Settings.Amplitude;
-			}
-
-			if (!FMath::IsNearlyZero(Settings.Offset))
-			{
-				OutSample += Settings.Offset;
-			}
-		}
-
-		return Index;
+		const float Index = Process(InTableView, OutSamplesView);
+		return FinalizeSingleSample(Index, InTableView.Num(), OutSamplesView, InTableView.SampleView.Last(), InTableView.FinalValue, InMode);
 	}
 
 	float FWaveTableSampler::Process(const FWaveTableView& InTableView, TArrayView<float> OutSamplesView)
@@ -381,6 +316,19 @@
 		return CurViewLastIndex;
 	}
 
+	float FWaveTableSampler::Process(const FWaveTableData& InTableData, float& OutSample, ESingleSampleMode InMode)
+	{
+		if (InTableData.IsEmpty())
+		{
+			return 0.f;
+		}
+
+		TArrayView<float> OutSamplesView { &OutSample, 1 };
+		const float Index = Process(InTableData, OutSamplesView);
+		return FinalizeSingleSample(Index, InTableData.GetNumSamples(), OutSamplesView, InTableData.GetLastValue(), InTableData.GetFinalValue(), InMode);
+
+	}
+
 	float FWaveTableSampler::Process(const FWaveTableData& InTableData, TArrayView<float> OutSamplesView)
 	{
 		return Process(InTableData, { }, { }, { }, OutSamplesView);
@@ -582,6 +530,75 @@
 		return OutIndicesView.Last();
 	}
 
+	float FWaveTableSampler::FinalizeSingleSample(float Index, int32 NumSamples, TArrayView<float> OutSampleView, float LastTableValue, float FinalValue, ESingleSampleMode InMode)
+	{
+		// Sampler functions by default cyclically and interpolates against initial value so
+		// early out as following index check and potential interp re-evaluation isn't necessary.
+		if (InMode == FWaveTableSampler::ESingleSampleMode::Loop)
+		{
+			return Index;
+		}
+
+		// If interpolating between last value and next, check
+		// to see if interpolation needs to be re-evaluated.
+		// This check is expensive in the context of buffer
+		// eval and thus only supported in "single-sample" mode.
+		const bool bIsFinalInterp = Index > NumSamples - 1;
+		const bool bHasCompleted = Index == 0 && (Settings.Phase > 0.0f || LastIndex > 0);
+		if (bIsFinalInterp || bHasCompleted)
+		{
+			float ModeValue = 0.f;
+
+			switch (InMode)
+			{
+				case ESingleSampleMode::Zero:
+				case ESingleSampleMode::Unit:
+				{
+					ModeValue = static_cast<float>(InMode);
+				}
+				break;
+
+				case ESingleSampleMode::Hold:
+				{
+					ModeValue = FinalValue;
+				}
+				break;
+
+				case ESingleSampleMode::Loop:
+				{
+					checkNoEntry(); // Should be handled above
+				}
+				break;
+
+				default:
+					break;
+			}
+
+			if (bIsFinalInterp)
+			{
+				TArray<float> FinalInterp = { LastTableValue, ModeValue };
+				OutSampleView.Last() = FMath::Frac(Index);
+				Interpolate(FinalInterp, OutSampleView, Settings.InterpolationMode);
+			}
+			else
+			{
+				OutSampleView.Last() = ModeValue;
+			}
+
+			if (!FMath::IsNearlyEqual(Settings.Amplitude, 1.0f))
+			{
+				OutSampleView.Last() *= Settings.Amplitude;
+			}
+
+			if (!FMath::IsNearlyZero(Settings.Offset))
+			{
+				OutSampleView.Last() += Settings.Offset;
+			}
+		}
+
+		return Index;
+	}
+
 	int32 FWaveTableSampler::GetIndexFinished() const
 	{
 		return OneShotData.IndexFinished;
diff --git a/Engine/Plugins/Runtime/WaveTable/Source/WaveTable/Public/WaveTable.h b/Engine/Plugins/Runtime/WaveTable/Source/WaveTable/Public/WaveTable.h
--- a/Engine/Plugins/Runtime/WaveTable/Source/WaveTable/Public/WaveTable.h
+++ b/Engine/Plugins/Runtime/WaveTable/Source/WaveTable/Public/WaveTable.h
@@ -51,6 +51,14 @@ namespace WaveTable
 		void SetFinalValue(const float InValue);
 		void SetNum(int32 InSize);
 		void Zero();
+
+		FORCEINLINE static void WrapIndexSmooth(int32 InMax, float& InOutIndex)
+		{
+			InOutIndex = FMath::Abs(InOutIndex); // Avoids fractional offset flip at 0 crossing
+			const int32 WrapIndex = FMath::TruncToInt32(InOutIndex) % InMax;
+			const float Fractional = FMath::Frac(InOutIndex);
+			InOutIndex = WrapIndex + Fractional;
+		}
 	};
 } // namespace WaveTable
 
@@ -97,6 +105,7 @@ private:
 	UPROPERTY()
 	float FinalValue = ::WaveTable::InvalidWaveTableValue;
 
+	// Returns the size of the underlying data's sample in number of bytes.
 	FORCEINLINE int32 GetSampleSize() const;
 
 public:
@@ -105,20 +114,43 @@ public:
 	// Provided buffer's size must match this WaveTableData's samples count.
 	void ArrayMixIn(TArrayView<float> OutputWaveSamples, float Gain = 1.0f) const;
 
+	// Empties the underlying data container, deallocating memory and invalidating the FinalValue.
 	void Empty();
 
 	EWaveTableBitDepth GetBitDepth() const;
+
+	// Returns read-only view of byte array that stores sample data.
 	const TArray<uint8>& GetRawData() const;
+
+	// Returns true if ArrayView of underlying Data in the proper sample format,
+	// setting OutData to a view of the Table's data. Returns false if the bit
+	// depth of the provided ArrayView does not match the Table's BitDepth property.
 	bool GetDataView(TArrayView<const float>& OutData) const;
 	bool GetDataView(TArrayView<float>& OutData);
 	bool GetDataView(TArrayView<const int16>& OutData) const;
 	bool GetDataView(TArrayView<int16>& OutData);
+
+	// Returns number of samples within the underlying DataView, not including a FinalValue if set.
+	// (Should not to be confused with the number of bytes used to represent the given DataView in
+	// the internal Data container).
 	int32 GetNumSamples() const;
+
+	// Returns the last data value value as a float.  If DataView is empty, returns 0.
+	float GetLastValue() const;
+
+	// Returns the value to interpolate to beyond the last value of the DataView.
+	// If FinalValue property is set to invalid (default), returns the last value in the 
+	// underlying DataView (see GetLastValue).
 	float GetFinalValue() const;
 
+	// Whether or not the underlying data is empty.
 	bool IsEmpty() const;
 
+	// Resamples the underlying data to a provided new sample rate.
 	bool Resample(int32 InCurrentSampleRate, int32 InNewSampleRate, ::WaveTable::FWaveTableSampler::EInterpolationMode InInterpMode = ::WaveTable::FWaveTableSampler::EInterpolationMode::Cubic);
+
+	// Resets the underlying data to the given number of samples. Does not change
+	// memory allocation of data unless NumSamples is larger than the current size.
 	void Reset(int32 NumSamples);
 
 	// Converts bit depth and sets internal data to the requested bit depth.
@@ -131,5 +163,7 @@ public:
 	void SetRawData(EWaveTableBitDepth InBitDepth, TArray<uint8>&& InData);
 	void SetFinalValue(float InFinalValue);
 
+	// Zeros the underlying data. If InNumSamples is set (optional),
+	// allocates space in underlying data array for the given number of samples.
 	void Zero(int32 InNumSamples = INDEX_NONE);
 };
diff --git a/Engine/Plugins/Runtime/WaveTable/Source/WaveTable/Public/WaveTableSampler.h b/Engine/Plugins/Runtime/WaveTable/Source/WaveTable/Public/WaveTableSampler.h
--- a/Engine/Plugins/Runtime/WaveTable/Source/WaveTable/Public/WaveTableSampler.h
+++ b/Engine/Plugins/Runtime/WaveTable/Source/WaveTable/Public/WaveTableSampler.h
@@ -74,6 +74,7 @@ namespace WaveTable
 
 		// Retrieves the sample value at the currently set phase, returning the floating point index which the current phase corresponds to.
 		float Process(const FWaveTableView& InTableView, float& OutSample, ESingleSampleMode InMode = ESingleSampleMode::Zero);
+		float Process(const FWaveTableData& InTableData, float& OutSample, ESingleSampleMode InMode = ESingleSampleMode::Zero);
 
 		// Resamples entire table length into given view
 		float Process(const FWaveTableView& InTableView, TArrayView<float> OutSamplesView);
@@ -104,6 +105,9 @@ namespace WaveTable
 		// Returns last positively increasing index if finished. If not finished, returns the last index value stored in the view.
 		int32 ComputeIndexFinished(TArrayView<const float> InSyncTriggers, TArrayView<float> OutIndicesView);
 
+		// Utility that takes in necessary data for finalizing a single sample index processing from either a TableData or a TableView.  Covers edge cases for final value.
+		float FinalizeSingleSample(float Index, int32 NumSamples, TArrayView<float> OutSample, float LastTableValue, float FinalValue, FWaveTableSampler::ESingleSampleMode InMode);
+
 		float LastIndex = 0.0f;
 
 		struct FOneShotData
