Move sound InitResources call in MovieSceneAudioSection to PostLoad and prevent registration of invalid MetaSound registry keys

diff --git a/Engine/Plugins/Runtime/Metasound/Source/MetasoundFrontend/Private/MetasoundFrontendRegistries.cpp b/Engine/Plugins/Runtime/Metasound/Source/MetasoundFrontend/Private/MetasoundFrontendRegistries.cpp
--- a/Engine/Plugins/Runtime/Metasound/Source/MetasoundFrontend/Private/MetasoundFrontendRegistries.cpp
+++ b/Engine/Plugins/Runtime/Metasound/Source/MetasoundFrontend/Private/MetasoundFrontendRegistries.cpp
@@ -373,18 +373,27 @@ namespace Metasound
 
 		FNodeRegistryKey FRegistryContainerImpl::RegisterGraph(const FSoftObjectPath& InAssetPath, const TScriptInterface<IMetaSoundDocumentInterface>& InDocumentInterface, bool bAsync)
 		{
-			METASOUND_TRACE_CPUPROFILER_EVENT_SCOPE(Metasound::FRegistryContainerImpl::RegisterGraph);
-
 			using namespace UE;
 
+			METASOUND_TRACE_CPUPROFILER_EVENT_SCOPE(Metasound::FRegistryContainerImpl::RegisterGraph);
+
 			check(InDocumentInterface);
 			check(IsInGameThread());
 
 			const FMetasoundFrontendDocument& Document = InDocumentInterface->GetConstDocument();
 			FNodeRegistryKey RegistryKey = FNodeRegistryKey(Document.RootGraph);
-			FNodeClassInfo NodeClassInfo(Document.RootGraph, InAssetPath);
+
+			if (!RegistryKey.IsValid())
+			{
+				// Do not attempt to build and register a MetaSound with an invalid registry key
+				UE_LOG(LogMetaSound, Warning, TEXT("Registry key is invalid when attemping to register graph for asset %s"), *InAssetPath.ToString());
+				return RegistryKey;
+			}
+
 			METASOUND_TRACE_CPUPROFILER_EVENT_SCOPE_TEXT(*FString::Printf(TEXT("FRegistryContainerImpl::RegisterGraph key:%s, asset %s"), *RegistryKey.ToString(), *InAssetPath.ToString()));
 
+			FNodeClassInfo NodeClassInfo(Document.RootGraph, InAssetPath);
+
 			// Proxies are created synchronously to avoid creating proxies in async tasks. Proxies
 			// are created from UObjects which need to be protected from GC and non-GT access.
 			FProxyDataCache ProxyDataCache;
diff --git a/Engine/Source/Runtime/MovieSceneTracks/Private/Sections/MovieSceneAudioSection.cpp b/Engine/Source/Runtime/MovieSceneTracks/Private/Sections/MovieSceneAudioSection.cpp
--- a/Engine/Source/Runtime/MovieSceneTracks/Private/Sections/MovieSceneAudioSection.cpp
+++ b/Engine/Source/Runtime/MovieSceneTracks/Private/Sections/MovieSceneAudioSection.cpp
@@ -105,16 +105,6 @@ namespace MovieSceneAudioSectionPrivate
 	}
 }
 
-void UMovieSceneAudioSection::Serialize(FArchive& Ar)
-{
-	Super::Serialize(Ar);
-
-	if (Ar.IsLoading())
-	{
-		CacheChannelProxy();
-	}
-}
-
 void UMovieSceneAudioSection::PostEditImport()
 {
 	Super::PostEditImport();
@@ -275,6 +265,8 @@ void UMovieSceneAudioSection::PostLoad()
 {
 	Super::PostLoad();
 
+	CacheChannelProxy();
+
 	if (AudioDilationFactor_DEPRECATED != AudioDeprecatedMagicNumber)
 	{
 		PitchMultiplier.SetDefault(AudioDilationFactor_DEPRECATED);
diff --git a/Engine/Source/Runtime/MovieSceneTracks/Public/Sections/MovieSceneAudioSection.h b/Engine/Source/Runtime/MovieSceneTracks/Public/Sections/MovieSceneAudioSection.h
--- a/Engine/Source/Runtime/MovieSceneTracks/Public/Sections/MovieSceneAudioSection.h
+++ b/Engine/Source/Runtime/MovieSceneTracks/Public/Sections/MovieSceneAudioSection.h
@@ -189,7 +189,6 @@ public:
 	MOVIESCENETRACKS_API virtual UMovieSceneSection* SplitSection(FQualifiedFrameTime SplitTime, bool bDeleteKeys) override;
 	MOVIESCENETRACKS_API virtual TOptional<FFrameTime> GetOffsetTime() const override;
 	MOVIESCENETRACKS_API virtual void MigrateFrameTimes(FFrameRate SourceRate, FFrameRate DestinationRate) override;
-	MOVIESCENETRACKS_API virtual void Serialize(FArchive& Ar) override;
 	MOVIESCENETRACKS_API virtual void PostEditImport() override;
 	MOVIESCENETRACKS_API virtual EMovieSceneChannelProxyType CacheChannelProxy() override;
 
