Implement feedback for/fix cached state leaving value stale in WaveTableBank eval
- LOCTEXT fix
#rnx

diff --git a/Engine/Plugins/Runtime/Metasound/Source/MetasoundEngine/Private/MetasoundWaveTableBankEvaluateNode.cpp b/Engine/Plugins/Runtime/Metasound/Source/MetasoundEngine/Private/MetasoundWaveTableBankEvaluateNode.cpp
--- a/Engine/Plugins/Runtime/Metasound/Source/MetasoundEngine/Private/MetasoundWaveTableBankEvaluateNode.cpp
+++ b/Engine/Plugins/Runtime/Metasound/Source/MetasoundEngine/Private/MetasoundWaveTableBankEvaluateNode.cpp
@@ -71,7 +71,7 @@ namespace Metasound
 					PluginNodeMissingPrompt,
 					GetDefaultInterface(),
 					{ NodeCategories::WaveTables },
-					{ NodeCategories::Envelopes, METASOUND_LOCTEXT("WaveTableEvaluateCurveKeyword", "Curve") },
+					{ NodeCategories::Envelopes, METASOUND_LOCTEXT("WaveTableBankEvaluateCurveKeyword", "Curve") },
 					{ }
 				};
 
@@ -154,8 +154,9 @@ namespace Metasound
 
 			const FWaveTableBankAsset& WaveTableBankAsset = *WaveTableBankReadRef;
 			FWaveTableBankAssetProxyPtr Proxy = WaveTableBankAsset.GetProxy();
-			float NewIndex = 0.0f;
-			if (!ResolveNextComputeIndex(Proxy, NewIndex))
+			float NewIndex = 0.f;
+			const float Input = FMath::Clamp(*InputReadRef, 0.f, 1.f);
+			if (!ResolveNextComputeIndex(Proxy, Input, NewIndex))
 			{
 				return;
 			}
@@ -168,7 +169,6 @@ namespace Metasound
 
 			float IndexFloorValue = 0.0f;
 			constexpr FWaveTableSampler::ESingleSampleMode SampleMode =  FWaveTableSampler::ESingleSampleMode::Hold;
-			const float Input = FMath::Clamp(*InputReadRef, 0.0f, 1.0f);
 
 			Sampler.Reset();
 			Sampler.SetPhase(Input);
@@ -177,65 +177,68 @@ namespace Metasound
 			if (IndexFloor != IndexCeil)
 			{
 				const FWaveTableData& IndexCeilTable = WaveTables[IndexCeil];
-				float IndexCeilValue = 0.0f;
+				float IndexCeilValue = 0.f;
 				Sampler.Reset();
 				Sampler.SetPhase(Input);
 				Sampler.Process(IndexCeilTable, IndexCeilValue, SampleMode);
 
 				const float Fractional = FMath::Frac(NewIndex);
-				LastValue = (IndexFloorValue * (1.0f - Fractional)) + (IndexCeilValue * Fractional);
+				CachedState.Value = (IndexFloorValue * (1.f - Fractional)) + (IndexCeilValue * Fractional);
 			}
 			else
 			{
-				LastValue = IndexFloorValue;
+				CachedState.Value = IndexFloorValue;
 			}
 
-			*OutWriteRef = LastValue;
+			*OutWriteRef = CachedState.Value;
 		}
 
 		void Reset(const IOperator::FResetParams& InParams)
 		{
 			using namespace WaveTable;
+
 			FWaveTableSampler::FSettings Settings;
 			Settings.Freq = 0.0f; // Sampler phase is manually progressed via this node
 			Sampler = FWaveTableSampler(MoveTemp(Settings));
 
+			CachedState = { };
 			*OutWriteRef = 0.f;
-
-			LastInterpolationMode = WaveTable::FWaveTableSampler::EInterpolationMode::COUNT;
-			LastIndex = -1.0f;
-			LastValue = 0.0f;
 		}
 
 	private:
 		// Returns true if new computation is required, setting OutIndex to index to compute.
 		// Returns false if new computation isn't required, resetting output & cached data accordingly.
-		bool ResolveNextComputeIndex(const FWaveTableBankAssetProxyPtr& Proxy, float& OutIndex)
+		bool ResolveNextComputeIndex(const FWaveTableBankAssetProxyPtr& Proxy, float Input, float& OutIndex)
 		{
 			using namespace WaveTable;
 
 			if (!Proxy.IsValid())
 			{
-				LastValue = 0.0f;
-				*OutWriteRef = LastValue;
+				CachedState = { };
+				*OutWriteRef = 0.f;
 				return false;
 			}
 
 			const float Index = *IndexReadRef;
 			const FWaveTableSampler::EInterpolationMode NewInterpolationMode = *InterpModeReadRef;
-			if (LastInterpolationMode == NewInterpolationMode)
+			if (CachedState.InterpolationMode == NewInterpolationMode)
 			{
-				if (FMath::IsNearlyEqual(Index, LastIndex))
+				if (FMath::IsNearlyEqual(Index, CachedState.Index))
 				{
-					*OutWriteRef = LastValue;
-					return false;
+					if (FMath::IsNearlyEqual(Input, CachedState.Input))
+					{
+						*OutWriteRef = CachedState.Value;
+						return false;
+					}
 				}
 			}
 
-			LastInterpolationMode = NewInterpolationMode;
-			LastIndex = Index;
 			Sampler.SetInterpolationMode(NewInterpolationMode);
 			OutIndex = FMath::Abs(Index); // Avoids fractional, interpolative flip at zero crossing
+
+			CachedState.InterpolationMode = NewInterpolationMode;
+			CachedState.Input = Input;
+			CachedState.Index = OutIndex;
 			return true;
 		}
 
@@ -248,9 +251,13 @@ namespace Metasound
 
 		FFloatWriteRef OutWriteRef;
 
-		WaveTable::FWaveTableSampler::EInterpolationMode LastInterpolationMode = WaveTable::FWaveTableSampler::EInterpolationMode::COUNT;
-		float LastIndex = -1.0f;
-		float LastValue = 0.0f;
+		struct FCachedState
+		{
+			WaveTable::FWaveTableSampler::EInterpolationMode InterpolationMode = WaveTable::FWaveTableSampler::EInterpolationMode::COUNT;
+			float Index = TNumericLimits<float>::Max();
+			float Input = TNumericLimits<float>::Max();
+			float Value = TNumericLimits<float>::Max();
+		} CachedState;
 	};
 
 	class FMetasoundWaveTableBankEvaluateNode : public FNodeFacade
