More LOCTEXT collision fixes, fix for missing bipolar clamp implementation
#rnx

diff --git a/Engine/Plugins/Runtime/Metasound/Source/MetasoundEngine/Private/MetasoundWaveTableBankEvaluateNode.cpp b/Engine/Plugins/Runtime/Metasound/Source/MetasoundEngine/Private/MetasoundWaveTableBankEvaluateNode.cpp
--- a/Engine/Plugins/Runtime/Metasound/Source/MetasoundEngine/Private/MetasoundWaveTableBankEvaluateNode.cpp
+++ b/Engine/Plugins/Runtime/Metasound/Source/MetasoundEngine/Private/MetasoundWaveTableBankEvaluateNode.cpp
@@ -20,10 +20,10 @@ namespace Metasound
 {
 	namespace WaveTableBankEvaluateNode
 	{
-		METASOUND_PARAM(WaveTableBankParam, "WaveTableBank", "The WaveTableBank to evaluate");
-		METASOUND_PARAM(InputParam, "Input", "[0, 1] the X input with which to evaluate the wavetable");
-		METASOUND_PARAM(IndexParam, "Index", "The table index to interpolate and evaluate the result of (wraps over number of entries)");
-		METASOUND_PARAM(OutParam, "Output", "The linearly mixed value of the provided WaveTableBank's applicable entries");
+		METASOUND_PARAM(WaveTableBankEval_BankParam, "WaveTableBank", "The WaveTableBank to evaluate");
+		METASOUND_PARAM(WaveTableBankEval_InputParam, "Input", "The X input with which to evaluate the wavetable ([-1, 1] or [0, 1] depending on Bank's 'bipolar' setting)");
+		METASOUND_PARAM(WaveTableBankEval_IndexParam, "Index", "The table index to interpolate and evaluate the result of (wraps over number of entries)");
+		METASOUND_PARAM(WaveTableBankEval_OutParam, "Output", "The linearly mixed value of the provided WaveTableBank's applicable entries");
 	} // WaveTableBankEvaluateNode
 	
 	class FMetasoundWaveTableBankEvaluateNodeOperator : public TExecutableOperator<FMetasoundWaveTableBankEvaluateNodeOperator>
@@ -36,9 +36,9 @@ namespace Metasound
 
 			static const FVertexInterface DefaultInterface(
 				FInputVertexInterface(
-					TInputDataVertex<FWaveTableBankAsset>(METASOUND_GET_PARAM_NAME_AND_METADATA(WaveTableBankParam)),
-					TInputDataVertex<float>(METASOUND_GET_PARAM_NAME_AND_METADATA(InputParam), 0.0f),
-					TInputDataVertex<float>(METASOUND_GET_PARAM_NAME_AND_METADATA(IndexParam), 0.0f),
+					TInputDataVertex<FWaveTableBankAsset>(METASOUND_GET_PARAM_NAME_AND_METADATA(WaveTableBankEval_BankParam)),
+					TInputDataVertex<float>(METASOUND_GET_PARAM_NAME_AND_METADATA(WaveTableBankEval_InputParam), 0.0f),
+					TInputDataVertex<float>(METASOUND_GET_PARAM_NAME_AND_METADATA(WaveTableBankEval_IndexParam), 0.0f),
 					TInputDataVertex<FEnumWaveTableInterpolationMode>("Interpolation", FDataVertexMetadata
 					{
 						LOCTEXT("MetasoundWaveTableBankEvaluateNode_InterpDescription", "How interpolation occurs between WaveTable values."),
@@ -47,7 +47,7 @@ namespace Metasound
 					}, static_cast<int32>(FWaveTableSampler::EInterpolationMode::Linear))
 				),
 				FOutputVertexInterface(
-					TOutputDataVertex<float>(METASOUND_GET_PARAM_NAME_AND_METADATA(OutParam))
+					TOutputDataVertex<float>(METASOUND_GET_PARAM_NAME_AND_METADATA(WaveTableBankEval_OutParam))
 				)
 			);
 
@@ -116,16 +116,16 @@ namespace Metasound
 		virtual void BindInputs(FInputVertexInterfaceData& InOutVertexData) override
 		{
 			using namespace WaveTableBankEvaluateNode;
-			InOutVertexData.BindReadVertex(METASOUND_GET_PARAM_NAME(WaveTableBankParam), WaveTableBankReadRef);
-			InOutVertexData.BindReadVertex(METASOUND_GET_PARAM_NAME(InputParam), InputReadRef);
-			InOutVertexData.BindReadVertex(METASOUND_GET_PARAM_NAME(IndexParam), IndexReadRef);
+			InOutVertexData.BindReadVertex(METASOUND_GET_PARAM_NAME(WaveTableBankEval_BankParam), WaveTableBankReadRef);
+			InOutVertexData.BindReadVertex(METASOUND_GET_PARAM_NAME(WaveTableBankEval_InputParam), InputReadRef);
+			InOutVertexData.BindReadVertex(METASOUND_GET_PARAM_NAME(WaveTableBankEval_IndexParam), IndexReadRef);
 			InOutVertexData.BindReadVertex("Interpolation", InterpModeReadRef);
 		}
 
 		virtual void BindOutputs(FOutputVertexInterfaceData& InOutVertexData) override
 		{
 			using namespace WaveTableBankEvaluateNode;
-			InOutVertexData.BindReadVertex(METASOUND_GET_PARAM_NAME(OutParam), OutWriteRef);
+			InOutVertexData.BindReadVertex(METASOUND_GET_PARAM_NAME(WaveTableBankEval_OutParam), OutWriteRef);
 		}
 		
 		virtual FDataReferenceCollection GetInputs() const override
@@ -155,7 +155,9 @@ namespace Metasound
 			const FWaveTableBankAsset& WaveTableBankAsset = *WaveTableBankReadRef;
 			FWaveTableBankAssetProxyPtr Proxy = WaveTableBankAsset.GetProxy();
 			float NewIndex = 0.f;
-			const float Input = FMath::Clamp(*InputReadRef, 0.f, 1.f);
+
+			const float Min = WaveTableBankAsset->IsBipolar() ? -1.f : 0.f;
+			const float Input = FMath::Clamp(*InputReadRef, Min, 1.f);
 			if (!ResolveNextComputeIndex(Proxy, Input, NewIndex))
 			{
 				return;
diff --git a/Engine/Plugins/Runtime/WaveTable/Source/WaveTable/Private/WaveTableBank.cpp b/Engine/Plugins/Runtime/WaveTable/Source/WaveTable/Private/WaveTableBank.cpp
--- a/Engine/Plugins/Runtime/WaveTable/Source/WaveTable/Private/WaveTableBank.cpp
+++ b/Engine/Plugins/Runtime/WaveTable/Source/WaveTable/Private/WaveTableBank.cpp
@@ -158,7 +158,7 @@ void UWaveTableBank::CopyToProxyData()
 	checkf(!ProxyData.IsValid(), TEXT("Cannot overwrite pre-existing proxy when calling UWaveTableBank::CopyToProxyData"));
 
 	PRAGMA_DISABLE_DEPRECATION_WARNINGS
-	ProxyData = MakeShared<FWaveTableBankAssetProxy, ESPMode::ThreadSafe>(GetUniqueID(), SampleMode, SampleRate, GetEntries());
+	ProxyData = MakeShared<FWaveTableBankAssetProxy, ESPMode::ThreadSafe>(GetUniqueID(), SampleMode, SampleRate, GetEntries(), bBipolar);
 	PRAGMA_ENABLE_DEPRECATION_WARNINGS
 }
 
@@ -167,7 +167,7 @@ void UWaveTableBank::MoveToProxyData()
 	checkf(!ProxyData.IsValid(), TEXT("Cannot overwrite pre-existing proxy when calling UWaveTableBank::MoveToProxyData"));
 
 	PRAGMA_DISABLE_DEPRECATION_WARNINGS
-	ProxyData = MakeShared<FWaveTableBankAssetProxy, ESPMode::ThreadSafe>(GetUniqueID(), SampleMode, SampleRate, MoveTemp(Entries));
+	ProxyData = MakeShared<FWaveTableBankAssetProxy, ESPMode::ThreadSafe>(GetUniqueID(), SampleMode, SampleRate, MoveTemp(Entries), bBipolar);
 	PRAGMA_ENABLE_DEPRECATION_WARNINGS
 }
 
@@ -268,28 +268,25 @@ void UWaveTableBank::Serialize(FArchive& Ar)
 #endif // WITH_EDITORONLY_DATA
 }
 
-FWaveTableBankAssetProxy::FWaveTableBankAssetProxy(const FWaveTableBankAssetProxy& InAssetProxy)
-	: ObjectId(InAssetProxy.ObjectId)
-	, SampleRate(InAssetProxy.SampleRate)
-	, SampleMode(InAssetProxy.SampleMode)
-	, WaveTableData(InAssetProxy.WaveTableData)
-{
-}
 
-FWaveTableBankAssetProxy::FWaveTableBankAssetProxy(uint32 InObjectId, EWaveTableSamplingMode InSamplingMode, int32 InSampleRate, const TArray<FWaveTableBankEntry>& InBankEntries)
+FWaveTableBankAssetProxy::FWaveTableBankAssetProxy(uint32 InObjectId, EWaveTableSamplingMode InSamplingMode, int32 InSampleRate, const TArray<FWaveTableBankEntry>& InBankEntries, bool bInBipolar)
+	: bBipolar(bInBipolar)
+	, ObjectId(InObjectId)
+	, SampleRate(InSampleRate)
+	, SampleMode(InSamplingMode)
 {
-	SampleMode = InSamplingMode;
-	SampleRate = InSampleRate;
 	Algo::Transform(InBankEntries, WaveTableData, [](const FWaveTableBankEntry& Entry)
 	{
 		return Entry.Transform.GetTableData();
 	});
 }
 
-FWaveTableBankAssetProxy::FWaveTableBankAssetProxy(uint32 InObjectId, EWaveTableSamplingMode InSamplingMode, int32 InSampleRate, TArray<FWaveTableBankEntry>&& InBankEntries)
+FWaveTableBankAssetProxy::FWaveTableBankAssetProxy(uint32 InObjectId, EWaveTableSamplingMode InSamplingMode, int32 InSampleRate, TArray<FWaveTableBankEntry>&& InBankEntries, bool bInBipolar)
+	: bBipolar(bInBipolar)
+	, ObjectId(InObjectId)
+	, SampleRate(InSampleRate)
+	, SampleMode(InSamplingMode)
 {
-	SampleMode = InSamplingMode;
-	SampleRate = InSampleRate;
 	for (FWaveTableBankEntry& Entry : InBankEntries)
 	{
 		WaveTableData.Add(MoveTemp(Entry.Transform.TableData));
@@ -298,11 +295,11 @@ FWaveTableBankAssetProxy::FWaveTableBankAssetProxy(uint32 InObjectId, EWaveTable
 }
 
 FWaveTableBankAssetProxy::FWaveTableBankAssetProxy(const UWaveTableBank& InWaveTableBank)
+	: bBipolar(InWaveTableBank.bBipolar)
+	, ObjectId(InWaveTableBank.GetUniqueID())
+	, SampleRate(InWaveTableBank.SampleRate)
+	, SampleMode(InWaveTableBank.SampleMode)
 {
-	ObjectId = InWaveTableBank.GetUniqueID();
-	SampleRate = InWaveTableBank.SampleRate;
-	SampleMode = InWaveTableBank.SampleMode;
-
 	PRAGMA_DISABLE_DEPRECATION_WARNINGS
 	Algo::Transform(InWaveTableBank.Entries, WaveTableData, [](const FWaveTableBankEntry& Entry)
 	{
diff --git a/Engine/Plugins/Runtime/WaveTable/Source/WaveTable/Public/WaveTableBank.h b/Engine/Plugins/Runtime/WaveTable/Source/WaveTable/Public/WaveTableBank.h
--- a/Engine/Plugins/Runtime/WaveTable/Source/WaveTable/Public/WaveTableBank.h
+++ b/Engine/Plugins/Runtime/WaveTable/Source/WaveTable/Public/WaveTableBank.h
@@ -94,9 +94,8 @@ public:
 	IMPL_AUDIOPROXY_CLASS(FWaveTableBankAssetProxy);
 
 	FWaveTableBankAssetProxy() = default;
-	FWaveTableBankAssetProxy(const FWaveTableBankAssetProxy& InAssetProxy);
-	FWaveTableBankAssetProxy(uint32 InObjectId, EWaveTableSamplingMode InSamplingMode, int32 InSampleRate, const TArray<FWaveTableBankEntry>& InBankEntries);
-	FWaveTableBankAssetProxy(uint32 InObjectId, EWaveTableSamplingMode InSamplingMode, int32 InSampleRate, TArray<FWaveTableBankEntry>&& InBankEntries);
+	FWaveTableBankAssetProxy(uint32 InObjectId, EWaveTableSamplingMode InSamplingMode, int32 InSampleRate, const TArray<FWaveTableBankEntry>& InBankEntries, bool bInBipolar = false);
+	FWaveTableBankAssetProxy(uint32 InObjectId, EWaveTableSamplingMode InSamplingMode, int32 InSampleRate, TArray<FWaveTableBankEntry>&& InBankEntries, bool bInBipolar = false);
 
 	UE_DEPRECATED(5.3, "Proxy generation & respective data translation is now entirely handled by UWaveTableBank calling the constructor variants above.")
 	FWaveTableBankAssetProxy(const UWaveTableBank& InWaveTableBank);
@@ -118,12 +117,18 @@ public:
 		return WaveTableData;
 	}
 
+	virtual bool IsBipolar() const
+	{
+		return bBipolar;
+	}
+
 	uint32 GetObjectId() const
 	{
 		return ObjectId;
 	}
 
 protected:
+	bool bBipolar = false;
 	uint32 ObjectId = INDEX_NONE;
 	int32 SampleRate = 48000;
 	EWaveTableSamplingMode SampleMode = EWaveTableSamplingMode::FixedResolution;
