Fix MetaSound trace debug print name for graph nodes

diff --git a/Engine/Plugins/Runtime/Metasound/Source/MetasoundGraphCore/Private/MetasoundNodeInterface.cpp b/Engine/Plugins/Runtime/Metasound/Source/MetasoundGraphCore/Private/MetasoundNodeInterface.cpp
--- a/Engine/Plugins/Runtime/Metasound/Source/MetasoundGraphCore/Private/MetasoundNodeInterface.cpp
+++ b/Engine/Plugins/Runtime/Metasound/Source/MetasoundGraphCore/Private/MetasoundNodeInterface.cpp
@@ -11,6 +11,8 @@ namespace Metasound
 	const FText PluginNodeMissingPrompt = FText::GetEmpty();
 #endif // WITH_EDITOR
 
+	const FNodeClassName FNodeClassName::InvalidNodeClassName;
+
 	FNodeClassName::FNodeClassName()
 	{
 	}
@@ -73,6 +75,11 @@ namespace Metasound
 		return FName(*FString::Printf(TEXT("%s.%s"), *InNamespace.ToString(), *InName.ToString()));
 	}
 
+	bool FNodeClassName::IsValid() const
+	{
+		return *this != InvalidNodeClassName;
+	}
+
 	bool operator==(const FOutputDataSource& InLeft, const FOutputDataSource& InRight)
 	{
 		return (InLeft.Node == InRight.Node) && (InLeft.Vertex == InRight.Vertex);
diff --git a/Engine/Plugins/Runtime/Metasound/Source/MetasoundGraphCore/Private/MetasoundOperatorBuilder.cpp b/Engine/Plugins/Runtime/Metasound/Source/MetasoundGraphCore/Private/MetasoundOperatorBuilder.cpp
--- a/Engine/Plugins/Runtime/Metasound/Source/MetasoundGraphCore/Private/MetasoundOperatorBuilder.cpp
+++ b/Engine/Plugins/Runtime/Metasound/Source/MetasoundGraphCore/Private/MetasoundOperatorBuilder.cpp
@@ -460,7 +460,12 @@ namespace Metasound
 			FGraphOperatorData::FOperatorInfo& OperatorInfo = OperatorMap.FindChecked(OperatorID);
 
 			{
-				METASOUND_TRACE_CPUPROFILER_EVENT_SCOPE_TEXT(*FString::Printf(TEXT("Metasound::FOperatorBuilder::CreateOperators::CreateAndBind %s"), *Node->GetMetadata().ClassName.GetFullName().ToString()));
+#if METASOUND_CPUPROFILERTRACE_ENABLED
+				// Use node class name if valid, otherwise (for example graph nodes) use instance name 
+				const FNodeClassName& NodeClassName = Node->GetMetadata().ClassName;
+				const FString& NodeTraceName = NodeClassName.IsValid() ? NodeClassName.GetFullName().ToString() : Node->GetInstanceName().ToString();
+				METASOUND_TRACE_CPUPROFILER_EVENT_SCOPE_TEXT(*FString::Printf(TEXT("Metasound::FOperatorBuilder::CreateOperators::CreateAndBind %s"), *NodeTraceName));
+#endif // METASOUND_CPUPROFILERTRACE_ENABLED
 
 				FBuildOperatorParams CreateParams{*Node, InOutContext.Settings, OperatorInfo.VertexData.GetInputs(), InOutContext.Environment, this};
 				FOperatorFactorySharedRef Factory = Node->GetDefaultOperatorFactory();
diff --git a/Engine/Plugins/Runtime/Metasound/Source/MetasoundGraphCore/Public/MetasoundNodeInterface.h b/Engine/Plugins/Runtime/Metasound/Source/MetasoundGraphCore/Public/MetasoundNodeInterface.h
--- a/Engine/Plugins/Runtime/Metasound/Source/MetasoundGraphCore/Public/MetasoundNodeInterface.h
+++ b/Engine/Plugins/Runtime/Metasound/Source/MetasoundGraphCore/Public/MetasoundNodeInterface.h
@@ -69,10 +69,21 @@ namespace Metasound
 		/** Namespace, name and variant of the node class. */
 		const FName& GetFullName() const;
 
+		/** Whether or not this instance of a node class name is a valid name.*/
+		bool IsValid() const;
+
 		static FName FormatFullName(const FName& InNamespace, const FName& InName, const FName& InVariant);
 
 		static FName FormatScopedName(const FName& InNamespace, const FName& InName);
 
+		/** Invalid form of node class name (i.e. empty namespace, name, and variant)*/
+		static const FNodeClassName InvalidNodeClassName;
+
+		friend FORCEINLINE bool operator==(const FNodeClassName& InLHS, const FNodeClassName& InRHS)
+		{
+			return (InLHS.Namespace == InRHS.Namespace) && (InLHS.Name == InRHS.Name) && (InLHS.Variant == InRHS.Variant);
+		}
+
 	private:
 
 		FName Namespace;
