Remove unneeded input binding in MetaSounds for perf

diff --git a/Engine/Plugins/Runtime/Metasound/Source/MetasoundGraphCore/Private/MetasoundOperatorBuilder.cpp b/Engine/Plugins/Runtime/Metasound/Source/MetasoundGraphCore/Private/MetasoundOperatorBuilder.cpp
--- a/Engine/Plugins/Runtime/Metasound/Source/MetasoundGraphCore/Private/MetasoundOperatorBuilder.cpp
+++ b/Engine/Plugins/Runtime/Metasound/Source/MetasoundGraphCore/Private/MetasoundOperatorBuilder.cpp
@@ -489,7 +489,12 @@ namespace Metasound
 				// Bind vertex to operator data
 				{
 					METASOUND_TRACE_CPUPROFILER_EVENT_SCOPE(Metasound::FOperatorBuilder::CreateOperators::BindInputsAndOutputs);
-					OperatorInfo.Operator->BindInputs(OperatorInfo.VertexData.GetInputs());
+					// Inputs don't need to be bound for all nodes unless they are dynamic 
+					// Inputs for input nodes will be bound separately in GatherGraphDataReferences
+					if (BuilderSettings.bEnableOperatorRebind)
+					{
+						OperatorInfo.Operator->BindInputs(OperatorInfo.VertexData.GetInputs());
+					}
 					OperatorInfo.Operator->BindOutputs(OperatorInfo.VertexData.GetOutputs());
 				}
 				
@@ -543,7 +548,6 @@ namespace Metasound
 		using FSourceElement = FOutputDataSourceCollection::ElementType;
 
 		FBuildStatus BuildStatus;
-
 		// Gather graph inputs
 		for (const FDestinationElement& Element : InOutContext.Graph.GetInputDataDestinations())
 		{
@@ -551,9 +555,11 @@ namespace Metasound
 			const FInputDataDestination& InputDestination = Element.Value;
 
 			const FOperatorID OperatorID = GetOperatorID(InputDestination.Node);
-			if (const FGraphOperatorData::FOperatorInfo* OperatorInfo = InOutContext.GraphOperatorData->OperatorMap.Find(OperatorID))
+			if (FGraphOperatorData::FOperatorInfo* OperatorInfo = InOutContext.GraphOperatorData->OperatorMap.Find(OperatorID))
 			{
-				const FInputVertexInterfaceData& NodeInputData = OperatorInfo->VertexData.GetInputs();
+				FInputVertexInterfaceData& NodeInputData = OperatorInfo->VertexData.GetInputs();
+				OperatorInfo->Operator->BindInputs(NodeInputData);
+
 				if (const FAnyDataReference* DataReference = NodeInputData.FindDataReference(InputDestination.Vertex.VertexName))
 				{
 					if (DataReference->GetDataTypeName() == InputDestination.Vertex.DataTypeName)
