Fix async references not yet registered when auto-updating in editor

diff --git a/Engine/Plugins/Runtime/Metasound/Source/MetasoundFrontend/Private/MetasoundFrontendTransform.cpp b/Engine/Plugins/Runtime/Metasound/Source/MetasoundFrontend/Private/MetasoundFrontendTransform.cpp
--- a/Engine/Plugins/Runtime/Metasound/Source/MetasoundFrontend/Private/MetasoundFrontendTransform.cpp
+++ b/Engine/Plugins/Runtime/Metasound/Source/MetasoundFrontend/Private/MetasoundFrontendTransform.cpp
@@ -613,16 +613,23 @@ namespace Metasound
 				using namespace Metasound::Frontend;
 
 				const FMetasoundFrontendClassMetadata& ClassMetadata = NodeHandle->GetClassMetadata();
+				const FNodeRegistryKey RegistryKey(ClassMetadata);
 
-				if (bIsPreset)
+				if (FMetasoundAssetBase* ReferencedMetaSoundAsset = IMetaSoundAssetManager::GetChecked().TryLoadAssetFromKey(RegistryKey))
 				{
-					const FNodeRegistryKey RegistryKey(ClassMetadata);
-					PresetReferencedMetaSoundAsset = IMetaSoundAssetManager::GetChecked().TryLoadAssetFromKey(RegistryKey);
-					if (!PresetReferencedMetaSoundAsset)
+ 					ReferencedMetaSoundAsset->WaitForAsyncGraphRegistration();
+					if (bIsPreset)
+					{
+						PresetReferencedMetaSoundAsset = ReferencedMetaSoundAsset;
+					}
+				}
+				else
+				{
+					if (bIsPreset)
 					{
 						UE_LOG(LogMetaSound, Error, TEXT("Auto-Updating preset '%s' failed: Referenced class '%s' missing."), *DebugAssetPath, *ClassMetadata.GetClassName().ToString());
+						return;
 					}
-					return;
 				}
 
 				FClassInterfaceUpdates InterfaceUpdates;
@@ -633,7 +640,7 @@ namespace Metasound
 
 				// Check if a updated minor version exists.
 				FMetasoundFrontendClass ClassWithHighestMinorVersion;
-				bool bFoundClassInSearchEngine = Frontend::ISearchEngine::Get().FindClassWithHighestMinorVersion(ClassMetadata.GetClassName(), ClassMetadata.GetVersion().Major, ClassWithHighestMinorVersion);
+				const bool bFoundClassInSearchEngine = Frontend::ISearchEngine::Get().FindClassWithHighestMinorVersion(ClassMetadata.GetClassName(), ClassMetadata.GetVersion().Major, ClassWithHighestMinorVersion);
 
 				if (bFoundClassInSearchEngine && (ClassWithHighestMinorVersion.Metadata.GetVersion() > ClassMetadata.GetVersion()))
 				{
