Only register assets requested from builder if never built or if builder has transactions logged requiring re-registration (CVar disabled, will enable in subsequent change)

diff --git a/Engine/Plugins/Runtime/Metasound/Source/MetasoundEngine/Private/MetasoundBuilderSubsystem.cpp b/Engine/Plugins/Runtime/Metasound/Source/MetasoundEngine/Private/MetasoundBuilderSubsystem.cpp
--- a/Engine/Plugins/Runtime/Metasound/Source/MetasoundEngine/Private/MetasoundBuilderSubsystem.cpp
+++ b/Engine/Plugins/Runtime/Metasound/Source/MetasoundEngine/Private/MetasoundBuilderSubsystem.cpp
@@ -33,6 +33,8 @@ namespace Metasound::Engine
 {
 	namespace BuilderSubsystemPrivate
 	{
+		int32 TransactionBasedRegistrationEnabled = 0;
+
 		template <typename TLiteralType>
 		FMetasoundFrontendLiteral CreatePODMetaSoundLiteral(const TLiteralType& Value, FName& OutDataType)
 		{
@@ -53,16 +55,31 @@ namespace Metasound::Engine
 			return IDataTypeRegistry::Get().CreateLiteralNode(DataType, MoveTemp(Params));
 		}
 	} // namespace BuilderSubsystemPrivate
+
+	FAutoConsoleVariableRef CVarMetaSoundBuilderForceAllFrontendRegistration(
+		TEXT("au.MetaSound.Builder.TransactionBasedRegistrationEnabled"),
+		BuilderSubsystemPrivate::TransactionBasedRegistrationEnabled,
+		TEXT("Forces all builder calls to register MetaSound objects with the Frontend.\n")
+		TEXT("Enabled (Default): !0, Disabled: 0"),
+		ECVF_Default);
 } // namespace Metasound::Engine
 
+
 void UMetaSoundBuilderBase::BeginDestroy()
 {
-	Super::BeginDestroy();
+	if (UMetaSoundBuilderSubsystem* Subsystem = UMetaSoundBuilderSubsystem::Get())
+	{
+		const FMetasoundFrontendDocument& Document = static_cast<const FMetaSoundFrontendDocumentBuilder&>(Builder).GetDocument();
+		const FMetasoundFrontendClassName& ClassName = Document.RootGraph.Metadata.GetClassName();
+		Subsystem->TransientBuilders.Remove(ClassName);
+	}
 
 	// Need to detach before destroying UPROPERTYs as the Builder 
 	// often holds a TScriptInterface<IMetaSoundDocumentInterface> of
 	// a UPROPERTY that lives on this or derived objects. 
 	Builder.FinishBuilding();
+
+	Super::BeginDestroy();
 }
 
 FMetaSoundBuilderNodeOutputHandle UMetaSoundBuilderBase::AddGraphInputNode(FName Name, FName DataType, FMetasoundFrontendLiteral DefaultValue, EMetaSoundBuilderResult& OutResult, bool bIsConstructorInput)
@@ -169,18 +186,37 @@ void UMetaSoundBuilderBase::AddInterface(FName InterfaceName, EMetaSoundBuilderR
 FMetaSoundNodeHandle UMetaSoundBuilderBase::AddNode(const TScriptInterface<IMetaSoundDocumentInterface>& NodeClass, EMetaSoundBuilderResult& OutResult)
 {
 	using namespace Metasound;
+	using namespace Metasound::Frontend;
 
 	FMetaSoundNodeHandle NewHandle;
 
 	if (NodeClass)
 	{
-		if (FMetasoundAssetBase* MetaSoundAsset = IMetasoundUObjectRegistry::Get().GetObjectAsAssetBase(NodeClass.GetObject()))
+		UObject* NodeClassObject = NodeClass.GetObject();
+		check(NodeClassObject);
+
+#if WITH_EDITOR
+		// Assets that may undergo serialization cannot reference transient objects
+		const bool bIsInvalidReference = !NodeClassObject->IsAsset() && Builder.CastDocumentObjectChecked<UObject>().IsAsset();
+#else
+		constexpr bool bIsInvalidReference = false;
+#endif // WITH_EDITOR
+
+		if (bIsInvalidReference)
 		{
-			MetaSoundAsset->RegisterGraphWithFrontend();
+			UObject& ThisBuildersObject = Builder.CastDocumentObjectChecked<UObject>();
+			UE_LOG(LogMetaSound, Warning,
+				TEXT("Failed to add node of transient asset '%s' to serialized asset '%s': "
+				"Transient object node class cannot be referenced from asset node class."),
+				*NodeClassObject->GetPathName(),
+				*ThisBuildersObject.GetPathName());
+		}
+		else
+		{
+			RegisterGraphIfOutstandingTransactions(*NodeClassObject);
 
 			const FMetasoundFrontendDocument& NodeClassDoc = NodeClass->GetConstDocument();
 			const FMetasoundFrontendGraphClass& NodeClassGraph = NodeClassDoc.RootGraph;
-
 			if (const FMetasoundFrontendNode* NewNode = Builder.AddGraphNode(NodeClassGraph))
 			{
 				NewHandle.NodeID = NewNode->GetID();
@@ -720,6 +756,45 @@ bool UMetaSoundBuilderBase::NodeOutputIsConnected(const FMetaSoundBuilderNodeOut
 	return Builder.IsNodeOutputConnected(OutputHandle.NodeID, OutputHandle.VertexID);
 }
 
+void UMetaSoundBuilderBase::RegisterGraphIfOutstandingTransactions(UObject& InMetaSound)
+{
+	using namespace Metasound;
+	using namespace Metasound::Engine;
+	using namespace Metasound::Frontend;
+	using namespace Metasound::Engine::BuilderSubsystemPrivate;
+
+	IMetaSoundAssetManager& AssetManager = IMetaSoundAssetManager::GetChecked();
+	FMetasoundAssetBase* MetaSoundAsset = IMetasoundUObjectRegistry::Get().GetObjectAsAssetBase(&InMetaSound);
+	check(MetaSoundAsset);
+
+	FMetaSoundAssetRegistrationOptions Options;
+	if (TransactionBasedRegistrationEnabled)
+	{
+		Options.bForceReregister = false;
+		Options.bRegisterDependencies = false; // Function handles registration via own recursive functionality below
+
+		TArray<FMetasoundAssetBase*> References = MetaSoundAsset->GetReferencedAssets();
+		for (FMetasoundAssetBase* Reference : References)
+		{
+			UObject* RefMetaSound = Reference->GetOwningAsset();
+			check(RefMetaSound);
+			AssetManager.AddOrUpdateAsset(*RefMetaSound);
+			RegisterGraphIfOutstandingTransactions(*RefMetaSound);
+		}
+
+		if (UMetaSoundBuilderBase* Builder = UMetaSoundBuilderSubsystem::GetChecked().FindBuilderOfDocument(&InMetaSound))
+		{
+			const int32 TransactionCount = Builder->Builder.GetTransactionCount();
+
+			// Force registration if transactions occurred since now and the last time the builder registered the asset.
+			Options.bForceReregister = Builder->LastTransactionRegistered != TransactionCount;
+			Builder->LastTransactionRegistered = TransactionCount;
+		}
+	}
+
+	MetaSoundAsset->RegisterGraphWithFrontend(Options);
+}
+
 void UMetaSoundBuilderBase::ReloadCache(bool bPrimeCache)
 {
 	Builder.ReloadCache();
@@ -982,10 +1057,7 @@ void UMetaSoundSourceBuilder::Audition(UObject* Parent, UAudioComponent* AudioCo
 	}
 
 	UMetaSoundSource& MetaSoundSource = GetMetaSoundSource();
-
-	FMetasoundAssetBase* MetaSoundAsset = IMetasoundUObjectRegistry::Get().GetObjectAsAssetBase(&MetaSoundSource);
-	check(MetaSoundAsset);
-	MetaSoundAsset->RegisterGraphWithFrontend();
+	RegisterGraphIfOutstandingTransactions(MetaSoundSource);
 
 	// Must be called post register as register ensures cached runtime data passed to transactor is up-to-date
 	MetaSoundSource.SetDynamicGeneratorEnabled(bLiveUpdatesEnabled);
@@ -1485,7 +1557,6 @@ UMetaSoundSourceBuilder* UMetaSoundBuilderSubsystem::CreateSourceBuilder(
 	AudioOutNodeInputs.Reset();
 
 	UMetaSoundSourceBuilder& NewBuilder = CreateTransientBuilder<UMetaSoundSourceBuilder>(BuilderName);
-
 	OutResult = EMetaSoundBuilderResult::Succeeded;
 	if (OutputFormat != EMetaSoundOutputAudioFormat::Mono)
 	{
@@ -1723,14 +1794,37 @@ FMetasoundFrontendLiteral UMetaSoundBuilderSubsystem::CreateMetaSoundLiteralFrom
 
 bool UMetaSoundBuilderSubsystem::DetachBuilderFromAsset(const FMetasoundFrontendClassName& InClassName) const
 {
-	return AssetBuilders.Remove(InClassName.GetFullName()) > 0;
+	using namespace Metasound;
+	using namespace Metasound::Frontend;
+
+	TWeakObjectPtr<UMetaSoundBuilderBase> Builder = AssetBuilders.FindRef(InClassName);
+	if (Builder.IsValid())
+	{
+		// If the builder has applied transactions to its document object that are not mirrored in the frontend registry,
+		// unregister version in registry. This will ensure that future requests for the builder's associated asset will
+		// register a fresh version from the object as the transaction history is intrinsically lost once this builder
+		// is destroyed.
+		if (Builder->LastTransactionRegistered != Builder->Builder.GetTransactionCount())
+		{
+			UObject& MetaSound = Builder->Builder.CastDocumentObjectChecked<UObject>();
+			if (FMetasoundAssetBase* MetaSoundAsset = IMetasoundUObjectRegistry::Get().GetObjectAsAssetBase(&MetaSound))
+			{
+				MetaSoundAsset->UnregisterGraphWithFrontend();
+			}
+		}
+
+		ensureAlways(AssetBuilders.Remove(InClassName));
+		return true;
+	}
+
+	return false;
 }
 
 void UMetaSoundBuilderSubsystem::InvalidateDocumentCache(const FMetasoundFrontendClassName& InClassName) const
 {
 	using namespace Metasound::Frontend;
 
-	TWeakObjectPtr<UMetaSoundBuilderBase> BuilderPtr = AssetBuilders.FindRef(InClassName.GetFullName());
+	TWeakObjectPtr<UMetaSoundBuilderBase> BuilderPtr = AssetBuilders.FindRef(InClassName);
 	if (BuilderPtr.IsValid())
 	{
 		BuilderPtr->InvalidateCache();
@@ -1742,6 +1836,33 @@ UMetaSoundBuilderBase* UMetaSoundBuilderSubsystem::FindBuilder(FName BuilderName
 	return NamedBuilders.FindRef(BuilderName);
 }
 
+UMetaSoundBuilderBase* UMetaSoundBuilderSubsystem::FindBuilderOfDocument(TScriptInterface<const IMetaSoundDocumentInterface> InMetaSound) const
+{
+	TWeakObjectPtr<UMetaSoundBuilderBase> Builder;
+	if (const UObject* MetaSoundObject = InMetaSound.GetObject())
+	{
+		const FMetasoundFrontendDocument& Document = InMetaSound->GetConstDocument();
+		const FMetasoundFrontendClassName& ClassName = Document.RootGraph.Metadata.GetClassName();
+
+		Builder = AssetBuilders.FindRef(ClassName);
+		if (Builder.IsValid())
+		{
+			ensureAlwaysMsgf(MetaSoundObject->IsAsset(), TEXT("MetaSound is asset but Builder was registered with Subsystem as transient"));
+		}
+		else
+		{
+			Builder = TransientBuilders.FindRef(ClassName);
+			if (Builder.IsValid())
+			{
+				ensureAlwaysMsgf(!MetaSoundObject->IsAsset(), TEXT("MetaSound is transient object but Builder was registered with Subsystem as asset"));
+			}
+		}
+
+	}
+
+	return Builder.Get();
+}
+
 UMetaSoundPatchBuilder* UMetaSoundBuilderSubsystem::FindPatchBuilder(FName BuilderName)
 {
 	if (UMetaSoundBuilderBase* Builder = FindBuilder(BuilderName))
@@ -1786,7 +1907,7 @@ bool UMetaSoundBuilderSubsystem::IsInterfaceRegistered(FName InInterfaceName) co
 #if WITH_EDITOR
 void UMetaSoundBuilderSubsystem::PostBuilderAssetTransaction(const FMetasoundFrontendClassName& InClassName)
 {
-	TWeakObjectPtr<UMetaSoundBuilderBase> BuilderPtr = AssetBuilders.FindRef(InClassName.GetFullName());
+	TWeakObjectPtr<UMetaSoundBuilderBase> BuilderPtr = AssetBuilders.FindRef(InClassName);
 	if (UMetaSoundBuilderBase* Builder = BuilderPtr.Get())
 	{
 		Builder->ReloadCache();
diff --git a/Engine/Plugins/Runtime/Metasound/Source/MetasoundEngine/Public/MetasoundBuilderSubsystem.h b/Engine/Plugins/Runtime/Metasound/Source/MetasoundEngine/Public/MetasoundBuilderSubsystem.h
--- a/Engine/Plugins/Runtime/Metasound/Source/MetasoundEngine/Public/MetasoundBuilderSubsystem.h
+++ b/Engine/Plugins/Runtime/Metasound/Source/MetasoundEngine/Public/MetasoundBuilderSubsystem.h
@@ -474,6 +474,11 @@ protected:
 	UE_DEPRECATED(5.4, "Use UMetaSoundBuilderDocument::Create instead")
 	UMetaSoundBuilderDocument* CreateTransientDocumentObject() const;
 
+	// Only registers provided MetaSound's graph class and referenced graphs recursively if
+	// it has yet to be registered or if it has an attached builder reporting outstanding
+	// transactions that have yet to be registered.
+	static void RegisterGraphIfOutstandingTransactions(UObject& InMetaSound);
+
 	UPROPERTY()
 	FMetaSoundFrontendDocumentBuilder Builder;
 
@@ -482,6 +487,9 @@ protected:
 	bool bIsAttached = false;
 #endif // WITH_EDITORONLY_DATA
 
+private:
+	int32 LastTransactionRegistered = 0;
+
 	// Friending allows for swapping the builder in certain circumstances where desired (eg. attaching a builder to an existing asset)
 	friend class UMetaSoundBuilderSubsystem;
 };
@@ -569,7 +577,10 @@ private:
 	TMap<FName, TObjectPtr<UMetaSoundBuilderBase>> NamedBuilders;
 
 	UPROPERTY()
-	mutable TMap<FName, TWeakObjectPtr<UMetaSoundBuilderBase>> AssetBuilders;
+	mutable TMap<FMetasoundFrontendClassName, TWeakObjectPtr<UMetaSoundBuilderBase>> AssetBuilders;
+
+	UPROPERTY()
+	mutable TMap<FMetasoundFrontendClassName, TWeakObjectPtr<UMetaSoundBuilderBase>> TransientBuilders;
 
 public:
 	virtual void Initialize(FSubsystemCollectionBase& Collection) override;
@@ -644,12 +655,19 @@ public:
 	UFUNCTION(BlueprintCallable, Category = "Audio|MetaSound|Builder", meta = (DisplayName = "Create MetaSound Literal From AudioParameter"))
 	UPARAM(DisplayName = "Param Literal") FMetasoundFrontendLiteral CreateMetaSoundLiteralFromParam(const FAudioParameter& Param);
 
-	UFUNCTION(BlueprintCallable, Category = "Audio|MetaSound|Builder")
+	// Returns the builder manually registered with the MetaSound Builder Subsystem with the provided custom name (if previously registered)
+	UFUNCTION(BlueprintCallable, Category = "Audio|MetaSound|Builder", meta = (DisplayName = "Find Builder By Name"))
 	UPARAM(DisplayName = "Builder") UMetaSoundBuilderBase* FindBuilder(FName BuilderName);
 
+	// Returns the builder associated with the given MetaSound (if one exists, transient or asset).
+	UFUNCTION(BlueprintCallable, Category = "Audio|MetaSound|Builder", meta = (DisplayName = "Find Builder By MetaSound"))
+	UPARAM(DisplayName = "Builder") UMetaSoundBuilderBase* FindBuilderOfDocument(TScriptInterface<const IMetaSoundDocumentInterface> InMetaSound) const;
+
+	// Returns the patch builder manually registered with the MetaSound Builder Subsystem with the provided custom name (if previously registered)
 	UFUNCTION(BlueprintCallable, Category = "Audio|MetaSound|Builder")
 	UPARAM(DisplayName = "Patch Builder") UMetaSoundPatchBuilder* FindPatchBuilder(FName BuilderName);
 
+	// Returns the source builder manually registered with the MetaSound Builder Subsystem with the provided custom name (if previously registered)
 	UFUNCTION(BlueprintCallable, Category = "Audio|MetaSound|Builder")
 	UPARAM(DisplayName = "Source Builder") UMetaSoundSourceBuilder* FindSourceBuilder(FName BuilderName);
 
@@ -692,6 +710,10 @@ private:
 		TObjectPtr<BuilderClass> NewBuilder = NewObject<BuilderClass>(TransientPackage, ObjectName, NewObjectFlags);
 		check(NewBuilder);
 		NewBuilder->CreateTransientBuilder();
+
+		const FMetasoundFrontendDocument& Document = NewBuilder->GetConstBuilder().GetDocument();
+		const FMetasoundFrontendClassName& ClassName = Document.RootGraph.Metadata.GetClassName();
+		TransientBuilders.Add(ClassName, NewBuilder);
 		return *NewBuilder.Get();
 	}
 
@@ -703,8 +725,8 @@ private:
 
 		TScriptInterface<IMetaSoundDocumentInterface> DocInterface = InMetaSoundObject;
 		const FMetasoundFrontendDocument& Document = DocInterface->GetConstDocument();
-		const FName FullClassName = Document.RootGraph.Metadata.GetClassName().GetFullName();
-		TWeakObjectPtr<UMetaSoundBuilderBase> Builder = AssetBuilders.FindRef(FullClassName);
+		const FMetasoundFrontendClassName& ClassName = Document.RootGraph.Metadata.GetClassName();
+		TWeakObjectPtr<UMetaSoundBuilderBase> Builder = AssetBuilders.FindRef(ClassName);
 		if (Builder.IsValid())
 		{
 			return *CastChecked<BuilderClass>(Builder.Get());
@@ -714,7 +736,9 @@ private:
 		check(NewBuilder);
 		NewBuilder->Builder = FMetaSoundFrontendDocumentBuilder(DocInterface);
 		TObjectPtr<UMetaSoundBuilderBase> NewBuilderBase = CastChecked<UMetaSoundBuilderBase>(NewBuilder);
-		AssetBuilders.Add(FullClassName, NewBuilderBase);
+		AssetBuilders.Add(ClassName, NewBuilderBase);
 		return *NewBuilder;
 	}
+
+	friend class UMetaSoundBuilderBase;
 };
diff --git a/Engine/Plugins/Runtime/Metasound/Source/MetasoundFrontend/Private/MetasoundFrontendDocumentBuilder.cpp b/Engine/Plugins/Runtime/Metasound/Source/MetasoundFrontend/Private/MetasoundFrontendDocumentBuilder.cpp
--- a/Engine/Plugins/Runtime/Metasound/Source/MetasoundFrontend/Private/MetasoundFrontendDocumentBuilder.cpp
+++ b/Engine/Plugins/Runtime/Metasound/Source/MetasoundFrontend/Private/MetasoundFrontendDocumentBuilder.cpp
@@ -1843,6 +1843,18 @@ EMetasoundFrontendVertexAccessType FMetaSoundFrontendDocumentBuilder::GetNodeOut
 	return EMetasoundFrontendVertexAccessType::Unset;
 }
 
+int32 FMetaSoundFrontendDocumentBuilder::GetTransactionCount() const
+{
+	using namespace Metasound::Frontend;
+
+	if (DocumentCache.IsValid())
+	{
+		return StaticCastSharedPtr<FDocumentCache>(DocumentCache)->GetTransactionCount();
+	}
+
+	return 0;
+}
+
 void FMetaSoundFrontendDocumentBuilder::InitGraphClassMetadata(FMetasoundFrontendClassMetadata& InOutMetadata, bool bResetVersion, const FMetasoundFrontendClassName* NewClassName)
 {
 	if (NewClassName)
@@ -2579,20 +2591,26 @@ bool FMetaSoundFrontendDocumentBuilder::SetGraphInputDefault(FName InputName, co
 
 	auto NameMatchesInput = [&InputName](const FMetasoundFrontendClassInput& Input) { return Input.Name == InputName; };
 	TArray<FMetasoundFrontendClassInput>& Inputs = GetDocument().RootGraph.Interface.Inputs;
-	if (FMetasoundFrontendClassInput* Input = Inputs.FindByPredicate(NameMatchesInput))
+
+	const int32 Index = Inputs.IndexOfByPredicate(NameMatchesInput);
+	if (Index != INDEX_NONE)
 	{
-		if (IDataTypeRegistry::Get().IsLiteralTypeSupported(Input->TypeName, InDefaultLiteral.GetType()))
+		FMetasoundFrontendClassInput& Input = Inputs[Index];
+		if (IDataTypeRegistry::Get().IsLiteralTypeSupported(Input.TypeName, InDefaultLiteral.GetType()))
 		{
-			Input->DefaultLiteral = InDefaultLiteral;
+			Input.DefaultLiteral = InDefaultLiteral;
+			DocumentDelegates->InterfaceDelegates.OnInputDefaultChanged.Broadcast(Index);
 
 			// Set the input as no longer inheriting default for presets
 			if (IsPreset())
 			{
-				return SetGraphInputInheritsDefault(InputName, /*bInputInheritsDefault=*/false);
+				constexpr bool bInputInheritsDefault = false;
+				return SetGraphInputInheritsDefault(InputName, bInputInheritsDefault);
 			}
+
 			return true;
 		}
-		UE_LOG(LogMetaSound, Error, TEXT("Attempting to set graph input of type '%s' with unsupported literal type"), *Input->TypeName.ToString());
+		UE_LOG(LogMetaSound, Error, TEXT("Attempting to set graph input of type '%s' with unsupported literal type"), *Input.TypeName.ToString());
 	}
 
 	return false;
diff --git a/Engine/Plugins/Runtime/Metasound/Source/MetasoundFrontend/Private/MetasoundFrontendDocumentCache.cpp b/Engine/Plugins/Runtime/Metasound/Source/MetasoundFrontend/Private/MetasoundFrontendDocumentCache.cpp
--- a/Engine/Plugins/Runtime/Metasound/Source/MetasoundFrontend/Private/MetasoundFrontendDocumentCache.cpp
+++ b/Engine/Plugins/Runtime/Metasound/Source/MetasoundFrontend/Private/MetasoundFrontendDocumentCache.cpp
@@ -70,6 +70,10 @@ namespace Metasound::Frontend
 		}
 	}
 
+	FDocumentCache::FDocumentCache()
+	{
+	}
+
 	FDocumentCache::FDocumentCache(const FMetasoundFrontendDocument& InDocument, TSharedRef<FDocumentModifyDelegates> Delegates)
 		: Document(&InDocument)
 		, ModifyDelegates(Delegates)
@@ -85,10 +89,11 @@ namespace Metasound::Frontend
 		{
 			METASOUND_TRACE_CPUPROFILER_EVENT_SCOPE(Metasound::Frontend::FDocumentCache::Init_Prime);
 
+			TSharedRef<const FDocumentCache> ThisShared = StaticCastSharedRef<const FDocumentCache>(AsShared());
 			DependencyCache = MakeShared<FDocumentDependencyCache>(GetDocument());
-			EdgeCache = FDocumentGraphEdgeCache::Create(AsShared(), ModifyDelegates->EdgeDelegates);
-			NodeCache = FDocumentGraphNodeCache::Create(AsShared(), ModifyDelegates->NodeDelegates);
-			InterfaceCache = FDocumentGraphInterfaceCache::Create(AsShared(), ModifyDelegates->InterfaceDelegates);
+			EdgeCache = FDocumentGraphEdgeCache::Create(ThisShared, ModifyDelegates->EdgeDelegates);
+			NodeCache = FDocumentGraphNodeCache::Create(ThisShared, ModifyDelegates->NodeDelegates);
+			InterfaceCache = FDocumentGraphInterfaceCache::Create(ThisShared, ModifyDelegates->InterfaceDelegates);
 		}
 
 		ModifyDelegates->OnDependencyAdded.AddSP(this, &FDocumentCache::OnDependencyAdded);
@@ -181,7 +186,8 @@ namespace Metasound::Frontend
 	{
 		if (!EdgeCache.IsValid())
 		{
-			EdgeCache = FDocumentGraphEdgeCache::Create(AsShared(), ModifyDelegates->EdgeDelegates);
+			TSharedRef<const FDocumentCache> ThisShared = StaticCastSharedRef<const FDocumentCache>(AsShared());
+			EdgeCache = FDocumentGraphEdgeCache::Create(ThisShared, ModifyDelegates->EdgeDelegates);
 		}
 		return *EdgeCache;
 	}
@@ -190,7 +196,8 @@ namespace Metasound::Frontend
 	{
 		if (!NodeCache.IsValid())
 		{
-			NodeCache = FDocumentGraphNodeCache::Create(AsShared(), ModifyDelegates->NodeDelegates);
+			TSharedRef<const FDocumentCache> ThisShared = StaticCastSharedRef<const FDocumentCache>(AsShared());
+			NodeCache = FDocumentGraphNodeCache::Create(ThisShared, ModifyDelegates->NodeDelegates);
 		}
 		return *NodeCache;
 	}
@@ -199,11 +206,33 @@ namespace Metasound::Frontend
 	{
 		if (!InterfaceCache.IsValid())
 		{
-			InterfaceCache = FDocumentGraphInterfaceCache::Create(AsShared(), ModifyDelegates->InterfaceDelegates);
+			TSharedRef<const FDocumentCache> ThisShared = StaticCastSharedRef<const FDocumentCache>(AsShared());
+			InterfaceCache = FDocumentGraphInterfaceCache::Create(ThisShared, ModifyDelegates->InterfaceDelegates);
 		}
 		return *InterfaceCache;
 	}
 
+	int32 FDocumentCache::GetTransactionCount() const
+	{
+		int32 TotalCount = TransactionCount;
+		if (NodeCache.IsValid())
+		{
+			TotalCount += NodeCache->GetTransactionCount();
+		}
+
+		if (EdgeCache.IsValid())
+		{
+			TotalCount += EdgeCache->GetTransactionCount();
+		}
+
+		if (InterfaceCache.IsValid())
+		{
+			TotalCount += InterfaceCache->GetTransactionCount();
+		}
+
+		return TotalCount;
+	}
+
 	void FDocumentCache::OnDependencyAdded(int32 NewIndex)
 	{
 		using namespace Metasound::Frontend;
@@ -214,6 +243,8 @@ namespace Metasound::Frontend
 		FDocumentDependencyCache& Cache = GetDependencyCache();
 		Cache.KeyToIndex.Add(Key, NewIndex);
 		Cache.IDToIndex.Add(Dependency.ID, NewIndex);
+
+		++TransactionCount;
 	}
 
 	void FDocumentCache::OnRemoveSwappingDependency(int32 SwapIndex, int32 LastIndex)
@@ -231,6 +262,8 @@ namespace Metasound::Frontend
 
 		DocumentCachePrivate::RemoveSwapMapIndexChecked(SwapKey, LastKey, Cache.KeyToIndex);
 		DocumentCachePrivate::RemoveSwapMapIndexChecked(SwapDependency.ID, LastDependency.ID, Cache.IDToIndex);
+
+		++TransactionCount;
 	}
 
 	void FDocumentCache::OnRenamingDependencyClass(const int32 IndexBeingRenamed, const FMetasoundFrontendClassName& NewName)
@@ -246,9 +279,11 @@ namespace Metasound::Frontend
 		FDocumentDependencyCache& Cache = GetDependencyCache();
 		Cache.KeyToIndex.Add(NewKey, IndexBeingRenamed);
 		Cache.KeyToIndex.Remove(OldKey);
+
+		++TransactionCount;
 	}
 
-	FDocumentGraphInterfaceCache::FDocumentGraphInterfaceCache(TSharedRef<const IDocumentCache> ParentCache)
+	FDocumentGraphInterfaceCache::FDocumentGraphInterfaceCache(TSharedRef<const FDocumentCache> ParentCache)
 		: Parent(ParentCache)
 	{
 		const FMetasoundFrontendGraphClass& GraphClass = Parent->GetDocument().RootGraph;
@@ -265,7 +300,7 @@ namespace Metasound::Frontend
 		};
 	}
 
-	TSharedRef<FDocumentGraphInterfaceCache> FDocumentGraphInterfaceCache::Create(TSharedRef<const IDocumentCache> ParentCache, FInterfaceModifyDelegates& OutDelegates)
+	TSharedRef<FDocumentGraphInterfaceCache> FDocumentGraphInterfaceCache::Create(TSharedRef<const FDocumentCache> ParentCache, FInterfaceModifyDelegates& OutDelegates)
 	{
 		METASOUND_TRACE_CPUPROFILER_EVENT_SCOPE(Metasound::Frontend::FDocumentGraphInterfaceCache::Create);
 
@@ -297,10 +332,15 @@ namespace Metasound::Frontend
 		return nullptr;
 	}
 
-	void FDocumentGraphInterfaceCache::Init(FInterfaceModifyDelegates& OutDelegates)
+	int32 FDocumentGraphInterfaceCache::GetTransactionCount() const
 	{
+		return TransactionCount;
+	}
 
+	void FDocumentGraphInterfaceCache::Init(FInterfaceModifyDelegates& OutDelegates)
+	{
 		OutDelegates.OnInputAdded.AddSP(this, &FDocumentGraphInterfaceCache::OnInputAdded);
+		OutDelegates.OnInputDefaultChanged.AddSP(this, &FDocumentGraphInterfaceCache::OnInputDefaultChanged);
 		OutDelegates.OnOutputAdded.AddSP(this, &FDocumentGraphInterfaceCache::OnOutputAdded);
 		OutDelegates.OnRemovingInput.AddSP(this, &FDocumentGraphInterfaceCache::OnRemovingInput);
 		OutDelegates.OnRemovingOutput.AddSP(this, &FDocumentGraphInterfaceCache::OnRemovingOutput);
@@ -312,6 +352,13 @@ namespace Metasound::Frontend
 		const FMetasoundFrontendGraphClass& GraphClass = Document.RootGraph;
 		const FMetasoundFrontendClassInput& Input = GraphClass.Interface.Inputs[NewIndex];
 		InputNameToIndex.Add(Input.Name, NewIndex);
+
+		++TransactionCount;
+	}
+
+	void FDocumentGraphInterfaceCache::OnInputDefaultChanged(int32 NewIndex)
+	{
+		++TransactionCount;
 	}
 
 	void FDocumentGraphInterfaceCache::OnOutputAdded(int32 NewIndex)
@@ -320,6 +367,8 @@ namespace Metasound::Frontend
 		const FMetasoundFrontendGraphClass& GraphClass = Document.RootGraph;
 		const FMetasoundFrontendClassOutput& Output = GraphClass.Interface.Outputs[NewIndex];
 		OutputNameToIndex.Add(Output.Name, NewIndex);
+
+		++TransactionCount;
 	}
 
 	void FDocumentGraphInterfaceCache::OnRemovingInput(int32 IndexBeingRemoved)
@@ -328,6 +377,8 @@ namespace Metasound::Frontend
 		const FMetasoundFrontendGraphClass& GraphClass = Document.RootGraph;
 		const FMetasoundFrontendClassInput& Input = GraphClass.Interface.Inputs[IndexBeingRemoved];
 		InputNameToIndex.Remove(Input.Name);
+
+		++TransactionCount;
 	}
 
 	void FDocumentGraphInterfaceCache::OnRemovingOutput(int32 IndexBeingRemoved)
@@ -336,9 +387,16 @@ namespace Metasound::Frontend
 		const FMetasoundFrontendGraphClass& GraphClass = Document.RootGraph;
 		const FMetasoundFrontendClassOutput& Output = GraphClass.Interface.Outputs[IndexBeingRemoved];
 		OutputNameToIndex.Remove(Output.Name);
+
+		++TransactionCount;
 	}
 
-	FDocumentGraphNodeCache::FDocumentGraphNodeCache(TSharedRef<const IDocumentCache> ParentCache)
+	FDocumentGraphNodeCache::FDocumentGraphNodeCache()
+		: Parent(MakeShared<FDocumentCache>())
+	{
+	}
+
+	FDocumentGraphNodeCache::FDocumentGraphNodeCache(TSharedRef<const FDocumentCache> ParentCache)
 		: Parent(ParentCache)
 	{
 		const FMetasoundFrontendGraph& Graph = Parent->GetDocument().RootGraph.Graph;
@@ -360,7 +418,7 @@ namespace Metasound::Frontend
 		return IDToIndex.Contains(InNodeID);
 	}
 
-	TSharedRef<FDocumentGraphNodeCache> FDocumentGraphNodeCache::Create(TSharedRef<const IDocumentCache> ParentCache, FNodeModifyDelegates& OutDelegates)
+	TSharedRef<FDocumentGraphNodeCache> FDocumentGraphNodeCache::Create(TSharedRef<const FDocumentCache> ParentCache, FNodeModifyDelegates& OutDelegates)
 	{
 		METASOUND_TRACE_CPUPROFILER_EVENT_SCOPE(Metasound::Frontend::FDocumentGraphNodeCache::Create);
 
@@ -399,6 +457,11 @@ namespace Metasound::Frontend
 		return nullptr;
 	}
 
+	int32 FDocumentGraphNodeCache::GetTransactionCount() const
+	{
+		return TransactionCount;
+	}
+
 	bool FDocumentGraphNodeCache::ContainsNodesOfClassID(const FGuid& InClassID) const
 	{
 		if (const TArray<int32>* NodeIndices = ClassIDToNodeIndices.Find(InClassID))
@@ -683,7 +746,9 @@ namespace Metasound::Frontend
 	void FDocumentGraphNodeCache::Init(FNodeModifyDelegates& OutDelegates)
 	{
 		OutDelegates.OnNodeAdded.AddSP(this, &FDocumentGraphNodeCache::OnNodeAdded);
+		OutDelegates.OnNodeInputLiteralSet.AddSP(this, &FDocumentGraphNodeCache::OnNodeInputLiteralSet);
 		OutDelegates.OnRemoveSwappingNode.AddSP(this, &FDocumentGraphNodeCache::OnRemoveSwappingNode);
+		OutDelegates.OnRemovingNodeInputLiteral.AddSP(this, &FDocumentGraphNodeCache::OnRemovingNodeInputLiteral);
 	}
 
 	void FDocumentGraphNodeCache::OnNodeAdded(int32 InNewIndex)
@@ -693,6 +758,13 @@ namespace Metasound::Frontend
 		const FMetasoundFrontendNode& Node = Graph.Nodes[InNewIndex];
 		IDToIndex.Add(Node.GetID(), InNewIndex);
 		ClassIDToNodeIndices.FindOrAdd(Node.ClassID).Add(InNewIndex);
+
+		++TransactionCount;
+	}
+
+	void FDocumentGraphNodeCache::OnNodeInputLiteralSet(int32 /* NodeIndex */, int32 /* VertexIndex */, int32 /* LiteralIndex */)
+	{
+		++TransactionCount;
 	}
 
 	void FDocumentGraphNodeCache::OnRemoveSwappingNode(int32 SwapIndex, int32 LastIndex)
@@ -706,9 +778,21 @@ namespace Metasound::Frontend
 
 		DocumentCachePrivate::RemoveSwapMapArrayIndexChecked(SwapNode.ClassID, SwapIndex, LastNode.ClassID, LastIndex, ClassIDToNodeIndices);
 		DocumentCachePrivate::RemoveSwapMapIndexChecked(SwapNode.GetID(), LastNode.GetID(), IDToIndex);
+
+		++TransactionCount;
+	}
+
+	void FDocumentGraphNodeCache::OnRemovingNodeInputLiteral(int32 /* NodeIndex */, int32 /* VertexIndex */, int32 /* LiteralIndex */)
+	{
+		++TransactionCount;
+	}
+
+	FDocumentGraphEdgeCache::FDocumentGraphEdgeCache()
+		: Parent(MakeShared<const FDocumentCache>())
+	{
 	}
 
-	FDocumentGraphEdgeCache::FDocumentGraphEdgeCache(TSharedRef<const IDocumentCache> ParentCache)
+	FDocumentGraphEdgeCache::FDocumentGraphEdgeCache(TSharedRef<const FDocumentCache> ParentCache)
 		: Parent(ParentCache)
 	{
 		const TArray<FMetasoundFrontendEdge>& Edges = Parent->GetDocument().RootGraph.Graph.Edges;
@@ -733,7 +817,7 @@ namespace Metasound::Frontend
 		return false;
 	}
 
-	TSharedRef<FDocumentGraphEdgeCache> FDocumentGraphEdgeCache::Create(TSharedRef<const IDocumentCache> ParentCache, FEdgeModifyDelegates& OutDelegates)
+	TSharedRef<FDocumentGraphEdgeCache> FDocumentGraphEdgeCache::Create(TSharedRef<const FDocumentCache> ParentCache, FEdgeModifyDelegates& OutDelegates)
 	{
 		METASOUND_TRACE_CPUPROFILER_EVENT_SCOPE(Metasound::Frontend::FDocumentGraphEdgeCache::Create);
 
@@ -779,6 +863,11 @@ namespace Metasound::Frontend
 		return { };
 	}
 
+	int32 FDocumentGraphEdgeCache::GetTransactionCount() const
+	{
+		return TransactionCount;
+	}
+
 	void FDocumentGraphEdgeCache::Init(FEdgeModifyDelegates& OutDelegates)
 	{
 		OutDelegates.OnEdgeAdded.AddSP(this, &FDocumentGraphEdgeCache::OnEdgeAdded);
@@ -802,6 +891,8 @@ namespace Metasound::Frontend
 		const FMetasoundFrontendEdge& NewEdge = Graph.Edges[InNewIndex];
 		InputToEdgeIndex.Add(NewEdge.GetToVertexHandle(), InNewIndex);
 		OutputToEdgeIndices.FindOrAdd(NewEdge.GetFromVertexHandle()).Add(InNewIndex);
+
+		++TransactionCount;
 	}
 
 	void FDocumentGraphEdgeCache::OnRemoveSwappingEdge(int32 SwapIndex, int32 LastIndex)
@@ -813,12 +904,12 @@ namespace Metasound::Frontend
 		const FMetasoundFrontendEdge& LastEdge = Edges[LastIndex];
 		DocumentCachePrivate::RemoveSwapMapIndexChecked(SwapEdge.GetToVertexHandle(), LastEdge.GetToVertexHandle(), InputToEdgeIndex);
 		DocumentCachePrivate::RemoveSwapMapArrayIndexChecked(SwapEdge.GetFromVertexHandle(), SwapIndex, LastEdge.GetFromVertexHandle(), LastIndex, OutputToEdgeIndices);
+
+		++TransactionCount;
 	}
 
 	FDocumentCache::FDocumentDependencyCache::FDocumentDependencyCache(const FMetasoundFrontendDocument& InDocument)
 	{
-		METASOUND_TRACE_CPUPROFILER_EVENT_SCOPE(Metasound::Frontend::FDocumentDependencyCache::Create);
-
 		for (int32 Index = 0; Index < InDocument.Dependencies.Num(); ++Index)
 		{
 			const FMetasoundFrontendClass& Class = InDocument.Dependencies[Index];
diff --git a/Engine/Plugins/Runtime/Metasound/Source/MetasoundFrontend/Private/MetasoundFrontendDocumentCache.h b/Engine/Plugins/Runtime/Metasound/Source/MetasoundFrontend/Private/MetasoundFrontendDocumentCache.h
--- a/Engine/Plugins/Runtime/Metasound/Source/MetasoundFrontend/Private/MetasoundFrontendDocumentCache.h
+++ b/Engine/Plugins/Runtime/Metasound/Source/MetasoundFrontend/Private/MetasoundFrontendDocumentCache.h
@@ -16,15 +16,25 @@ namespace Metasound::Frontend
 	// Forward Declarations
 	class FDocumentCache;
 
+	struct FDocumentCacheTransaction
+	{
+		int32 DocumentTransactionID = INDEX_NONE;
+		int32 InterfaceTransactionID = INDEX_NONE;
+		int32 NodeTransactionID = INDEX_NONE;
+		int32 EdgeTransactionID = INDEX_NONE;
+	};
+
 	class FDocumentGraphEdgeCache : public IDocumentGraphEdgeCache
 	{
 		template <typename ObjectType, ESPMode Mode>
 		friend class SharedPointerInternals::TIntrusiveReferenceController;
 
-		FDocumentGraphEdgeCache(TSharedRef<const IDocumentCache> ParentCache);
+		// Only exists to make TSharedFromThis happy, but should never be called
+		FDocumentGraphEdgeCache();
+		FDocumentGraphEdgeCache(TSharedRef<const FDocumentCache> ParentCache);
 
 	public:
-		static TSharedRef<FDocumentGraphEdgeCache> Create(TSharedRef<const IDocumentCache> ParentCache, FEdgeModifyDelegates& OutDelegates);
+		static TSharedRef<FDocumentGraphEdgeCache> Create(TSharedRef<const FDocumentCache> ParentCache, FEdgeModifyDelegates& OutDelegates);
 		virtual ~FDocumentGraphEdgeCache() = default;
 
 		// IDocumentGraphEdgeCache implementation
@@ -37,6 +47,7 @@ namespace Metasound::Frontend
 		virtual const int32* FindEdgeIndexToNodeInput(const FGuid& InNodeID, const FGuid& InVertexID) const override;
 		virtual const TArrayView<const int32> FindEdgeIndicesFromNodeOutput(const FGuid& InNodeID, const FGuid& InVertexID) const override;
 
+		int32 GetTransactionCount() const;
 
 	private:
 		void Init(FEdgeModifyDelegates& OutDelegates);
@@ -49,7 +60,9 @@ namespace Metasound::Frontend
 		// Cache of Input NodeId/VertexId pairs to associated edge indices
 		TMap<FMetasoundFrontendVertexHandle, int32> InputToEdgeIndex;
 
-		TSharedPtr<const IDocumentCache> Parent;
+		TSharedRef<const FDocumentCache> Parent;
+
+		int32 TransactionCount = 0;
 	};
 
 
@@ -58,10 +71,12 @@ namespace Metasound::Frontend
 		template <typename ObjectType, ESPMode Mode>
 		friend class SharedPointerInternals::TIntrusiveReferenceController;
 
-		FDocumentGraphNodeCache(TSharedRef<const IDocumentCache> ParentCache);
+		// Only exists to make TSharedFromThis happy, but should never be called
+		FDocumentGraphNodeCache();
+		FDocumentGraphNodeCache(TSharedRef<const FDocumentCache> ParentCache);
 
 	public:
-		static TSharedRef<FDocumentGraphNodeCache> Create(TSharedRef<const IDocumentCache> ParentCache, FNodeModifyDelegates& OutDelegates);
+		static TSharedRef<FDocumentGraphNodeCache> Create(TSharedRef<const FDocumentCache> ParentCache, FNodeModifyDelegates& OutDelegates);
 		virtual ~FDocumentGraphNodeCache() = default;
 
 		// IDocumentGraphNodeCache implementation
@@ -85,10 +100,14 @@ namespace Metasound::Frontend
 		virtual const FMetasoundFrontendVertex* FindReroutedOutputVertex(const FGuid& InNodeID, const FGuid& InVertexID, const FMetasoundFrontendNode** ConnectedNodes = nullptr, bool* bOutIsRerouted = nullptr) const override;
 		virtual const FMetasoundFrontendVertex* FindReroutedOutputVertex(const FGuid& InNodeID, FName InVertexName, const FMetasoundFrontendNode** ConnectedNodes = nullptr, bool* bOutIsRerouted = nullptr) const override;
 
+		int32 GetTransactionCount() const;
+
 	private:
 		void Init(FNodeModifyDelegates& OutDelegates);
 		void OnNodeAdded(int32 NewIndex);
+		void OnNodeInputLiteralSet(int32 NodeIndex, int32 VertexIndex, int32 LiteralIndex);
 		void OnRemoveSwappingNode(int32 IndexBeingRemoved, int32 LastIndex);
+		void OnRemovingNodeInputLiteral(int32 NodeIndex, int32 VertexIndex, int32 LiteralIndex);
 
 		// Cache of NodeId to array index of node
 		TSortedMap<FGuid, int32> IDToIndex;
@@ -96,7 +115,9 @@ namespace Metasound::Frontend
 		// Cache of ClassID to referencing node indices
 		TSortedMap<FGuid, TArray<int32>> ClassIDToNodeIndices;
 
-		TSharedPtr<const IDocumentCache> Parent;
+		TSharedRef<const FDocumentCache> Parent;
+
+		int32 TransactionCount = 0;
 	};
 
 
@@ -105,19 +126,22 @@ namespace Metasound::Frontend
 		template <typename ObjectType, ESPMode Mode>
 		friend class SharedPointerInternals::TIntrusiveReferenceController;
 
-		FDocumentGraphInterfaceCache(TSharedRef<const IDocumentCache> ParentCache);
+		FDocumentGraphInterfaceCache(TSharedRef<const FDocumentCache> ParentCache);
 
 	public:
-		static TSharedRef<FDocumentGraphInterfaceCache> Create(TSharedRef<const IDocumentCache> ParentCache, FInterfaceModifyDelegates& OutDelegates);
+		static TSharedRef<FDocumentGraphInterfaceCache> Create(TSharedRef<const FDocumentCache> ParentCache, FInterfaceModifyDelegates& OutDelegates);
 		virtual ~FDocumentGraphInterfaceCache() = default;
 
 		// IDocumentGraphInterfaceCache implementation
 		virtual const FMetasoundFrontendClassInput* FindInput(FName InputName) const override;
 		virtual const FMetasoundFrontendClassOutput* FindOutput(FName OutputName) const override;
 
+		int32 GetTransactionCount() const;
+
 	private:
 		void Init(FInterfaceModifyDelegates& OutDelegates);
 		void OnInputAdded(int32 NewIndex);
+		void OnInputDefaultChanged(int32 NewIndex);
 		void OnOutputAdded(int32 NewIndex);
 		void OnRemovingInput(int32 IndexBeingRemoved);
 		void OnRemovingOutput(int32 IndexBeingRemoved);
@@ -128,7 +152,9 @@ namespace Metasound::Frontend
 		// Cache of Output name to array index of output
 		TMap<FName, int32> OutputNameToIndex;
 
-		TSharedPtr<const IDocumentCache> Parent;
+		TSharedRef<const FDocumentCache> Parent;
+
+		int32 TransactionCount = 0;
 	};
 
 
@@ -137,6 +163,8 @@ namespace Metasound::Frontend
 		template <typename ObjectType, ESPMode Mode>
 		friend class SharedPointerInternals::TIntrusiveReferenceController;
 
+		// Only exists to make TSharedFromThis happy, but should never be called
+		FDocumentCache();
 		FDocumentCache(const FMetasoundFrontendDocument& InDocument, TSharedRef<FDocumentModifyDelegates> Delegates);
 
 	public:
@@ -156,6 +184,8 @@ namespace Metasound::Frontend
 		virtual const IDocumentGraphNodeCache& GetNodeCache() const override;
 		virtual const IDocumentGraphInterfaceCache& GetInterfaceCache() const override;
 
+		int32 GetTransactionCount() const;
+
 	private:
 		void Init(bool bPrimeCache);
 
@@ -189,5 +219,8 @@ namespace Metasound::Frontend
 
 		const FMetasoundFrontendDocument* Document = nullptr;
 		TSharedRef<FDocumentModifyDelegates> ModifyDelegates;
+
+		// Number of transactions processed since builder was instantiated
+		int32 TransactionCount = 0;
 	};
 } // namespace Metasound::Frontend
diff --git a/Engine/Plugins/Runtime/Metasound/Source/MetasoundFrontend/Public/MetasoundFrontendDocumentBuilder.h b/Engine/Plugins/Runtime/Metasound/Source/MetasoundFrontend/Public/MetasoundFrontendDocumentBuilder.h
--- a/Engine/Plugins/Runtime/Metasound/Source/MetasoundFrontend/Public/MetasoundFrontendDocumentBuilder.h
+++ b/Engine/Plugins/Runtime/Metasound/Source/MetasoundFrontend/Public/MetasoundFrontendDocumentBuilder.h
@@ -213,6 +213,9 @@ public:
 
 	const Metasound::Frontend::FDocumentModifyDelegates& GetDocumentDelegates() const;
 	const IMetaSoundDocumentInterface& GetDocumentInterface() const;
+
+	int32 GetTransactionCount() const;
+
 	EMetasoundFrontendVertexAccessType GetNodeInputAccessType(const FGuid& InNodeID, const FGuid& InVertexID) const;
 	const FMetasoundFrontendLiteral* GetNodeInputClassDefault(const FGuid& InNodeID, const FGuid& InVertexID) const;
 	const FMetasoundFrontendLiteral* GetNodeInputDefault(const FGuid& InNodeID, const FGuid& InVertexID) const;
diff --git a/Engine/Plugins/Runtime/Metasound/Source/MetasoundFrontend/Public/MetasoundFrontendDocumentModifyDelegates.h b/Engine/Plugins/Runtime/Metasound/Source/MetasoundFrontend/Public/MetasoundFrontendDocumentModifyDelegates.h
--- a/Engine/Plugins/Runtime/Metasound/Source/MetasoundFrontend/Public/MetasoundFrontendDocumentModifyDelegates.h
+++ b/Engine/Plugins/Runtime/Metasound/Source/MetasoundFrontend/Public/MetasoundFrontendDocumentModifyDelegates.h
@@ -21,6 +21,7 @@ namespace Metasound::Frontend
 		FOnMetaSoundFrontendDocumentMutateInterfaceArray OnRemovingInterface;
 
 		FOnMetaSoundFrontendDocumentMutateArray OnInputAdded;
+		FOnMetaSoundFrontendDocumentMutateArray OnInputDefaultChanged;
 		FOnMetaSoundFrontendDocumentMutateArray OnRemovingInput;
 
 		FOnMetaSoundFrontendDocumentMutateArray OnOutputAdded;
