[Metasound] Follow up changes to operator pooling behavior.

#swarm https://p4-swarm.epicgames.net/reviews/29255703

diff --git a/Engine/Plugins/Runtime/Metasound/Source/MetasoundEngine/Private/MetasoundOperatorCacheSubsystem.cpp b/Engine/Plugins/Runtime/Metasound/Source/MetasoundEngine/Private/MetasoundOperatorCacheSubsystem.cpp
--- a/Engine/Plugins/Runtime/Metasound/Source/MetasoundEngine/Private/MetasoundOperatorCacheSubsystem.cpp
+++ b/Engine/Plugins/Runtime/Metasound/Source/MetasoundEngine/Private/MetasoundOperatorCacheSubsystem.cpp
@@ -15,14 +15,15 @@ static TOptional<Metasound::FMetasoundGeneratorInitParams> CreateInitParams(UMet
 	using namespace Metasound::Engine;
 	using namespace Metasound::SourcePrivate;
 
+	// InMetaSound was null
 	if (!ensure(InMetaSound))
 	{
-		return {}; // InMetaSound was null
+		return {};
 	}
-
-	if (ensure(!InMetaSound->IsDynamic()))
+	// we cannot precache dynamic metasounds
+	if (!ensure(!InMetaSound->IsDynamic()))
 	{
-		return {}; // we cannot precache dynamic metasounds
+		return {};
 	}
 
 	FOperatorSettings InSettings = InMetaSound->GetOperatorSettings(static_cast<FSampleRate>(InParams.SampleRate));
@@ -87,11 +88,12 @@ void UMetaSoundCacheSubsystem::PrecacheMetaSound(UMetaSoundSource* InMetaSound,
 
 	TSharedPtr<FOperatorPool> OperatorPool;
 	IMetasoundGeneratorModule* Module = FModuleManager::GetModulePtr<IMetasoundGeneratorModule>("MetasoundGenerator");
-	if (ensure(Module))
+	if (!ensure(Module))
 	{
-		OperatorPool = Module->GetOperatorPool();
+		return;
 	}
 
+	OperatorPool = Module->GetOperatorPool();
 	const FMixerDevice* MixerDevice = GetMixerDevice();
 
 	if (!ensure(MixerDevice && InMetaSound && OperatorPool))
@@ -123,8 +125,56 @@ void UMetaSoundCacheSubsystem::PrecacheMetaSound(UMetaSoundSource* InMetaSound,
 		  MoveTemp(InitParams.GetValue())
 		, InMetaSound->GetRegistryKey()
 		, FSoftObjectPath(InMetaSound->GetOwningAsset())
+		, InMetaSound->AssetClassID
 		, InNumInstances
 	);
 
 	OperatorPool->BuildAndAddOperator(MoveTemp(Data));
+}
+
+void UMetaSoundCacheSubsystem::TouchOrPrecacheMetaSound(UMetaSoundSource* InMetaSound, int32 InNumInstances)
+{
+	using namespace Metasound;
+
+	TSharedPtr<FOperatorPool> OperatorPool;
+	IMetasoundGeneratorModule* Module = FModuleManager::GetModulePtr<IMetasoundGeneratorModule>("MetasoundGenerator");
+	if (!ensure(Module))
+	{
+		return;
+	}
+
+	OperatorPool = Module->GetOperatorPool();
+	if (!ensure(OperatorPool))
+	{
+		return;
+	}
+
+	// get the number of instances already in the cache
+	const int32 NumInCache = OperatorPool->GetNumCachedOperatorsWithAssetClassID(InMetaSound->AssetClassID);
+
+	// move pre-existing to the top of the cache
+	OperatorPool->TouchOperatorsViaAssetClassID(InMetaSound->AssetClassID, FMath::Min(NumInCache, InNumInstances));
+
+	// build the difference (InNumInstances - existing)
+	const int32 NumToBuild = InNumInstances - NumInCache;
+	if (NumToBuild > 0)
+	{
+		PrecacheMetaSound(InMetaSound, NumToBuild);
+	}
+}
+
+void UMetaSoundCacheSubsystem::RemoveCachedOperatorsForMetaSound(UMetaSoundSource* InMetaSound)
+{
+	using namespace Metasound;
+
+	IMetasoundGeneratorModule* Module = FModuleManager::GetModulePtr<IMetasoundGeneratorModule>("MetasoundGenerator");
+	if (!ensure(Module))
+	{
+		return;
+	}
+
+	if (TSharedPtr<FOperatorPool> OperatorPool = Module->GetOperatorPool())
+	{
+		OperatorPool->RemoveOperatorsWithAssetClassID(InMetaSound->AssetClassID);
+	}
 }
\ No newline at end of file
diff --git a/Engine/Plugins/Runtime/Metasound/Source/MetasoundEngine/Public/MetasoundOperatorCacheSubsystem.h b/Engine/Plugins/Runtime/Metasound/Source/MetasoundEngine/Public/MetasoundOperatorCacheSubsystem.h
--- a/Engine/Plugins/Runtime/Metasound/Source/MetasoundEngine/Public/MetasoundOperatorCacheSubsystem.h
+++ b/Engine/Plugins/Runtime/Metasound/Source/MetasoundEngine/Public/MetasoundOperatorCacheSubsystem.h
@@ -33,9 +33,19 @@ public:
 	 
 	/* Builds the requested number of MetaSound operators (asynchronously) and puts them in the pool for playback.
 	(If these operators are not yet available when the MetaSound attepmts to play, one will be created Independent of this request.) */
-	UFUNCTION(BlueprintCallable, Category = "Audio")
+	UFUNCTION(BlueprintCallable, Category = "MetaSound")
 	void PrecacheMetaSound(UMetaSoundSource* InMetaSound, int32 InNumInstances = 1);
 
+	/* same as PrecacheMetaSound except cached operator that already exists in the cache will be moved to the top instead of building,
+	any operators that we couldn't move to the top, will be built.
+	(i.e. if 2 operators are already cached and Num Instances is 4, it will construct 2 and move the existing 2 to the top of the cache) */
+	UFUNCTION(BlueprintCallable, Category = "MetaSound")
+	void TouchOrPrecacheMetaSound(UMetaSoundSource* InMetaSound, int32 InNumInstances = 1);
+
+	/* Clear the operator pool of any operators associated with the given MetaSound */
+	UFUNCTION(BlueprintCallable, Category = "MetaSound")
+	void RemoveCachedOperatorsForMetaSound(UMetaSoundSource* InMetaSound);
+
 private:
 	FSoundGeneratorInitParams BuildParams;
 }; //UMetaSoundCacheSubsystem
diff --git a/Engine/Plugins/Runtime/Metasound/Source/MetasoundGenerator/Private/MetasoundGenerator.cpp b/Engine/Plugins/Runtime/Metasound/Source/MetasoundGenerator/Private/MetasoundGenerator.cpp
--- a/Engine/Plugins/Runtime/Metasound/Source/MetasoundGenerator/Private/MetasoundGenerator.cpp
+++ b/Engine/Plugins/Runtime/Metasound/Source/MetasoundGenerator/Private/MetasoundGenerator.cpp
@@ -816,13 +816,12 @@ PRAGMA_ENABLE_DEPRECATION_WARNINGS
 		if (bIsOperatorPoolEnabled)
 		{
 			bUseOperatorPool = ConsoleVariables::bEnableExperimentalOperatorPool || MetasoundGeneratorPrivate::HasOneShotInterface(InParams.Graph->GetVertexInterface());
-
-			if (bUseOperatorPool)
-			{
-				bDidUseCachedOperator = TryUseCachedOperator(InParams, true /* bTriggerGenerator */);
-			}
 		}
 
+		// check the cache for manually pre-cached operators
+		// even if the cvars disable automatic cache population.
+		bDidUseCachedOperator = TryUseCachedOperator(InParams, true /* bTriggerGenerator */);
+
 		if (!bDidUseCachedOperator)
 		{
 			BuildGraph(MoveTemp(InParams));
@@ -857,7 +856,8 @@ PRAGMA_ENABLE_DEPRECATION_WARNINGS
 			FOperatorAndInputs GraphOperatorAndInputs = OperatorPool->ClaimOperator(OperatorID);
 			if (GraphOperatorAndInputs.Operator.IsValid())
 			{
-				UE_LOG(LogMetasoundGenerator, VeryVerbose, TEXT("Using cached operator %s for MetaSound %s"), *LexToString(OperatorID), *InInitParams.MetaSoundName); 
+				UE_LOG(LogMetasoundGenerator, VeryVerbose, TEXT("Using cached operator %s for MetaSound %s"), *LexToString(OperatorID), *InInitParams.MetaSoundName);
+				bUseOperatorPool = true; // raise this flag to make sure we put the operator back in the pool (regardless of CVAR state)
 
 				// Apply and default inputs to the operator.
 				GeneratorBuilder::ApplyAudioParameters(OperatorSettings, MoveTemp(InInitParams.DefaultParameters), GraphOperatorAndInputs.Inputs);
@@ -890,16 +890,14 @@ PRAGMA_ENABLE_DEPRECATION_WARNINGS
 		{
 			FMetasoundGeneratorModule& Module = FModuleManager::GetModuleChecked<FMetasoundGeneratorModule>("MetasoundGenerator");
 			TSharedPtr<FOperatorPool> OperatorPool = Module.GetOperatorPool();
-			if (OperatorPool.IsValid())
+			TUniquePtr<IOperator> GraphOperator = ReleaseGraphOperator();
+
+			if (OperatorPool.IsValid() && GraphOperator.IsValid())
 			{
-				TUniquePtr<IOperator> GraphOperator = ReleaseGraphOperator();
+				// Release graph operator and input data to the cache
+				UE_LOG(LogMetasoundGenerator, VeryVerbose, TEXT("Caching operator %s"), *LexToString(OperatorID));
+				OperatorPool->AddOperator(OperatorID, MoveTemp(GraphOperator), ReleaseInputVertexData());
 
-				if (GraphOperator.IsValid())
-				{
-					// Release graph operator and input data to the cache
-					UE_LOG(LogMetasoundGenerator, VeryVerbose, TEXT("Caching operator %s"), *LexToString(OperatorID));
-					OperatorPool->AddOperator(OperatorID, MoveTemp(GraphOperator), ReleaseInputVertexData());
-				}
 			}
 		}
 
diff --git a/Engine/Plugins/Runtime/Metasound/Source/MetasoundGenerator/Private/MetasoundOperatorCache.cpp b/Engine/Plugins/Runtime/Metasound/Source/MetasoundGenerator/Private/MetasoundOperatorCache.cpp
--- a/Engine/Plugins/Runtime/Metasound/Source/MetasoundGenerator/Private/MetasoundOperatorCache.cpp
+++ b/Engine/Plugins/Runtime/Metasound/Source/MetasoundGenerator/Private/MetasoundOperatorCache.cpp
@@ -154,11 +154,13 @@ namespace Metasound
 		  FMetasoundGeneratorInitParams&& InInitParams
 		, Frontend::FNodeRegistryKey InRegistryKey
 		, FSoftObjectPath InAssetPath
+		, FGuid InAssetClassID
 		, int32 InNumInstances
 	)
 	: InitParams(InInitParams)
 	, RegistryKey(InRegistryKey)
 	, AssetPath(InAssetPath)
+	, AssetClassID(InAssetClassID)
 	, NumInstances(InNumInstances)
 	{}
 
@@ -234,7 +236,6 @@ namespace Metasound
 			return;
 		}
 
-
 		FMetasoundGeneratorModule& Module = FModuleManager::GetModuleChecked<FMetasoundGeneratorModule>("MetasoundGenerator");
 		AsyncTask(ENamedThreads::AnyThread, [PreCacheData = MoveTemp(InBuildData), OperatorPool = Module.GetOperatorPool()] ()
 		{
@@ -257,11 +258,47 @@ namespace Metasound
 				FOperatorAndInputs OperatorAndInputs = GeneratorBuilder::BuildGraphOperator(PreCacheData->InitParams.OperatorSettings, PreCacheData->InitParams, BuildResults);
 				GeneratorBuilder::LogBuildErrors(PreCacheData->InitParams.MetaSoundName, BuildResults);
 	
-				OperatorPool->AddOperator(PreCacheData->InitParams.Graph->GetInstanceID(), MoveTemp(OperatorAndInputs));	
+				const FGuid& GraphID = PreCacheData->InitParams.Graph->GetInstanceID();
+				OperatorPool->AddOperator(GraphID, MoveTemp(OperatorAndInputs));
+				OperatorPool->AddAssetIdToGraphIdLookUp(PreCacheData->AssetClassID, GraphID);
 			}
 		});
 	}
 
+	void FOperatorPool::TouchOperators(const FGuid& InOpeoratorID, const int32& NumToTouch)
+	{
+		const int32 NumToMove = FMath::Min(NumToTouch, GetNumCachedOperatorsWithID(InOpeoratorID));
+		if (!NumToMove)
+		{
+			return;
+		}
+
+		FScopeLock Lock(&CriticalSection);
+
+		// add to the "top" (end)
+		for (int32 i = 0; i < NumToMove; ++i)
+		{
+			Stack.Add(InOpeoratorID);
+		}
+
+		// remove from the "bottom" (begining)
+		for (int32 i = 0; i < NumToMove; ++i)
+		{
+			Stack.RemoveSingle(InOpeoratorID);
+		}
+
+	}
+
+	void FOperatorPool::TouchOperatorsViaAssetClassID(const FGuid& InAssetClassID, const int32& NumToTouch)
+	{
+		FScopeLock Lock(&CriticalSection);
+		FGuid* GraphIdPtr = AssetIdToGraphIdLookUp.Find(InAssetClassID);
+		if (GraphIdPtr)
+		{
+			TouchOperators(*GraphIdPtr, NumToTouch);
+		}
+	}
+
 	void FOperatorPool::RemoveOperatorsWithID(const FGuid& InOperatorID)
 	{
 		FScopeLock Lock(&CriticalSection);
@@ -269,6 +306,17 @@ namespace Metasound
 		Stack.Remove(InOperatorID);
 	}
 
+	void FOperatorPool::RemoveOperatorsWithAssetClassID(const FGuid& InAssetClassID)
+	{
+		FScopeLock Lock(&CriticalSection);
+		FGuid* GraphIdPtr = AssetIdToGraphIdLookUp.Find(InAssetClassID);
+		if (GraphIdPtr)
+		{
+			RemoveOperatorsWithID(*GraphIdPtr);
+			AssetIdToGraphIdLookUp.Remove(InAssetClassID);
+		}
+	}
+		
 	void FOperatorPool::SetMaxNumOperators(uint32 InMaxNumOperators)
 	{
 		FScopeLock Lock(&CriticalSection);
@@ -276,6 +324,34 @@ namespace Metasound
 		Trim();
 	}
 
+	int32 FOperatorPool::GetNumCachedOperatorsWithID(const FGuid& InOperatorID) const
+	{
+		FScopeLock Lock(&CriticalSection);
+		if (TArray<FOperatorAndInputs> const* OperatorsWithID = Operators.Find(InOperatorID))
+		{
+			return OperatorsWithID->Num();
+		}
+
+		return 0;
+	}
+
+	int32 FOperatorPool::GetNumCachedOperatorsWithAssetClassID(const FGuid& InAssetClassID) const
+	{
+		FScopeLock Lock(&CriticalSection);
+		if (const FGuid* GraphIdPtr = AssetIdToGraphIdLookUp.Find(InAssetClassID))
+		{
+			return GetNumCachedOperatorsWithID(*GraphIdPtr);
+		}
+
+		return 0;
+	}
+
+	void FOperatorPool::AddAssetIdToGraphIdLookUp(const FGuid& InAssetClassID, const FGuid& InOperatorID)
+	{
+		FScopeLock Lock(&CriticalSection);
+		AssetIdToGraphIdLookUp.Add(InAssetClassID, InOperatorID);
+	}
+
 #if METASOUND_OPERATORCACHEPROFILER_ENABLED
 	void FOperatorPool::UpdateHitRateTracker()
 	{
@@ -285,6 +361,7 @@ namespace Metasound
 
 	void FOperatorPool::Trim()
 	{
+		FScopeLock Lock(&CriticalSection);
 		int32 NumToTrim = Stack.Num() - Settings.MaxNumOperators;
 		if (NumToTrim > 0)
 		{
@@ -308,5 +385,7 @@ namespace Metasound
 			}
 			Stack.RemoveAt(0, NumToTrim);
 		}
+
+		// todo: prune AssetIdToGraphIdLookUp?
 	}
 } // namespace Metasound
\ No newline at end of file
diff --git a/Engine/Plugins/Runtime/Metasound/Source/MetasoundGenerator/Public/MetasoundOperatorCache.h b/Engine/Plugins/Runtime/Metasound/Source/MetasoundGenerator/Public/MetasoundOperatorCache.h
--- a/Engine/Plugins/Runtime/Metasound/Source/MetasoundGenerator/Public/MetasoundOperatorCache.h
+++ b/Engine/Plugins/Runtime/Metasound/Source/MetasoundGenerator/Public/MetasoundOperatorCache.h
@@ -68,6 +68,7 @@ namespace Metasound
 		FMetasoundGeneratorInitParams InitParams;
 		Frontend::FNodeRegistryKey RegistryKey;
 		FSoftObjectPath AssetPath;
+		FGuid AssetClassID;
 		int32 NumInstances;
 
 		FOperatorBuildData() = delete;
@@ -75,6 +76,7 @@ namespace Metasound
 			  FMetasoundGeneratorInitParams&& InInitParams
 			, Frontend::FNodeRegistryKey InRegistryKey
 			, FSoftObjectPath InAssetPath
+			, FGuid InAssetID
 			, int32 InNumInstances = 1
 		);
 
@@ -86,11 +88,24 @@ namespace Metasound
 	{
 	public:
 		FOperatorPool(const FOperatorPoolSettings& InSettings);
+
 		FOperatorAndInputs ClaimOperator(const FGuid& InOperatorID);
+
 		void AddOperator(const FGuid& InOperatorID, TUniquePtr<IOperator>&& InOperator, FInputVertexInterfaceData&& InputData);
 		void AddOperator(const FGuid& InOperatorID, FOperatorAndInputs && OperatorAndInputs);
 		void BuildAndAddOperator(TUniquePtr<FOperatorBuildData> InBuildData);
+
+		void TouchOperators(const FGuid& InOpeoratorID, const int32& NumToTouch = 1);
+		void TouchOperatorsViaAssetClassID(const FGuid& InAssetClassID, const int32& NumToTouch = 1);
+
 		void RemoveOperatorsWithID(const FGuid& InOperatorID);
+		void RemoveOperatorsWithAssetClassID(const FGuid& InAssetClassID);
+
+		int32 GetNumCachedOperatorsWithID(const FGuid& InOperatorID) const;
+		int32 GetNumCachedOperatorsWithAssetClassID(const FGuid& InAssetClassID) const;
+
+		void AddAssetIdToGraphIdLookUp(const FGuid& InAssetClassID, const FGuid& InOperatorID);
+
 		void SetMaxNumOperators(uint32 InMaxNumOperators);
 #if METASOUND_OPERATORCACHEPROFILER_ENABLED
 		void UpdateHitRateTracker();
@@ -100,13 +115,14 @@ namespace Metasound
 		void Trim();
 
 		FOperatorPoolSettings Settings;
-		FCriticalSection CriticalSection;
+		mutable FCriticalSection CriticalSection;
 
 #if METASOUND_OPERATORCACHEPROFILER_ENABLED
 		OperatorPoolPrivate::FWindowedHitRate HitRateTracker;
 #endif // #if METASOUND_OPERATORCACHEPROFILER_ENABLED
 
 		TMap<FGuid, TArray<FOperatorAndInputs>> Operators;
+		TMap<FGuid, FGuid> AssetIdToGraphIdLookUp;
 		TArray<FGuid> Stack;
 	};
 
