[Metasounds] Optimize memory allocations in DSP objects for recycled metasounds.

#swarm https://p4-swarm.epicgames.net/reviews/29150979

diff --git a/Engine/Plugins/Runtime/Metasound/Source/MetasoundStandardNodes/Private/MetasoundDelayNode.cpp b/Engine/Plugins/Runtime/Metasound/Source/MetasoundStandardNodes/Private/MetasoundDelayNode.cpp
--- a/Engine/Plugins/Runtime/Metasound/Source/MetasoundStandardNodes/Private/MetasoundDelayNode.cpp
+++ b/Engine/Plugins/Runtime/Metasound/Source/MetasoundStandardNodes/Private/MetasoundDelayNode.cpp
@@ -155,7 +155,6 @@ namespace Metasound
 		FeedbackSample = 0.f;
 		PrevDelayTimeMsec = GetInputDelayTimeMsec();
 
-		DelayBuffer.Reset();
 		DelayBuffer.Init(InParams.OperatorSettings.GetSampleRate(), MaxDelayTimeSeconds);
 		DelayBuffer.SetDelayMsec(PrevDelayTimeMsec);
 
diff --git a/Engine/Plugins/Runtime/Metasound/Source/MetasoundStandardNodes/Private/MetasoundStereoDelayNode.cpp b/Engine/Plugins/Runtime/Metasound/Source/MetasoundStandardNodes/Private/MetasoundStereoDelayNode.cpp
--- a/Engine/Plugins/Runtime/Metasound/Source/MetasoundStandardNodes/Private/MetasoundStereoDelayNode.cpp
+++ b/Engine/Plugins/Runtime/Metasound/Source/MetasoundStandardNodes/Private/MetasoundStereoDelayNode.cpp
@@ -29,10 +29,13 @@ namespace Metasound
 		METASOUND_PARAM(InDryLevel, "Dry Level", "The dry level of the delay.")
 		METASOUND_PARAM(InWetLevel, "Wet Level", "The wet level of the delay.")
 		METASOUND_PARAM(InFeedbackAmount, "Feedback", "Feedback amount.")
+		METASOUND_PARAM(InParamMaxDelayTime, "Max Delay Time", "The maximum amount of time to delay the audio.")
 		METASOUND_PARAM(OutAudioLeft, "Out Left", "Left channel audio output.")
 		METASOUND_PARAM(OutAudioRight, "Out Right", "Right channel audio output.")
 
-		static float MaxDelaySeconds = 5.0f;
+		static constexpr float MinDelaySeconds = 0.001f;
+		static constexpr float MaxDelaySeconds = 1000.f;
+		static constexpr float DefaultMaxDelaySeconds = 2.5f;
 	}
 
 	enum class EStereoDelayMode
@@ -67,7 +70,8 @@ namespace Metasound
 			const FFloatReadRef& InDelayRatio,
 			const FFloatReadRef& InDryLevel, 
 			const FFloatReadRef& InWetLevel, 
-			const FFloatReadRef& InFeedback);
+			const FFloatReadRef& InFeedback,
+			float InMaxDelayTimeSeconds);
 
 
 		virtual void BindInputs(FInputVertexInterfaceData& InOutVertexData) override;
@@ -114,9 +118,13 @@ namespace Metasound
 		// The current delay time and delay ratio
 		float PrevDelayTimeMsec;
 		float PrevDelayRatio;
+
+		// Maximum delay time
+		float MaxDelayTimeSeconds = StereoDelay::DefaultMaxDelaySeconds;
 	};
 
-	FStereoDelayOperator::FStereoDelayOperator(const FOperatorSettings& InSettings, 
+	FStereoDelayOperator::FStereoDelayOperator(
+		const FOperatorSettings& InSettings, 
 		const FAudioBufferReadRef& InLeftAudioInput, 
 		const FAudioBufferReadRef& InRightAudioInput,
 		const FStereoDelayModeReadRef& InStereoDelayMode,
@@ -124,7 +132,8 @@ namespace Metasound
 		const FFloatReadRef& InDelayRatio,
 		const FFloatReadRef& InDryLevel,
 		const FFloatReadRef& InWetLevel, 
-		const FFloatReadRef& InFeedback)
+		const FFloatReadRef& InFeedback,
+		float InMaxDelayTimeSeconds)
 
 		: LeftAudioInput(InLeftAudioInput)
 		, RightAudioInput(InRightAudioInput)
@@ -138,10 +147,11 @@ namespace Metasound
 		, RightAudioOutput(FAudioBufferWriteRef::CreateNew(InSettings))
 		, PrevDelayTimeMsec(GetInputDelayTimeMsecClamped())
 		, PrevDelayRatio(GetInputDelayRatioClamped())
+		, MaxDelayTimeSeconds(FMath::Clamp(InMaxDelayTimeSeconds, StereoDelay::MinDelaySeconds, StereoDelay::MaxDelaySeconds))
 	{
-		LeftDelayBuffer.Init(InSettings.GetSampleRate(), StereoDelay::MaxDelaySeconds);
+		LeftDelayBuffer.Init(InSettings.GetSampleRate(), MaxDelayTimeSeconds);
 		LeftDelayBuffer.SetDelayMsec(PrevDelayTimeMsec * (1.0f + PrevDelayRatio));
-		RightDelayBuffer.Init(InSettings.GetSampleRate(), StereoDelay::MaxDelaySeconds);
+		RightDelayBuffer.Init(InSettings.GetSampleRate(), MaxDelayTimeSeconds);
 		RightDelayBuffer.SetDelayMsec(PrevDelayTimeMsec * (1.0f - PrevDelayRatio));
 	}
 
@@ -158,6 +168,7 @@ namespace Metasound
 		InOutVertexData.BindReadVertex(METASOUND_GET_PARAM_NAME(InDryLevel), DryLevel);
 		InOutVertexData.BindReadVertex(METASOUND_GET_PARAM_NAME(InWetLevel), WetLevel);
 		InOutVertexData.BindReadVertex(METASOUND_GET_PARAM_NAME(InFeedbackAmount), Feedback);
+		InOutVertexData.SetValue(METASOUND_GET_PARAM_NAME(InParamMaxDelayTime), FTime::FromSeconds(MaxDelayTimeSeconds));
 	}
 
 	void FStereoDelayOperator::BindOutputs(FOutputVertexInterfaceData& InOutVertexData)
@@ -187,7 +198,7 @@ namespace Metasound
 	float FStereoDelayOperator::GetInputDelayTimeMsecClamped() const
 	{
 		// Clamp the delay time to the max delay allowed
-		return 1000.0f * FMath::Clamp((float)DelayTime->GetSeconds(), 0.0f, StereoDelay::MaxDelaySeconds);
+		return 1000.0f * FMath::Clamp((float)DelayTime->GetSeconds(), 0.0f, MaxDelayTimeSeconds);
 	}
 
 	float FStereoDelayOperator::GetInputDelayRatioClamped() const
@@ -317,7 +328,10 @@ namespace Metasound
 
 	const FVertexInterface& FStereoDelayOperator::GetVertexInterface()
 	{
-		using namespace StereoDelay; 
+		using namespace StereoDelay;
+
+		FDataVertexMetadata MaxDelayTimeMetadata = METASOUND_GET_PARAM_METADATA(InParamMaxDelayTime);
+		MaxDelayTimeMetadata.bIsAdvancedDisplay = true;
 
 		static const FVertexInterface Interface(
 			FInputVertexInterface(
@@ -328,7 +342,8 @@ namespace Metasound
 				TInputDataVertex<float>(METASOUND_GET_PARAM_NAME_AND_METADATA(InDelayRatio), 0.1f),
 				TInputDataVertex<float>(METASOUND_GET_PARAM_NAME_AND_METADATA(InDryLevel), 0.0f),
 				TInputDataVertex<float>(METASOUND_GET_PARAM_NAME_AND_METADATA(InWetLevel), 1.0f),
-				TInputDataVertex<float>(METASOUND_GET_PARAM_NAME_AND_METADATA(InFeedbackAmount), 0.0f)
+				TInputDataVertex<float>(METASOUND_GET_PARAM_NAME_AND_METADATA(InFeedbackAmount), 0.0f),
+				TInputConstructorVertex<FTime>(METASOUND_GET_PARAM_NAME(InParamMaxDelayTime), MaxDelayTimeMetadata, DefaultMaxDelaySeconds)
 			),
 			FOutputVertexInterface(
 				TOutputDataVertex<FAudioBuffer>(METASOUND_GET_PARAM_NAME_AND_METADATA(OutAudioLeft)),
@@ -374,8 +389,9 @@ namespace Metasound
 		FFloatReadRef DryLevel = InputData.GetOrCreateDefaultDataReadReference<float>(METASOUND_GET_PARAM_NAME(InDryLevel), InParams.OperatorSettings);
 		FFloatReadRef WetLevel = InputData.GetOrCreateDefaultDataReadReference<float>(METASOUND_GET_PARAM_NAME(InWetLevel), InParams.OperatorSettings);
 		FFloatReadRef Feedback = InputData.GetOrCreateDefaultDataReadReference<float>(METASOUND_GET_PARAM_NAME(InFeedbackAmount), InParams.OperatorSettings);
+		FTime MaxDelayTime = InputData.GetOrCreateDefaultValue<FTime>(METASOUND_GET_PARAM_NAME(InParamMaxDelayTime), InParams.OperatorSettings);
 
-		return MakeUnique<FStereoDelayOperator>(InParams.OperatorSettings, LeftAudioIn, RightAudioIn, StereoDelayMode, DelayTime, DelayRatio, DryLevel, WetLevel, Feedback);
+		return MakeUnique<FStereoDelayOperator>(InParams.OperatorSettings, LeftAudioIn, RightAudioIn, StereoDelayMode, DelayTime, DelayRatio, DryLevel, WetLevel, Feedback, MaxDelayTime.GetSeconds());
 	}
 
 
diff --git a/Engine/Plugins/Runtime/Synthesis/Source/Synthesis/Private/SubmixEffectTapDelay.cpp b/Engine/Plugins/Runtime/Synthesis/Source/Synthesis/Private/SubmixEffectTapDelay.cpp
--- a/Engine/Plugins/Runtime/Synthesis/Source/Synthesis/Private/SubmixEffectTapDelay.cpp
+++ b/Engine/Plugins/Runtime/Synthesis/Source/Synthesis/Private/SubmixEffectTapDelay.cpp
@@ -212,7 +212,6 @@ void FSubmixEffectTapDelay::UpdateParameters()
 		// If we have a new max delay line length, reallocate the delay line
 		if (MaxDelayLineLength != NewSettings.MaximumDelayLength)
 		{
-			DelayLine.Reset();
 			DelayLine.Init(SampleRate, NewSettings.MaximumDelayLength / 1000.0f);
 			MaxDelayLineLength = NewSettings.MaximumDelayLength;
 		}
diff --git a/Engine/Source/Runtime/SignalProcessing/Private/Delay.cpp b/Engine/Source/Runtime/SignalProcessing/Private/Delay.cpp
--- a/Engine/Source/Runtime/SignalProcessing/Private/Delay.cpp
+++ b/Engine/Source/Runtime/SignalProcessing/Private/Delay.cpp
@@ -2,12 +2,20 @@
 
 #include "DSP/Delay.h"
 #include "DSP/Dsp.h"
+#include "HAL/IConsoleManager.h"
+
+static float FDelayInitialAllocationSecondsCVar = -1.0f;
+FAutoConsoleVariableRef CVarFDelayInitialAllocationSeconds(
+	TEXT("au.DSP.InitialFDelayAllocationSeconds"),
+	FDelayInitialAllocationSecondsCVar,
+	TEXT("Override the inital delay line allocation in seconds, it will grow up to InBufferLengthSec.\n"),
+	//TEXT("The default is -1.  A value less than zero will allocate the full InBufferLengthSec\n"),
+	ECVF_Default);
 
 namespace Audio
 {
 	FDelay::FDelay()
-		: AudioBuffer(nullptr)
-		, AudioBufferSize(0)
+		: AudioBufferSize(0)
 		, ReadIndex(0)
 		, WriteIndex(0)
 		, SampleRate(0)
@@ -19,36 +27,30 @@ namespace Audio
 		Reset();
 	}
 
-	FDelay::~FDelay()
-	{
-		if (AudioBuffer)
-		{
-			delete[] AudioBuffer;
-			AudioBuffer = nullptr;
-		}
-	}
-
+	// update metadata, call Reset()
 	void FDelay::Init(const float InSampleRate, const float InBufferLengthSec)
 	{
 		SampleRate = InSampleRate;
-		AudioBufferSize = (int32) (InBufferLengthSec * (float)InSampleRate) + 1;
+		
+		// we cache this value because it is where we cap resizing the buffer
+		MaxBufferLengthSamples = InBufferLengthSec * (float)InSampleRate + 1;
 
-		if (AudioBuffer)
+		float InitialBufferSizeSeconds = InBufferLengthSec;
+
+		if (FDelayInitialAllocationSecondsCVar > 0.f)
 		{
-			delete[] AudioBuffer;
+			InitialBufferSizeSeconds = FMath::Min(InBufferLengthSec, FDelayInitialAllocationSecondsCVar);
 		}
-
-		AudioBuffer = new float[AudioBufferSize];
-
+		
+		AudioBufferSize = (int32)(InitialBufferSizeSeconds * (float)InSampleRate) + 1;
 		Reset();
 	}
 
+	// resize AudioBuffer, zero-out delay line, reset indicies
 	void FDelay::Reset()
 	{
-		if (AudioBuffer)
-		{
-			FMemory::Memzero(AudioBuffer, sizeof(float) * AudioBufferSize);
-		}
+		AudioBuffer.Reset(AudioBufferSize);
+		AudioBuffer.AddZeroed(AudioBufferSize);
 
 		WriteIndex = 0;
 		ReadIndex = 0;
@@ -60,23 +62,29 @@ namespace Audio
 	{
 		// Directly set the delay
 		const float NewDelayInSamples = InDelayMsec * SampleRate * 0.001f;
-		DelayInSamples = FMath::Min(NewDelayInSamples, static_cast<float>(AudioBufferSize));
+		DelayInSamples = FMath::Min(NewDelayInSamples, MaxBufferLengthSamples);
+		ResizeIfNeeded(DelayInSamples);
 		Update(true);
 	}
 
 	void FDelay::SetDelaySamples(const float InDelaySamples)
 	{
-		DelayInSamples = InDelaySamples;
+		DelayInSamples = FMath::Min(InDelaySamples, MaxBufferLengthSamples);
+		ResizeIfNeeded(DelayInSamples);
 		Update(true);
 	}
 
 	void FDelay::SetEasedDelayMsec(const float InDelayMsec, const bool bIsInit)
 	{
+		const float DesiredDelayInSamples = InDelayMsec * SampleRate * 0.001f;
+		const float TargetDelayInSamples = FMath::Min(DesiredDelayInSamples, MaxBufferLengthSamples);
+		ResizeIfNeeded(TargetDelayInSamples);
+
 		EaseDelayMsec.SetValue(InDelayMsec, bIsInit);
 		if (bIsInit)
 		{
 			const float NewDelayInSamples = InDelayMsec * SampleRate * 0.001f;
-			DelayInSamples = FMath::Min(NewDelayInSamples, static_cast<float>(AudioBufferSize));
+			DelayInSamples = TargetDelayInSamples;
 		}
 		Update(bIsInit);
 	}
@@ -211,4 +219,39 @@ namespace Audio
 			}
 		}
 	}
+
+	void FDelay::ResizeIfNeeded(const int32 InNewNumSamples)
+	{
+		// should be clamped by callers
+		ensure(InNewNumSamples <= MaxBufferLengthSamples);
+
+		// already large enough
+		if (InNewNumSamples <= AudioBufferSize)
+		{
+			return;
+		}
+
+		// resize the buffer
+		const int32 OldBufferSize = AudioBufferSize;
+		AudioBufferSize = FMath::Min(AudioBufferSize * 2, MaxBufferLengthSamples);
+		AudioBuffer.SetNumUninitialized(AudioBufferSize);
+
+		// see if we need to copy data to the end
+		if (ReadIndex < WriteIndex)
+		{
+			// no action needed, we will write over the uninitialzed data
+			// before we read from it.
+			return;
+		}
+
+		// (WriteIndex <= ReadIndex): our soon-to-be-read-data is in two chunks.
+		// we need to copy the second chunk to the end of the now resized array
+		// and update the read index.
+
+		// note: we can leave alone the old data since we will write to it before its read
+		const int32 SamplesToCopy = OldBufferSize - WriteIndex;
+		const int32 OldReadIndex = ReadIndex;
+		ReadIndex = AudioBufferSize - SamplesToCopy;
+		FMemory::Memmove(&AudioBuffer[WriteIndex], &AudioBuffer[AudioBufferSize - OldBufferSize], SamplesToCopy * sizeof(float));
+	}
 }
diff --git a/Engine/Source/Runtime/SignalProcessing/Private/Reverb.cpp b/Engine/Source/Runtime/SignalProcessing/Private/Reverb.cpp
--- a/Engine/Source/Runtime/SignalProcessing/Private/Reverb.cpp
+++ b/Engine/Source/Runtime/SignalProcessing/Private/Reverb.cpp
@@ -2,6 +2,8 @@
 
 #include "DSP/Reverb.h"
 
+#include "DSP/WaveTableOsc.h"
+
 namespace Audio
 {
 	struct FReverbSettingsInternal
diff --git a/Engine/Source/Runtime/SignalProcessing/Public/DSP/Delay.h b/Engine/Source/Runtime/SignalProcessing/Public/DSP/Delay.h
--- a/Engine/Source/Runtime/SignalProcessing/Public/DSP/Delay.h
+++ b/Engine/Source/Runtime/SignalProcessing/Public/DSP/Delay.h
@@ -3,8 +3,8 @@
 #pragma once
 
 #include "CoreMinimal.h"
+#include "DSP/AlignedBuffer.h"
 #include "DSP/Dsp.h"
-#include "DSP/WaveTableOsc.h"
 
 namespace Audio
 {
@@ -16,13 +16,14 @@ namespace Audio
 		SIGNALPROCESSING_API FDelay();
 
 		// Virtual Destructor
-		SIGNALPROCESSING_API virtual ~FDelay();
+		SIGNALPROCESSING_API virtual ~FDelay() = default;
 
-	public:
 		// Initialization of the delay with given sample rate and max buffer size in samples.
+		// (calls Reset())
 		SIGNALPROCESSING_API void Init(const float InSampleRate, const float InBufferLengthSec = 2.0f);
 
 		// Resets the delay line state, flushes buffer and resets read/write pointers.
+		// (called by Init())
 		SIGNALPROCESSING_API void Reset();
 
 		// Sets the delay line length. Will clamp to within range of the max initialized delay line length (won't resize).
@@ -58,12 +59,14 @@ namespace Audio
 		SIGNALPROCESSING_API virtual void ProcessAudioBuffer(const float* InAudio, int32 InNumSamples, float* OutAudio);
 
 	protected:
-
 		// Updates delay line based on any recent changes to settings
 		SIGNALPROCESSING_API void Update(bool bForce = false);
 
+		// Updates delay line based on any recent changes to settings
+		SIGNALPROCESSING_API void ResizeIfNeeded(const int32 InNewNumSamples);
+
 		// Pointer to the circular buffer of audio.
-		float* AudioBuffer;
+		FAlignedFloatBuffer AudioBuffer;
 
 		// Max length of buffer (in samples)
 		int32 AudioBufferSize;
@@ -80,6 +83,9 @@ namespace Audio
 		// Delay in samples; float supports fractional delay
 		float DelayInSamples;
 
+		// Upper limit of where we will resize the array
+		int32 MaxBufferLengthSamples;
+
 		// Eased delay in msec
 		FExponentialEase EaseDelayMsec;
 
diff --git a/Engine/Source/Runtime/SignalProcessing/Public/DSP/Dsp.h b/Engine/Source/Runtime/SignalProcessing/Public/DSP/Dsp.h
--- a/Engine/Source/Runtime/SignalProcessing/Public/DSP/Dsp.h
+++ b/Engine/Source/Runtime/SignalProcessing/Public/DSP/Dsp.h
@@ -787,6 +787,52 @@ namespace Audio
 		mutable FCriticalSection CritSect;
 	};
 
+	template <typename SampleType>
+	struct DisjointedArrayView
+	{
+		DisjointedArrayView(TArrayView<SampleType> && InFirstBuffer, TArrayView<SampleType> && InSecondBuffer)
+		: FirstBuffer(MoveTemp(InFirstBuffer))
+		, SecondBuffer(MoveTemp(InSecondBuffer))
+		{}
+
+		template <typename OtherSampleType>
+		DisjointedArrayView<OtherSampleType> SplitOtherToMatch(OtherSampleType* Other, int32 InNum) const
+		{
+			ensure(InNum == Num());
+			const int32 FirstChunkNum = FirstNum();
+
+			return DisjointedArrayView<OtherSampleType>(
+				TArrayView<OtherSampleType>(Other, FirstChunkNum)
+				, TArrayView<OtherSampleType>(Other + FirstChunkNum, InNum - FirstChunkNum)
+			);
+		}
+
+		int32 CopyIntoBuffer(SampleType* InDestination, int32 InNumSamples)
+		{
+			check(InNumSamples >= Num());
+			const int32 FirstCopySize = FirstNum() * sizeof(SampleType);
+			const int32 SecondCopySize = SecondNum() * sizeof(SampleType);
+
+			FMemory::Memcpy(InDestination, FirstBuffer.GetData(), FirstCopySize);
+
+			if (SecondCopySize)
+			{
+				FMemory::Memcpy(InDestination + FirstNum(), SecondBuffer.GetData(), SecondCopySize);
+			}
+
+			return Num();
+		}
+
+		int32 FirstNum() const { return FirstBuffer.Num(); }
+		int32 SecondNum() const { return SecondBuffer.Num(); }
+		int32 Num() const { return FirstBuffer.Num() + SecondBuffer.Num(); }
+
+		// data:
+		TArrayView<SampleType> FirstBuffer;
+		TArrayView<SampleType> SecondBuffer;
+
+	}; // struct DisjointedArrayView
+
 	/**
 	 * Basic implementation of a circular buffer built for pushing and popping arbitrary amounts of data at once.
 	 * Designed to be thread safe for SPSC; However, if Push() and Pop() are both trying to access an overlapping area of the buffer,
@@ -935,7 +981,7 @@ namespace Audio
 			}
 		}
 
-		bool Push(SampleType&& InElement)
+		bool Push(SampleType && InElement)
 		{
 			if (Remainder() == 0)
 			{
@@ -973,6 +1019,24 @@ namespace Audio
 			return NumToCopy;
 		}
 
+		// same Peek(), but provides a (possibly) disjointed view of the memory in-place
+		// Push calls while the returned view is being accessed is undefined behavior
+		DisjointedArrayView <const SampleType> PeekInPlace(uint32 NumSamples) const
+		{
+			const SampleType* SrcBuffer = InternalBuffer.GetData();
+			const uint32 ReadIndex = ReadCounter.GetValue();
+			const uint32 WriteIndex = WriteCounter.GetValue();
+
+			int32 NumToView = FMath::Min<int32>(NumSamples, Num());
+			const int32 NumRead = FMath::Min<int32>(NumToView, Capacity - ReadIndex);
+			check(NumSamples < ((uint32)TNumericLimits<int32>::Max()));
+
+			return DisjointedArrayView < const SampleType > (
+				TArrayView<const SampleType>(SrcBuffer + ReadIndex, NumRead)
+				, TArrayView<const SampleType>(SrcBuffer, (NumToView - NumRead))
+			);
+		}
+
 		// Peeks a single element.
 		// returns false if the element is empty.
 		bool Peek(SampleType& OutElement) const
@@ -1004,6 +1068,19 @@ namespace Audio
 			return NumSamplesRead;
 		}
 
+		// Same as Pop(), but provides a (possibly) disjinted view of memory in-place
+		// Push calls while the returned view is being accessed is undefined behavior
+		DisjointedArrayView<const SampleType> PopInPlace(uint32 NumSamples)
+		{
+			check(NumSamples < ((uint32)TNumericLimits<int32>::Max()));
+
+			DisjointedArrayView<const SampleType> View = PeekInPlace(NumSamples);
+			const int32 NumSamplesRead = View.Num();
+			ReadCounter.Set((ReadCounter.GetValue() + NumSamplesRead) % Capacity);
+
+			return View;
+		}
+
 		// Pops some amount of samples into this circular buffer.
 		// Returns the amount of samples read.
 		int32 Pop(uint32 NumSamples)
diff --git a/Engine/Source/Runtime/SignalProcessing/Public/DSP/Reverb.h b/Engine/Source/Runtime/SignalProcessing/Public/DSP/Reverb.h
--- a/Engine/Source/Runtime/SignalProcessing/Public/DSP/Reverb.h
+++ b/Engine/Source/Runtime/SignalProcessing/Public/DSP/Reverb.h
@@ -9,6 +9,9 @@
 
 namespace Audio
 {
+	// forward
+	class FWaveTableOsc;
+
 	struct FEarlyReflectionsSettings
 	{
 		// Early reflections gain
