Enable preset metasound sources to use their base graph instead of themselves.

diff --git a/Engine/Plugins/Runtime/Metasound/Source/MetasoundEngine/Private/MetasoundSource.cpp b/Engine/Plugins/Runtime/Metasound/Source/MetasoundEngine/Private/MetasoundSource.cpp
--- a/Engine/Plugins/Runtime/Metasound/Source/MetasoundEngine/Private/MetasoundSource.cpp
+++ b/Engine/Plugins/Runtime/Metasound/Source/MetasoundEngine/Private/MetasoundSource.cpp
@@ -51,6 +51,11 @@
 
 namespace Metasound
 {
+	namespace ConsoleVariables
+	{
+		bool bEnableExperimentalRuntimePresetGraphInflation = false;
+	}
+
 	namespace SourcePrivate
 	{
 		// Holds onto a global static TSet for tracking which error/warning logs have been
@@ -155,9 +160,17 @@ namespace Metasound
 			FCriticalSection DataChannelMapCS;
 			TSortedMap<FAudioDeviceIDAndInstanceID, FQueueState> DataChannels;
 		};
+
 	} // namespace SourcePrivate
 } // namespace Metasound
 
+FAutoConsoleVariableRef CVarMetaSoundEnableExperimentalRUntimePresetGraphInflation(
+	TEXT("au.MetaSound.Experimental.EnableRuntimePresetGraphInflation"),
+	Metasound::ConsoleVariables::bEnableExperimentalRuntimePresetGraphInflation,
+	TEXT("Enables experimental feature of MetaSounds which reduces overhead of preset graphs\n")
+	TEXT("Default: false"),
+	ECVF_Default);
+
 
 UMetaSoundSource::UMetaSoundSource(const FObjectInitializer& ObjectInitializer)
 	: Super(ObjectInitializer)
@@ -692,28 +705,72 @@ ISoundGeneratorPtr UMetaSoundSource::CreateSoundGenerator(const FSoundGeneratorI
 
 		Generator = MoveTemp(DynamicGenerator);
 	}
-	else
+	else 
 	{
-		TSharedPtr<const IGraph> MetasoundGraph = GetRegisteredGraph();
-		if (!MetasoundGraph.IsValid())
+		// By default, the sound generator for a metasound preset uses a graph specifically
+		// associated with the UMetaSoundSource_Preset. The overridden defaults for that
+		// preset are baked into the IGraph. Unfortunately, this makes the MetaSound
+		// operator pool less efficient because it associates the operator with the IGraph.
+		// The way the presets use the IGraph mean that there is less sharing of cached
+		// operators.
+		//
+		// To improve the efficiency of the operator pool, we have presets use their
+		// base IGraphs so that more MetaSounds utilize the same IGraph. This requires
+		// us to retrieve that specific graph. We also supply the parameters that were overridden
+		// in the preset to the FMetaSoundGenerator, because they are not backed into
+		// the base IGraph. 
+		if (ConsoleVariables::bEnableExperimentalRuntimePresetGraphInflation && RootMetasoundDocument.RootGraph.PresetOptions.bIsPreset)
 		{
-			return ISoundGeneratorPtr(nullptr);
+			// Get the graph associated with base graph which this preset wraps .
+			TSharedPtr<const IGraph> MetasoundGraph = TryGetMetaSoundPresetBaseGraph();
+
+			if (MetasoundGraph.IsValid())
+			{
+				// Combine the supplied parameters with the overridden parameters set on the preset.
+				TArray<FAudioParameter> MergedParameters;
+				MergePresetOverridesAndSuppliedDefaults(InDefaultParameters, MergedParameters);
+
+				// Create generator.
+				FMetasoundGeneratorInitParams InitParams
+				{
+					InSettings,
+					MoveTemp(BuilderSettings),
+					MetasoundGraph,
+					Environment,
+					GetName(),
+					GetOutputAudioChannelOrder(),
+					MoveTemp(MergedParameters),
+					bBuildSynchronous,
+					DataChannel
+				};
+
+				Generator = MakeShared<FMetasoundConstGraphGenerator>(MoveTemp(InitParams));
+			}
 		}
 
-		FMetasoundGeneratorInitParams InitParams
+		if (!Generator.IsValid())
 		{
-			InSettings,
-			MoveTemp(BuilderSettings),
-			MetasoundGraph,
-			Environment,
-			GetName(),
-			GetOutputAudioChannelOrder(),
-			MoveTemp(InDefaultParameters),
-			bBuildSynchronous,
-			DataChannel
-		};
+			TSharedPtr<const IGraph> MetasoundGraph = GetRegisteredGraph();
+			if (!MetasoundGraph.IsValid())
+			{
+				return ISoundGeneratorPtr(nullptr);
+			}
 
-		Generator = MakeShared<FMetasoundConstGraphGenerator>(MoveTemp(InitParams));
+			FMetasoundGeneratorInitParams InitParams
+			{
+				InSettings,
+				MoveTemp(BuilderSettings),
+				MetasoundGraph,
+				Environment,
+				GetName(),
+				GetOutputAudioChannelOrder(),
+				MoveTemp(InDefaultParameters),
+				bBuildSynchronous,
+				DataChannel
+			};
+
+			Generator = MakeShared<FMetasoundConstGraphGenerator>(MoveTemp(InitParams));
+		}
 	}
 
 	if (Generator.IsValid())
@@ -1553,4 +1610,176 @@ void UMetaSoundSource::InvalidateCachedRuntimeInputData()
 	RuntimeInputData.bIsValid.store(false);
 }
 
+TSharedPtr<const Metasound::IGraph> UMetaSoundSource::TryGetMetaSoundPresetBaseGraph() const
+{
+	using namespace Metasound;
+
+	TSharedPtr<const IGraph> MetasoundGraph;
+	if (ReferencedAssetClassObjects.Num() == 1)
+	{
+		// Get first element from TSet<>
+		TObjectPtr<const UObject> BaseGraph;
+		for (const TObjectPtr<UObject>& ReferencedGraph : ReferencedAssetClassObjects)
+		{
+			BaseGraph = ReferencedGraph;
+			break;
+		}
+
+		// Get the reference graph as a UMetaSoundSource
+		TObjectPtr<const UMetaSoundSource> BaseMetaSoundSource = Cast<const UMetaSoundSource>(BaseGraph);
+		if (BaseMetaSoundSource)
+		{
+			// Get registered graph of base metasound source. 
+			FSoftObjectPath BaseMetaSoundSourcePath(BaseMetaSoundSource->GetOwningAsset());
+			MetasoundGraph = FMetasoundFrontendRegistryContainer::Get()->GetGraph(BaseMetaSoundSource->GetRegistryKey(), FSoftObjectPath(BaseMetaSoundSource->GetOwningAsset()));
+		}
+	}
+	else
+	{
+		UE_LOG(LogMetaSound, Warning, TEXT("Attempt to reference parent of metasound preset failed due to unexpected number of reference asses (%d) from MetaSound Preset %s"), ReferencedAssetClassObjects.Num(), *GetOwningAssetName());
+	}
+
+	return MetasoundGraph;
+}
+
+void UMetaSoundSource::MergePresetOverridesAndSuppliedDefaults(const TArray<FAudioParameter>& InSuppliedDefaults, TArray<FAudioParameter>& OutMerged)
+{
+	using namespace Metasound;
+
+	if(!RuntimeInputData.bIsValid.load())
+	{
+		UE_CLOG(SourcePrivate::HasNotBeenLoggedForThisObject(*this, __LINE__), LogMetaSound, Warning, TEXT("Failed to use experimental preset inflation due to invalid runtime data on MetaSound %s. Ensure that UMetaSoundSource::InitResources() is executed on the game thread before CreateSoundGenerator"), *GetOwningAssetName());
+	}
+	else
+	{
+		// Get parameters that are overridden in the Preset.
+		for (const TPair<FVertexName, FRuntimeInput>& Pair : RuntimeInputData.InputMap)
+		{
+			if (!RootMetasoundDocument.RootGraph.PresetOptions.InputsInheritingDefault.Contains(Pair.Key))
+			{
+				OutMerged.Add(Pair.Value.ToAudioParameter());
+			}
+		}
+
+		// Instead of using `FAudioParameter::Merge(...)` we roll custom merge logic because
+		// FAudioParameter::Merge(...) requires us to move the InSuppliedDefaults into the
+		// OutMerged array. This unwanted in this case because we want to maintain InSuppliedDefaults
+		// for the scenario that the experimental preset override fails to create a generator
+		// in which case we are forced to use the prior codepath which does not merge these
+		// two parameter arrays. 
+		for (const FAudioParameter& SuppliedParameter : InSuppliedDefaults)
+		{
+			auto HasSameName = [&Name=SuppliedParameter.ParamName](const FAudioParameter& InOtherParam) -> bool
+			{
+				return InOtherParam.ParamName == Name;
+			};
+
+			// The supplied defaults generally come from BP and should override any other parameters.
+			if (FAudioParameter* Param = OutMerged.FindByPredicate(HasSameName))
+			{
+				*Param = SuppliedParameter;
+			}
+			else
+			{
+				OutMerged.Add(SuppliedParameter);
+			}
+		}
+	}
+}
+
+FAudioParameter UMetaSoundSource::FRuntimeInput::ToAudioParameter() const
+{
+	FAudioParameter Params;
+	Params.ParamName = Name;
+	Params.TypeName = TypeName;
+
+	switch (DefaultLiteral.GetType())
+	{
+		case EMetasoundFrontendLiteralType::Boolean:
+		{
+			static const FName TriggerName = "Trigger";
+			if (Params.TypeName == TriggerName)
+			{
+				Params.ParamType = EAudioParameterType::Trigger;
+			}
+			else
+			{
+				Params.ParamType = EAudioParameterType::Boolean;
+			}
+				
+			ensure(DefaultLiteral.TryGet(Params.BoolParam));
+		}
+		break;
+
+		case EMetasoundFrontendLiteralType::BooleanArray:
+		{
+			Params.ParamType = EAudioParameterType::BooleanArray;
+			ensure(DefaultLiteral.TryGet(Params.ArrayBoolParam));
+		}
+		break;
+
+		case EMetasoundFrontendLiteralType::Integer:
+		{
+			Params.ParamType = EAudioParameterType::Integer;
+			ensure(DefaultLiteral.TryGet(Params.IntParam));
+		}
+		break;
+
+		case EMetasoundFrontendLiteralType::IntegerArray:
+		{
+			Params.ParamType = EAudioParameterType::IntegerArray;
+			ensure(DefaultLiteral.TryGet(Params.ArrayIntParam));
+		}
+		break;
+
+		case EMetasoundFrontendLiteralType::Float:
+		{
+			Params.ParamType = EAudioParameterType::Float;
+			ensure(DefaultLiteral.TryGet(Params.FloatParam));
+		}
+		break;
+
+		case EMetasoundFrontendLiteralType::FloatArray:
+		{
+			Params.ParamType = EAudioParameterType::FloatArray;
+			ensure(DefaultLiteral.TryGet(Params.ArrayFloatParam));
+		}
+		break;
+
+		case EMetasoundFrontendLiteralType::String:
+		{
+			Params.ParamType = EAudioParameterType::String;
+			ensure(DefaultLiteral.TryGet(Params.StringParam));
+		}
+		break;
+
+		case EMetasoundFrontendLiteralType::StringArray:
+		{
+			Params.ParamType = EAudioParameterType::StringArray;
+			ensure(DefaultLiteral.TryGet(Params.ArrayStringParam));
+		}
+		break;
+
+		case EMetasoundFrontendLiteralType::UObject:
+		{
+			Params.ParamType = EAudioParameterType::Object;
+			UObject* Object = nullptr;
+			ensure(DefaultLiteral.TryGet(Object));
+			Params.ObjectParam = Object;
+		}
+		break;
+
+		case EMetasoundFrontendLiteralType::UObjectArray:
+		{
+			Params.ParamType = EAudioParameterType::ObjectArray;
+			ensure(DefaultLiteral.TryGet(MutableView(Params.ArrayObjectParam)));
+		}
+		break;
+
+		default:
+		break;
+	}
+
+	return Params;
+}
 #undef LOCTEXT_NAMESPACE // MetaSound
diff --git a/Engine/Plugins/Runtime/Metasound/Source/MetasoundEngine/Public/MetasoundSource.h b/Engine/Plugins/Runtime/Metasound/Source/MetasoundEngine/Public/MetasoundSource.h
--- a/Engine/Plugins/Runtime/Metasound/Source/MetasoundEngine/Public/MetasoundSource.h
+++ b/Engine/Plugins/Runtime/Metasound/Source/MetasoundEngine/Public/MetasoundSource.h
@@ -77,6 +77,8 @@ class METASOUNDENGINE_API UMetaSoundSource : public USoundWaveProcedural, public
 		FMetasoundFrontendLiteral DefaultLiteral;
 		// True if the data type is transmittable. False otherwise.
 		bool bIsTransmittable;
+
+		FAudioParameter ToAudioParameter() const;
 	};
 
 	struct FRuntimeInputData
@@ -292,6 +294,9 @@ public:
 	Metasound::FMetasoundEnvironment CreateEnvironment(const FSoundGeneratorInitParams& InParams) const;
 	const TArray<Metasound::FVertexName>& GetOutputAudioChannelOrder() const;
 private:
+
+	TSharedPtr<const Metasound::IGraph> TryGetMetaSoundPresetBaseGraph() const;
+	void MergePresetOverridesAndSuppliedDefaults(const TArray<FAudioParameter>& InSuppliedDefaults, TArray<FAudioParameter>& OutMerged);
 	
 	Metasound::FMetasoundEnvironment CreateEnvironment() const;
 	Metasound::FMetasoundEnvironment CreateEnvironment(const Audio::FParameterTransmitterInitParams& InParams) const;
