Default graph preset inflation to "true"

diff --git a/Engine/Plugins/Runtime/Metasound/Source/MetasoundEngine/Private/MetasoundSource.cpp b/Engine/Plugins/Runtime/Metasound/Source/MetasoundEngine/Private/MetasoundSource.cpp
--- a/Engine/Plugins/Runtime/Metasound/Source/MetasoundEngine/Private/MetasoundSource.cpp
+++ b/Engine/Plugins/Runtime/Metasound/Source/MetasoundEngine/Private/MetasoundSource.cpp
@@ -53,7 +53,7 @@ namespace Metasound
 {
 	namespace ConsoleVariables
 	{
-		bool bEnableExperimentalRuntimePresetGraphInflation = false;
+		bool bEnableExperimentalRuntimePresetGraphInflation = true;
 	}
 
 	namespace SourcePrivate
@@ -168,7 +168,7 @@ FAutoConsoleVariableRef CVarMetaSoundEnableExperimentalRUntimePresetGraphInflati
 	TEXT("au.MetaSound.Experimental.EnableRuntimePresetGraphInflation"),
 	Metasound::ConsoleVariables::bEnableExperimentalRuntimePresetGraphInflation,
 	TEXT("Enables experimental feature of MetaSounds which reduces overhead of preset graphs\n")
-	TEXT("Default: false"),
+	TEXT("Default: true"),
 	ECVF_Default);
 
 
