Add cvars for MetaSound wave player node buffer sizes, change MaxDecodeSizeInFrames from 8192 to 1024 for memory savings

diff --git a/Engine/Plugins/Runtime/Metasound/Source/MetasoundEngine/Private/MetasoundWavePlayerNode.cpp b/Engine/Plugins/Runtime/Metasound/Source/MetasoundEngine/Private/MetasoundWavePlayerNode.cpp
--- a/Engine/Plugins/Runtime/Metasound/Source/MetasoundEngine/Private/MetasoundWavePlayerNode.cpp
+++ b/Engine/Plugins/Runtime/Metasound/Source/MetasoundEngine/Private/MetasoundWavePlayerNode.cpp
@@ -311,6 +311,24 @@ namespace Metasound
 		FTimeReadRef LoopDuration;
 	};
 
+	// Maximum decode size in frames. 
+	static int32 MaxDecodeSizeInFrames = 1024;
+	FAutoConsoleVariableRef CVarMetaSoundWavePlayerMaxDecodeSizeInFrames(
+		TEXT("au.MetaSound.WavePlayer.MaxDecodeSizeInFrames"),
+		MaxDecodeSizeInFrames,
+		TEXT("Max size in frames used for decoding audio in the MetaSound wave player node.\n")
+		TEXT("Default: 1024"),
+		ECVF_Default);
+
+	// Block size for deinterleaving audio. 
+	static int32 DeinterleaveBlockSizeInFrames = 512;
+	FAutoConsoleVariableRef CVarMetaSoundWavePlayerDeinterleaveBlockSizeInFrames(
+		TEXT("au.MetaSound.WavePlayer.DeinterleaveBlockSizeInFrames"),
+		DeinterleaveBlockSizeInFrames,
+		TEXT("Block size in frames used for deinterleaving audio in the MetaSound wave player node.\n")
+		TEXT("Default: 512"),
+		ECVF_Default);
+
 	/** MetaSound operator for the wave player node. */
 	class FWavePlayerOperator : public TExecutableOperator<FWavePlayerOperator>
 	{	
@@ -318,10 +336,6 @@ namespace Metasound
 
 		// Maximum absolute pitch shift in octaves. 
 		static constexpr float MaxAbsPitchShiftInOctaves = 6.0f;
-		// Maximum decode size in frames. 
-		static constexpr int32 MaxDecodeSizeInFrames = 8192;
-		// Block size for deinterleaving audio. 
-		static constexpr int32 DeinterleaveBlockSizeInFrames = 512;
 
 		FWavePlayerOperator(const FWavePlayerOpArgs& InArgs)
 			: OperatorSettings(InArgs.Settings)
@@ -707,7 +721,9 @@ namespace Metasound
 				
 				// Create the wave proxy reader.
 				FSoundWaveProxyReader::FSettings WaveReaderSettings;
-				WaveReaderSettings.MaxDecodeSizeInFrames = MaxDecodeSizeInFrames;
+				WaveReaderSettings.MaxDecodeSizeInFrames = FMath::IsPowerOfTwo(MaxDecodeSizeInFrames) ? 
+					MaxDecodeSizeInFrames : FMath::RoundUpToPowerOfTwo(MaxDecodeSizeInFrames);
+
 				WaveReaderSettings.StartTimeInSeconds = StartTime->GetSeconds();
 				WaveReaderSettings.LoopStartTimeInSeconds = LoopStartTime->GetSeconds();
 				WaveReaderSettings.LoopDurationInSeconds = LoopDuration->GetSeconds(); 
@@ -715,6 +731,9 @@ namespace Metasound
 
 				WaveProxyReader = FSoundWaveProxyReader::Create(WaveProxy.ToSharedRef(), WaveReaderSettings);
 
+				DeinterleaveBufferBlockSizeInFrames = FMath::IsPowerOfTwo(DeinterleaveBlockSizeInFrames) ? 
+					DeinterleaveBlockSizeInFrames : FMath::RoundUpToPowerOfTwo(DeinterleaveBlockSizeInFrames);
+
 				if (WaveProxyReader.IsValid())
 				{
 					UpdateSampleRateFrameRatio();
@@ -723,7 +742,7 @@ namespace Metasound
 					if (WaveProxyNumChannels > 0)
 					{
 						// Create buffer for interleaved audio
-						int32 InterleavedBufferNumSamples = WaveProxyNumChannels * DeinterleaveBlockSizeInFrames;
+						int32 InterleavedBufferNumSamples = WaveProxyNumChannels * DeinterleaveBufferBlockSizeInFrames;
 						InterleavedBuffer.Reset(InterleavedBufferNumSamples);
 						InterleavedBuffer.AddUninitialized(InterleavedBufferNumSamples);
 
@@ -739,10 +758,10 @@ namespace Metasound
 						// better control.
 						ConvertDeinterleaveParams.MonoUpmixMethod = Audio::EChannelMapMonoUpmixMethod::FullVolume;
 						ConvertDeinterleave = Audio::IConvertDeinterleave::Create(ConvertDeinterleaveParams);
-						Audio::SetMultichannelBufferSize(NumDeinterleaveChannels, DeinterleaveBlockSizeInFrames, DeinterleavedBuffer);
+						Audio::SetMultichannelBufferSize(NumDeinterleaveChannels, DeinterleaveBufferBlockSizeInFrames, DeinterleavedBuffer);
 
 						// Initialize source buffer
-						int32 FrameCapacity = DeinterleaveBlockSizeInFrames + FMath::CeilToInt(GetMaxFrameRatio() * OperatorSettings.GetNumFramesPerBlock());
+						int32 FrameCapacity = DeinterleaveBufferBlockSizeInFrames + FMath::CeilToInt(GetMaxFrameRatio() * OperatorSettings.GetNumFramesPerBlock());
 						Audio::SetMultichannelCircularBufferCapacity(NumOutputChannels, FrameCapacity, SourceCircularBuffer);
 						SourceState = FSourceBufferState(*WaveProxyReader, SourceCircularBuffer);
 
@@ -786,7 +805,7 @@ namespace Metasound
 			if (bIsPlaying)
 			{
 				const int32 NumExistingFrames = Audio::GetMultichannelBufferNumFrames(OutBuffer);
-				const int32 NumSamplesToGenerate = DeinterleaveBlockSizeInFrames * WaveProxyReader->GetNumChannels();
+				const int32 NumSamplesToGenerate = DeinterleaveBufferBlockSizeInFrames * WaveProxyReader->GetNumChannels();
  				check(NumSamplesToGenerate == InterleavedBuffer.Num())
 
 				WaveProxyReader->PopAudio(InterleavedBuffer);
@@ -796,7 +815,7 @@ namespace Metasound
 				{
 					OutBuffer[ChannelIndex].Push(DeinterleavedBuffer[ChannelIndex]);
 				}
-				OutSourceState.Append(DeinterleaveBlockSizeInFrames, *bLoop);
+				OutSourceState.Append(DeinterleaveBufferBlockSizeInFrames, *bLoop);
 			}
 			else
 			{
@@ -1044,7 +1063,8 @@ namespace Metasound
 		int32 NumDeinterleaveChannels;
 		bool bOnNearlyDoneTriggeredForWave = false;
 		bool bIsPlaying = false;
-		
+		// Cached from cvar 
+		int32 DeinterleaveBufferBlockSizeInFrames;
 	};
 
 	class FWavePlayerOperatorFactory : public IOperatorFactory
diff --git a/Engine/Source/Runtime/AudioCodecEngine/Public/SoundWaveProxyReader.h b/Engine/Source/Runtime/AudioCodecEngine/Public/SoundWaveProxyReader.h
--- a/Engine/Source/Runtime/AudioCodecEngine/Public/SoundWaveProxyReader.h
+++ b/Engine/Source/Runtime/AudioCodecEngine/Public/SoundWaveProxyReader.h
@@ -48,7 +48,7 @@ public:
 	/** Settings for a FSoundWaveProxyReader. */
 	struct FSettings
 	{
-		uint32 MaxDecodeSizeInFrames = 8192;
+		uint32 MaxDecodeSizeInFrames = 1024;
 		float StartTimeInSeconds = 0.f;
 		bool bIsLooping = false;
 		float LoopStartTimeInSeconds = 0.f;
