Fix default parameter values when skipping preset
#rnx

diff --git a/Engine/Plugins/Runtime/Metasound/Source/MetasoundEngine/Private/MetasoundSource.cpp b/Engine/Plugins/Runtime/Metasound/Source/MetasoundEngine/Private/MetasoundSource.cpp
--- a/Engine/Plugins/Runtime/Metasound/Source/MetasoundEngine/Private/MetasoundSource.cpp
+++ b/Engine/Plugins/Runtime/Metasound/Source/MetasoundEngine/Private/MetasoundSource.cpp
@@ -53,11 +53,13 @@ namespace Metasound
 {
 	namespace ConsoleVariables
 	{
-		bool bEnableExperimentalRuntimePresetGraphInflation = false;
+		bool bEnableExperimentalRuntimePresetGraphInflation = true;
 	}
 
 	namespace SourcePrivate
 	{
+		static const FLazyName TriggerName = "Trigger";
+
 		// Holds onto a global static TSet for tracking which error/warning logs have been
 		// trigger in order to avoid log spam.
 		bool HasNotBeenLoggedForThisObject(const UMetaSoundSource& InMetaSound, uint32 InLogLineNumber)
@@ -161,6 +163,150 @@ namespace Metasound
 			TSortedMap<FAudioDeviceIDAndInstanceID, FQueueState> DataChannels;
 		};
 
+		void CreateUObjectProxies(const Frontend::IDataTypeRegistry& InRegistry, FName InVertexTypeName, bool bClearUObjectPointers, FAudioParameter& InOutParamToInit)
+		{
+			using namespace Metasound;
+
+			switch (InOutParamToInit.ParamType)
+			{
+				case EAudioParameterType::Object:
+				{
+					TSharedPtr<Audio::IProxyData> ProxyPtr = InRegistry.CreateProxyFromUObject(InVertexTypeName, InOutParamToInit.ObjectParam);
+					InOutParamToInit.ObjectProxies.Emplace(MoveTemp(ProxyPtr));
+
+					if (bClearUObjectPointers)
+					{
+						InOutParamToInit.ObjectParam = nullptr;
+					}
+				}
+				break;
+
+				case EAudioParameterType::ObjectArray:
+				{
+					const FName ElementTypeName = CreateElementTypeNameFromArrayTypeName(InVertexTypeName);
+					for (TObjectPtr<UObject>& Object : InOutParamToInit.ArrayObjectParam)
+					{
+						TSharedPtr<Audio::IProxyData> ProxyPtr = InRegistry.CreateProxyFromUObject(ElementTypeName, Object);
+						InOutParamToInit.ObjectProxies.Emplace(MoveTemp(ProxyPtr));
+					}
+
+					if (bClearUObjectPointers)
+					{
+						InOutParamToInit.ArrayObjectParam.Reset();
+					}
+				}
+				break;
+
+				default:
+					break;
+			}
+		}
+
+		FAudioParameter MakeAudioParameter(const Frontend::IDataTypeRegistry& InRegistry, FName InParamName, FName InTypeName, const FMetasoundFrontendLiteral& InLiteral, bool bCreateUObjectProxies) 
+		{
+			constexpr bool bClearUObjectPointers = false;
+
+			FAudioParameter Params;
+			Params.ParamName = InParamName;
+			Params.TypeName = InTypeName;
+
+			switch (InLiteral.GetType())
+			{
+				case EMetasoundFrontendLiteralType::Boolean:
+				{
+					if (Params.TypeName == TriggerName)
+					{
+						Params.ParamType = EAudioParameterType::Trigger;
+					}
+					else
+					{
+						Params.ParamType = EAudioParameterType::Boolean;
+					}
+						
+					ensure(InLiteral.TryGet(Params.BoolParam));
+				}
+				break;
+
+				case EMetasoundFrontendLiteralType::BooleanArray:
+				{
+					Params.ParamType = EAudioParameterType::BooleanArray;
+					ensure(InLiteral.TryGet(Params.ArrayBoolParam));
+				}
+				break;
+
+				case EMetasoundFrontendLiteralType::Integer:
+				{
+					Params.ParamType = EAudioParameterType::Integer;
+					ensure(InLiteral.TryGet(Params.IntParam));
+				}
+				break;
+
+				case EMetasoundFrontendLiteralType::IntegerArray:
+				{
+					Params.ParamType = EAudioParameterType::IntegerArray;
+					ensure(InLiteral.TryGet(Params.ArrayIntParam));
+				}
+				break;
+
+				case EMetasoundFrontendLiteralType::Float:
+				{
+					Params.ParamType = EAudioParameterType::Float;
+					ensure(InLiteral.TryGet(Params.FloatParam));
+				}
+				break;
+
+				case EMetasoundFrontendLiteralType::FloatArray:
+				{
+					Params.ParamType = EAudioParameterType::FloatArray;
+					ensure(InLiteral.TryGet(Params.ArrayFloatParam));
+				}
+				break;
+
+				case EMetasoundFrontendLiteralType::String:
+				{
+					Params.ParamType = EAudioParameterType::String;
+					ensure(InLiteral.TryGet(Params.StringParam));
+				}
+				break;
+
+				case EMetasoundFrontendLiteralType::StringArray:
+				{
+					Params.ParamType = EAudioParameterType::StringArray;
+					ensure(InLiteral.TryGet(Params.ArrayStringParam));
+				}
+				break;
+
+				case EMetasoundFrontendLiteralType::UObject:
+				{
+					Params.ParamType = EAudioParameterType::Object;
+					UObject* Object = nullptr;
+					ensure(InLiteral.TryGet(Object));
+					Params.ObjectParam = Object;
+					if (bCreateUObjectProxies)
+					{
+						CreateUObjectProxies(InRegistry, InTypeName, bClearUObjectPointers, Params);
+					}
+				}
+				break;
+
+				case EMetasoundFrontendLiteralType::UObjectArray:
+				{
+					Params.ParamType = EAudioParameterType::ObjectArray;
+					ensure(InLiteral.TryGet(MutableView(Params.ArrayObjectParam)));
+					if (bCreateUObjectProxies)
+					{
+						CreateUObjectProxies(InRegistry, InTypeName, bClearUObjectPointers, Params);
+					}
+				}
+				break;
+
+				default:
+				break;
+			}
+
+			return Params;
+		}
+
 	} // namespace SourcePrivate
 } // namespace Metasound
 
@@ -168,7 +314,7 @@ FAutoConsoleVariableRef CVarMetaSoundEnableExperimentalRUntimePresetGraphInflati
 	TEXT("au.MetaSound.Experimental.EnableRuntimePresetGraphInflation"),
 	Metasound::ConsoleVariables::bEnableExperimentalRuntimePresetGraphInflation,
 	TEXT("Enables experimental feature of MetaSounds which reduces overhead of preset graphs\n")
-	TEXT("Default: false"),
+	TEXT("Default: true"),
 	ECVF_Default);
 
 
@@ -523,9 +669,14 @@ void UMetaSoundSource::InitParameters(TArray<FAudioParameter>& ParametersToInit,
 	METASOUND_LLM_SCOPE;
 	METASOUND_TRACE_CPUPROFILER_EVENT_SCOPE(UMetaSoundSource::InitParameters);
 
+
 	if (bIsBuilderActive)	
 	{
-		InitParametersInternal(CreateRuntimeInputMap(), ParametersToInit, InFeatureName);
+		// Do not create UObject proxies in the runtime input map because they proxies 
+		// stored there will not be used. The necessary proxies in the ParametersToInit 
+		// will be created and used instead. 
+		constexpr bool bCreateUObjectProxiesInRuntimeInputMap = false; 
+		InitParametersInternal(CreateRuntimeInputMap(bCreateUObjectProxiesInRuntimeInputMap), ParametersToInit, InFeatureName);
 	}
 	else
 	{
@@ -542,7 +693,12 @@ void UMetaSoundSource::InitParameters(TArray<FAudioParameter>& ParametersToInit,
 			// suffer the consequences of incurring significant performance losses 
 			// each time a parameter is set on the MetaSound. 
 			UE_CLOG(HasNotBeenLoggedForThisObject(*this, __LINE__), LogMetaSound, Warning, TEXT("Initializing parameters on uninitialized UMetaSoundSource %s will result in slower performance. UMetaSoundSource::InitResources should finish executing on the game thread before attempting to call UMetaSoundSource::InitParameters(...)"), *GetOwningAssetName());
-			InitParametersInternal(CreateRuntimeInputMap(), ParametersToInit, InFeatureName);
+
+			// Do not create UObject proxies in the runtime input map because they proxies 
+			// stored there will not be used. The necessary proxies in the ParametersToInit 
+			// will be created and used instead. 
+			constexpr bool bCreateUObjectProxiesInRuntimeInputMap = false; 
+			InitParametersInternal(CreateRuntimeInputMap(bCreateUObjectProxiesInRuntimeInputMap), ParametersToInit, InFeatureName);
 		}
 	}
 }
@@ -800,105 +956,9 @@ bool UMetaSoundSource::GetAllDefaultParameters(TArray<FAudioParameter>& OutParam
 		return false;
 	}
 
-
 	for(const TPair<FVertexName, FRuntimeInput>& Pair : RuntimeInputData.InputMap)
 	{
-		const FRuntimeInput& Input = Pair.Value;
-		FAudioParameter Params;
-		Params.ParamName = Input.Name;
-		Params.TypeName = Input.TypeName;
-
-		switch (Input.DefaultLiteral.GetType())
-		{
-			case EMetasoundFrontendLiteralType::Boolean:
-			{
-				static const FName TriggerName = "Trigger";
-				if (Params.TypeName == TriggerName)
-				{
-					Params.ParamType = EAudioParameterType::Trigger;
-				}
-				else
-				{
-					Params.ParamType = EAudioParameterType::Boolean;
-				}
-					
-				ensure(Input.DefaultLiteral.TryGet(Params.BoolParam));
-			}
-			break;
-
-			case EMetasoundFrontendLiteralType::BooleanArray:
-			{
-				Params.ParamType = EAudioParameterType::BooleanArray;
-				ensure(Input.DefaultLiteral.TryGet(Params.ArrayBoolParam));
-			}
-			break;
-
-			case EMetasoundFrontendLiteralType::Integer:
-			{
-				Params.ParamType = EAudioParameterType::Integer;
-				ensure(Input.DefaultLiteral.TryGet(Params.IntParam));
-			}
-			break;
-
-			case EMetasoundFrontendLiteralType::IntegerArray:
-			{
-				Params.ParamType = EAudioParameterType::IntegerArray;
-				ensure(Input.DefaultLiteral.TryGet(Params.ArrayIntParam));
-			}
-			break;
-
-			case EMetasoundFrontendLiteralType::Float:
-			{
-				Params.ParamType = EAudioParameterType::Float;
-				ensure(Input.DefaultLiteral.TryGet(Params.FloatParam));
-			}
-			break;
-
-			case EMetasoundFrontendLiteralType::FloatArray:
-			{
-				Params.ParamType = EAudioParameterType::FloatArray;
-				ensure(Input.DefaultLiteral.TryGet(Params.ArrayFloatParam));
-			}
-			break;
-
-			case EMetasoundFrontendLiteralType::String:
-			{
-				Params.ParamType = EAudioParameterType::String;
-				ensure(Input.DefaultLiteral.TryGet(Params.StringParam));
-			}
-			break;
-
-			case EMetasoundFrontendLiteralType::StringArray:
-			{
-				Params.ParamType = EAudioParameterType::StringArray;
-				ensure(Input.DefaultLiteral.TryGet(Params.ArrayStringParam));
-			}
-			break;
-
-			case EMetasoundFrontendLiteralType::UObject:
-			{
-				Params.ParamType = EAudioParameterType::Object;
-				UObject* Object = nullptr;
-				ensure(Input.DefaultLiteral.TryGet(Object));
-				Params.ObjectParam = Object;
-			}
-			break;
-
-			case EMetasoundFrontendLiteralType::UObjectArray:
-			{
-				Params.ParamType = EAudioParameterType::ObjectArray;
-				ensure(Input.DefaultLiteral.TryGet(MutableView(Params.ArrayObjectParam)));
-			}
-			break;
-
-			default:
-			break;
-		}
-
-		if (Params.ParamType != EAudioParameterType::None)
-		{
-			OutParameters.Add(Params);
-		}
+		OutParameters.Add(Pair.Value.DefaultParameter);
 	}
 	return true;
 }
@@ -998,42 +1058,6 @@ void UMetaSoundSource::InitParametersInternal(const Metasound::TSortedVertexName
 		}
 	};
 
-	auto ConstructProxies = [&DataTypeRegistry](FAudioParameter& OutParamToInit, FName VertexTypeName)
-	{
-		using namespace Metasound;
-
-		switch (OutParamToInit.ParamType)
-		{
-			case EAudioParameterType::Object:
-			{
-				TSharedPtr<Audio::IProxyData> ProxyPtr = DataTypeRegistry.CreateProxyFromUObject(VertexTypeName, OutParamToInit.ObjectParam);
-				OutParamToInit.ObjectProxies.Emplace(MoveTemp(ProxyPtr));
-
-				// Null out param as it is no longer needed (nor desired to be accessed once passed to the Audio Thread)
-				OutParamToInit.ObjectParam = nullptr;
-			}
-			break;
-
-			case EAudioParameterType::ObjectArray:
-			{
-				const FName ElementTypeName = CreateElementTypeNameFromArrayTypeName(VertexTypeName);
-				for (TObjectPtr<UObject>& Object : OutParamToInit.ArrayObjectParam)
-				{
-					TSharedPtr<Audio::IProxyData> ProxyPtr = DataTypeRegistry.CreateProxyFromUObject(ElementTypeName, Object);
-					OutParamToInit.ObjectProxies.Emplace(MoveTemp(ProxyPtr));
-				}
-				// Reset param array as it is no longer needed (nor desired to be accessed once passed to the Audio Thread).
-				// All object manipulation hereafter should be done via proxies
-				OutParamToInit.ArrayObjectParam.Reset();
-			}
-			break;
-
-			default:
-				break;
-		}
-	};
-
-
 	for (int32 i = ParametersToInit.Num() - 1; i >= 0; --i)
 	{
 		bool bIsParameterValid = false;
@@ -1044,7 +1068,8 @@ void UMetaSoundSource::InitParametersInternal(const Metasound::TSortedVertexName
 			if (IsParameterValidInternal(Parameter, Input->TypeName, DataTypeRegistry))
 			{
 				Sanitize(Parameter);
-				ConstructProxies(Parameter, Input->TypeName);
+				constexpr bool bClearUObjectPointers = true; // protect against leaking UObject ptrs to the audio thread 
+				SourcePrivate::CreateUObjectProxies(DataTypeRegistry, Input->TypeName, bClearUObjectPointers, Parameter);
 				bIsParameterValid = true;
 			}
 		}
@@ -1262,7 +1287,11 @@ TSharedPtr<Audio::IParameterTransmitter> UMetaSoundSource::CreateParameterTransm
 			UE_LOG(LogMetaSound, Warning, TEXT("Creating a Parameter Transmiiter on uninitialized UMetaSoundSource %s will result in slower performance. UMetaSoundSource::InitResources should finish executing on the game thread before attempting to call UMetaSoundSource::CreateParameterTransmitter(...)"), *GetOwningAssetName());
 		}
 
-		return CreateParameterTransmitterInternal(CreateRuntimeInputMap(), InParams);
+		// Do not create UObject proxies in the runtime input map because they proxies 
+		// stored there will not be used. The necessary proxies in the ParametersToInit 
+		// will be created and used instead. 
+		constexpr bool bCreateUObjectProxiesInRuntimeInputMap = false; 
+		return CreateParameterTransmitterInternal(CreateRuntimeInputMap(bCreateUObjectProxiesInRuntimeInputMap), InParams);
 	}
 	else
 	{
@@ -1537,7 +1566,7 @@ TSharedPtr<Metasound::DynamicGraph::FDynamicOperatorTransactor> UMetaSoundSource
 	return DynamicTransactor;
 }
 
-Metasound::TSortedVertexNameMap<UMetaSoundSource::FRuntimeInput> UMetaSoundSource::CreateRuntimeInputMap() const
+Metasound::TSortedVertexNameMap<UMetaSoundSource::FRuntimeInput> UMetaSoundSource::CreateRuntimeInputMap(bool bCreateUObjectProxies) const
 {
 	using namespace Metasound;
 	using namespace Metasound::Frontend;
@@ -1583,13 +1612,18 @@ Metasound::TSortedVertexNameMap<UMetaSoundSource::FRuntimeInput> UMetaSoundSourc
 				UE_LOG(LogMetaSound, Warning, TEXT("Failed to find data type '%s' in registry. Assuming data type is not transmittable"), *Input.TypeName.ToString());
 
 			}
-			PublicInputs.Add(Input.Name, FRuntimeInput{Input.Name, Input.TypeName, Input.AccessType, Input.DefaultLiteral, bIsTransmittable});
+
+			
+			FAudioParameter DefaultParameter = SourcePrivate::MakeAudioParameter(Registry, Input.Name, Input.TypeName, Input.DefaultLiteral, bCreateUObjectProxies) ;
+
+			PublicInputs.Add(Input.Name, FRuntimeInput{Input.Name, Input.TypeName, Input.AccessType, DefaultParameter, bIsTransmittable});
 		}
 	}
 
 	// Add the parameter pack input that ALL Metasounds have
 	FMetasoundFrontendClassInput ParameterPackInput = UMetasoundParameterPack::GetClassInput();
-	PublicInputs.Add(ParameterPackInput.Name, FRuntimeInput{ParameterPackInput.Name, ParameterPackInput.TypeName, ParameterPackInput.AccessType, ParameterPackInput.DefaultLiteral, true /* bIsTransmittable */});
+	FAudioParameter ParameterPackDefaultParameter = SourcePrivate::MakeAudioParameter(Registry, ParameterPackInput.Name, ParameterPackInput.TypeName, ParameterPackInput.DefaultLiteral, bCreateUObjectProxies) ;
+	PublicInputs.Add(ParameterPackInput.Name, FRuntimeInput{ParameterPackInput.Name, ParameterPackInput.TypeName, ParameterPackInput.AccessType, ParameterPackDefaultParameter, true /* bIsTransmittable */});
 	
 	return PublicInputs;
 }
@@ -1601,7 +1635,8 @@ void UMetaSoundSource::CacheRuntimeInputData()
 		UE_LOG(LogMetaSound, Warning, TEXT("Skipping caching of runtime inputs for UMetaSoundSource %s because there is an active builder"), *GetOwningAssetName());
 	}
 
-	RuntimeInputData.InputMap = CreateRuntimeInputMap();
+	constexpr bool bCreateUObjectProxies = true; 
+	RuntimeInputData.InputMap = CreateRuntimeInputMap(bCreateUObjectProxies);
 	RuntimeInputData.bIsValid.store(true);
 }
 
@@ -1657,7 +1692,7 @@ void UMetaSoundSource::MergePresetOverridesAndSuppliedDefaults(const TArray<FAud
 		{
 			if (!RootMetasoundDocument.RootGraph.PresetOptions.InputsInheritingDefault.Contains(Pair.Key))
 			{
-				OutMerged.Add(Pair.Value.ToAudioParameter());
+				OutMerged.Add(Pair.Value.DefaultParameter);
 			}
 		}
 
@@ -1687,99 +1722,4 @@ void UMetaSoundSource::MergePresetOverridesAndSuppliedDefaults(const TArray<FAud
 	}
 }
 
-FAudioParameter UMetaSoundSource::FRuntimeInput::ToAudioParameter() const
-{
-	FAudioParameter Params;
-	Params.ParamName = Name;
-	Params.TypeName = TypeName;
-
-	switch (DefaultLiteral.GetType())
-	{
-		case EMetasoundFrontendLiteralType::Boolean:
-		{
-			static const FName TriggerName = "Trigger";
-			if (Params.TypeName == TriggerName)
-			{
-				Params.ParamType = EAudioParameterType::Trigger;
-			}
-			else
-			{
-				Params.ParamType = EAudioParameterType::Boolean;
-			}
-				
-			ensure(DefaultLiteral.TryGet(Params.BoolParam));
-		}
-		break;
-
-		case EMetasoundFrontendLiteralType::BooleanArray:
-		{
-			Params.ParamType = EAudioParameterType::BooleanArray;
-			ensure(DefaultLiteral.TryGet(Params.ArrayBoolParam));
-		}
-		break;
-
-		case EMetasoundFrontendLiteralType::Integer:
-		{
-			Params.ParamType = EAudioParameterType::Integer;
-			ensure(DefaultLiteral.TryGet(Params.IntParam));
-		}
-		break;
-
-		case EMetasoundFrontendLiteralType::IntegerArray:
-		{
-			Params.ParamType = EAudioParameterType::IntegerArray;
-			ensure(DefaultLiteral.TryGet(Params.ArrayIntParam));
-		}
-		break;
-
-		case EMetasoundFrontendLiteralType::Float:
-		{
-			Params.ParamType = EAudioParameterType::Float;
-			ensure(DefaultLiteral.TryGet(Params.FloatParam));
-		}
-		break;
-
-		case EMetasoundFrontendLiteralType::FloatArray:
-		{
-			Params.ParamType = EAudioParameterType::FloatArray;
-			ensure(DefaultLiteral.TryGet(Params.ArrayFloatParam));
-		}
-		break;
-
-		case EMetasoundFrontendLiteralType::String:
-		{
-			Params.ParamType = EAudioParameterType::String;
-			ensure(DefaultLiteral.TryGet(Params.StringParam));
-		}
-		break;
-
-		case EMetasoundFrontendLiteralType::StringArray:
-		{
-			Params.ParamType = EAudioParameterType::StringArray;
-			ensure(DefaultLiteral.TryGet(Params.ArrayStringParam));
-		}
-		break;
-
-		case EMetasoundFrontendLiteralType::UObject:
-		{
-			Params.ParamType = EAudioParameterType::Object;
-			UObject* Object = nullptr;
-			ensure(DefaultLiteral.TryGet(Object));
-			Params.ObjectParam = Object;
-		}
-		break;
-
-		case EMetasoundFrontendLiteralType::UObjectArray:
-		{
-			Params.ParamType = EAudioParameterType::ObjectArray;
-			ensure(DefaultLiteral.TryGet(MutableView(Params.ArrayObjectParam)));
-		}
-		break;
-
-		default:
-		break;
-	}
-
-	return Params;
-}
 #undef LOCTEXT_NAMESPACE // MetaSound
diff --git a/Engine/Plugins/Runtime/Metasound/Source/MetasoundEngine/Public/MetasoundSource.h b/Engine/Plugins/Runtime/Metasound/Source/MetasoundEngine/Public/MetasoundSource.h
--- a/Engine/Plugins/Runtime/Metasound/Source/MetasoundEngine/Public/MetasoundSource.h
+++ b/Engine/Plugins/Runtime/Metasound/Source/MetasoundEngine/Public/MetasoundSource.h
@@ -73,12 +73,10 @@ class METASOUNDENGINE_API UMetaSoundSource : public USoundWaveProcedural, public
 		FName TypeName;
 		// Access type of input vertex.
 		EMetasoundFrontendVertexAccessType AccessType;
-		// Default literal of input vertex.
-		FMetasoundFrontendLiteral DefaultLiteral;
+		// Default parameter of input vertex.
+		FAudioParameter DefaultParameter;
 		// True if the data type is transmittable. False otherwise.
 		bool bIsTransmittable;
-
-		FAudioParameter ToAudioParameter() const;
 	};
 
 	struct FRuntimeInputData
@@ -306,7 +304,7 @@ private:
 	void TrackGenerator(uint64 Id, TSharedPtr<Metasound::FMetasoundGenerator> Generator);
 	void ForgetGenerator(ISoundGeneratorPtr Generator);
 
-	Metasound::TSortedVertexNameMap<FRuntimeInput> CreateRuntimeInputMap() const;
+	Metasound::TSortedVertexNameMap<FRuntimeInput> CreateRuntimeInputMap(bool bCreateUObjectProxies) const;
 	void CacheRuntimeInputData();
 	void InvalidateCachedRuntimeInputData();
 
diff --git a/Engine/Plugins/Runtime/Metasound/Source/MetasoundGenerator/Private/MetasoundGenerator.cpp b/Engine/Plugins/Runtime/Metasound/Source/MetasoundGenerator/Private/MetasoundGenerator.cpp
--- a/Engine/Plugins/Runtime/Metasound/Source/MetasoundGenerator/Private/MetasoundGenerator.cpp
+++ b/Engine/Plugins/Runtime/Metasound/Source/MetasoundGenerator/Private/MetasoundGenerator.cpp
@@ -860,7 +860,7 @@ PRAGMA_ENABLE_DEPRECATION_WARNINGS
 				bUseOperatorPool = true; // raise this flag to make sure we put the operator back in the pool (regardless of CVAR state)
 
 				// Apply and default inputs to the operator.
-				GeneratorBuilder::ApplyAudioParameters(OperatorSettings, MoveTemp(InInitParams.DefaultParameters), GraphOperatorAndInputs.Inputs);
+				GeneratorBuilder::ResetGraphOperatorInputs(OperatorSettings, MoveTemp(InInitParams.DefaultParameters), GraphOperatorAndInputs.Inputs);
 
 				// Reset operator internal state before playing it.
 				if (IOperator::FResetFunction Reset = GraphOperatorAndInputs.Operator->GetResetFunction())
diff --git a/Engine/Plugins/Runtime/Metasound/Source/MetasoundGenerator/Private/MetasoundGeneratorBuilder.cpp b/Engine/Plugins/Runtime/Metasound/Source/MetasoundGenerator/Private/MetasoundGeneratorBuilder.cpp
--- a/Engine/Plugins/Runtime/Metasound/Source/MetasoundGenerator/Private/MetasoundGeneratorBuilder.cpp
+++ b/Engine/Plugins/Runtime/Metasound/Source/MetasoundGenerator/Private/MetasoundGeneratorBuilder.cpp
@@ -339,22 +339,35 @@ namespace Metasound
 			};
 		}
 
-		void ApplyAudioParameters(const FOperatorSettings& InOperatorSettings, TArray<FAudioParameter>&& InParameters, FInputVertexInterfaceData& InInterface)
+		void ResetGraphOperatorInputs(const FOperatorSettings& InOperatorSettings, TArray<FAudioParameter> InParameterOverrides, FInputVertexInterfaceData& InOutInterface)
 		{
-			METASOUND_TRACE_CPUPROFILER_EVENT_SCOPE(GeneratorBuilder::ApplyAudioParameters);
+			METASOUND_TRACE_CPUPROFILER_EVENT_SCOPE(GeneratorBuilder::ResetGraphOperatorInputs);
+
 			Frontend::IDataTypeRegistry& DataTypeRegistry = Frontend::IDataTypeRegistry::Get();
-			for (FAudioParameter& Parameter : InParameters)
+
+			for (MetasoundVertexDataPrivate::FInputBinding& Binding : InOutInterface)
 			{
-				if (const FAnyDataReference* Ref = InInterface.FindDataReference(Parameter.ParamName))
+				if (const FAnyDataReference* Ref = Binding.GetDataReference())
 				{
 					if (EDataReferenceAccessType::Write == Ref->GetAccessType())
 					{
 						Frontend::FLiteralAssignmentFunction LiteralSetter = DataTypeRegistry.GetLiteralAssignmentFunction(Ref->GetDataTypeName());
-
 						if (LiteralSetter)
 						{
-							FLiteral Literal = Frontend::ConvertParameterToLiteral(MoveTemp(Parameter)); 
-							LiteralSetter(InOperatorSettings, Literal, *Ref);
+							auto IsParameterNameEqualToVertexName = [&Name=Binding.GetVertex().VertexName](const FAudioParameter& InParam) -> bool
+							{
+								return Name == InParam.ParamName;
+							};
+
+							if (FAudioParameter* Parameter = InParameterOverrides.FindByPredicate(IsParameterNameEqualToVertexName))
+							{
+								FLiteral Literal = Frontend::ConvertParameterToLiteral(MoveTemp(*Parameter)); 
+								LiteralSetter(InOperatorSettings, Literal, *Ref);
+							}
+							else
+							{
+								LiteralSetter(InOperatorSettings, Binding.GetVertex().GetDefaultLiteral(), *Ref);
+							}
 						}
 					}
 				}
diff --git a/Engine/Plugins/Runtime/Metasound/Source/MetasoundGenerator/Private/MetasoundGeneratorBuilder.h b/Engine/Plugins/Runtime/Metasound/Source/MetasoundGenerator/Private/MetasoundGeneratorBuilder.h
--- a/Engine/Plugins/Runtime/Metasound/Source/MetasoundGenerator/Private/MetasoundGeneratorBuilder.h
+++ b/Engine/Plugins/Runtime/Metasound/Source/MetasoundGenerator/Private/MetasoundGeneratorBuilder.h
@@ -110,8 +110,8 @@ namespace Metasound
 		// Create all the data required to run a FMetasoundGenerator
 		MetasoundGeneratorPrivate::FMetasoundGeneratorData BuildGeneratorData(const FOperatorSettings& InOperatorSettings, const FMetasoundGeneratorInitParams& InInitParams, FOperatorAndInputs&& InGraphOperatorAndInputs, TUniquePtr<Frontend::FGraphAnalyzer> InAnalyzer);
 
-		// Apply audio parameters to the input vertex data
-		void ApplyAudioParameters(const FOperatorSettings& InOperatorSettings, TArray<FAudioParameter>&& InParameters, FInputVertexInterfaceData& InInterface);
+		// Resets FInputVertexInterfaceData to it's initial default state using the default literals stored ther-in. If a parameter exists in InParmaeterOverrides, that value is used instead. 
+		void ResetGraphOperatorInputs(const FOperatorSettings& InOperatorSettings, TArray<FAudioParameter> InParameterOverrides, FInputVertexInterfaceData& InOutInterface);
 
 		// Initialize the maps of parameter setters for applying async input data to metasound inputs.
 		void InitializeParameterSetters(FInputVertexInterfaceData& InputData, MetasoundGeneratorPrivate::FParameterSetterSortedMap& OutParamSetters, TMap<FName, MetasoundGeneratorPrivate::FParameterPackSetter>& OutParamPackSetters);
