[Metasound] gaurd for null deref
#rnx

diff --git a/Engine/Plugins/Runtime/Metasound/Source/MetasoundEngine/Private/MetasoundWaveTableBankEvaluateNode.cpp b/Engine/Plugins/Runtime/Metasound/Source/MetasoundEngine/Private/MetasoundWaveTableBankEvaluateNode.cpp
--- a/Engine/Plugins/Runtime/Metasound/Source/MetasoundEngine/Private/MetasoundWaveTableBankEvaluateNode.cpp
+++ b/Engine/Plugins/Runtime/Metasound/Source/MetasoundEngine/Private/MetasoundWaveTableBankEvaluateNode.cpp
@@ -156,6 +156,12 @@ namespace Metasound
 			FWaveTableBankAssetProxyPtr Proxy = WaveTableBankAsset.GetProxy();
 			float NewIndex = 0.f;
 
+			if (!WaveTableBankAsset.IsValid())
+			{
+				*OutWriteRef = 0.f; // same as Reset() state
+				return;
+			}
+
 			const float Min = WaveTableBankAsset->IsBipolar() ? -1.f : 0.f;
 			const float Input = FMath::Clamp(*InputReadRef, Min, 1.f);
 			if (!ResolveNextComputeIndex(Proxy, Input, NewIndex))
