Fix ensure when trying to get an array of integers
#rnx

diff --git a/Engine/Plugins/Runtime/Metasound/Source/MetasoundFrontend/Private/MetasoundFrontendLiteral.cpp b/Engine/Plugins/Runtime/Metasound/Source/MetasoundFrontend/Private/MetasoundFrontendLiteral.cpp
--- a/Engine/Plugins/Runtime/Metasound/Source/MetasoundFrontend/Private/MetasoundFrontendLiteral.cpp
+++ b/Engine/Plugins/Runtime/Metasound/Source/MetasoundFrontend/Private/MetasoundFrontendLiteral.cpp
@@ -598,7 +598,7 @@ bool FMetasoundFrontendLiteral::TryGet(int32& OutValue) const
 
 bool FMetasoundFrontendLiteral::TryGet(TArray<int32>& OutValue) const
 {
-	if (Type == EMetasoundFrontendLiteralType::Integer)
+	if (Type == EMetasoundFrontendLiteralType::IntegerArray)
 	{
 		OutValue = AsInteger;
 		return true;
