Fix use of pooled metasounds using preset graph inflation and constructor inputs
#rnx

diff --git a/Engine/Plugins/Runtime/Metasound/Source/MetasoundEngine/Private/MetasoundSource.cpp b/Engine/Plugins/Runtime/Metasound/Source/MetasoundEngine/Private/MetasoundSource.cpp
--- a/Engine/Plugins/Runtime/Metasound/Source/MetasoundEngine/Private/MetasoundSource.cpp
+++ b/Engine/Plugins/Runtime/Metasound/Source/MetasoundEngine/Private/MetasoundSource.cpp
@@ -1,6 +1,7 @@
 // Copyright Epic Games, Inc. All Rights Reserved.
 #include "MetasoundSource.h"
 
+#include "Algo/AnyOf.h"
 #include "Algo/Find.h"
 #include "Algo/Transform.h"
 #include "AssetRegistry/AssetRegistryModule.h"
@@ -875,7 +876,7 @@ ISoundGeneratorPtr UMetaSoundSource::CreateSoundGenerator(const FSoundGeneratorI
 		// us to retrieve that specific graph. We also supply the parameters that were overridden
 		// in the preset to the FMetaSoundGenerator, because they are not backed into
 		// the base IGraph. 
-		if (ConsoleVariables::bEnableExperimentalRuntimePresetGraphInflation && RootMetasoundDocument.RootGraph.PresetOptions.bIsPreset)
+		if (ConsoleVariables::bEnableExperimentalRuntimePresetGraphInflation && bIsPresetGraphInflationSupported)
 		{
 			// Get the graph associated with base graph which this preset wraps .
 			TSharedPtr<const IGraph> MetasoundGraph = TryGetMetaSoundPresetBaseGraph();
@@ -1630,6 +1631,8 @@ Metasound::TSortedVertexNameMap<UMetaSoundSource::FRuntimeInput> UMetaSoundSourc
 
 void UMetaSoundSource::CacheRuntimeInputData()
 {
+	using namespace Metasound;
+
 	if (bIsBuilderActive)
 	{
 		UE_LOG(LogMetaSound, Warning, TEXT("Skipping caching of runtime inputs for UMetaSoundSource %s because there is an active builder"), *GetOwningAssetName());
@@ -1637,11 +1640,33 @@ void UMetaSoundSource::CacheRuntimeInputData()
 
 	constexpr bool bCreateUObjectProxies = true; 
 	RuntimeInputData.InputMap = CreateRuntimeInputMap(bCreateUObjectProxies);
+
+	// Determine if preset graph inflation is possible 
+	//
+	// Constructor inputs conflict with `Preset Graph Inflation` and `Operator Caching`. 
+	// This logic protects against attempting to use preset graph inflation when the preset 
+	// graph has overridden constructor pins. 
+	// 
+	// Operator caching of base preset graphs fail when there are constructor inputs because
+	// constructor inputs set on the preset cannot be updated after the base operator is 
+	// cached. 
+	auto IsOverriddenConstructorInput = [&InputsInheritingDefault=RootMetasoundDocument.RootGraph.PresetOptions.InputsInheritingDefault](const TPair<FVertexName, FRuntimeInput>& Pair) 
+	{
+		if (Pair.Value.AccessType == EMetasoundFrontendVertexAccessType::Value)
+		{
+			return !InputsInheritingDefault.Contains(Pair.Key);
+		}
+		return false;
+	};
+
+	bIsPresetGraphInflationSupported = RootMetasoundDocument.RootGraph.PresetOptions.bIsPreset && !Algo::AnyOf(RuntimeInputData.InputMap, IsOverriddenConstructorInput);
+
 	RuntimeInputData.bIsValid.store(true);
 }
 
 void UMetaSoundSource::InvalidateCachedRuntimeInputData()
 {
+	bIsPresetGraphInflationSupported = false; 
 	RuntimeInputData.bIsValid.store(false);
 }
 
diff --git a/Engine/Plugins/Runtime/Metasound/Source/MetasoundEngine/Public/MetasoundSource.h b/Engine/Plugins/Runtime/Metasound/Source/MetasoundEngine/Public/MetasoundSource.h
--- a/Engine/Plugins/Runtime/Metasound/Source/MetasoundEngine/Public/MetasoundSource.h
+++ b/Engine/Plugins/Runtime/Metasound/Source/MetasoundEngine/Public/MetasoundSource.h
@@ -337,4 +337,9 @@ private:
 	mutable Metasound::FSampleRate CachedAudioDeviceSampleRate = 0;
 	
 	bool bIsBuilderActive = false;
+
+	// Preset graph inflation is a performance optimization intended for use with the MetaSoundOperatorPool. If multiple presets 
+	// utilize the same base MetaSound, they may be able to share their operators in the operator pool. This makes for a more
+	// efficient use of the operator pool.
+	bool bIsPresetGraphInflationSupported = false;
 };
