[metasound builder] low-hanging fruit pre-allocation of containers and moving Stringbuilding / FName conversion to static init.
#rnx

diff --git a/Engine/Plugins/Runtime/Metasound/Source/MetasoundGraphCore/Private/MetasoundOperatorBuilder.cpp b/Engine/Plugins/Runtime/Metasound/Source/MetasoundGraphCore/Private/MetasoundOperatorBuilder.cpp
--- a/Engine/Plugins/Runtime/Metasound/Source/MetasoundGraphCore/Private/MetasoundOperatorBuilder.cpp
+++ b/Engine/Plugins/Runtime/Metasound/Source/MetasoundGraphCore/Private/MetasoundOperatorBuilder.cpp
@@ -330,6 +330,11 @@ namespace Metasound
 		// Create FOperatorInfos from Nodes
 		TSortedMap<FOperatorID, FGraphOperatorData::FOperatorInfo>& OperatorMap = InOutGraphOperatorData.OperatorMap;
 		TArray<FOperatorID>& OperatorOrder = InOutGraphOperatorData.OperatorOrder;
+
+		const int32 NumNodes = InSortedNodes.Num();
+		OperatorMap.Reserve(OperatorMap.Num() + NumNodes);
+		OperatorOrder.Reserve(OperatorOrder.Num() + NumNodes);
+
 		{
 			METASOUND_TRACE_CPUPROFILER_EVENT_SCOPE(Metasound::FOperatorBuilder::InitializeOperatorInfo::Nodes);
 			for (const INode* Node : InSortedNodes)
@@ -462,9 +467,11 @@ namespace Metasound
 			{
 #if METASOUND_CPUPROFILERTRACE_ENABLED
 				// Use node class name if valid, otherwise (for example graph nodes) use instance name 
+				TStringBuilder<256> TraceNamePtr;
 				const FNodeClassName& NodeClassName = Node->GetMetadata().ClassName;
-				const FString& NodeTraceName = NodeClassName.IsValid() ? NodeClassName.GetFullName().ToString() : Node->GetInstanceName().ToString();
-				METASOUND_TRACE_CPUPROFILER_EVENT_SCOPE_TEXT(*FString::Printf(TEXT("Metasound::FOperatorBuilder::CreateOperators::CreateAndBind %s"), *NodeTraceName));
+				const FName& NodeTraceName = NodeClassName.IsValid() ? NodeClassName.GetFullName() : Node->GetInstanceName();
+				TraceNamePtr << "Metasound::FOperatorBuilder::CreateOperators::CreateAndBind " << NodeTraceName;
+				METASOUND_TRACE_CPUPROFILER_EVENT_SCOPE_TEXT(*TraceNamePtr);
 #endif // METASOUND_CPUPROFILERTRACE_ENABLED
 
 				FBuildOperatorParams CreateParams{*Node, InOutContext.Settings, OperatorInfo.VertexData.GetInputs(), InOutContext.Environment, this};
diff --git a/Engine/Plugins/Runtime/Metasound/Source/MetasoundStandardNodes/Private/MetasoundBandsplitterNode.cpp b/Engine/Plugins/Runtime/Metasound/Source/MetasoundStandardNodes/Private/MetasoundBandsplitterNode.cpp
--- a/Engine/Plugins/Runtime/Metasound/Source/MetasoundStandardNodes/Private/MetasoundBandsplitterNode.cpp
+++ b/Engine/Plugins/Runtime/Metasound/Source/MetasoundStandardNodes/Private/MetasoundBandsplitterNode.cpp
@@ -128,7 +128,7 @@ namespace Metasound
 
 			for (uint32 Chan = 0; Chan < NumChannels; ++Chan)
 			{
-				InOutVertexData.BindReadVertex(GetAudioInputName(Chan), AudioInputs[Chan]);
+				InOutVertexData.BindReadVertex(AudioInputNames[Chan], AudioInputs[Chan]);
 			}
 
 			InOutVertexData.BindReadVertex(METASOUND_GET_PARAM_NAME(InputFilterOrder), FilterOrder);
@@ -136,7 +136,7 @@ namespace Metasound
 
 			for (uint32 BandIndex = 0; BandIndex < NumBands - 1; ++BandIndex)
 			{
-				InOutVertexData.BindReadVertex(GetCrossoverInputName(BandIndex), CrossoverFrequencies[BandIndex]);
+				InOutVertexData.BindReadVertex(CrossoverInputNames[BandIndex], CrossoverFrequencies[BandIndex]);
 			}
 		}
 
@@ -189,7 +189,7 @@ namespace Metasound
 					const FDataVertexMetadata AudioInputMetadata;
 #endif // WITH_EDITOR
 
-					InputInterface.Add(TInputDataVertex<FAudioBuffer>(GetAudioInputName(ChannelIndex), AudioInputMetadata));
+					InputInterface.Add(TInputDataVertex<FAudioBuffer>(AudioInputNames[ChannelIndex], AudioInputMetadata));
 				}
 
 				InputInterface.Add(TInputDataVertex<FEnumBandSplitterFilterOrder>(METASOUND_GET_PARAM_NAME_AND_METADATA(InputFilterOrder)));
@@ -208,7 +208,7 @@ namespace Metasound
 					const FDataVertexMetadata CrossoverInputMetadata;
 #endif // WITH_EDITOR
 
-					InputInterface.Add(TInputDataVertex<float>(GetCrossoverInputName(InputIndex), CrossoverInputMetadata, (1 + InputIndex) * 500.0f));
+					InputInterface.Add(TInputDataVertex<float>(CrossoverInputNames[InputIndex], CrossoverInputMetadata, (1 + InputIndex) * 500.0f));
 				}
 
 				// outputs
@@ -252,12 +252,12 @@ namespace Metasound
 
 			for (uint32 Chan = 0; Chan < NumChannels; Chan++)
 			{
-				InputBuffers.Add(InputData.GetOrConstructDataReadReference<FAudioBuffer>(GetAudioInputName(Chan), InParams.OperatorSettings));
+				InputBuffers.Add(InputData.GetOrConstructDataReadReference<FAudioBuffer>(AudioInputNames[Chan], InParams.OperatorSettings));
 			}
 
 			for (uint32 Band = 0; Band < NumBands - 1; Band++)
 			{
-				InputCrossovers.Add(InputData.GetOrCreateDefaultDataReadReference<float>(GetCrossoverInputName(Band), InParams.OperatorSettings));
+				InputCrossovers.Add(InputData.GetOrCreateDefaultDataReadReference<float>(CrossoverInputNames[Band], InParams.OperatorSettings));
 			}
 
 			return MakeUnique<TBandSplitterOperator<NumBands, NumChannels>>(InParams.OperatorSettings, MoveTemp(InputBuffers), MoveTemp(InputCrossovers), FilterOrderIn, bPhaseCompensateIn);
@@ -361,6 +361,72 @@ namespace Metasound
 			return bDidUpdate;
 		}
 
+		
+		static TArray<FVertexName> InitializeAudioInputNames()
+		{
+			TArray<FVertexName> Names;
+			Names.AddUninitialized(NumChannels);
+
+			for (int ChanIndex = 0; ChanIndex < NumChannels; ++ChanIndex)
+			{
+				if (NumChannels == 1)
+				{
+					Names[ChanIndex] = *FString::Printf(TEXT("In"));
+				}
+				else if (NumChannels == 2)
+				{
+					Names[ChanIndex] = *FString::Printf(TEXT("In %s"), (ChanIndex == 0) ? TEXT("L") : TEXT("R"));
+				}
+				else
+				{
+					Names[ChanIndex] = *FString::Printf(TEXT("In %i"), ChanIndex);
+				}
+			}
+
+			return Names;
+		}
+
+		static TArray<FVertexName> InitializeCrossoverInputNames()
+		{
+			TArray<FVertexName> Names;
+			Names.AddUninitialized(NumBands);
+
+			for (int InputIndex = 0; InputIndex < NumBands; ++InputIndex)
+			{
+				Names[InputIndex] = *FString::Printf(TEXT("Crossover %i"), InputIndex);
+			}
+
+			return Names;
+		}
+
+		static TArray<FVertexName> InitializeAudioOutputNames()
+		{
+			TArray<FVertexName> Names;
+			Names.AddUninitialized(NumBands * NumChannels);
+
+			for (int BandIndex = 0; BandIndex < NumBands; ++BandIndex)
+			{
+				for (int ChanIndex = 0; ChanIndex < NumChannels; ++ChanIndex)
+				{
+					if (NumChannels == 1)
+					{
+						Names[BandIndex * NumChannels + ChanIndex] = *FString::Printf(TEXT("Band %i Out"), BandIndex);
+					}
+					else if (NumChannels == 2)
+					{
+						Names[BandIndex * NumChannels + ChanIndex] = *FString::Printf(TEXT("Band %i %s"), BandIndex, (ChanIndex == 0) ? TEXT("L") : TEXT("R"));
+					}
+					else
+					{
+						Names[BandIndex * NumChannels + ChanIndex] = *FString::Printf(TEXT("Band %i Out %i"), BandIndex, ChanIndex);
+					}
+				}
+			}
+
+			return Names;
+		}
+
+
 		TArray<FAudioBufferReadRef> AudioInputs;
 		TArray<FFloatReadRef> CrossoverFrequencies;
 
@@ -377,6 +443,10 @@ namespace Metasound
 		TArray<Audio::FLinkwitzRileyBandSplitter> Filters;
 		TArray<Audio::FMultibandBuffer> MultiBandBuffers;
 
+		static const inline TArray<FVertexName> AudioInputNames = InitializeAudioInputNames();
+		static const inline TArray<FVertexName> CrossoverInputNames = InitializeCrossoverInputNames();
+		static const inline TArray<FVertexName> AudioOutputNames = InitializeAudioOutputNames();
+
 		FOperatorSettings Settings;
 
 		static FNodeClassMetadata CreateNodeClassMetadata(const FName& InOperatorName, const FText& InDisplayName, const FText& InDescription, const FVertexInterface& InDefaultInterface)
@@ -399,37 +469,9 @@ namespace Metasound
 			return Metadata;
 		}
 
-		static const FVertexName GetAudioInputName(uint32 ChannelIndex)
-		{
-			if (NumChannels == 1)
-			{
-				return *FString::Printf(TEXT("In"));
-			}
-			else if (NumChannels == 2)
-			{
-				return *FString::Printf(TEXT("In %s"), (ChannelIndex == 0) ? TEXT("L") : TEXT("R"));
-			}
-
-			return *FString::Printf(TEXT("In %i"), ChannelIndex);
-		}
-
-		static const FVertexName GetCrossoverInputName(uint32 InputIndex)
-		{
-			return *FString::Printf(TEXT("Crossover %i"), InputIndex);
-		}
-
 		static const FVertexName GetAudioOutputName(uint32 OutputIndex, uint32 ChannelIndex)
 		{
-			if (NumChannels == 1)
-			{
-				return *FString::Printf(TEXT("Band %i Out"), OutputIndex);
-			}
-			else if (NumChannels == 2)
-			{
-				return *FString::Printf(TEXT("Band %i %s"), OutputIndex, (ChannelIndex == 0) ? TEXT("L") : TEXT("R"));
-			}
-
-			return *FString::Printf(TEXT("Band %i Out %i"), OutputIndex, ChannelIndex);
+			return AudioOutputNames[OutputIndex * NumChannels + ChannelIndex];
 		}
 
 #if WITH_EDITOR
diff --git a/Engine/Plugins/Runtime/Metasound/Source/MetasoundStandardNodes/Private/MetasoundMixerNode.cpp b/Engine/Plugins/Runtime/Metasound/Source/MetasoundStandardNodes/Private/MetasoundMixerNode.cpp
--- a/Engine/Plugins/Runtime/Metasound/Source/MetasoundStandardNodes/Private/MetasoundMixerNode.cpp
+++ b/Engine/Plugins/Runtime/Metasound/Source/MetasoundStandardNodes/Private/MetasoundMixerNode.cpp
@@ -19,7 +19,6 @@
 
 namespace Metasound
 {
-# pragma region Operator Declaration
 	template<uint32 NumInputs, uint32 NumChannels>
 	class TAudioMixerNodeOperator : public TExecutableOperator<TAudioMixerNodeOperator<NumInputs, NumChannels>>
 	{
@@ -78,7 +77,7 @@ namespace Metasound
 #else 
 					FDataVertexMetadata GainPinMetaData;
 #endif // WITH_EDITOR
-					TInputDataVertex<float>GainVertexModel(GetGainInputName(InputIndex), GainPinMetaData, 1.0f);
+					TInputDataVertex<float>GainVertexModel(GainInputNames[InputIndex], GainPinMetaData, 1.0f);
 
 					InputInterface.Add(GainVertexModel);
 				}
@@ -96,7 +95,7 @@ namespace Metasound
 #else 
 					const FDataVertexMetadata AudioOutputMetadata;
 #endif // WITH_EDITOR
-					OutputInterface.Add(TOutputDataVertex<FAudioBuffer>(GetAudioOutputName(i), AudioOutputMetadata));
+					OutputInterface.Add(TOutputDataVertex<FAudioBuffer>(AudioOutputNames[i], AudioOutputMetadata));
 				}
 
 				return FVertexInterface(InputInterface, OutputInterface);
@@ -160,7 +159,7 @@ namespace Metasound
 					InputBuffers.Add(InputData.GetOrConstructDataReadReference<FAudioBuffer>(GetAudioInputName(i, Chan), InParams.OperatorSettings));
 				}
 
-				InputGains.Add(InputData.GetOrCreateDefaultDataReadReference<float>(GetGainInputName(i), InParams.OperatorSettings));
+				InputGains.Add(InputData.GetOrCreateDefaultDataReadReference<float>(GainInputNames[i], InParams.OperatorSettings));
 			}
 
 			return MakeUnique<TAudioMixerNodeOperator<NumInputs, NumChannels>>(InParams, MoveTemp(InputBuffers), MoveTemp(InputGains));
@@ -175,7 +174,7 @@ namespace Metasound
 					InOutVertexData.BindReadVertex(GetAudioInputName(i, Chan), Inputs[i * NumChannels + Chan]);
 				}
 
-				InOutVertexData.BindReadVertex(GetGainInputName(i), Gains[i]);
+				InOutVertexData.BindReadVertex(GainInputNames[i], Gains[i]);
 			}
 		}
 
@@ -183,7 +182,7 @@ namespace Metasound
 		{
 			for (uint32 i = 0; i < NumChannels; ++i)
 			{
-				InOutVertexData.BindReadVertex(GetAudioOutputName(i), Outputs[i]);
+				InOutVertexData.BindReadVertex(AudioOutputNames[i], Outputs[i]);
 			}
 		}
 
@@ -240,6 +239,94 @@ namespace Metasound
 
 
 	private:
+		static TArray<FVertexName> InitializeAudioInputNames()
+		{
+			TStringBuilder<32> InputNameStr;
+			TArray<FVertexName> Names;
+			Names.AddUninitialized(NumInputs * NumChannels);
+
+			for (int InputIndex = 0; InputIndex < NumInputs; ++InputIndex)
+			{
+				for (int ChannelIndex = 0; ChannelIndex < NumChannels; ++ChannelIndex)
+				{
+					InputNameStr << "In " << InputIndex;
+
+					if (NumChannels == 1)
+					{
+					}
+					else if (NumChannels == 2)
+					{
+						InputNameStr << ' ' << ((!ChannelIndex)? 'L' : 'R');
+					}
+					else
+					{
+						InputNameStr << ", " << ChannelIndex;
+					}
+
+					Names[InputIndex * NumChannels + ChannelIndex] = *InputNameStr;
+					InputNameStr.Reset();
+				}
+			}
+
+			return Names;
+		}
+
+		static TArray<FVertexName> InitializeAudioOutputNames()
+		{
+			TStringBuilder<32> AudioOutNameStr;
+			TArray<FVertexName> Names;
+			Names.AddUninitialized(NumChannels);
+
+			for (int ChanIndex = 0; ChanIndex < NumChannels; ++ChanIndex)
+			{
+				AudioOutNameStr << "Out";
+
+				if (NumChannels == 1)
+				{
+				}
+				else if (NumChannels == 2)
+				{
+					AudioOutNameStr << ' ' << ((!ChanIndex)? 'L' : 'R');
+				}
+				else
+				{
+					AudioOutNameStr << ' ' << ChanIndex;
+				}
+			
+				Names[ChanIndex] = *AudioOutNameStr;
+				AudioOutNameStr.Reset();
+			}
+
+			return Names;
+		}
+
+		static TArray<FVertexName> InitializeGainInputNames()
+		{
+			TStringBuilder<32> GainNameStr;
+			TArray<FVertexName> Names;
+			Names.AddUninitialized(NumInputs);
+
+			for (int GainIndex = 0; GainIndex < NumInputs; ++GainIndex)
+			{
+				GainNameStr << "Gain " << GainIndex;
+				Names[GainIndex] = * GainNameStr;
+
+				GainNameStr.Reset();
+			}
+
+			return Names;
+		}
+
+		static const FVertexName& GetAudioInputName(uint32 InputIndex, uint32 ChannelIndex)
+		{
+			return AudioInputNames[InputIndex * NumChannels + ChannelIndex];
+		}
+
+
+		static inline const TArray<FVertexName> AudioInputNames = InitializeAudioInputNames();
+		static inline const TArray<FVertexName> AudioOutputNames = InitializeAudioOutputNames();
+		static inline const TArray<FVertexName> GainInputNames = InitializeGainInputNames();
+
 		TArray<FFloatReadRef> Gains;
 		TArray<FAudioBufferReadRef> Inputs;
 		TArray<FAudioBufferWriteRef> Outputs;
@@ -266,39 +353,6 @@ namespace Metasound
 			return Metadata;
 		}
 
-#pragma region Name Gen
-		static const FVertexName GetAudioInputName(uint32 InputIndex, uint32 ChannelIndex)
-		{
-			if (NumChannels == 1)
-			{
-				return *FString::Printf(TEXT("In %i"), InputIndex);
-			}
-			else if (NumChannels == 2)
-			{
-				return *FString::Printf(TEXT("In %i %s"), InputIndex, (ChannelIndex == 0) ? TEXT("L") : TEXT("R"));
-			}
-
-			return *FString::Printf(TEXT("In %i, %i"), InputIndex, ChannelIndex);
-		}
-
-		static const FVertexName GetGainInputName(uint32 InputIndex)
-		{
-			return *FString::Printf(TEXT("Gain %i"), InputIndex);
-		}
-
-		static const FVertexName GetAudioOutputName(uint32 ChannelIndex)
-		{
-			if (NumChannels == 1)
-			{
-				return TEXT("Out");
-			}
-			else if (NumChannels == 2)
-			{
-				return *FString::Printf(TEXT("Out %s"), (ChannelIndex == 0) ? TEXT("L") : TEXT("R"));
-			}
-
-			return *FString::Printf(TEXT("Out %i"), ChannelIndex);
-		}
 
 #if WITH_EDITOR
 		static const FText GetAudioInputDescription(uint32 InputIndex, uint32 ChannelIndex)
@@ -361,12 +415,9 @@ namespace Metasound
 			return METASOUND_LOCTEXT_FORMAT("AudioMixerAudioOutputDescription", "Summed output for channel: {0}", ChannelIndex);
 		}
 #endif // WITH_EDITOR
-#pragma endregion
 	}; // class TAudioMixerNodeOperator
-#pragma endregion
 
 
-#pragma region Node Definition
 	template<uint32 NumInputs, uint32 NumChannels>
 	class METASOUNDSTANDARDNODES_API TAudioMixerNode : public FNodeFacade
 	{
@@ -380,10 +431,8 @@ namespace Metasound
 
 		virtual ~TAudioMixerNode() = default;
 	};
-#pragma endregion
 
 
-#pragma region Node Registration
 	#define REGISTER_AUDIOMIXER_NODE(A, B) \
 		using FAudioMixerNode_##A ## _ ##B = TAudioMixerNode<A, B>; \
 		METASOUND_REGISTER_NODE(FAudioMixerNode_##A ## _ ##B) \
@@ -410,11 +459,6 @@ namespace Metasound
 	// test
 //	REGISTER_AUDIOMIXER_NODE(8, 6)
 
-#pragma endregion
-
-
-
-
 } // namespace Metasound
 
 #undef LOCTEXT_NAMESPACE // "MetasoundStandardNodes_MixerNode"
diff --git a/Engine/Source/Runtime/SignalProcessing/Private/DirectedGraphAlgo.cpp b/Engine/Source/Runtime/SignalProcessing/Private/DirectedGraphAlgo.cpp
--- a/Engine/Source/Runtime/SignalProcessing/Private/DirectedGraphAlgo.cpp
+++ b/Engine/Source/Runtime/SignalProcessing/Private/DirectedGraphAlgo.cpp
@@ -143,10 +143,14 @@ namespace Audio
 	{
 		using namespace DirectedGraphAlgoPrivate;
 
-		
 		// Initialize side info about which vertices have been visited and which
 		// have not.
 		TMap<int32, FDepthFirstSortInfo> SortInfos;
+
+		const int32 MaxVertexNum = InUniqueVertices.Num();
+		SortInfos.Reserve(SortInfos.Num() + MaxVertexNum);
+		OutVertexOrder.Reserve(OutVertexOrder.Num() + MaxVertexNum);
+
 		for (int32 Vertex : InUniqueVertices)
 		{
 			SortInfos.Add(Vertex);
