Deprecate FCreateOperatorParams and use FBuildOperatorParams instead for MetaSound perf (part 3)

diff --git a/Engine/Plugins/Runtime/Metasound/Source/MetasoundGraphCore/Private/MetasoundFacade.cpp b/Engine/Plugins/Runtime/Metasound/Source/MetasoundGraphCore/Private/MetasoundFacade.cpp
--- a/Engine/Plugins/Runtime/Metasound/Source/MetasoundGraphCore/Private/MetasoundFacade.cpp
+++ b/Engine/Plugins/Runtime/Metasound/Source/MetasoundGraphCore/Private/MetasoundFacade.cpp
@@ -25,8 +25,10 @@ namespace Metasound
 		return [InCreateFunc](const FBuildOperatorParams& InParams, FBuildResults& OutResults)
 		{
 			FDataReferenceCollection InputDataReferences = InParams.InputData.ToDataReferenceCollection();
+PRAGMA_DISABLE_DEPRECATION_WARNINGS // Temporary for the deprecation of FCreateOperatorParams
 			FCreateOperatorParams DeprecatedParams{ InParams.Node, InParams.OperatorSettings, InputDataReferences, InParams.Environment, InParams.Builder };
 			return InCreateFunc(DeprecatedParams, OutResults.Errors);
+PRAGMA_ENABLE_DEPRECATION_WARNINGS
 		};
 	}
 
diff --git a/Engine/Plugins/Runtime/Metasound/Source/MetasoundGraphCore/Public/MetasoundBuilderInterface.h b/Engine/Plugins/Runtime/Metasound/Source/MetasoundGraphCore/Public/MetasoundBuilderInterface.h
--- a/Engine/Plugins/Runtime/Metasound/Source/MetasoundGraphCore/Public/MetasoundBuilderInterface.h
+++ b/Engine/Plugins/Runtime/Metasound/Source/MetasoundGraphCore/Public/MetasoundBuilderInterface.h
@@ -61,6 +61,8 @@ namespace Metasound
 	/** FCreateOperatorParams holds the parameters provided to operator factories
 	 * during the creation of an IOperator
 	 */
+	struct UE_DEPRECATED(5.4, "Update APIs to use FBuildOperatorParams. Operator creation routines should have the following signature \"TUniquePtr<IOperator> CreateOperator(const FBuildOperatorParams&, FBuildResults&)\"") FCreateOperatorParams;
+
 	struct FCreateOperatorParams
 	{
 		/** The node associated with this factory and the desired IOperator. */
diff --git a/Engine/Plugins/Runtime/Metasound/Source/MetasoundGraphCore/Public/MetasoundFacade.h b/Engine/Plugins/Runtime/Metasound/Source/MetasoundGraphCore/Public/MetasoundFacade.h
--- a/Engine/Plugins/Runtime/Metasound/Source/MetasoundGraphCore/Public/MetasoundFacade.h
+++ b/Engine/Plugins/Runtime/Metasound/Source/MetasoundGraphCore/Public/MetasoundFacade.h
@@ -10,51 +10,56 @@
 
 namespace Metasound
 {
-	// Helper template to determine whether a static class function is declared
-	// for a given template class.
-	template <typename U>
-	class TIsFactoryMethodDeclared
+	namespace MetasoundFacadePrivate
 	{
-		private:
-			template<typename T, T> 
-			struct Helper;
+		// Helper template to determine whether a static class function is declared
+		// for a given template class.
+		template <typename U>
+		class TIsFactoryMethodDeclared
+		{
+			private:
+				template<typename T, T> 
+				struct Helper;
 
-			// Check for "static TUniquePtr<IOperator> U::CreateOperator(const FCreateOperatorParams& Inparams, FBuildErrorArray& OutErrors)"
-			template<typename T>
-			static uint8 Check(Helper<TUniquePtr<IOperator>(*)(const FCreateOperatorParams& InParams, FBuildErrorArray& OutErrors), &T::CreateOperator>*);
+				PRAGMA_DISABLE_DEPRECATION_WARNINGS // Temporary for deprecation of FCreateOperatorParams
+				// Check for "static TUniquePtr<IOperator> U::CreateOperator(const FCreateOperatorParams& Inparams, FBuildErrorArray& OutErrors)"
+				template<typename T>
+				static uint8 Check(Helper<TUniquePtr<IOperator>(*)(const FCreateOperatorParams& InParams, FBuildErrorArray& OutErrors), &T::CreateOperator>*);
+				PRAGMA_ENABLE_DEPRECATION_WARNINGS
 
-			// Check for "static TUniquePtr<IOperator> U::CreateOperator(const FBuildOperatorParams& Inparams, FBuildErrorArray& OutErrors)"
-			template<typename T>
-			static uint8 Check(Helper<TUniquePtr<IOperator>(*)(const FBuildOperatorParams& InParams, FBuildResults& OutResults), &T::CreateOperator>*);
+				// Check for "static TUniquePtr<IOperator> U::CreateOperator(const FBuildOperatorParams& Inparams, FBuildErrorArray& OutErrors)"
+				template<typename T>
+				static uint8 Check(Helper<TUniquePtr<IOperator>(*)(const FBuildOperatorParams& InParams, FBuildResults& OutResults), &T::CreateOperator>*);
 
-			template<typename T> static uint16 Check(...);
+				template<typename T> static uint16 Check(...);
 
-		public:
+			public:
 
-			// If the function exists, then "Value" is true. Otherwise "Value" is false.
-			static constexpr bool Value = sizeof(Check<U>(0)) == sizeof(uint8_t);
-	};
+				// If the function exists, then "Value" is true. Otherwise "Value" is false.
+				static constexpr bool Value = sizeof(Check<U>(0)) == sizeof(uint8_t);
+		};
 
-	// Helper template to determine whether a static class function is declared
-	// for a given template class.
-	template <typename U>
-	class TIsNodeInfoDeclared
-	{
-		private:
-			template<typename T, T> 
-			struct Helper;
+		// Helper template to determine whether a static class function is declared
+		// for a given template class.
+		template <typename U>
+		class TIsNodeInfoDeclared
+		{
+			private:
+				template<typename T, T> 
+				struct Helper;
 
-			// Check for "static const FNodeClassMetadata& U::GetNodeInfo()"
-			template<typename T>
-			static uint8 Check(Helper<const FNodeClassMetadata&(*)(), &T::GetNodeInfo>*);
+				// Check for "static const FNodeClassMetadata& U::GetNodeInfo()"
+				template<typename T>
+				static uint8 Check(Helper<const FNodeClassMetadata&(*)(), &T::GetNodeInfo>*);
 
-			template<typename T> static uint16 Check(...);
+				template<typename T> static uint16 Check(...);
 
-		public:
+			public:
 
-			// If the function exists, then "Value" is true. Otherwise "Value" is false.
-			static constexpr bool Value = sizeof(Check<U>(0)) == sizeof(uint8_t);
-	};
+				// If the function exists, then "Value" is true. Otherwise "Value" is false.
+				static constexpr bool Value = sizeof(Check<U>(0)) == sizeof(uint8_t);
+		};
+	}
 
 	/** TFacadeOperatorClass encapsulates an operator type and checks that the
 	 * required static functions exist to build the facade operator class.  It 
@@ -67,20 +72,11 @@ namespace Metasound
 		static_assert(std::is_base_of<IOperator, OperatorType>::value, "To use the FNodeFacade constructor, the OperatorType must be derived from IOperator");
 
 		// Require static TUniquePtr<IOperator> OperatorType::CreateOperator(const FCreateOperatorParams&, TArray<TUniquePtr<IOperatorBuildError>>&) exists.
-		static_assert(TIsFactoryMethodDeclared<OperatorType>::Value, "To use the FNodeFacade constructor, the OperatorType must have the static function \"static TUniquePtr<IOperator> OperatorType::CreateOperator(const FCreateOperatorParams&, TArray<TUniquePtr<IOperatorBuildError>>&)\"");
-
-		// Require static const FNodeClassMetadata& OperatorType::GetNodeInfo() exists.
-		static_assert(TIsNodeInfoDeclared<OperatorType>::Value, "To use the FNodeFacade constructor, the OperatorType must have the static function \"static const FNodeClassMetadata& OperatorType::GetNodeInfo()\"");
+		static_assert(MetasoundFacadePrivate::TIsFactoryMethodDeclared<OperatorType>::Value, "To use the FNodeFacade constructor, the OperatorType must have the static function \"static TUniquePtr<IOperator> OperatorType::CreateOperator(const FCreateOperatorParams&, TArray<TUniquePtr<IOperatorBuildError>>&)\"");
 
-		static const FNodeClassMetadata& GetNodeInfo() 
-		{ 
-			return OperatorType::GetNodeInfo(); 
-		}
 
-		static TUniquePtr<TUniquePtr<IOperator>> CreateOperator(const FCreateOperatorParams& InParams, FBuildErrorArray& OutErrors) 
-		{
-			return OperatorType::CreateOperator(InParams, OutErrors);
-		}
+		// Require static const FNodeClassMetadata& OperatorType::GetNodeInfo() exists.
+		static_assert(MetasoundFacadePrivate::TIsNodeInfoDeclared<OperatorType>::Value, "To use the FNodeFacade constructor, the OperatorType must have the static function \"static const FNodeClassMetadata& OperatorType::GetNodeInfo()\"");
 
 		typedef OperatorType Type;
 	};
@@ -99,7 +95,10 @@ namespace Metasound
 			class METASOUNDGRAPHCORE_API FFactory : public IOperatorFactory
 			{
 				using FCreateOperatorFunction = TFunction<TUniquePtr<IOperator>(const FBuildOperatorParams& InParams, FBuildResults& OutResults)>;
+
+				PRAGMA_DISABLE_DEPRECATION_WARNINGS // Temporary for deprecation of FCreateOperatorParams
 				using FOriginalCreateOperatorFunction = TFunction<TUniquePtr<IOperator>(const FCreateOperatorParams& InParams, FBuildErrorArray& OutErrors)>;
+				PRAGMA_ENABLE_DEPRECATION_WARNINGS
 
 				public:
 					FFactory(FCreateOperatorFunction InCreateFunc);
