CIS fix for MetaSound runtime data being created at cook

diff --git a/Engine/Plugins/Runtime/Metasound/Source/MetasoundEngine/Private/MetasoundSource.cpp b/Engine/Plugins/Runtime/Metasound/Source/MetasoundEngine/Private/MetasoundSource.cpp
--- a/Engine/Plugins/Runtime/Metasound/Source/MetasoundEngine/Private/MetasoundSource.cpp
+++ b/Engine/Plugins/Runtime/Metasound/Source/MetasoundEngine/Private/MetasoundSource.cpp
@@ -744,7 +744,8 @@ void UMetaSoundSource::RegisterGraphWithFrontend(Metasound::Frontend::FMetaSound
 
 	FMetasoundAssetBase::RegisterGraphWithFrontend(InRegistrationOptions);
 	const bool bIsRuntimeInputDataValid = RuntimeInputData.bIsValid.load();
-	if (!bIsRuntimeInputDataValid)
+	// Runtime data does not need to and should not be created at cook
+	if (!bIsRuntimeInputDataValid && !IsRunningCookCommandlet())
 	{
 		CacheRuntimeInputData();
 	}
