Add quality, sample rate & block rate override setters to MS builder API
#rnx

diff --git a/Engine/Plugins/Runtime/Metasound/Source/MetasoundEngine/Private/MetasoundBuilderSubsystem.cpp b/Engine/Plugins/Runtime/Metasound/Source/MetasoundEngine/Private/MetasoundBuilderSubsystem.cpp
--- a/Engine/Plugins/Runtime/Metasound/Source/MetasoundEngine/Private/MetasoundBuilderSubsystem.cpp
+++ b/Engine/Plugins/Runtime/Metasound/Source/MetasoundEngine/Private/MetasoundBuilderSubsystem.cpp
@@ -25,6 +25,7 @@
 #include "MetasoundTrace.h"
 #include "MetasoundUObjectRegistry.h"
 #include "MetasoundVertex.h"
+#include "PerPlatformProperties.h"
 
 #include UE_INLINE_GENERATED_CPP_BY_NAME(MetasoundBuilderSubsystem)
 
@@ -1491,6 +1492,11 @@ void UMetaSoundSourceBuilder::OnRemovingOutput(int32 OutputIndex) const
 	});
 }
 
+void UMetaSoundSourceBuilder::SetBlockRateOverride(float BlockRate)
+{
+	GetMetaSoundSource().BlockRateOverride.Default = BlockRate;
+}
+
 void UMetaSoundSourceBuilder::SetFormat(EMetaSoundOutputAudioFormat OutputFormat, EMetaSoundBuilderResult& OutResult)
 {
 	using namespace Metasound::Engine;
@@ -1534,6 +1540,28 @@ void UMetaSoundSourceBuilder::SetFormat(EMetaSoundOutputAudioFormat OutputFormat
 	OutResult = bSuccess ? EMetaSoundBuilderResult::Succeeded : EMetaSoundBuilderResult::Failed;
 }
 
+#if WITH_EDITORONLY_DATA
+void UMetaSoundSourceBuilder::SetPlatformBlockRateOverride(const FPerPlatformFloat& PlatformBlockRate)
+{
+	GetMetaSoundSource().BlockRateOverride = PlatformBlockRate;
+}
+
+void UMetaSoundSourceBuilder::SetPlatformSampleRateOverride(const FPerPlatformInt& PlatformSampleRate)
+{
+	GetMetaSoundSource().SampleRateOverride = PlatformSampleRate;
+}
+#endif // WITH_EDITORONLY_DATA
+
+void UMetaSoundSourceBuilder::SetQuality(FName Quality)
+{
+	GetMetaSoundSource().QualitySetting = Quality;
+}
+
+void UMetaSoundSourceBuilder::SetSampleRateOverride(int32 SampleRate)
+{
+	GetMetaSoundSource().SampleRateOverride.Default = SampleRate;
+}
+
 UMetaSoundPatchBuilder* UMetaSoundBuilderSubsystem::CreatePatchBuilder(FName BuilderName, EMetaSoundBuilderResult& OutResult)
 {
 	OutResult = EMetaSoundBuilderResult::Succeeded;
diff --git a/Engine/Plugins/Runtime/Metasound/Source/MetasoundEngine/Public/MetasoundBuilderSubsystem.h b/Engine/Plugins/Runtime/Metasound/Source/MetasoundEngine/Public/MetasoundBuilderSubsystem.h
--- a/Engine/Plugins/Runtime/Metasound/Source/MetasoundEngine/Public/MetasoundBuilderSubsystem.h
+++ b/Engine/Plugins/Runtime/Metasound/Source/MetasoundEngine/Public/MetasoundBuilderSubsystem.h
@@ -22,6 +22,7 @@
 
 
 // Forward Declarations
+class FMetasoundAssetBase;
 class UAudioComponent;
 class UMetaSound;
 class UMetaSoundPatch;
@@ -29,7 +30,8 @@ class UMetaSoundSource;
 
 struct FMetasoundFrontendClassName;
 struct FMetasoundFrontendVersion;
-class FMetasoundAssetBase;
+struct FPerPlatformFloat;
+struct FPerPlatformInt;
 
 enum class EMetaSoundOutputAudioFormat : uint8;
 
@@ -529,18 +531,37 @@ public:
 	UFUNCTION(BlueprintCallable, Category = "Audio|MetaSound|Builder")
 	bool GetLiveUpdatesEnabled() const;
 
+	// Sets the MetaSound's BlockRate override
+	UFUNCTION(BlueprintCallable, Category = "Audio|MetaSound|Builder")
+	void SetBlockRateOverride(float BlockRate);
+
 	// Sets the output audio format of the source
 	UFUNCTION(BlueprintCallable, Category = "Audio|MetaSound|Builder", meta = (ExpandEnumAsExecs = "OutResult"))
 	void SetFormat(EMetaSoundOutputAudioFormat OutputFormat, EMetaSoundBuilderResult& OutResult);
 
+	// Sets the MetaSound's SampleRate override
+	UFUNCTION(BlueprintCallable, Category = "Audio|MetaSound|Builder")
+	void SetSampleRateOverride(int32 SampleRate);
+
 	const Metasound::Engine::FOutputAudioFormatInfoPair* FindOutputAudioFormatInfo() const;
 
 	virtual const UClass& GetBuilderUClass() const override;
 
+#if WITH_EDITORONLY_DATA
+	// Sets the MetaSound's BlockRate override (editor only, to allow setting per-platform values)
+	void SetPlatformBlockRateOverride(const FPerPlatformFloat& PlatformFloat);
+
+	// Sets the MetaSound's BlockRate override (editor only, to allow setting per-platform values)
+	void SetPlatformSampleRateOverride(const FPerPlatformInt& PlatformInt);
+#endif // WITH_EDITORONLY_DATA
+
+	// Sets the MetaSound's Quality level
+	UFUNCTION(BlueprintCallable, Category = "Audio|MetaSound|Builder")
+	void SetQuality(FName Quality);
+
 protected:
 	virtual void CreateTransientBuilder() override;
 
-
 private:
 	static TOptional<Metasound::FAnyDataReference> CreateDataReference(const Metasound::FOperatorSettings& InOperatorSettings, FName DataType, const Metasound::FLiteral& InLiteral, Metasound::EDataReferenceAccessType AccessType);
 
