Use BusyWait instead of Wait in tasks run on worker threads.

On low core count machines where there are only a couple worker threads, waiting from a worker thread with a blocking Wait can mean it ends up deadlocking. BusyWait will ensure the worker thread processes outstanding tasks while it waits and makes forward progress.

diff --git a/Engine/Plugins/Runtime/Metasound/Source/MetasoundFrontend/Private/MetasoundFrontendRegistries.cpp b/Engine/Plugins/Runtime/Metasound/Source/MetasoundFrontend/Private/MetasoundFrontendRegistries.cpp
--- a/Engine/Plugins/Runtime/Metasound/Source/MetasoundFrontend/Private/MetasoundFrontendRegistries.cpp
+++ b/Engine/Plugins/Runtime/Metasound/Source/MetasoundFrontend/Private/MetasoundFrontendRegistries.cpp
@@ -27,6 +27,13 @@
 #include "MetasoundRouter.h"
 #include "MetasoundTrace.h"
 
+bool bBusyWaitOnAsyncRegistrationTasks = true;
+static FAutoConsoleVariableRef CVarAsyncRegistrationTasksBusyWait(
+	TEXT("au.MetaSound.BusyWaitOnAsyncRegistrationTasks"),
+	bBusyWaitOnAsyncRegistrationTasks,
+	TEXT("Use TaskGraph BusyWait instead of simple Wait. Required to avoid hangs on platforms with low number of cores."),
+	ECVF_Default);
+
 namespace Metasound
 {
 	namespace Frontend
@@ -956,7 +963,14 @@ namespace Metasound
 			if (ActiveRegistrationTask.IsValid())
 			{
 				METASOUND_TRACE_CPUPROFILER_EVENT_SCOPE(Metasound::FRegistryContainerImpl::WaitForRegistrationTaskToComplete);
-				ActiveRegistrationTask.Wait();
+				if (bBusyWaitOnAsyncRegistrationTasks)
+				{
+					ActiveRegistrationTask.BusyWait();
+				}
+				else
+				{
+					ActiveRegistrationTask.Wait();
+				}
 			}
 		}
 
