Add cvar to allow resetting metasound operators on cache insertion to potentially deallocate memory.

#rnx

#localization none
#tests local LLM A/B in PIE to confirm regression fixed w/o affecting weapon audio functionality

diff --git a/Engine/Plugins/Runtime/Metasound/Source/MetasoundEngine/Private/MetasoundWavePlayerNode.cpp b/Engine/Plugins/Runtime/Metasound/Source/MetasoundEngine/Private/MetasoundWavePlayerNode.cpp
--- a/Engine/Plugins/Runtime/Metasound/Source/MetasoundEngine/Private/MetasoundWavePlayerNode.cpp
+++ b/Engine/Plugins/Runtime/Metasound/Source/MetasoundEngine/Private/MetasoundWavePlayerNode.cpp
@@ -490,7 +490,7 @@ namespace Metasound
 			SortedCuePoints.Reset();
 			for (Audio::TCircularAudioBuffer<float>& Buffer : SourceCircularBuffer)
 			{
-				Buffer.SetNum(0);
+				Buffer.Empty();
 			}
 
 			SourceState = WavePlayerNodePrivate::FSourceBufferState();
diff --git a/Engine/Plugins/Runtime/Metasound/Source/MetasoundGenerator/Private/MetasoundGenerator.cpp b/Engine/Plugins/Runtime/Metasound/Source/MetasoundGenerator/Private/MetasoundGenerator.cpp
--- a/Engine/Plugins/Runtime/Metasound/Source/MetasoundGenerator/Private/MetasoundGenerator.cpp
+++ b/Engine/Plugins/Runtime/Metasound/Source/MetasoundGenerator/Private/MetasoundGenerator.cpp
@@ -33,6 +33,7 @@ namespace Metasound
 		static bool bEnableAsyncMetaSoundGeneratorBuilder = true;
 		static bool bEnableExperimentalOneShotOperatorPool = false;
 		static bool bEnableExperimentalOperatorPool = false;
+		static bool bEnableResetOnOperatorPoolInsertion = true;
 #if ENABLE_METASOUNDGENERATOR_INVALID_SAMPLE_VALUE_LOGGING
 		static bool bEnableMetaSoundGeneratorNonFiniteLogging = false;
 		static bool bEnableMetaSoundGeneratorInvalidSampleValueLogging = false;
@@ -210,6 +211,14 @@ FAutoConsoleVariableRef CVarMetaSoundEnableExperimentalOperatorPool(
 	TEXT("Default: false"),
 	ECVF_Default);
 
+FAutoConsoleVariableRef CVarMetaSoundEnableResetOnOperatorPoolInsertion(
+	TEXT("au.MetaSound.Experimental.EnableResetOnOperatorPoolInsertion"),
+	Metasound::ConsoleVariables::bEnableResetOnOperatorPoolInsertion,
+	TEXT("Enables reset of MetaSound operators as they are inserted into the Operator pool.\n")
+	TEXT("This can save a significant amount memory at the cost of reseting pooled operators twice.\n")
+	TEXT("Default: true"),
+	ECVF_Default);
+
 #if ENABLE_METASOUNDGENERATOR_INVALID_SAMPLE_VALUE_LOGGING
 
 FAutoConsoleVariableRef CVarMetaSoundEnableGeneratorNonFiniteLogging(
@@ -865,7 +874,8 @@ PRAGMA_ENABLE_DEPRECATION_WARNINGS
 				// Reset operator internal state before playing it.
 				if (IOperator::FResetFunction Reset = GraphOperatorAndInputs.Operator->GetResetFunction())
 				{
-					IOperator::FResetParams ResetParams {OperatorSettings, InInitParams.Environment};
+					EnvironmentPtr = MakeUnique<FMetasoundEnvironment>(InInitParams.Environment);
+					IOperator::FResetParams ResetParams {OperatorSettings, *EnvironmentPtr};
 					Reset(GraphOperatorAndInputs.Operator.Get(), ResetParams);
 				}
 				
@@ -896,6 +906,17 @@ PRAGMA_ENABLE_DEPRECATION_WARNINGS
 			{
 				// Release graph operator and input data to the cache
 				UE_LOG(LogMetasoundGenerator, VeryVerbose, TEXT("Caching operator %s"), *LexToString(OperatorID));
+
+				// give the operator a chance to reduce its memory footprint before being cached
+				if (ConsoleVariables::bEnableResetOnOperatorPoolInsertion && EnvironmentPtr.IsValid())
+				{
+					if (IOperator::FResetFunction Reset = GraphOperator->GetResetFunction())
+					{
+						IOperator::FResetParams ResetParams {OperatorSettings, *EnvironmentPtr.Get()};
+						Reset(GraphOperator.Get(), ResetParams);
+					}
+				}
+
 				OperatorPool->AddOperator(OperatorID, MoveTemp(GraphOperator), ReleaseInputVertexData());
 
 			}
diff --git a/Engine/Plugins/Runtime/Metasound/Source/MetasoundGenerator/Private/MetasoundOperatorCache.cpp b/Engine/Plugins/Runtime/Metasound/Source/MetasoundGenerator/Private/MetasoundOperatorCache.cpp
--- a/Engine/Plugins/Runtime/Metasound/Source/MetasoundGenerator/Private/MetasoundOperatorCache.cpp
+++ b/Engine/Plugins/Runtime/Metasound/Source/MetasoundGenerator/Private/MetasoundOperatorCache.cpp
@@ -20,6 +20,7 @@
 namespace Metasound
 {
 #if METASOUND_OPERATORCACHEPROFILER_ENABLED
+	TRACE_DECLARE_INT_COUNTER(MetaSound_OperatorPool_NumOperators, TEXT("MetaSound/OperatorPool/NumOperatorsInPool"));
 	TRACE_DECLARE_FLOAT_COUNTER(MetaSound_OperatorPool_HitRatio, TEXT("MetaSound/OperatorPool/HitRatio"));
 	TRACE_DECLARE_FLOAT_COUNTER(MetaSound_OperatorPool_WindowedHitRatio, TEXT("MetaSound/OperatorPool/WindowedHitRatio"));
 
@@ -186,6 +187,7 @@ namespace Metasound
 #if METASOUND_OPERATORCACHEPROFILER_ENABLED
 				OperatorPoolPrivate::CacheHitCount++;
 				bCacheHit = true;
+				TRACE_COUNTER_DECREMENT(MetaSound_OperatorPool_NumOperators);
 #endif
 			}
 		}
@@ -212,6 +214,9 @@ namespace Metasound
 
 		FScopeLock Lock(&CriticalSection);
 		Stack.Add(InOperatorID);
+#if METASOUND_OPERATORCACHEPROFILER_ENABLED
+		TRACE_COUNTER_INCREMENT(MetaSound_OperatorPool_NumOperators);
+#endif // #if METASOUND_OPERATORCACHEPROFILER_ENABLED
 
 		if (TArray<FOperatorAndInputs>* OperatorArray = Operators.Find(InOperatorID))
 		{
@@ -303,7 +308,11 @@ namespace Metasound
 	{
 		FScopeLock Lock(&CriticalSection);
 		Operators.Remove(InOperatorID);
-		Stack.Remove(InOperatorID);
+		const int32 NumRemoved = Stack.Remove(InOperatorID);
+
+#if METASOUND_OPERATORCACHEPROFILER_ENABLED
+		TRACE_COUNTER_SUBTRACT(MetaSound_OperatorPool_NumOperators, int64(NumRemoved));
+#endif // #if METASOUND_OPERATORCACHEPROFILER_ENABLED
 	}
 
 	void FOperatorPool::RemoveOperatorsWithAssetClassID(const FGuid& InAssetClassID)
@@ -384,6 +393,9 @@ namespace Metasound
 				}
 			}
 			Stack.RemoveAt(0, NumToTrim);
+#if METASOUND_OPERATORCACHEPROFILER_ENABLED
+			TRACE_COUNTER_DECREMENT(MetaSound_OperatorPool_NumOperators);
+#endif // #if METASOUND_OPERATORCACHEPROFILER_ENABLED
 		}
 
 		// todo: prune AssetIdToGraphIdLookUp?
diff --git a/Engine/Plugins/Runtime/Metasound/Source/MetasoundGenerator/Public/MetasoundGenerator.h b/Engine/Plugins/Runtime/Metasound/Source/MetasoundGenerator/Public/MetasoundGenerator.h
--- a/Engine/Plugins/Runtime/Metasound/Source/MetasoundGenerator/Public/MetasoundGenerator.h
+++ b/Engine/Plugins/Runtime/Metasound/Source/MetasoundGenerator/Public/MetasoundGenerator.h
@@ -365,6 +365,7 @@ namespace Metasound
 		bool TryUseCachedOperator(FMetasoundGeneratorInitParams& InParams, bool bTriggerGenerator);
 		void ReleaseOperatorToCache();
 
+		TUniquePtr<FMetasoundEnvironment> EnvironmentPtr;
 		TUniquePtr<FAsyncTaskBase> BuilderTask;
 		FGuid OperatorID;
 		bool bUseOperatorPool = false;
diff --git a/Engine/Source/Runtime/SignalProcessing/Public/DSP/Dsp.h b/Engine/Source/Runtime/SignalProcessing/Public/DSP/Dsp.h
--- a/Engine/Source/Runtime/SignalProcessing/Public/DSP/Dsp.h
+++ b/Engine/Source/Runtime/SignalProcessing/Public/DSP/Dsp.h
@@ -880,6 +880,13 @@ namespace Audio
 			SetCapacity(InCapacity);
 		}
 
+		void Empty()
+		{
+			ReadCounter.Set(0);
+			WriteCounter.Set(0);
+			InternalBuffer.Empty();
+		}
+
 		void SetCapacity(uint32 InCapacity)
 		{
 			checkf(InCapacity < (uint32)TNumericLimits<int32>::Max(), TEXT("Max capacity for this buffer is 2,147,483,647 samples. Otherwise our index arithmetic will not work."));
