Minor clean-up from prior CR
#rnx

diff --git a/Engine/Plugins/Runtime/Metasound/Source/MetasoundFrontend/Private/MetasoundFrontendDocumentCache.cpp b/Engine/Plugins/Runtime/Metasound/Source/MetasoundFrontend/Private/MetasoundFrontendDocumentCache.cpp
--- a/Engine/Plugins/Runtime/Metasound/Source/MetasoundFrontend/Private/MetasoundFrontendDocumentCache.cpp
+++ b/Engine/Plugins/Runtime/Metasound/Source/MetasoundFrontend/Private/MetasoundFrontendDocumentCache.cpp
@@ -910,6 +910,8 @@ namespace Metasound::Frontend
 
 	FDocumentCache::FDocumentDependencyCache::FDocumentDependencyCache(const FMetasoundFrontendDocument& InDocument)
 	{
+		METASOUND_TRACE_CPUPROFILER_EVENT_SCOPE(Metasound::Frontend::FDocumentDependencyCache::FDocumentDependencyCache);
+
 		for (int32 Index = 0; Index < InDocument.Dependencies.Num(); ++Index)
 		{
 			const FMetasoundFrontendClass& Class = InDocument.Dependencies[Index];
diff --git a/Engine/Plugins/Runtime/Metasound/Source/MetasoundFrontend/Private/MetasoundFrontendDocumentCache.h b/Engine/Plugins/Runtime/Metasound/Source/MetasoundFrontend/Private/MetasoundFrontendDocumentCache.h
--- a/Engine/Plugins/Runtime/Metasound/Source/MetasoundFrontend/Private/MetasoundFrontendDocumentCache.h
+++ b/Engine/Plugins/Runtime/Metasound/Source/MetasoundFrontend/Private/MetasoundFrontendDocumentCache.h
@@ -16,14 +16,6 @@ namespace Metasound::Frontend
 	// Forward Declarations
 	class FDocumentCache;
 
-	struct FDocumentCacheTransaction
-	{
-		int32 DocumentTransactionID = INDEX_NONE;
-		int32 InterfaceTransactionID = INDEX_NONE;
-		int32 NodeTransactionID = INDEX_NONE;
-		int32 EdgeTransactionID = INDEX_NONE;
-	};
-
 	class FDocumentGraphEdgeCache : public IDocumentGraphEdgeCache
 	{
 		template <typename ObjectType, ESPMode Mode>
