Enable transaction-based registration in Builder API
#rnx

diff --git a/Engine/Plugins/Runtime/Metasound/Source/MetasoundEngine/Private/MetasoundBuilderSubsystem.cpp b/Engine/Plugins/Runtime/Metasound/Source/MetasoundEngine/Private/MetasoundBuilderSubsystem.cpp
--- a/Engine/Plugins/Runtime/Metasound/Source/MetasoundEngine/Private/MetasoundBuilderSubsystem.cpp
+++ b/Engine/Plugins/Runtime/Metasound/Source/MetasoundEngine/Private/MetasoundBuilderSubsystem.cpp
@@ -34,7 +34,7 @@ namespace Metasound::Engine
 {
 	namespace BuilderSubsystemPrivate
 	{
-		int32 TransactionBasedRegistrationEnabled = 0;
+		int32 TransactionBasedRegistrationEnabled = 1;
 
 		template <typename TLiteralType>
 		FMetasoundFrontendLiteral CreatePODMetaSoundLiteral(const TLiteralType& Value, FName& OutDataType)
