Show source or MetaSound settings by default in MetaSound details panel, autoexpand graph member menu by default

diff --git a/Engine/Plugins/Runtime/Metasound/Source/MetasoundEditor/Private/MetasoundDetailCustomization.cpp b/Engine/Plugins/Runtime/Metasound/Source/MetasoundEditor/Private/MetasoundDetailCustomization.cpp
--- a/Engine/Plugins/Runtime/Metasound/Source/MetasoundEditor/Private/MetasoundDetailCustomization.cpp
+++ b/Engine/Plugins/Runtime/Metasound/Source/MetasoundEditor/Private/MetasoundDetailCustomization.cpp
@@ -96,12 +96,6 @@ namespace Metasound
 		{
 			using namespace Frontend;
 
-			EMetasoundActiveDetailView DetailsView = EMetasoundActiveDetailView::Metasound;
-			if (const UMetasoundEditorSettings* EditorSettings = GetDefault<UMetasoundEditorSettings>())
-			{
-				DetailsView = EditorSettings->DetailView;
-			}
-
 			TArray<TWeakObjectPtr<UObject>> Objects;
 			DetailLayout.GetObjectsBeingCustomized(Objects);
 
@@ -119,6 +113,18 @@ namespace Metasound
 
 			TWeakObjectPtr<UMetaSoundSource> MetaSoundSource = Cast<UMetaSoundSource>(MetaSound.Get());
 
+			// MetaSound patches don't have source settings, so view MetaSound settings by default 
+			EMetasoundActiveDetailView DetailsView = EMetasoundActiveDetailView::Metasound;
+			if (MetaSoundSource.IsValid())
+			{
+				// Show source settings by default unless previously set
+				DetailsView = EMetasoundActiveDetailView::General;
+				if (const UMetasoundEditorSettings* EditorSettings = GetDefault<UMetasoundEditorSettings>())
+				{
+					DetailsView = EditorSettings->DetailView;
+				}
+			}
+
 			switch (DetailsView)
 			{
 				case EMetasoundActiveDetailView::Metasound:
@@ -261,20 +267,13 @@ namespace Metasound
 				default:
 					DetailLayout.HideCategory("MetaSound");
 
-					const bool bShouldBeInitiallyCollapsed = true;
 					TArray<TSharedRef<IPropertyHandle>>DeveloperProperties;
 					TArray<TSharedRef<IPropertyHandle>>SoundProperties;
 
 					DetailLayout.EditCategory("Sound")
-						.InitiallyCollapsed(bShouldBeInitiallyCollapsed)
 						.GetDefaultProperties(SoundProperties);
-					DetailLayout.EditCategory("Attenuation").InitiallyCollapsed(bShouldBeInitiallyCollapsed);
 					DetailLayout.EditCategory("Developer")
-						.InitiallyCollapsed(bShouldBeInitiallyCollapsed)
 						.GetDefaultProperties(DeveloperProperties);
-					DetailLayout.EditCategory("Effects").InitiallyCollapsed(bShouldBeInitiallyCollapsed);
-					DetailLayout.EditCategory("Modulation").InitiallyCollapsed(bShouldBeInitiallyCollapsed);
-					DetailLayout.EditCategory("Voice Management").InitiallyCollapsed(bShouldBeInitiallyCollapsed);
 
 					auto HideProperties = [](const TSet<FName>& PropsToHide, const TArray<TSharedRef<IPropertyHandle>>& Properties)
 					{
diff --git a/Engine/Plugins/Runtime/Metasound/Source/MetasoundEditor/Private/MetasoundEditor.cpp b/Engine/Plugins/Runtime/Metasound/Source/MetasoundEditor/Private/MetasoundEditor.cpp
--- a/Engine/Plugins/Runtime/Metasound/Source/MetasoundEditor/Private/MetasoundEditor.cpp
+++ b/Engine/Plugins/Runtime/Metasound/Source/MetasoundEditor/Private/MetasoundEditor.cpp
@@ -1114,6 +1114,7 @@ namespace Metasound
 
 			SAssignNew(GraphMembersMenu, SGraphActionMenu, false)
 				.AlphaSortItems(true)
+				.AutoExpandActionMenu(true)
 				.OnActionDoubleClicked(this, &FEditor::OnMemberActionDoubleClicked)
 				.OnActionDragged(this, &FEditor::OnActionDragged)
 				.OnActionMatchesName(this, &FEditor::HandleActionMatchesName)
@@ -1135,6 +1136,8 @@ namespace Metasound
 
 			FPropertyEditorModule& PropertyModule = FModuleManager::LoadModuleChecked<FPropertyEditorModule>("PropertyEditor");
 			MetasoundDetails = PropertyModule.CreateDetailView(Args);
+			// Set details selection to the MetaSound's source settings 
+			SetSelection({ Metasound });
 			InterfacesDetails = PropertyModule.CreateDetailView(Args);
 			if (InterfacesDetails.IsValid())
 			{
diff --git a/Engine/Plugins/Runtime/Metasound/Source/MetasoundEditor/Private/MetasoundEditorSettings.h b/Engine/Plugins/Runtime/Metasound/Source/MetasoundEditor/Private/MetasoundEditorSettings.h
--- a/Engine/Plugins/Runtime/Metasound/Source/MetasoundEditor/Private/MetasoundEditorSettings.h
+++ b/Engine/Plugins/Runtime/Metasound/Source/MetasoundEditor/Private/MetasoundEditorSettings.h
@@ -160,5 +160,5 @@ public:
 
 	/** Determines which details view to show in Metasounds Editor */
 	UPROPERTY(Transient)
-	EMetasoundActiveDetailView DetailView = EMetasoundActiveDetailView::Metasound;
+	EMetasoundActiveDetailView DetailView = EMetasoundActiveDetailView::General;
 };
