Fix MetaSound reset to default arrow still showing even after literal has been set to default

diff --git a/Engine/Plugins/Runtime/Metasound/Source/MetasoundFrontend/Private/MetasoundFrontendNodeController.cpp b/Engine/Plugins/Runtime/Metasound/Source/MetasoundFrontend/Private/MetasoundFrontendNodeController.cpp
--- a/Engine/Plugins/Runtime/Metasound/Source/MetasoundFrontend/Private/MetasoundFrontendNodeController.cpp
+++ b/Engine/Plugins/Runtime/Metasound/Source/MetasoundFrontend/Private/MetasoundFrontendNodeController.cpp
@@ -123,32 +123,35 @@ namespace Metasound
 					{
 						// Type None implies forwarding to the node instance to default construct the literal 
 						// so setting to/from that class default literal is an exception to requiring a type match
-						const bool bIsMatchingType = VertexLiteral.Value.GetType() == InVertexLiteral.Value.GetType();
 						bool bSettingToClassDefaultLiteral = false;
 						bool bSettingFromClassDefaultLiteralToMatchingType = false;
 
-						if (!bIsMatchingType)
+						auto GetInputWithVertexID = [InVertexID = VertexLiteral.VertexID](const FMetasoundFrontendClassInput& ClassInput)
 						{
-							auto GetInputWithVertexID = [InVertexID = VertexLiteral.VertexID](const FMetasoundFrontendClassInput& ClassInput)
-							{
-								return ClassInput.VertexID == InVertexID;
-							};
-							const FMetasoundFrontendClassInput* ClassInput = ClassPtr.Get()->Interface.Inputs.FindByPredicate(GetInputWithVertexID);
-							if (ClassInput)
-							{
-								// Check if setting back to class default literal
-								const FMetasoundFrontendLiteral& ClassDefaultLiteral = ClassInput->DefaultLiteral;
-								bSettingToClassDefaultLiteral = ClassDefaultLiteral.IsEqual(InVertexLiteral.Value);
-
-								// Check if setting from class default literal (which may have a None type) to an appropriate type 
-								FDataTypeRegistryInfo DataTypeInfo;
-								IDataTypeRegistry::Get().GetDataTypeInfo(ClassInput->TypeName, DataTypeInfo);
-								const EMetasoundFrontendLiteralType ClassInputLiteralType = static_cast<EMetasoundFrontendLiteralType>(DataTypeInfo.PreferredLiteralType);
-								bSettingFromClassDefaultLiteralToMatchingType |= ClassDefaultLiteral.IsEqual(VertexLiteral.Value) && ClassInputLiteralType == InVertexLiteral.Value.GetType();
-							}
+							return ClassInput.VertexID == InVertexID;
+						};
+						const FMetasoundFrontendClassInput* ClassInput = ClassPtr.Get()->Interface.Inputs.FindByPredicate(GetInputWithVertexID);
+						if (ClassInput)
+						{
+							// Check if setting back to class default literal
+							const FMetasoundFrontendLiteral& ClassDefaultLiteral = ClassInput->DefaultLiteral;
+							bSettingToClassDefaultLiteral = ClassDefaultLiteral.IsEqual(InVertexLiteral.Value);
+
+							// Check if setting from class default literal (which may have a None type) to an appropriate type 
+							FDataTypeRegistryInfo DataTypeInfo;
+							IDataTypeRegistry::Get().GetDataTypeInfo(ClassInput->TypeName, DataTypeInfo);
+							const EMetasoundFrontendLiteralType ClassInputLiteralType = static_cast<EMetasoundFrontendLiteralType>(DataTypeInfo.PreferredLiteralType);
+							bSettingFromClassDefaultLiteralToMatchingType |= ClassDefaultLiteral.IsEqual(VertexLiteral.Value) && ClassInputLiteralType == InVertexLiteral.Value.GetType();
 						}
 
-						if (ensure(bIsMatchingType || bSettingToClassDefaultLiteral || bSettingFromClassDefaultLiteralToMatchingType))
+						if (bSettingToClassDefaultLiteral)
+						{
+							ClearInputLiteral(VertexLiteral.VertexID);
+							return;
+						}
+						
+						const bool bIsMatchingType = VertexLiteral.Value.GetType() == InVertexLiteral.Value.GetType();
+						if (ensure(bIsMatchingType || bSettingFromClassDefaultLiteralToMatchingType))
 						{
 							VertexLiteral = InVertexLiteral;
 						}
