Added better naming to new quality settings.

diff --git a/Engine/Plugins/Runtime/Metasound/Source/MetasoundEngine/Private/MetasoundSettings.cpp b/Engine/Plugins/Runtime/Metasound/Source/MetasoundEngine/Private/MetasoundSettings.cpp
new file mode 100644
--- /dev/null
+++ b/Engine/Plugins/Runtime/Metasound/Source/MetasoundEngine/Private/MetasoundSettings.cpp
@@ -0,0 +1,37 @@
+// Copyright Epic Games, Inc. All Rights Reserved.
+
+#include "MetasoundSettings.h"
+
+#if WITH_EDITORONLY_DATA
+FMetaSoundQualitySettings::FMetaSoundQualitySettings()
+	: UniqueId(FGuid::NewGuid())
+	, Name(UMetaSoundQuality::GenerateNewName())
+{
+}
+#endif //WITH_EDITORONLY_DATA
+
+FName UMetaSoundQuality::GenerateNewName()
+{
+	const TArray<FName> Names = GetQualityList();
+	const TCHAR* BaseName = TEXT("New Quality");
+	FString NewName = BaseName;
+	
+	for (int32 Postfix=0; Names.Contains(*NewName); ++Postfix)
+	{
+		NewName = FString::Format(TEXT("{0} {1}"), { BaseName, Postfix });
+	}
+	return FName(*NewName);
+}
+
+TArray<FName> UMetaSoundQuality::GetQualityList()
+{
+	TArray<FName> Names;
+	if (const UMetaSoundSettings* Settings = GetDefault<UMetaSoundSettings>())
+	{
+		Algo::Transform(Settings->QualitySettings, Names, [](const FMetaSoundQualitySettings& Quality) -> FName
+		{
+			return Quality.Name;
+		});
+	}
+	return Names;
+}
diff --git a/Engine/Plugins/Runtime/Metasound/Source/MetasoundEngine/Public/MetasoundSettings.h b/Engine/Plugins/Runtime/Metasound/Source/MetasoundEngine/Public/MetasoundSettings.h
--- a/Engine/Plugins/Runtime/Metasound/Source/MetasoundEngine/Public/MetasoundSettings.h
+++ b/Engine/Plugins/Runtime/Metasound/Source/MetasoundEngine/Public/MetasoundSettings.h
@@ -34,6 +34,19 @@ struct METASOUNDENGINE_API FDefaultMetaSoundAssetAutoUpdateSettings
 	FSoftObjectPath MetaSound;
 };
 
+UCLASS()
+class UMetaSoundQuality : public UObject
+{
+	GENERATED_BODY()
+public:
+	static FName GenerateNewName();
+	/**
+	* Returns a list of quality settings to present to a combobox
+	* */
+	UFUNCTION()
+	static TArray<FName> GetQualityList();
+};
+
 USTRUCT()
 struct METASOUNDENGINE_API FMetaSoundQualitySettings
 {
@@ -41,15 +54,18 @@ struct METASOUNDENGINE_API FMetaSoundQualitySettings
 
 #if WITH_EDITORONLY_DATA
 
+	/* Editor only constructor, that will generate a new guid. And a new name */
+	FMetaSoundQualitySettings();
+	
 	/** A hidden GUID that will be generated once when adding a new entry. This prevents orphaning of renamed entries. **/
 	UPROPERTY(meta = (IgnoreForMemberInitializationTest))
-	FGuid UniqueId = FGuid::NewGuid();
+	FGuid UniqueId = {};
 
 #endif //WITH_EDITORONLY_DATA
 
 	/** Name of this quality setting. This will appear in the quality dropdown list.
 		The names should be unique and adequately describe the Entry. "High", "Low" etc. **/
-	UPROPERTY(EditAnywhere, Category = "Quality")
+	UPROPERTY(EditAnywhere, Category = "Quality", meta = (IgnoreForMemberInitializationTest))
 	FName Name;
 
 	/** Sample Rate (in Hz). NOTE: A Zero value will have no effect and use the Device Rate. **/
@@ -110,28 +126,3 @@ public:
 #endif // WITH_EDITOR
 };
 
-UCLASS()
-class UMetaSoundQuality : public UObject
-{
-	GENERATED_BODY()
-public:
-
-	/**
-	* Returns a list of quality settings to present to a combobox
-	* */
-	UFUNCTION()
-	static TArray<FName> GetQualityList()
-	{
-		TArray<FName> Names;
-
-		if (const UMetaSoundSettings* Settings = GetDefault<UMetaSoundSettings>())
-		{
-			Algo::Transform(Settings->QualitySettings, Names, [](const FMetaSoundQualitySettings& Quality) -> FName 
-			{
-				return Quality.Name;
-			});
-		}
-		return Names;
-	}
-
-};
