[Backout] - CL29792222
Original CL Desc
-----------------------------------------------------------------
Added better naming to new quality settings.

diff --git a/Engine/Plugins/Runtime/Metasound/Source/MetasoundEngine/Private/MetasoundSettings.cpp b/Engine/Plugins/Runtime/Metasound/Source/MetasoundEngine/Private/MetasoundSettings.cpp
deleted file mode 100644
diff --git a/Engine/Plugins/Runtime/Metasound/Source/MetasoundEngine/Public/MetasoundSettings.h b/Engine/Plugins/Runtime/Metasound/Source/MetasoundEngine/Public/MetasoundSettings.h
--- a/Engine/Plugins/Runtime/Metasound/Source/MetasoundEngine/Public/MetasoundSettings.h
+++ b/Engine/Plugins/Runtime/Metasound/Source/MetasoundEngine/Public/MetasoundSettings.h
@@ -34,19 +34,6 @@ struct METASOUNDENGINE_API FDefaultMetaSoundAssetAutoUpdateSettings
 	FSoftObjectPath MetaSound;
 };
 
-UCLASS()
-class UMetaSoundQuality : public UObject
-{
-	GENERATED_BODY()
-public:
-	static FName GenerateNewName();
-	/**
-	* Returns a list of quality settings to present to a combobox
-	* */
-	UFUNCTION()
-	static TArray<FName> GetQualityList();
-};
-
 USTRUCT()
 struct METASOUNDENGINE_API FMetaSoundQualitySettings
 {
@@ -54,18 +41,15 @@ struct METASOUNDENGINE_API FMetaSoundQualitySettings
 
 #if WITH_EDITORONLY_DATA
 
-	/* Editor only constructor, that will generate a new guid. And a new name */
-	FMetaSoundQualitySettings();
-	
 	/** A hidden GUID that will be generated once when adding a new entry. This prevents orphaning of renamed entries. **/
 	UPROPERTY(meta = (IgnoreForMemberInitializationTest))
-	FGuid UniqueId = {};
+	FGuid UniqueId = FGuid::NewGuid();
 
 #endif //WITH_EDITORONLY_DATA
 
 	/** Name of this quality setting. This will appear in the quality dropdown list.
 		The names should be unique and adequately describe the Entry. "High", "Low" etc. **/
-	UPROPERTY(EditAnywhere, Category = "Quality", meta = (IgnoreForMemberInitializationTest))
+	UPROPERTY(EditAnywhere, Category = "Quality")
 	FName Name;
 
 	/** Sample Rate (in Hz). NOTE: A Zero value will have no effect and use the Device Rate. **/
@@ -126,3 +110,28 @@ public:
 #endif // WITH_EDITOR
 };
 
+UCLASS()
+class UMetaSoundQuality : public UObject
+{
+	GENERATED_BODY()
+public:
+
+	/**
+	* Returns a list of quality settings to present to a combobox
+	* */
+	UFUNCTION()
+	static TArray<FName> GetQualityList()
+	{
+		TArray<FName> Names;
+
+		if (const UMetaSoundSettings* Settings = GetDefault<UMetaSoundSettings>())
+		{
+			Algo::Transform(Settings->QualitySettings, Names, [](const FMetaSoundQualitySettings& Quality) -> FName 
+			{
+				return Quality.Name;
+			});
+		}
+		return Names;
+	}
+
+};
