Add fix for PostExecuteInput types (i.e. triggers) throwing an assert when dynamically connecting a new owner to a referenced graph from attempting to rebind a read only data reference to a writable one.

diff --git a/Engine/Plugins/Runtime/Metasound/Source/MetasoundFrontend/Public/MetasoundInputNode.h b/Engine/Plugins/Runtime/Metasound/Source/MetasoundFrontend/Public/MetasoundInputNode.h
--- a/Engine/Plugins/Runtime/Metasound/Source/MetasoundFrontend/Public/MetasoundInputNode.h
+++ b/Engine/Plugins/Runtime/Metasound/Source/MetasoundFrontend/Public/MetasoundInputNode.h
@@ -176,22 +176,22 @@ namespace Metasound
 
 		public:
 			using FDataWriteReference = TDataWriteReference<DataType>;
-			using FDataWriteReferenceFactory = TDataWriteReferenceLiteralFactory<DataType>;
 
 			TPostExecutableInputOperator(const FVertexName& InDataReferenceName, TDataWriteReference<DataType> InValue)
 				: DataReferenceName(InDataReferenceName)
-				, Value(InValue)
+				, DataRef(InValue)
 			{
 			}
 
 			virtual void BindInputs(FInputVertexInterfaceData& InOutVertexData) override
 			{
-				InOutVertexData.BindWriteVertex(DataReferenceName, Value);
+				InOutVertexData.BindVertex(DataReferenceName, DataRef);
 			}
 
 			virtual void BindOutputs(FOutputVertexInterfaceData& InOutVertexData) override
 			{
-				InOutVertexData.BindReadVertex(DataReferenceName, Value);
+				TDataReadReference<DataType> DataReadRef = DataRef.GetDataReadReference<DataType>();
+				InOutVertexData.BindReadVertex(DataReferenceName, DataReadRef);
 			}
 
 			virtual FExecuteFunction GetExecuteFunction() override
@@ -210,7 +210,6 @@ namespace Metasound
 			}
 
 		protected:
-
 			static void PostExecute(IOperator* InOperator)
 			{
 				using FPostExecutableInputOperator = TPostExecutableInputOperator<DataType>;
@@ -218,11 +217,19 @@ namespace Metasound
 				FPostExecutableInputOperator* DerivedOperator = static_cast<FPostExecutableInputOperator*>(InOperator);
 				check(nullptr != DerivedOperator);
 
-				TPostExecutableDataType<DataType>::PostExecute(*(DerivedOperator->Value));
+				// This condition is checked at runtime as its possible dynamic graphs may reassign ownership
+				// of underlying data to operate on in post execute. In this case, the expectation is that the
+				// data reference is now owned by another provider/operator.
+				FAnyDataReference& DerivedDataRef = DerivedOperator->DataRef;
+				if (DerivedDataRef.GetAccessType() == EDataReferenceAccessType::Write)
+				{
+					FDataWriteReference DataWriteRef = DerivedDataRef.GetDataWriteReference<DataType>();
+					TPostExecutableDataType<DataType>::PostExecute(*DataWriteRef);
+				}
 			}
 
 			FVertexName DataReferenceName;
-			FDataWriteReference Value;
+			FAnyDataReference DataRef;
 		};
 
 		template<typename DataType>
@@ -252,7 +259,12 @@ namespace Metasound
 				FResetablePostExecutableInputOperator* Operator = static_cast<FResetablePostExecutableInputOperator*>(InOperator);
 				check(nullptr != Operator);
 
-				*Operator->Value = TDataTypeLiteralFactory<DataType>::CreateExplicitArgs(InParams.OperatorSettings, Operator->Literal);
+				// If DataRef is not writable, reference is assumed to be reset by another owning operator.
+				if (Operator->DataRef.GetAccessType() == EDataReferenceAccessType::Write)
+				{
+					FAnyDataReference& Ref = Operator->DataRef;
+					*Ref.GetDataWriteReference<DataType>() = TDataTypeLiteralFactory<DataType>::CreateExplicitArgs(InParams.OperatorSettings, Operator->Literal);
+				}
 			}
 
 			FLiteral Literal;
