[MetaSounds] Report changes to transmitter and runtime input map when live edits are enabled during Builder audition
#rnx
#tests Audio Engine Tests including ConnectInputAfterAuditionAndThenSetParameter, Audio Unit Tests

diff --git a/Engine/Plugins/Runtime/Metasound/Source/MetasoundEngine/Private/MetasoundBuilderSubsystem.cpp b/Engine/Plugins/Runtime/Metasound/Source/MetasoundEngine/Private/MetasoundBuilderSubsystem.cpp
--- a/Engine/Plugins/Runtime/Metasound/Source/MetasoundEngine/Private/MetasoundBuilderSubsystem.cpp
+++ b/Engine/Plugins/Runtime/Metasound/Source/MetasoundEngine/Private/MetasoundBuilderSubsystem.cpp
@@ -3,6 +3,7 @@
 
 #include "Algo/Find.h"
 #include "Algo/Transform.h"
+#include "AudioDevice.h"
 #include "Components/AudioComponent.h"
 #include "Engine/Engine.h"
 #include "HAL/IConsoleManager.h"
@@ -1073,9 +1074,25 @@ void UMetaSoundSourceBuilder::Audition(UObject* Parent, UAudioComponent* AudioCo
 		CreateGenerator.Execute(NewHandle);
 	}
 
+	if (bLiveUpdatesEnabled)
+	{
+		LiveComponentIDs.Add(AudioComponent->GetAudioComponentID());
+		LiveComponentHandle = AudioComponent->OnAudioFinishedNative.AddUObject(this, &UMetaSoundSourceBuilder::OnLiveComponentFinished);
+	}
+
 	AudioComponent->Play();
 }
 
+void UMetaSoundSourceBuilder::OnLiveComponentFinished(UAudioComponent* AudioComponent)
+{
+	constexpr bool bAllowShrinking = false;
+	LiveComponentIDs.RemoveSwap(AudioComponent->GetAudioComponentID(), bAllowShrinking);
+	if (LiveComponentIDs.IsEmpty())
+	{
+		AudioComponent->OnAudioFinishedNative.Remove(LiveComponentHandle);
+	}
+}
+
 bool UMetaSoundSourceBuilder::ExecuteAuditionableTransaction(FAuditionableTransaction Transaction) const
 {
 	using namespace Metasound::Engine::BuilderSubsystemPrivate;
@@ -1146,7 +1163,7 @@ UMetaSoundSource& UMetaSoundSourceBuilder::GetMetaSoundSource()
 	return Builder.CastDocumentObjectChecked<UMetaSoundSource>();
 }
 
-void UMetaSoundSourceBuilder::InitDelegates(Metasound::Frontend::FDocumentModifyDelegates& OutDocumentDelegates) const
+void UMetaSoundSourceBuilder::InitDelegates(Metasound::Frontend::FDocumentModifyDelegates& OutDocumentDelegates)
 {
 	OutDocumentDelegates.EdgeDelegates.OnEdgeAdded.AddUObject(this, &UMetaSoundSourceBuilder::OnEdgeAdded);
 	OutDocumentDelegates.EdgeDelegates.OnRemoveSwappingEdge.AddUObject(this, &UMetaSoundSourceBuilder::OnRemoveSwappingEdge);
@@ -1195,7 +1212,7 @@ TOptional<Metasound::FAnyDataReference> UMetaSoundSourceBuilder::CreateDataRefer
 	return IDataTypeRegistry::Get().CreateDataReference(DataType, AccessType, InLiteral, InOperatorSettings);
 };
 
-void UMetaSoundSourceBuilder::OnInputAdded(int32 InputIndex) const
+void UMetaSoundSourceBuilder::OnInputAdded(int32 InputIndex)
 {
 	using namespace Metasound::DynamicGraph;
 
@@ -1204,9 +1221,29 @@ void UMetaSoundSourceBuilder::OnInputAdded(int32 InputIndex) const
 		using namespace Metasound;
 		using namespace Metasound::Frontend;
 
-		const FMetasoundFrontendDocument& Doc = Builder.GetDocument();
+		const FMetaSoundFrontendDocumentBuilder& ConstBuilder = Builder;
+		const FMetasoundFrontendDocument& Doc = ConstBuilder.GetDocument();
 		const FMetasoundFrontendGraphClass& GraphClass = Doc.RootGraph;
 		const FMetasoundFrontendClassInput& NewInput = GraphClass.Interface.Inputs[InputIndex];
+
+		constexpr bool bCreateUObjectProxies = true;
+		UMetaSoundSource& Source = GetMetaSoundSource();
+		Source.RuntimeInputData.InputMap.Add(NewInput.Name, UMetaSoundSource::CreateRuntimeInput(IDataTypeRegistry::Get(), NewInput, bCreateUObjectProxies));
+
+		for (uint64 AudioComponentID : LiveComponentIDs)
+		{
+			if (UAudioComponent* AudioComponent = UAudioComponent::GetAudioComponentFromID(AudioComponentID))
+			{
+				if (FAudioDevice* AudioDevice = AudioComponent->GetAudioDevice())
+				{
+					AudioDevice->SendCommandToActiveSounds(AudioComponentID, [NewInputName = NewInput.Name](FActiveSound& ActiveSound)
+					{
+						static_cast<FMetaSoundParameterTransmitter*>(ActiveSound.GetTransmitter())->AddAvailableParameter(NewInputName);
+					});
+				}
+			}
+		}
+
 		const FLiteral NewInputLiteral = NewInput.DefaultLiteral.ToLiteral(NewInput.TypeName);
 		Transactor.AddInputDataDestination(NewInput.NodeID, NewInput.Name, NewInputLiteral, &UMetaSoundSourceBuilder::CreateDataReference);
 		return true;
@@ -1404,19 +1441,39 @@ void UMetaSoundSourceBuilder::OnRemoveSwappingEdge(int32 SwapIndex, int32 LastIn
 	});
 }
 
-void UMetaSoundSourceBuilder::OnRemovingInput(int32 InputIndex) const
+void UMetaSoundSourceBuilder::OnRemovingInput(int32 InputIndex)
 {
 	using namespace Metasound::DynamicGraph;
 
 	ExecuteAuditionableTransaction([this, InputIndex](FDynamicOperatorTransactor& Transactor)
 	{
+		using namespace Metasound;
 		using namespace Metasound::Frontend;
 
-		const FMetasoundFrontendDocument& Doc = Builder.GetDocument();
+		const FMetaSoundFrontendDocumentBuilder& ConstBuilder = Builder;
+		const FMetasoundFrontendDocument& Doc = ConstBuilder.GetDocument();
 		const FMetasoundFrontendGraphClass& GraphClass = Doc.RootGraph;
 		const FMetasoundFrontendClassInput& InputBeingRemoved = GraphClass.Interface.Inputs[InputIndex];
 
+		UMetaSoundSource& Source = GetMetaSoundSource();
+		Source.RuntimeInputData.InputMap.Remove(InputBeingRemoved.Name);
+
 		Transactor.RemoveInputDataDestination(InputBeingRemoved.Name);
+
+		for (uint64 AudioComponentID : LiveComponentIDs)
+		{
+			if (UAudioComponent* AudioComponent = UAudioComponent::GetAudioComponentFromID(AudioComponentID))
+			{
+				if (FAudioDevice* AudioDevice = AudioComponent->GetAudioDevice())
+				{
+					AudioDevice->SendCommandToActiveSounds(AudioComponentID, [InputRemoved = InputBeingRemoved.Name](FActiveSound& ActiveSound)
+						{
+							static_cast<FMetaSoundParameterTransmitter*>(ActiveSound.GetTransmitter())->RemoveAvailableParameter(InputRemoved);
+						});
+				}
+			}
+		}
+
 		return true;
 	});
 }
diff --git a/Engine/Plugins/Runtime/Metasound/Source/MetasoundEngine/Private/MetasoundSource.cpp b/Engine/Plugins/Runtime/Metasound/Source/MetasoundEngine/Private/MetasoundSource.cpp
--- a/Engine/Plugins/Runtime/Metasound/Source/MetasoundEngine/Private/MetasoundSource.cpp
+++ b/Engine/Plugins/Runtime/Metasound/Source/MetasoundEngine/Private/MetasoundSource.cpp
@@ -1568,6 +1568,26 @@ TSharedPtr<Metasound::DynamicGraph::FDynamicOperatorTransactor> UMetaSoundSource
 	return DynamicTransactor;
 }
 
+UMetaSoundSource::FRuntimeInput UMetaSoundSource::CreateRuntimeInput(const Metasound::Frontend::IDataTypeRegistry& Registry, const FMetasoundFrontendClassInput& Input, bool bCreateUObjectProxies)
+{
+	using namespace Metasound;
+	using namespace Metasound::Frontend;
+
+	bool bIsTransmittable = false;
+	if (const IDataTypeRegistryEntry* RegistryEntry = Registry.FindDataTypeRegistryEntry(Input.TypeName))
+	{
+		bIsTransmittable = RegistryEntry->GetDataTypeInfo().bIsTransmittable;
+	}
+	else
+	{
+		UE_LOG(LogMetaSound, Warning, TEXT("Failed to find data type '%s' in registry. Assuming data type is not transmittable"), *Input.TypeName.ToString());
+	}
+
+	FAudioParameter DefaultParameter = SourcePrivate::MakeAudioParameter(Registry, Input.Name, Input.TypeName, Input.DefaultLiteral, bCreateUObjectProxies);
+
+	return FRuntimeInput { Input.Name, Input.TypeName, Input.AccessType, DefaultParameter, bIsTransmittable };
+}
+
 Metasound::TSortedVertexNameMap<UMetaSoundSource::FRuntimeInput> UMetaSoundSource::CreateRuntimeInputMap(bool bCreateUObjectProxies) const
 {
 	using namespace Metasound;
@@ -1577,7 +1597,7 @@ Metasound::TSortedVertexNameMap<UMetaSoundSource::FRuntimeInput> UMetaSoundSourc
 
 	auto GetInputName = [](const FMetasoundFrontendClassInput& InInput) { return InInput.Name; };
 
-	IDataTypeRegistry& Registry = IDataTypeRegistry::Get();
+	const IDataTypeRegistry& Registry = IDataTypeRegistry::Get();
 	const FMetasoundFrontendDocument& Doc = GetConstDocument();
 
 	TArray<const IInterfaceRegistryEntry*> Interfaces;
@@ -1604,21 +1624,7 @@ Metasound::TSortedVertexNameMap<UMetaSoundSource::FRuntimeInput> UMetaSoundSourc
 	{
 		if (!PrivateInputs.Contains(Input.Name))
 		{
-			bool bIsTransmittable = false;
-			if (const IDataTypeRegistryEntry* RegistryEntry = Registry.FindDataTypeRegistryEntry(Input.TypeName))
-			{
-				bIsTransmittable = RegistryEntry->GetDataTypeInfo().bIsTransmittable;	
-			}
-			else
-			{
-				UE_LOG(LogMetaSound, Warning, TEXT("Failed to find data type '%s' in registry. Assuming data type is not transmittable"), *Input.TypeName.ToString());
-
-			}
-
-			
-			FAudioParameter DefaultParameter = SourcePrivate::MakeAudioParameter(Registry, Input.Name, Input.TypeName, Input.DefaultLiteral, bCreateUObjectProxies) ;
-
-			PublicInputs.Add(Input.Name, FRuntimeInput{Input.Name, Input.TypeName, Input.AccessType, DefaultParameter, bIsTransmittable});
+			PublicInputs.Add(Input.Name, CreateRuntimeInput(Registry, Input, bCreateUObjectProxies));
 		}
 	}
 
diff --git a/Engine/Plugins/Runtime/Metasound/Source/MetasoundEngine/Public/MetasoundBuilderSubsystem.h b/Engine/Plugins/Runtime/Metasound/Source/MetasoundEngine/Public/MetasoundBuilderSubsystem.h
--- a/Engine/Plugins/Runtime/Metasound/Source/MetasoundEngine/Public/MetasoundBuilderSubsystem.h
+++ b/Engine/Plugins/Runtime/Metasound/Source/MetasoundEngine/Public/MetasoundBuilderSubsystem.h
@@ -568,15 +568,16 @@ private:
 	const UMetaSoundSource& GetMetaSoundSource() const;
 	UMetaSoundSource& GetMetaSoundSource();
 
-	void InitDelegates(Metasound::Frontend::FDocumentModifyDelegates& OutDocumentDelegates) const;
+	void InitDelegates(Metasound::Frontend::FDocumentModifyDelegates& OutDocumentDelegates);
 
 	void OnEdgeAdded(int32 EdgeIndex) const;
-	void OnInputAdded(int32 InputIndex) const;
+	void OnInputAdded(int32 InputIndex);
+	void OnLiveComponentFinished(UAudioComponent* AudioComponent);
 	void OnNodeAdded(int32 NodeIndex) const;
 	void OnNodeInputLiteralSet(int32 NodeIndex, int32 VertexIndex, int32 LiteralIndex) const;
 	void OnOutputAdded(int32 OutputIndex) const;
 	void OnRemoveSwappingEdge(int32 SwapIndex, int32 LastIndex) const;
-	void OnRemovingInput(int32 InputIndex) const;
+	void OnRemovingInput(int32 InputIndex);
 	void OnRemoveSwappingNode(int32 SwapIndex, int32 LastIndex) const;
 	void OnRemovingNodeInputLiteral(int32 NodeIndex, int32 VertexIndex, int32 LiteralIndex) const;
 	void OnRemovingOutput(int32 OutputIndex) const;
@@ -584,6 +585,9 @@ private:
 	using FAuditionableTransaction = TFunctionRef<bool(Metasound::DynamicGraph::FDynamicOperatorTransactor&)>;
 	bool ExecuteAuditionableTransaction(FAuditionableTransaction Transaction) const;
 
+	TArray<uint64> LiveComponentIDs;
+	FDelegateHandle LiveComponentHandle;
+
 	friend class UMetaSoundBuilderSubsystem;
 };
 
diff --git a/Engine/Plugins/Runtime/Metasound/Source/MetasoundEngine/Public/MetasoundSource.h b/Engine/Plugins/Runtime/Metasound/Source/MetasoundEngine/Public/MetasoundSource.h
--- a/Engine/Plugins/Runtime/Metasound/Source/MetasoundEngine/Public/MetasoundSource.h
+++ b/Engine/Plugins/Runtime/Metasound/Source/MetasoundEngine/Public/MetasoundSource.h
@@ -23,32 +23,36 @@
 
 
 // Forward Declarations
+class UMetaSoundSettings;
+
+namespace Audio
+{
+	using DeviceID = uint32;
+} // namespace Audio
+
 namespace Metasound
 {
 	struct FMetaSoundEngineAssetHelper;
+	struct FMetasoundGeneratorInitParams;
+
 	class FMetasoundGenerator;
+
+	namespace Frontend
+	{
+		class IDataTypeRegistry;
+	} // namespace Frontend
+
 	namespace SourcePrivate
 	{
 		class FParameterRouter;
-	}
+	} // namespace SourcePrivate
 
 	namespace DynamicGraph
 	{
 		class FDynamicOperatorTransactor;
-	}
+	} // namespace DynamicGraph
 } // namespace Metasound
 
-namespace Audio
-{
-	using DeviceID = uint32;
-}
-
-class UMetaSoundSettings;
-
-namespace Metasound
-{
-	struct FMetasoundGeneratorInitParams;
-} // namespace Metasound
 
 DECLARE_TS_MULTICAST_DELEGATE_TwoParams(FOnGeneratorInstanceCreated, uint64, TSharedPtr<Metasound::FMetasoundGenerator>);
 DECLARE_TS_MULTICAST_DELEGATE_TwoParams(FOnGeneratorInstanceDestroyed, uint64, TSharedPtr<Metasound::FMetasoundGenerator>);
@@ -291,8 +295,8 @@ private:
 public:
 	Metasound::FMetasoundEnvironment CreateEnvironment(const FSoundGeneratorInitParams& InParams) const;
 	const TArray<Metasound::FVertexName>& GetOutputAudioChannelOrder() const;
-private:
 
+private:
 	TSharedPtr<const Metasound::IGraph> TryGetMetaSoundPresetBaseGraph() const;
 	void MergePresetOverridesAndSuppliedDefaults(const TArray<FAudioParameter>& InSuppliedDefaults, TArray<FAudioParameter>& OutMerged);
 	
@@ -304,6 +308,7 @@ private:
 	void TrackGenerator(uint64 Id, TSharedPtr<Metasound::FMetasoundGenerator> Generator);
 	void ForgetGenerator(ISoundGeneratorPtr Generator);
 
+	static FRuntimeInput CreateRuntimeInput(const Metasound::Frontend::IDataTypeRegistry& Registry, const FMetasoundFrontendClassInput& Input, bool bCreateUObjectProxies);
 	Metasound::TSortedVertexNameMap<FRuntimeInput> CreateRuntimeInputMap(bool bCreateUObjectProxies) const;
 	void CacheRuntimeInputData();
 	void InvalidateCachedRuntimeInputData();
diff --git a/Engine/Plugins/Runtime/Metasound/Source/MetasoundFrontend/Private/MetasoundParameterTransmitter.cpp b/Engine/Plugins/Runtime/Metasound/Source/MetasoundFrontend/Private/MetasoundParameterTransmitter.cpp
--- a/Engine/Plugins/Runtime/Metasound/Source/MetasoundFrontend/Private/MetasoundParameterTransmitter.cpp
+++ b/Engine/Plugins/Runtime/Metasound/Source/MetasoundFrontend/Private/MetasoundParameterTransmitter.cpp
@@ -258,6 +258,16 @@ namespace Metasound
 		return true;
 	}
 
+	void FMetaSoundParameterTransmitter::AddAvailableParameter(FName InName)
+	{
+		AvailableParameterNames.Add(InName);
+	}
+
+	void FMetaSoundParameterTransmitter::RemoveAvailableParameter(FName InName)
+	{
+		AvailableParameterNames.Remove(InName);
+	}
+
 	bool FMetaSoundParameterTransmitter::SetParameters(TArray<FAudioParameter>&& InParameters)
 	{
 		bool bSuccess = true;
diff --git a/Engine/Plugins/Runtime/Metasound/Source/MetasoundFrontend/Public/MetasoundParameterTransmitter.h b/Engine/Plugins/Runtime/Metasound/Source/MetasoundFrontend/Public/MetasoundParameterTransmitter.h
--- a/Engine/Plugins/Runtime/Metasound/Source/MetasoundFrontend/Public/MetasoundParameterTransmitter.h
+++ b/Engine/Plugins/Runtime/Metasound/Source/MetasoundFrontend/Public/MetasoundParameterTransmitter.h
@@ -123,6 +123,9 @@ namespace Metasound
 		UE_DEPRECATED(5.2, "Use ResetParameters() or OnDeleteActiveSound() instead depending on use case.")
 		bool Reset() override;
 		
+		void AddAvailableParameter(FName InName);
+		void RemoveAvailableParameter(FName InName);
+
 		/** Sets parameters using array of AudioParameter structs
 		 *
 		 * @param InParameter - Parameter to set.
