Retrieve MetaSound graph before building graph operator in async task (cvarable) and add nullptr checks

#localization none
#tests preflight 656ceb1d8b7aa8578b8b235a

diff --git a/Engine/Plugins/Runtime/Metasound/Source/MetasoundGenerator/Private/MetasoundGeneratorBuilder.cpp b/Engine/Plugins/Runtime/Metasound/Source/MetasoundGenerator/Private/MetasoundGeneratorBuilder.cpp
--- a/Engine/Plugins/Runtime/Metasound/Source/MetasoundGenerator/Private/MetasoundGeneratorBuilder.cpp
+++ b/Engine/Plugins/Runtime/Metasound/Source/MetasoundGenerator/Private/MetasoundGeneratorBuilder.cpp
@@ -194,6 +194,12 @@ namespace Metasound
 				}
 			};
 
+			if (!InInitParams.Graph)
+			{
+				UE_LOG(LogMetaSound, Error, TEXT("Unable to build graph operator inputs for null graph in MetaSoundSource [%s]"), *InInitParams.MetaSoundName);
+				return FInputVertexInterfaceData();
+			}
+
 			const FInputVertexInterface& InputInterface = InInitParams.Graph->GetVertexInterface().GetInputInterface();
 			FInputVertexInterfaceData InputData(InputInterface);
 
@@ -254,6 +260,11 @@ namespace Metasound
 			FOperatorAndInputs OpAndInputs;
 			OpAndInputs.Inputs = BuildGraphOperatorInputs(InOperatorSettings, InInitParams);
 
+			if (!InInitParams.Graph)
+			{
+				UE_LOG(LogMetaSound, Error, TEXT("Unable to build graph operator for null graph in MetaSoundSource [%s]"), *InInitParams.MetaSoundName);
+				return OpAndInputs;
+			}
 			// Create an instance of the new graph operator
 			FBuildGraphOperatorParams BuildParams { *InInitParams.Graph, InOperatorSettings, OpAndInputs.Inputs, InInitParams.Environment };
 			FOperatorBuilder Builder(InInitParams.BuilderSettings);
diff --git a/Engine/Plugins/Runtime/Metasound/Source/MetasoundGenerator/Private/MetasoundOperatorCache.cpp b/Engine/Plugins/Runtime/Metasound/Source/MetasoundGenerator/Private/MetasoundOperatorCache.cpp
--- a/Engine/Plugins/Runtime/Metasound/Source/MetasoundGenerator/Private/MetasoundOperatorCache.cpp
+++ b/Engine/Plugins/Runtime/Metasound/Source/MetasoundGenerator/Private/MetasoundOperatorCache.cpp
@@ -23,12 +23,21 @@ namespace Metasound
 	TRACE_DECLARE_INT_COUNTER(MetaSound_OperatorPool_NumOperators, TEXT("MetaSound/OperatorPool/NumOperatorsInPool"));
 	TRACE_DECLARE_FLOAT_COUNTER(MetaSound_OperatorPool_HitRatio, TEXT("MetaSound/OperatorPool/HitRatio"));
 	TRACE_DECLARE_FLOAT_COUNTER(MetaSound_OperatorPool_WindowedHitRatio, TEXT("MetaSound/OperatorPool/WindowedHitRatio"));
+#endif // #if METASOUND_OPERATORCACHEPROFILER_ENABLED
 
 	namespace OperatorPoolPrivate
 	{
+		static bool bMetasoundPoolSyncGraphRetrieval = true;
+		FAutoConsoleVariableRef CVarMetasoundPoolSyncGraphRetrieval(
+			TEXT("au.MetaSound.OperatorPoolSyncGraphRetrieval"),
+			bMetasoundPoolSyncGraphRetrieval,
+			TEXT("Retrieves graph on the requesting thread prior to asynchronous task to create instance.\n"),
+			ECVF_Default);
+
+#if METASOUND_OPERATORCACHEPROFILER_ENABLED
 		static std::atomic<uint32> CacheHitCount = 0;
 		static std::atomic<uint32> CacheAttemptCount = 0;
-	 
+
 		static float MetasoundPoolHitRateWindowSecondsCVar = 1.0f;
 		FAutoConsoleVariableRef CVarMetasoundPoolHitRateWindowSeconds(
 			TEXT("au.MetaSound.OperatorPoolHitRateWindowSeconds"),
@@ -147,8 +156,8 @@ namespace Metasound
 				}
 			}
 		}
-	} // namespace OperatorPoolPrivate
 #endif // #if METASOUND_OPERATORCACHEPROFILER_ENABLED
+	} // namespace OperatorPoolPrivate
 
 
 	FOperatorBuildData::FOperatorBuildData(
@@ -236,14 +245,30 @@ namespace Metasound
 
 	void FOperatorPool::BuildAndAddOperator(TUniquePtr<FOperatorBuildData> InBuildData)
 	{
+		using namespace OperatorPoolPrivate;
+
 		if (!ensure(InBuildData))
 		{
 			return;
 		}
 
+		TSharedPtr<const Metasound::FGraph> Graph;
+		if (bMetasoundPoolSyncGraphRetrieval)
+		{
+			// get the metasound graph and add to init params (might wait for async registration to complete)
+			Graph = FMetasoundFrontendRegistryContainer::Get()->GetGraph(InBuildData->RegistryKey, InBuildData->AssetPath);
+			if (!Graph.IsValid())
+			{
+				UE_LOG(LogMetasoundGenerator, Error, TEXT("Failed to retrieve graph '%s' synchronously when attempting to BuildAndAddOperator to pool"), *InBuildData->AssetPath.ToString());
+				return;
+			}
+		}
+
 		FMetasoundGeneratorModule& Module = FModuleManager::GetModuleChecked<FMetasoundGeneratorModule>("MetasoundGenerator");
-		AsyncTask(ENamedThreads::AnyThread, [PreCacheData = MoveTemp(InBuildData), OperatorPool = Module.GetOperatorPool()] ()
+		AsyncTask(ENamedThreads::AnyThread, [Graph = Graph, PreCacheData = MoveTemp(InBuildData), OperatorPool = Module.GetOperatorPool()] ()
 		{
+			using namespace OperatorPoolPrivate;
+
 			METASOUND_LLM_SCOPE;
 			METASOUND_TRACE_CPUPROFILER_EVENT_SCOPE(Metasound::FOperatorPool::AsyncOperatorPrecache)
 
@@ -252,10 +277,29 @@ namespace Metasound
 				return;
 			}
 
-			// get the metasound graph and add to init params (might wait for async registration to complete)
-			PreCacheData->InitParams.Graph = FMetasoundFrontendRegistryContainer::Get()->GetGraph(PreCacheData->RegistryKey, PreCacheData->AssetPath);
-			ensure(PreCacheData->InitParams.Graph);
-	
+			if (Graph.IsValid())
+			{
+				PreCacheData->InitParams.Graph = Graph;
+			}
+			else
+			{
+				if (bMetasoundPoolSyncGraphRetrieval)
+				{
+					return;
+				}
+				else
+				{
+					// get the metasound graph and add to init params (might wait for async registration to complete)
+					PreCacheData->InitParams.Graph = FMetasoundFrontendRegistryContainer::Get()->GetGraph(PreCacheData->RegistryKey, PreCacheData->AssetPath);
+				}
+			}
+
+			if (!PreCacheData->InitParams.Graph)
+			{
+				UE_LOG(LogMetasoundGenerator, Error, TEXT("Failed to retrieve graph '%s' async when attempting to BuildAndAddOperator to pool"), *PreCacheData->AssetPath.ToString());
+				return;
+			}
+
 			const int32 NumInstances = PreCacheData->NumInstances;
 			for (int32 i = 0; i < NumInstances; ++i)
 			{
