Add cvar to disable manual MetaSoundPrecache all together
#rnx

diff --git a/Engine/Plugins/Runtime/Metasound/Source/MetasoundEngine/Private/MetasoundOperatorCacheSubsystem.cpp b/Engine/Plugins/Runtime/Metasound/Source/MetasoundEngine/Private/MetasoundOperatorCacheSubsystem.cpp
--- a/Engine/Plugins/Runtime/Metasound/Source/MetasoundEngine/Private/MetasoundOperatorCacheSubsystem.cpp
+++ b/Engine/Plugins/Runtime/Metasound/Source/MetasoundEngine/Private/MetasoundOperatorCacheSubsystem.cpp
@@ -86,6 +86,11 @@ void UMetaSoundCacheSubsystem::PrecacheMetaSound(UMetaSoundSource* InMetaSound,
 	using namespace Audio;
 	using namespace Metasound;
 
+	if (!FMetasoundGenerator::GetManualPrecacheEnabled())
+	{
+		return;
+	}
+
 	TSharedPtr<FOperatorPool> OperatorPool;
 	IMetasoundGeneratorModule* Module = FModuleManager::GetModulePtr<IMetasoundGeneratorModule>("MetasoundGenerator");
 	if (!ensure(Module))
diff --git a/Engine/Plugins/Runtime/Metasound/Source/MetasoundGenerator/Private/MetasoundGenerator.cpp b/Engine/Plugins/Runtime/Metasound/Source/MetasoundGenerator/Private/MetasoundGenerator.cpp
--- a/Engine/Plugins/Runtime/Metasound/Source/MetasoundGenerator/Private/MetasoundGenerator.cpp
+++ b/Engine/Plugins/Runtime/Metasound/Source/MetasoundGenerator/Private/MetasoundGenerator.cpp
@@ -33,6 +33,7 @@ namespace Metasound
 		static bool bEnableAsyncMetaSoundGeneratorBuilder = true;
 		static bool bEnableExperimentalOneShotOperatorPool = false;
 		static bool bEnableExperimentalOperatorPool = false;
+		static bool bEnableExperimentalOperatorPoolManualPrecache = true;
 		static bool bEnableResetOnOperatorPoolInsertion = true;
 #if ENABLE_METASOUNDGENERATOR_INVALID_SAMPLE_VALUE_LOGGING
 		static bool bEnableMetaSoundGeneratorNonFiniteLogging = false;
@@ -211,6 +212,13 @@ FAutoConsoleVariableRef CVarMetaSoundEnableExperimentalOperatorPool(
 	TEXT("Default: false"),
 	ECVF_Default);
 
+FAutoConsoleVariableRef CVarMetaSoundEnableExperimentalOperatorPoolManualPrecache(
+	TEXT("au.MetaSound.Experimental.EnableOperatorPoolManualPrecache"),
+	Metasound::ConsoleVariables::bEnableExperimentalOperatorPoolManualPrecache,
+	TEXT("Enables manual pre-caching of explicit MetaSound asset operators.\n")
+	TEXT("Default: true"),
+	ECVF_Default);
+
 FAutoConsoleVariableRef CVarMetaSoundEnableResetOnOperatorPoolInsertion(
 	TEXT("au.MetaSound.Experimental.EnableResetOnOperatorPoolInsertion"),
 	Metasound::ConsoleVariables::bEnableResetOnOperatorPoolInsertion,
@@ -682,6 +690,11 @@ PRAGMA_ENABLE_DEPRECATION_WARNINGS
 		return RenderTime;
 	}
 
+	bool FMetasoundGenerator::GetManualPrecacheEnabled()
+	{
+		return ConsoleVariables::bEnableExperimentalOperatorPoolManualPrecache;
+	}
+
 	int32 FMetasoundGenerator::FillWithBuffer(const Audio::FAlignedFloatBuffer& InBuffer, float* OutAudio, int32 MaxNumOutputSamples)
 	{
 		int32 InNum = InBuffer.Num();
diff --git a/Engine/Plugins/Runtime/Metasound/Source/MetasoundGenerator/Public/MetasoundGenerator.h b/Engine/Plugins/Runtime/Metasound/Source/MetasoundGenerator/Public/MetasoundGenerator.h
--- a/Engine/Plugins/Runtime/Metasound/Source/MetasoundGenerator/Public/MetasoundGenerator.h
+++ b/Engine/Plugins/Runtime/Metasound/Source/MetasoundGenerator/Public/MetasoundGenerator.h
@@ -226,6 +226,9 @@ namespace Metasound
 		/** Fraction of a single CPU core used to render audio on a scale of 0.0 to 1.0 */
 		double GetCPUCoreUtilization() const;
 
+		/** Returns whether or not operator pool manual precache feature is enabled (true) or not (false). */
+		static bool GetManualPrecacheEnabled();
+
 		// Called when a new graph has been "compiled" and set up as this generator's graph.
 		// Note: We don't allow direct assignment to the FOnSetGraph delegate
 		// because we want to give the Delegate an initial immediate callback if the generator 
