[Metasound] Notify listeners only once per render block when a dynamic generator's vertex interface changes.

diff --git a/Engine/Plugins/Runtime/Metasound/Source/MetasoundGenerator/Private/MetasoundGenerator.cpp b/Engine/Plugins/Runtime/Metasound/Source/MetasoundGenerator/Private/MetasoundGenerator.cpp
--- a/Engine/Plugins/Runtime/Metasound/Source/MetasoundGenerator/Private/MetasoundGenerator.cpp
+++ b/Engine/Plugins/Runtime/Metasound/Source/MetasoundGenerator/Private/MetasoundGenerator.cpp
@@ -448,7 +448,7 @@ PRAGMA_ENABLE_DEPRECATION_WARNINGS
 
 		bIsWaitingForFirstGraph = false;
 		OnSetGraph.Broadcast();
-		OnVertexInterfaceDataUpdated.Broadcast(VertexInterfaceData);
+		bVertexInterfaceHasChanged.store(true);
 	}
 
 	TUniquePtr<IOperator> FMetasoundGenerator::ReleaseGraphOperator()
@@ -471,7 +471,7 @@ PRAGMA_ENABLE_DEPRECATION_WARNINGS
 		ParameterSetters.Reset();
 		ParameterPackSetters.Reset();
 
-		OnVertexInterfaceDataUpdated.Broadcast(VertexInterfaceData);
+		bVertexInterfaceHasChanged.store(true);
 	}
 
 	void FMetasoundGenerator::QueueParameterPack(TSharedPtr<FMetasoundParameterPackStorage> ParameterPack)
@@ -672,6 +672,15 @@ PRAGMA_ENABLE_DEPRECATION_WARNINGS
 			}
 		}
 
+		// If the vertex interface changed, notify listeners
+		{
+			bool bExpected = true;
+			if (bVertexInterfaceHasChanged.compare_exchange_weak(bExpected, false))
+			{
+				OnVertexInterfaceDataUpdated.Broadcast(VertexInterfaceData);
+			}
+		}
+
 		return NumSamplesWritten;
 	}
 
@@ -1006,7 +1015,7 @@ PRAGMA_ENABLE_DEPRECATION_WARNINGS
 		VertexInterfaceData.GetInputs() = InGraphInputs;
 		GeneratorBuilder::AddParameterSetterIfWritable(InVertexName, VertexInterfaceData.GetInputs(), ParameterSetters, ParameterPackSetters);
 
-		OnVertexInterfaceDataUpdated.Broadcast(VertexInterfaceData);
+		bVertexInterfaceHasChanged.store(true);
 	}
 
 	void FMetasoundDynamicGraphGenerator::OnInputRemoved(const FVertexName& InVertexName, const FInputVertexInterfaceData& InGraphInputs)
@@ -1014,7 +1023,7 @@ PRAGMA_ENABLE_DEPRECATION_WARNINGS
 		VertexInterfaceData.GetInputs() = InGraphInputs;
 		ParameterSetters.Remove(InVertexName);
 		ParameterPackSetters.Remove(InVertexName);
-		OnVertexInterfaceDataUpdated.Broadcast(VertexInterfaceData);
+		bVertexInterfaceHasChanged.store(true);
 	}
 
 	void FMetasoundDynamicGraphGenerator::OnOutputAdded(const FVertexName& InVertexName, const FOutputVertexInterfaceData& InGraphOutputs)
@@ -1026,7 +1035,7 @@ PRAGMA_ENABLE_DEPRECATION_WARNINGS
 		{
 			OnFinishedTriggerRef = InGraphOutputs.GetDataReadReference<FTrigger>(InVertexName);
 		}
-		OnVertexInterfaceDataUpdated.Broadcast(VertexInterfaceData);
+		bVertexInterfaceHasChanged.store(true);
 	}
 
 	void FMetasoundDynamicGraphGenerator::OnOutputUpdated(const FVertexName& InVertexName, const FOutputVertexInterfaceData& InGraphOutputs)
@@ -1044,7 +1053,8 @@ PRAGMA_ENABLE_DEPRECATION_WARNINGS
 		{
 			OnFinishedTriggerRef = InGraphOutputs.GetDataReadReference<FTrigger>(InVertexName);
 		}
-		OnVertexInterfaceDataUpdated.Broadcast(VertexInterfaceData);
+		
+		bVertexInterfaceHasChanged.store(true);
 	}
 
 	void FMetasoundDynamicGraphGenerator::OnOutputRemoved(const FVertexName& InVertexName, const FOutputVertexInterfaceData& InGraphOutputs)
@@ -1052,7 +1062,7 @@ PRAGMA_ENABLE_DEPRECATION_WARNINGS
 		MetasoundGeneratorPrivate::LogErrorIfIsAudioVertex(InVertexName, TEXT("The MetaSound Generator cannot dynamically remove audio outputs"));
 
 		VertexInterfaceData.GetOutputs() = InGraphOutputs;
-		OnVertexInterfaceDataUpdated.Broadcast(VertexInterfaceData);
+		bVertexInterfaceHasChanged.store(true);
 	}
 
 	TUniquePtr<IOperator> FMetasoundDynamicGraphGenerator::ReleaseGraphOperator()
diff --git a/Engine/Plugins/Runtime/Metasound/Source/MetasoundGenerator/Public/MetasoundGenerator.h b/Engine/Plugins/Runtime/Metasound/Source/MetasoundGenerator/Public/MetasoundGenerator.h
--- a/Engine/Plugins/Runtime/Metasound/Source/MetasoundGenerator/Public/MetasoundGenerator.h
+++ b/Engine/Plugins/Runtime/Metasound/Source/MetasoundGenerator/Public/MetasoundGenerator.h
@@ -257,6 +257,8 @@ namespace Metasound
 		void ClearGraph();
 		bool UpdateGraphIfPending();
 
+		std::atomic<bool> bVertexInterfaceHasChanged{ false };
+
 	private:
 
 		friend class FAsyncMetaSoundBuilderBase;
