[Readd] Added better naming to new quality settings.

diff --git a/Engine/Plugins/Runtime/Metasound/Source/MetasoundEngine/Private/MetasoundSettings.cpp b/Engine/Plugins/Runtime/Metasound/Source/MetasoundEngine/Private/MetasoundSettings.cpp
new file mode 100644
--- /dev/null
+++ b/Engine/Plugins/Runtime/Metasound/Source/MetasoundEngine/Private/MetasoundSettings.cpp
@@ -0,0 +1,127 @@
+// Copyright Epic Games, Inc. All Rights Reserved.
+
+#include "MetasoundSettings.h"
+#include "Algo/Count.h"
+
+#if WITH_EDITORONLY_DATA
+
+namespace MetaSoundSettingsPrivate
+{
+	/** Generate new name for the item. **/
+	static FName GenerateNextName(const TCHAR* InBaseName)
+	{
+		const TArray<FName> Names = UMetaSoundQualityHelper::GetQualityList();	
+		FString NewName = InBaseName;
+		
+		// Start adding a postfix starting at 1.
+		for (int32 Postfix=1; Names.Contains(*NewName); ++Postfix)
+		{
+			NewName = FString::Format(TEXT("{0} {1}"), { InBaseName, Postfix });
+		}
+		return FName(*NewName);
+	}
+	
+	/** When adding a quality new entry in the editor call this **/	
+	static void OnCreateNewQualitySettings(FMetaSoundQualitySettings& InNewItem)
+	{
+		InNewItem.Name = GenerateNextName(TEXT("New Quality"));
+		InNewItem.UniqueId = FGuid::NewGuid();
+	}
+
+	static void OnDuplicateQualitySettings(FMetaSoundQualitySettings& InNewlyDuplicatedItem)
+	{
+		InNewlyDuplicatedItem.Name = GenerateNextName(*FString::Printf(TEXT("%s (Duplicate)"), *InNewlyDuplicatedItem.Name.ToString()));
+		InNewlyDuplicatedItem.UniqueId = FGuid::NewGuid();
+	}
+
+	static void OnPasteQualitySettings(FMetaSoundQualitySettings& InNewlyPastedItem)
+	{
+		InNewlyPastedItem.Name = GenerateNextName(*FString::Printf(TEXT("%s (Copied)"), *InNewlyPastedItem.Name.ToString()));
+		InNewlyPastedItem.UniqueId = FGuid::NewGuid();
+	}
+	static void OnRenameQualitySettings(FMetaSoundQualitySettings& InRenamed)
+	{
+		// More than one?
+		if (Algo::Count(UMetaSoundQualityHelper::GetQualityList(), InRenamed.Name) > 1)
+		{
+			// add something unique to the name.
+			InRenamed.Name = GenerateNextName(*InRenamed.Name.ToString());
+		}	
+	}
+	
+}//MetaSoundSettingsPrivate
+
+void UMetaSoundSettings::PostEditChangeChainProperty(FPropertyChangedChainEvent& PostEditChangeChainProperty)
+{
+	const int32 QualityItemIndex = PostEditChangeChainProperty.GetArrayIndex(GET_MEMBER_NAME_STRING_CHECKED(UMetaSoundSettings, QualitySettings));
+
+	// Quality item changed..
+	if (QualityItemIndex != INDEX_NONE && QualitySettings.IsValidIndex(QualityItemIndex))
+	{
+		FMetaSoundQualitySettings& Item = QualitySettings[QualityItemIndex];
+
+		// Name has changed.
+		if (PostEditChangeChainProperty.GetPropertyName() == GET_MEMBER_NAME_CHECKED(FMetaSoundQualitySettings, Name))
+		{
+			MetaSoundSettingsPrivate::OnRenameQualitySettings(Item);
+		}
+		
+		// Array change.
+		else if (PostEditChangeChainProperty.GetPropertyName() == GET_MEMBER_NAME_CHECKED(UMetaSoundSettings, QualitySettings))
+		{
+			// Add
+			if (PostEditChangeChainProperty.ChangeType == EPropertyChangeType::ArrayAdd)
+			{
+				MetaSoundSettingsPrivate::OnCreateNewQualitySettings(Item);
+			}
+			// Duplicate.
+			else if (PostEditChangeChainProperty.ChangeType == EPropertyChangeType::Duplicate)
+			{
+				MetaSoundSettingsPrivate::OnDuplicateQualitySettings(Item);
+			}
+		}	
+	}
+	
+	// Handle pasting separately as we might not have a valid index in the case of pasting when quality array is empty.
+	if (PostEditChangeChainProperty.GetPropertyName() == GET_MEMBER_NAME_CHECKED(UMetaSoundSettings, QualitySettings))
+	{
+		// Paste...
+		if (PostEditChangeChainProperty.ChangeType == EPropertyChangeType::ValueSet)
+		{
+			const int32 IndexOfPastedItem = QualityItemIndex != INDEX_NONE ? QualityItemIndex : 0;
+			if (QualitySettings.IsValidIndex(IndexOfPastedItem))
+			{
+				FMetaSoundQualitySettings& PastedItem = QualitySettings[IndexOfPastedItem];
+				MetaSoundSettingsPrivate::OnPasteQualitySettings(PastedItem);
+			}
+		}
+	}
+
+	Super::PostEditChangeChainProperty(PostEditChangeChainProperty);
+}
+
+void UMetaSoundSettings::PostEditChangeProperty(FPropertyChangedEvent& PropertyChangedEvent)
+{	
+	Super::PostEditChangeProperty(PropertyChangedEvent);
+	
+	DenyListCacheChangeID++;
+}
+
+#endif //WITH_EDITORONLY_DATA
+
+TArray<FName> UMetaSoundQualityHelper::GetQualityList()
+{
+	TArray<FName> Names;
+
+#if WITH_EDITORONLY_DATA
+	if (const UMetaSoundSettings* Settings = GetDefault<UMetaSoundSettings>())
+	{
+		Algo::Transform(Settings->QualitySettings, Names, [](const FMetaSoundQualitySettings& Quality) -> FName
+		{
+			return Quality.Name;
+		});
+	}
+#endif //WITH_EDITORONLY_DATA
+
+	return Names;
+}
diff --git a/Engine/Plugins/Runtime/Metasound/Source/MetasoundEngine/Public/MetasoundSettings.h b/Engine/Plugins/Runtime/Metasound/Source/MetasoundEngine/Public/MetasoundSettings.h
--- a/Engine/Plugins/Runtime/Metasound/Source/MetasoundEngine/Public/MetasoundSettings.h
+++ b/Engine/Plugins/Runtime/Metasound/Source/MetasoundEngine/Public/MetasoundSettings.h
@@ -14,7 +14,7 @@
 
 // Forward Declarations
 struct FMetasoundFrontendClassName;
-
+struct FPropertyChangedChainEvent;
 
 UENUM()
 enum class EMetaSoundMessageLevel : uint8
@@ -34,23 +34,36 @@ struct METASOUNDENGINE_API FDefaultMetaSoundAssetAutoUpdateSettings
 	FSoftObjectPath MetaSound;
 };
 
+UCLASS(Hidden)
+class METASOUNDENGINE_API UMetaSoundQualityHelper : public UObject
+{
+	GENERATED_BODY()
+
+public:
+	/**
+	* Returns a list of quality settings to present to a combobox
+	* */
+	UFUNCTION()
+	static TArray<FName> GetQualityList();
+};
+
 USTRUCT()
 struct METASOUNDENGINE_API FMetaSoundQualitySettings
 {
 	GENERATED_BODY()
-
+	
 #if WITH_EDITORONLY_DATA
 
 	/** A hidden GUID that will be generated once when adding a new entry. This prevents orphaning of renamed entries. **/
-	UPROPERTY(meta = (IgnoreForMemberInitializationTest))
-	FGuid UniqueId = FGuid::NewGuid();
+	UPROPERTY()
+	FGuid UniqueId = {};
 
 #endif //WITH_EDITORONLY_DATA
 
 	/** Name of this quality setting. This will appear in the quality dropdown list.
 		The names should be unique and adequately describe the Entry. "High", "Low" etc. **/
 	UPROPERTY(EditAnywhere, Category = "Quality")
-	FName Name;
+	FName Name = {};
 
 	/** Sample Rate (in Hz). NOTE: A Zero value will have no effect and use the Device Rate. **/
 	UPROPERTY(EditAnywhere, Category = "Quality", meta = (ClampMin = "0", ClampMax="96000"))
@@ -101,37 +114,11 @@ public:
 	int32 DenyListCacheChangeID = 0;	
 
 #if WITH_EDITOR
-	virtual void PostEditChangeProperty(FPropertyChangedEvent& PropertyChangedEvent) override
-	{		
-		Super::PostEditChangeProperty(PropertyChangedEvent);
+private:
+	
+	virtual void PostEditChangeChainProperty(FPropertyChangedChainEvent& PropertyChangedEvent) override;
 
-		DenyListCacheChangeID++;
-	}
+	virtual void PostEditChangeProperty(FPropertyChangedEvent& PropertyChangedEvent) override;
 #endif // WITH_EDITOR
 };
 
-UCLASS()
-class UMetaSoundQuality : public UObject
-{
-	GENERATED_BODY()
-public:
-
-	/**
-	* Returns a list of quality settings to present to a combobox
-	* */
-	UFUNCTION()
-	static TArray<FName> GetQualityList()
-	{
-		TArray<FName> Names;
-
-		if (const UMetaSoundSettings* Settings = GetDefault<UMetaSoundSettings>())
-		{
-			Algo::Transform(Settings->QualitySettings, Names, [](const FMetaSoundQualitySettings& Quality) -> FName 
-			{
-				return Quality.Name;
-			});
-		}
-		return Names;
-	}
-
-};
diff --git a/Engine/Plugins/Runtime/Metasound/Source/MetasoundEngine/Public/MetasoundSource.h b/Engine/Plugins/Runtime/Metasound/Source/MetasoundEngine/Public/MetasoundSource.h
--- a/Engine/Plugins/Runtime/Metasound/Source/MetasoundEngine/Public/MetasoundSource.h
+++ b/Engine/Plugins/Runtime/Metasound/Source/MetasoundEngine/Public/MetasoundSource.h
@@ -115,7 +115,7 @@ public:
 	EMetaSoundOutputAudioFormat OutputFormat;
 
 	// The Quality this Metasound will use. These are defined in the MetaSounds project settings.
-	UPROPERTY(config, EditAnywhere, BlueprintReadWrite, meta = (GetOptions="MetasoundEngine.MetaSoundQuality.GetQualityList"), Category = "Metasound")
+	UPROPERTY(config, EditAnywhere, BlueprintReadWrite, meta = (GetOptions="MetasoundEngine.MetaSoundQualityHelper.GetQualityList"), Category = "Metasound")
 	FName QualitySetting;
 
 #if WITH_EDITORONLY_DATA
