Delete AudioCodecEngineModule

#tests PIE

diff --git a/Engine/Plugins/Runtime/Metasound/Source/MetasoundEngine/MetasoundEngine.Build.cs b/Engine/Plugins/Runtime/Metasound/Source/MetasoundEngine/MetasoundEngine.Build.cs
--- a/Engine/Plugins/Runtime/Metasound/Source/MetasoundEngine/MetasoundEngine.Build.cs
+++ b/Engine/Plugins/Runtime/Metasound/Source/MetasoundEngine/MetasoundEngine.Build.cs
@@ -39,7 +39,6 @@ namespace UnrealBuildTool.Rules
 					"Engine",
 					"MetasoundFrontend",
 					"MetasoundStandardNodes",
-					"AudioCodecEngine",
 					"Serialization",
 					"AudioPlatformConfiguration",
 					"WaveTable"
diff --git a/Engine/Plugins/Runtime/Metasound/Source/MetasoundEngine/Private/MetasoundEngineModule.cpp b/Engine/Plugins/Runtime/Metasound/Source/MetasoundEngine/Private/MetasoundEngineModule.cpp
--- a/Engine/Plugins/Runtime/Metasound/Source/MetasoundEngine/Private/MetasoundEngineModule.cpp
+++ b/Engine/Plugins/Runtime/Metasound/Source/MetasoundEngine/Private/MetasoundEngineModule.cpp
@@ -82,7 +82,6 @@ public:
 		FModuleManager::Get().LoadModuleChecked("MetasoundFrontend");
 		FModuleManager::Get().LoadModuleChecked("MetasoundStandardNodes");
 		FModuleManager::Get().LoadModuleChecked("MetasoundGenerator");
-		FModuleManager::Get().LoadModuleChecked("AudioCodecEngine");
 		FModuleManager::Get().LoadModuleChecked("WaveTable");
 		
 		// Set GCObject referencer for metasound frontend node registry. The MetaSound
diff --git a/Engine/Source/Runtime/AudioCodecEngine/AudioCodecEngine.Build.cs b/Engine/Source/Runtime/AudioCodecEngine/AudioCodecEngine.Build.cs
deleted file mode 100644
diff --git a/Engine/Source/Runtime/AudioCodecEngine/Private/AudioCodecEngineModule.cpp b/Engine/Source/Runtime/AudioCodecEngine/Private/AudioCodecEngineModule.cpp
deleted file mode 100644
