Delete IAudioCodec and IAudioCodecRegistry

#tests PIE

diff --git a/Engine/Plugins/Runtime/Metasound/Source/MetasoundEngine/Private/MetasoundWave.cpp b/Engine/Plugins/Runtime/Metasound/Source/MetasoundEngine/Private/MetasoundWave.cpp
--- a/Engine/Plugins/Runtime/Metasound/Source/MetasoundEngine/Private/MetasoundWave.cpp
+++ b/Engine/Plugins/Runtime/Metasound/Source/MetasoundEngine/Private/MetasoundWave.cpp
@@ -6,8 +6,6 @@
 #include "AudioDeviceManager.h"
 #include "ContentStreaming.h"
 #include "DSP/ParamInterpolator.h"
-#include "IAudioCodec.h"
-#include "IAudioCodecRegistry.h"
 #include "MetasoundPrimitives.h"
 #include "MetasoundTrace.h"
 #include "Sound/SoundWave.h"
diff --git a/Engine/Plugins/Runtime/Metasound/Source/MetasoundEngine/Private/MetasoundWavePlayerNode.cpp b/Engine/Plugins/Runtime/Metasound/Source/MetasoundEngine/Private/MetasoundWavePlayerNode.cpp
--- a/Engine/Plugins/Runtime/Metasound/Source/MetasoundEngine/Private/MetasoundWavePlayerNode.cpp
+++ b/Engine/Plugins/Runtime/Metasound/Source/MetasoundEngine/Private/MetasoundWavePlayerNode.cpp
@@ -4,7 +4,6 @@
 #include "DSP/ConvertDeinterleave.h"
 #include "DSP/MultichannelBuffer.h"
 #include "DSP/MultichannelLinearResampler.h"
-#include "IAudioCodec.h"
 #include "MetasoundBuildError.h"
 #include "MetasoundBuilderInterface.h"
 #include "MetasoundEngineNodesNames.h"
diff --git a/Engine/Plugins/Runtime/Metasound/Source/MetasoundEngine/Public/MetasoundWave.h b/Engine/Plugins/Runtime/Metasound/Source/MetasoundEngine/Public/MetasoundWave.h
--- a/Engine/Plugins/Runtime/Metasound/Source/MetasoundEngine/Public/MetasoundWave.h
+++ b/Engine/Plugins/Runtime/Metasound/Source/MetasoundEngine/Public/MetasoundWave.h
@@ -6,8 +6,6 @@
 #include "MetasoundDataTypeRegistrationMacro.h"
 #include "IAudioProxyInitializer.h"
 #include "Sound/SoundWave.h"
-#include "IAudioCodecRegistry.h"
-#include "IAudioCodec.h"
 #include "DSP/InterpolatedLinearPitchShifter.h"
 
 
diff --git a/Engine/Plugins/Runtime/Metasound/Source/MetasoundEngineTest/Private/EngineTestMetasoundWavePlayer.cpp b/Engine/Plugins/Runtime/Metasound/Source/MetasoundEngineTest/Private/EngineTestMetasoundWavePlayer.cpp
--- a/Engine/Plugins/Runtime/Metasound/Source/MetasoundEngineTest/Private/EngineTestMetasoundWavePlayer.cpp
+++ b/Engine/Plugins/Runtime/Metasound/Source/MetasoundEngineTest/Private/EngineTestMetasoundWavePlayer.cpp
@@ -14,9 +14,6 @@
 #include "MetasoundWave.h"
 #include "MetasoundAudioBuffer.h"
 
-#include "IAudioCodec.h"
-#include "IAudioCodecRegistry.h"
-
 #if WITH_DEV_AUTOMATION_TESTS
 
 // Disabled for now until waveplayer implementation is extracted and place in a better location for testing
diff --git a/Engine/Source/Runtime/AudioExtensions/Private/IAudioCodec.cpp b/Engine/Source/Runtime/AudioExtensions/Private/IAudioCodec.cpp
deleted file mode 100644
diff --git a/Engine/Source/Runtime/AudioExtensions/Private/IAudioCodecRegistry.cpp b/Engine/Source/Runtime/AudioExtensions/Private/IAudioCodecRegistry.cpp
deleted file mode 100644
diff --git a/Engine/Source/Runtime/AudioExtensions/Public/IAudioCodec.h b/Engine/Source/Runtime/AudioExtensions/Public/IAudioCodec.h
deleted file mode 100644
diff --git a/Engine/Source/Runtime/AudioExtensions/Public/IAudioCodecRegistry.h b/Engine/Source/Runtime/AudioExtensions/Public/IAudioCodecRegistry.h
deleted file mode 100644
diff --git a/Engine/Source/Runtime/Engine/Classes/Sound/SoundWaveProxyReader.h b/Engine/Source/Runtime/Engine/Classes/Sound/SoundWaveProxyReader.h
--- a/Engine/Source/Runtime/Engine/Classes/Sound/SoundWaveProxyReader.h
+++ b/Engine/Source/Runtime/Engine/Classes/Sound/SoundWaveProxyReader.h
@@ -5,7 +5,6 @@
 #include "DSP/BufferVectorOperations.h"
 #include "AudioDecompress.h"
 #include "HAL/Platform.h"
-#include "IAudioCodec.h"
 #include "Templates/SharedPointer.h"
 #include "Templates/UniquePtr.h"
 
diff --git a/Engine/Source/Runtime/Engine/Private/SoundWaveProxyReader.cpp b/Engine/Source/Runtime/Engine/Private/SoundWaveProxyReader.cpp
--- a/Engine/Source/Runtime/Engine/Private/SoundWaveProxyReader.cpp
+++ b/Engine/Source/Runtime/Engine/Private/SoundWaveProxyReader.cpp
@@ -7,6 +7,7 @@
 #include "Containers/ArrayView.h"
 #include "DSP/BufferVectorOperations.h"
 #include "DSP/Dsp.h"
+#include "DSP/FloatArrayMath.h"
 #include "HAL/IConsoleManager.h"
 #include "HAL/LowLevelMemTracker.h"
 #include "HAL/Platform.h"
