- Registration Race Condition Fixes
- Part 1/2 fixing inconsistent file naming that is causing code query tool slowdowns
- Replace SoftObjectPaths with TopLevelAssetPaths for perf (ObjPaths do string copies all over, TopLevelAssetPaths are just two FNames)
#rnx
#tests AudioQA MetaSound Automated Tests, AudioTests

diff --git a/Engine/Plugins/Runtime/Metasound/Source/MetasoundEngine/Private/Metasound.cpp b/Engine/Plugins/Runtime/Metasound/Source/MetasoundEngine/Private/Metasound.cpp
--- a/Engine/Plugins/Runtime/Metasound/Source/MetasoundEngine/Private/Metasound.cpp
+++ b/Engine/Plugins/Runtime/Metasound/Source/MetasoundEngine/Private/Metasound.cpp
@@ -154,7 +154,6 @@ const UEdGraph& UMetaSoundPatch::GetGraphChecked() const
 	check(Graph);
 	return *Graph;
 }
-
 FText UMetaSoundPatch::GetDisplayName() const
 {
 	FString TypeName = UMetaSoundPatch::StaticClass()->GetName();
@@ -168,6 +167,12 @@ void UMetaSoundPatch::SetRegistryAssetClassInfo(const Metasound::Frontend::FNode
 }
 #endif // WITH_EDITORONLY_DATA
 
+
+FTopLevelAssetPath UMetaSoundPatch::GetAssetPathChecked() const
+{
+	return Metasound::FMetaSoundEngineAssetHelper::GetAssetPathChecked(*this);
+}
+
 void UMetaSoundPatch::PostLoad() 
 {
 	Super::PostLoad();
diff --git a/Engine/Plugins/Runtime/Metasound/Source/MetasoundEngine/Private/MetasoundAssetSubsystem.cpp b/Engine/Plugins/Runtime/Metasound/Source/MetasoundEngine/Private/MetasoundAssetSubsystem.cpp
--- a/Engine/Plugins/Runtime/Metasound/Source/MetasoundEngine/Private/MetasoundAssetSubsystem.cpp
+++ b/Engine/Plugins/Runtime/Metasound/Source/MetasoundEngine/Private/MetasoundAssetSubsystem.cpp
@@ -35,7 +35,7 @@ namespace Metasound
 			bool bSuccess = true;
 
 			OutInfo.Type = EMetasoundFrontendClassType::External;
-			OutInfo.AssetPath = InAssetData.GetSoftObjectPath();
+			OutInfo.AssetPath = InAssetData.AssetClassPath;
 			FString AssetClassID;
 			bSuccess &= InAssetData.GetTagValue(AssetTags::AssetClassID, AssetClassID);
 			OutInfo.AssetClassID = FGuid(AssetClassID);
diff --git a/Engine/Plugins/Runtime/Metasound/Source/MetasoundEngine/Private/MetasoundBuilderSubsystem.cpp b/Engine/Plugins/Runtime/Metasound/Source/MetasoundEngine/Private/MetasoundBuilderSubsystem.cpp
--- a/Engine/Plugins/Runtime/Metasound/Source/MetasoundEngine/Private/MetasoundBuilderSubsystem.cpp
+++ b/Engine/Plugins/Runtime/Metasound/Source/MetasoundEngine/Private/MetasoundBuilderSubsystem.cpp
@@ -1062,7 +1062,7 @@ void UMetaSoundSourceBuilder::Audition(UObject* Parent, UAudioComponent* AudioCo
 	RegisterGraphIfOutstandingTransactions(MetaSoundSource);
 
 	// Must be called post register as register ensures cached runtime data passed to transactor is up-to-date
-	MetaSoundSource.SetDynamicGeneratorEnabled(bLiveUpdatesEnabled);
+	MetaSoundSource.SetDynamicGeneratorEnabled(MetaSoundSource.GetAssetPathChecked(), bLiveUpdatesEnabled);
 	MetaSoundSource.ConformObjectDataToInterfaces();
 
 	AudioComponent->SetSound(&MetaSoundSource);
diff --git a/Engine/Plugins/Runtime/Metasound/Source/MetasoundEngine/Private/MetasoundEngineAsset.h b/Engine/Plugins/Runtime/Metasound/Source/MetasoundEngine/Private/MetasoundEngineAsset.h
--- a/Engine/Plugins/Runtime/Metasound/Source/MetasoundEngine/Private/MetasoundEngineAsset.h
+++ b/Engine/Plugins/Runtime/Metasound/Source/MetasoundEngine/Private/MetasoundEngineAsset.h
@@ -62,6 +62,14 @@ namespace Metasound
 		}
 #endif // WITH_EDITOR
 
+		template <typename TMetaSoundObject>
+		static FTopLevelAssetPath GetAssetPathChecked(TMetaSoundObject& InMetaSound)
+		{
+			FTopLevelAssetPath Path;
+			ensureAlwaysMsgf(Path.TrySetPath(&InMetaSound), TEXT("Failed to set TopLevelAssetPath from MetaSound '%s'. MetaSound must be highest level object in package."), *InMetaSound.GetPathName());
+			return Path;
+		}
+
 		template <typename TMetaSoundObject>
 		static TArray<FMetasoundAssetBase*> GetReferencedAssets(TMetaSoundObject& InMetaSound)
 		{
diff --git a/Engine/Plugins/Runtime/Metasound/Source/MetasoundEngine/Private/MetasoundOperatorCacheSubsystem.cpp b/Engine/Plugins/Runtime/Metasound/Source/MetasoundEngine/Private/MetasoundOperatorCacheSubsystem.cpp
--- a/Engine/Plugins/Runtime/Metasound/Source/MetasoundEngine/Private/MetasoundOperatorCacheSubsystem.cpp
+++ b/Engine/Plugins/Runtime/Metasound/Source/MetasoundEngine/Private/MetasoundOperatorCacheSubsystem.cpp
@@ -129,7 +129,7 @@ void UMetaSoundCacheSubsystem::PrecacheMetaSound(UMetaSoundSource* InMetaSound,
 	TUniquePtr<FOperatorBuildData> Data = MakeUnique<FOperatorBuildData>(
 		  MoveTemp(InitParams.GetValue())
 		, InMetaSound->GetRegistryKey()
-		, FSoftObjectPath(InMetaSound->GetOwningAsset())
+		, InMetaSound->GetAssetPathChecked()
 		, InMetaSound->AssetClassID
 		, InNumInstances
 	);
diff --git a/Engine/Plugins/Runtime/Metasound/Source/MetasoundEngine/Private/MetasoundSource.cpp b/Engine/Plugins/Runtime/Metasound/Source/MetasoundEngine/Private/MetasoundSource.cpp
--- a/Engine/Plugins/Runtime/Metasound/Source/MetasoundEngine/Private/MetasoundSource.cpp
+++ b/Engine/Plugins/Runtime/Metasound/Source/MetasoundEngine/Private/MetasoundSource.cpp
@@ -490,6 +490,11 @@ bool UMetaSoundSource::ConformObjectDataToInterfaces()
 	return bDidAlterObjectData;
 }
 
+FTopLevelAssetPath UMetaSoundSource::GetAssetPathChecked() const
+{
+	return Metasound::FMetaSoundEngineAssetHelper::GetAssetPathChecked(*this);
+}
+
 void UMetaSoundSource::BeginDestroy()
 {
 	UnregisterGraphWithFrontend();
@@ -879,7 +884,7 @@ ISoundGeneratorPtr UMetaSoundSource::CreateSoundGenerator(const FSoundGeneratorI
 		// the base IGraph. 
 		if (ConsoleVariables::bEnableExperimentalRuntimePresetGraphInflation && bIsPresetGraphInflationSupported)
 		{
-			// Get the graph associated with base graph which this preset wraps .
+			// Get the graph associated with base graph which this preset wraps.
 			TSharedPtr<const IGraph> MetasoundGraph = TryGetMetaSoundPresetBaseGraph();
 
 			if (MetasoundGraph.IsValid())
@@ -908,7 +913,7 @@ ISoundGeneratorPtr UMetaSoundSource::CreateSoundGenerator(const FSoundGeneratorI
 
 		if (!Generator.IsValid())
 		{
-			TSharedPtr<const IGraph> MetasoundGraph = GetRegisteredGraph();
+			TSharedPtr<const FGraph> MetasoundGraph = FMetasoundFrontendRegistryContainer::Get()->GetGraph(GetRegistryKey(), GetAssetPathChecked());
 			if (!MetasoundGraph.IsValid())
 			{
 				return ISoundGeneratorPtr(nullptr);
@@ -1520,7 +1525,7 @@ void UMetaSoundSource::OnFinishActiveBuilder()
 	bIsBuilderActive = false;
 }
 
-TSharedPtr<Metasound::DynamicGraph::FDynamicOperatorTransactor> UMetaSoundSource::SetDynamicGeneratorEnabled(bool bInIsEnabled)
+TSharedPtr<Metasound::DynamicGraph::FDynamicOperatorTransactor> UMetaSoundSource::SetDynamicGeneratorEnabled(const FTopLevelAssetPath& InAssetPath, bool bInIsEnabled)
 {
 	using namespace Metasound;
 	using namespace Metasound::DynamicGraph;
@@ -1538,7 +1543,8 @@ TSharedPtr<Metasound::DynamicGraph::FDynamicOperatorTransactor> UMetaSoundSource
 			// graph to see if any FGraph already exists. 
 			if (IsRegistered())
 			{
-				TSharedPtr<const FGraph> CurrentGraph = GetRegisteredGraph();
+				TSharedPtr<const FGraph> CurrentGraph = FMetasoundFrontendRegistryContainer::Get()->GetGraph(GetRegistryKey(), InAssetPath);
+
 				if (CurrentGraph.IsValid())
 				{
 					DynamicTransactor = MakeShared<FDynamicOperatorTransactor>(*CurrentGraph);
@@ -1696,9 +1702,9 @@ TSharedPtr<const Metasound::IGraph> UMetaSoundSource::TryGetMetaSoundPresetBaseG
 		TObjectPtr<const UMetaSoundSource> BaseMetaSoundSource = Cast<const UMetaSoundSource>(BaseGraph);
 		if (BaseMetaSoundSource)
 		{
-			// Get registered graph of base metasound source. 
-			FSoftObjectPath BaseMetaSoundSourcePath(BaseMetaSoundSource->GetOwningAsset());
-			MetasoundGraph = FMetasoundFrontendRegistryContainer::Get()->GetGraph(BaseMetaSoundSource->GetRegistryKey(), FSoftObjectPath(BaseMetaSoundSource->GetOwningAsset()));
+			// Get registered graph of base metasound source.
+			const FTopLevelAssetPath BaseMetaSoundSourcePath = BaseMetaSoundSource->GetAssetPathChecked();
+			MetasoundGraph = FMetasoundFrontendRegistryContainer::Get()->GetGraph(BaseMetaSoundSource->GetRegistryKey(), BaseMetaSoundSourcePath);
 		}
 	}
 	else
diff --git a/Engine/Plugins/Runtime/Metasound/Source/MetasoundEngine/Public/Metasound.h b/Engine/Plugins/Runtime/Metasound/Source/MetasoundEngine/Public/Metasound.h
--- a/Engine/Plugins/Runtime/Metasound/Source/MetasoundEngine/Public/Metasound.h
+++ b/Engine/Plugins/Runtime/Metasound/Source/MetasoundEngine/Public/Metasound.h
@@ -127,9 +127,9 @@ public:
 	{
 		Graph = CastChecked<UMetasoundEditorGraphBase>(InGraph);
 	}
-
 #endif // #if WITH_EDITORONLY_DATA
 
+	virtual FTopLevelAssetPath GetAssetPathChecked() const override;
 	virtual const UClass& GetBaseMetaSoundUClass() const final override;
 	virtual const FMetasoundFrontendDocument& GetConstDocument() const override;
 
diff --git a/Engine/Plugins/Runtime/Metasound/Source/MetasoundEngine/Public/MetasoundSource.h b/Engine/Plugins/Runtime/Metasound/Source/MetasoundEngine/Public/MetasoundSource.h
--- a/Engine/Plugins/Runtime/Metasound/Source/MetasoundEngine/Public/MetasoundSource.h
+++ b/Engine/Plugins/Runtime/Metasound/Source/MetasoundEngine/Public/MetasoundSource.h
@@ -230,6 +230,8 @@ public:
 
 	virtual bool ConformObjectDataToInterfaces() override;
 
+	virtual FTopLevelAssetPath GetAssetPathChecked() const override;
+
 	UObject* GetOwningAsset() override
 	{
 		return this;
@@ -323,7 +325,7 @@ private:
 	 * Note: Disabling the dynamic generator will sever the communication between any active generators
 	 * even if the dynamic generator is re-enabled during the lifetime of the active generators
 	 */
-	TSharedPtr<Metasound::DynamicGraph::FDynamicOperatorTransactor> SetDynamicGeneratorEnabled(bool bInIsEnabled);
+	TSharedPtr<Metasound::DynamicGraph::FDynamicOperatorTransactor> SetDynamicGeneratorEnabled(const FTopLevelAssetPath& InAssetPath, bool bInIsEnabled);
 
 	/** Get dynamic transactor
 	 *
diff --git a/Engine/Plugins/Runtime/Metasound/Source/MetasoundFrontend/Private/MetasoundAssetBase.cpp b/Engine/Plugins/Runtime/Metasound/Source/MetasoundFrontend/Private/MetasoundAssetBase.cpp
--- a/Engine/Plugins/Runtime/Metasound/Source/MetasoundFrontend/Private/MetasoundAssetBase.cpp
+++ b/Engine/Plugins/Runtime/Metasound/Source/MetasoundFrontend/Private/MetasoundAssetBase.cpp
@@ -129,7 +129,7 @@ namespace Metasound
 			class FDocumentNodeRegistryEntryForCook : public INodeRegistryEntry
 			{
 			public:
-				FDocumentNodeRegistryEntryForCook(const FMetasoundFrontendDocument& InDocument, const FSoftObjectPath& InAssetPath)
+				FDocumentNodeRegistryEntryForCook(const FMetasoundFrontendDocument& InDocument, const FTopLevelAssetPath& InAssetPath)
 					: Interfaces(InDocument.Interfaces)
 					, FrontendClass(InDocument.RootGraph)
 					, ClassInfo(InDocument.RootGraph, InAssetPath)
@@ -320,13 +320,12 @@ void FMetasoundAssetBase::RegisterGraphWithFrontend(Metasound::Frontend::FMetaSo
 	check(Owner);
 	const FString AssetName = Owner->GetName();
 
-	{
-		TScriptInterface<IMetaSoundDocumentInterface> RegistryDocInterface = AssetBasePrivate::BuildRegistryDocument(Owner);
-		// Clear before reregistration
-		RegistryKey = FNodeRegistryKey();
-		RegisteredGraphAssetPath = FSoftObjectPath();
-		RegistryKey = CacheRuntimeData(RegistryDocInterface);
-	}
+	// If the IMetaSoundDocumentInterface actively modified by a builder, then
+	// we must build synchronously to avoid a race condition on reading/writing
+	// the IMetaSoundDocumentInterface
+	const bool bAsync = !(IsBuilderActive() || ConsoleVariables::bDisableAsyncGraphRegistration);
+	const TScriptInterface<IMetaSoundDocumentInterface> RegistryDocInterface = AssetBasePrivate::BuildRegistryDocument(Owner);
+	RegistryKey = FMetasoundFrontendRegistryContainer::Get()->RegisterGraph(RegistryDocInterface, bAsync);
 
 	if (RegistryKey.IsValid())
 	{
@@ -343,11 +342,6 @@ void FMetasoundAssetBase::RegisterGraphWithFrontend(Metasound::Frontend::FMetaSo
 	}
 }
 
-TSharedPtr<const Metasound::FGraph> FMetasoundAssetBase::GetRegisteredGraph() const
-{
-	return FMetasoundFrontendRegistryContainer::Get()->GetGraph(RegistryKey, RegisteredGraphAssetPath);
-}
-
 void FMetasoundAssetBase::CookMetaSound()
 {
 	using namespace Metasound;
@@ -393,8 +387,12 @@ void FMetasoundAssetBase::CookMetaSound()
 		// creating proxies. To do so, we use a special node registration object which reflects the necessary information
 		// for the node registry, but does not create INodes.
 		UnregisterGraphWithFrontend();
-		const FMetasoundFrontendDocument& Document = TScriptInterface<IMetaSoundDocumentInterface>(Owner)->GetConstDocument();
-		TUniquePtr<INodeRegistryEntry> RegistryEntry = MakeUnique<AssetBasePrivate::FDocumentNodeRegistryEntryForCook>(Document, FSoftObjectPath(Owner));
+		TScriptInterface<IMetaSoundDocumentInterface> DocInterface(Owner);
+		const FMetasoundFrontendDocument& Document = DocInterface->GetConstDocument();
+		TUniquePtr<INodeRegistryEntry> RegistryEntry = MakeUnique<AssetBasePrivate::FDocumentNodeRegistryEntryForCook>(
+			Document,
+			DocInterface->GetAssetPathChecked()
+		);
 		RegistryKey = FMetasoundFrontendRegistryContainer::Get()->RegisterNode(MoveTemp(RegistryEntry));
 	}
 
@@ -419,34 +417,24 @@ void FMetasoundAssetBase::UnregisterGraphWithFrontend()
 	METASOUND_TRACE_CPUPROFILER_EVENT_SCOPE(MetaSoundAssetBase::UnregisterGraphWithFrontend);
 
 	check(IsInGameThread());
-	if (!RegistryKey.IsValid())
-	{
-		return;
-	}
-
-	const UObject* OwningAsset = GetOwningAsset();
-	if (!ensureAlways(OwningAsset))
-	{
-		return;
-	}
-
-	bool bSuccess = false;
-	if (RegisteredGraphAssetPath.IsValid())
-	{
-		bSuccess = FMetasoundFrontendRegistryContainer::Get()->UnregisterGraph(RegistryKey, RegisteredGraphAssetPath);
-	}
-	else
+	if (RegistryKey.IsValid())
 	{
-		bSuccess = FMetasoundFrontendRegistryContainer::Get()->UnregisterNode(RegistryKey);
-	}
+		UObject* OwningAsset = GetOwningAsset();
+		if (ensureAlways(OwningAsset))
+		{
+			// If the IMetaSoundDocumentInterface is actively modified by a builder, then
+			// we must release built graph synchronously to avoid a race condition on
+			// reading/writing the IMetaSoundDocumentInterface
+			const bool bAsync = !(IsBuilderActive() || ConsoleVariables::bDisableAsyncGraphRegistration);
+			const bool bSuccess = FMetasoundFrontendRegistryContainer::Get()->UnregisterGraph(OwningAsset, bAsync);
+			if (!bSuccess)
+			{
+				UE_LOG(LogMetaSound, Warning, TEXT("Failed to unregister node with key %s for asset %s. No registry entry exists with that key."), *RegistryKey.ToString(), *GetOwningAssetName());
+			}
+		}
 
-	if (!bSuccess)
-	{
-		UE_LOG(LogMetaSound, Warning, TEXT("Failed to unregister node with key %s for asset %s. No registry entry exists with that key."), *RegistryKey.ToString(), *GetOwningAssetName());
+		RegistryKey = FNodeRegistryKey();
 	}
-
-	RegistryKey = FNodeRegistryKey();
-	RegisteredGraphAssetPath = FSoftObjectPath();
 }
 
 void FMetasoundAssetBase::SetMetadata(FMetasoundFrontendClassMetadata& InMetadata)
@@ -477,8 +465,8 @@ Metasound::Frontend::FNodeClassInfo FMetasoundAssetBase::GetAssetClassInfo() con
 
 	const UObject* Owner = GetOwningAsset();
 	check(Owner);
-
-	return FNodeClassInfo { GetDocumentChecked().RootGraph, FSoftObjectPath(Owner) };
+	TScriptInterface<const IMetaSoundDocumentInterface> DocInterface((UObject*)Owner);
+	return FNodeClassInfo { GetDocumentChecked().RootGraph, DocInterface->GetAssetPathChecked() };
 }
 
 void FMetasoundAssetBase::SetDocument(const FMetasoundFrontendDocument& InDocument, bool bMarkDirty)
@@ -958,9 +946,10 @@ void FMetasoundAssetBase::UpdateAssetRegistry()
 	using namespace Metasound;
 	using namespace Metasound::Frontend;
 
-	const UObject* Owner = GetOwningAsset();
+	UObject* Owner = GetOwningAsset();
 	check(Owner);
-	FNodeClassInfo AssetClassInfo { GetDocumentChecked().RootGraph, FSoftObjectPath(Owner) };
+	TScriptInterface<IMetaSoundDocumentInterface> DocInterface(Owner);
+	FNodeClassInfo AssetClassInfo(GetDocumentChecked().RootGraph, DocInterface->GetAssetPathChecked());
 
 	// Refresh Asset Registry Info if successfully registered with Frontend
 	const FMetasoundFrontendGraphClass& DocumentClassGraph = GetDocumentHandle()->GetRootGraphClass();
@@ -986,32 +975,14 @@ TSharedPtr<FMetasoundFrontendDocument> FMetasoundAssetBase::PreprocessDocument()
 	return nullptr;
 }
 
-Metasound::Frontend::FNodeRegistryKey FMetasoundAssetBase::CacheRuntimeData(const TScriptInterface<IMetaSoundDocumentInterface>& InDocumentInterface)
-{
-	using namespace Metasound::Frontend;
-
-	METASOUND_TRACE_CPUPROFILER_EVENT_SCOPE(MetaSoundAssetBase::CacheRuntimeData);
-
-	const FMetasoundFrontendDocument& Document = InDocumentInterface->GetConstDocument();
-
-	// If the IMetaSoundDocumentInterface actively modified by a builder, then 
-	// we must build synchronously to avoid a race condition on reading/writing 
-	// the IMetaSoundDocumentInterface
-	const bool bAsync = !(IsBuilderActive() || ConsoleVariables::bDisableAsyncGraphRegistration);
-
-	RegisteredGraphAssetPath = FSoftObjectPath(GetOwningAsset());
-	FNodeRegistryKey NewRegistryKey = FMetasoundFrontendRegistryContainer::Get()->RegisterGraph(RegisteredGraphAssetPath, InDocumentInterface, bAsync);
-
-
-	return NewRegistryKey;
-}
-
 void FMetasoundAssetBase::WaitForAsyncGraphRegistration()
 {
 	using namespace Metasound::Frontend;
 	if (RegistryKey.IsValid())
 	{
-		FMetasoundFrontendRegistryContainer::Get()->WaitForAsyncGraphRegistration(RegistryKey, RegisteredGraphAssetPath);
+		TScriptInterface<IMetaSoundDocumentInterface> DocInterface(GetOwningAsset());
+		check(DocInterface.GetObject());
+		FMetasoundFrontendRegistryContainer::Get()->WaitForAsyncGraphRegistration(RegistryKey, DocInterface->GetAssetPathChecked());
 	}
 }
 
diff --git a/Engine/Plugins/Runtime/Metasound/Source/MetasoundFrontend/Private/MetasoundFrontend.cpp b/Engine/Plugins/Runtime/Metasound/Source/MetasoundFrontend/Private/MetasoundFrontend.cpp
--- a/Engine/Plugins/Runtime/Metasound/Source/MetasoundFrontend/Private/MetasoundFrontend.cpp
+++ b/Engine/Plugins/Runtime/Metasound/Source/MetasoundFrontend/Private/MetasoundFrontend.cpp
@@ -9,9 +9,8 @@
 #include "MetasoundDataTypeRegistrationMacro.h"
 #include "MetasoundFrontendDocument.h"
 #include "MetasoundFrontendGraph.h"
-#include "MetasoundFrontendNodeRegistryPrivate.h"
 #include "MetasoundFrontendNodeTemplateRegistry.h"
-#include "MetasoundFrontendRegistries.h"
+#include "MetasoundFrontendRegistryContainer.h"
 #include "MetasoundFrontendRegistryTransaction.h"
 #include "MetasoundJsonBackend.h"
 #include "MetasoundOperatorBuilder.h"
diff --git a/Engine/Plugins/Runtime/Metasound/Source/MetasoundFrontend/Private/MetasoundFrontendDocumentBuilder.cpp b/Engine/Plugins/Runtime/Metasound/Source/MetasoundFrontend/Private/MetasoundFrontendDocumentBuilder.cpp
--- a/Engine/Plugins/Runtime/Metasound/Source/MetasoundFrontend/Private/MetasoundFrontendDocumentBuilder.cpp
+++ b/Engine/Plugins/Runtime/Metasound/Source/MetasoundFrontend/Private/MetasoundFrontendDocumentBuilder.cpp
@@ -588,6 +588,13 @@ UMetaSoundBuilderDocument& UMetaSoundBuilderDocument::Create(const IMetaSoundDoc
 	return *DocObject;
 }
 
+FTopLevelAssetPath UMetaSoundBuilderDocument::GetAssetPathChecked() const
+{
+	FTopLevelAssetPath Path;
+	ensureAlwaysMsgf(Path.TrySetPath(this), TEXT("Failed to set TopLevelAssetPath from transient MetaSound '%s'. MetaSound must be highest level object in package."), *GetPathName());
+	return Path;
+}
+
 const FMetasoundFrontendDocument& UMetaSoundBuilderDocument::GetConstDocument() const
 {
 	return Document;
diff --git a/Engine/Plugins/Runtime/Metasound/Source/MetasoundFrontend/Private/MetasoundFrontendNodeRegistryPrivate.h b/Engine/Plugins/Runtime/Metasound/Source/MetasoundFrontend/Private/MetasoundFrontendNodeRegistryPrivate.h
deleted file mode 100644
diff --git a/Engine/Plugins/Runtime/Metasound/Source/MetasoundFrontend/Private/MetasoundFrontendNodeTemplateRegistry.cpp b/Engine/Plugins/Runtime/Metasound/Source/MetasoundFrontend/Private/MetasoundFrontendNodeTemplateRegistry.cpp
--- a/Engine/Plugins/Runtime/Metasound/Source/MetasoundFrontend/Private/MetasoundFrontendNodeTemplateRegistry.cpp
+++ b/Engine/Plugins/Runtime/Metasound/Source/MetasoundFrontend/Private/MetasoundFrontendNodeTemplateRegistry.cpp
@@ -3,8 +3,7 @@
 #include "MetasoundFrontendNodeTemplateRegistry.h"
 
 #include "MetasoundFrontendDocument.h"
-#include "MetasoundFrontendNodeRegistryPrivate.h"
-#include "MetasoundFrontendRegistries.h"
+#include "MetasoundFrontendRegistryContainerImpl.h"
 #include "MetasoundFrontendRegistryTransaction.h"
 #include "MetasoundLog.h"
 #include "MetasoundTrace.h"
diff --git a/Engine/Plugins/Runtime/Metasound/Source/MetasoundFrontend/Private/MetasoundFrontendQuerySteps.cpp b/Engine/Plugins/Runtime/Metasound/Source/MetasoundFrontend/Private/MetasoundFrontendQuerySteps.cpp
--- a/Engine/Plugins/Runtime/Metasound/Source/MetasoundFrontend/Private/MetasoundFrontendQuerySteps.cpp
+++ b/Engine/Plugins/Runtime/Metasound/Source/MetasoundFrontend/Private/MetasoundFrontendQuerySteps.cpp
@@ -5,9 +5,9 @@
 #include "Algo/MaxElement.h"
 #include "MetasoundFrontend.h"
 #include "MetasoundFrontendDocument.h"
-#include "MetasoundFrontendNodeRegistryPrivate.h"
 #include "MetasoundFrontendQuery.h"
-#include "MetasoundFrontendRegistries.h"
+#include "MetasoundFrontendRegistryContainer.h"
+#include "MetasoundFrontendRegistryContainerImpl.h"
 #include "MetasoundFrontendRegistryTransaction.h"
 #include "MetasoundLog.h"
 
diff --git a/Engine/Plugins/Runtime/Metasound/Source/MetasoundFrontend/Private/MetasoundFrontendRegistries.cpp b/Engine/Plugins/Runtime/Metasound/Source/MetasoundFrontend/Private/MetasoundFrontendRegistries.cpp
deleted file mode 100644
diff --git a/Engine/Plugins/Runtime/Metasound/Source/MetasoundFrontend/Private/MetasoundFrontendRegistryContainer.cpp b/Engine/Plugins/Runtime/Metasound/Source/MetasoundFrontend/Private/MetasoundFrontendRegistryContainer.cpp
new file mode 100644
--- /dev/null
+++ b/Engine/Plugins/Runtime/Metasound/Source/MetasoundFrontend/Private/MetasoundFrontendRegistryContainer.cpp
@@ -0,0 +1,203 @@
+// Copyright Epic Games, Inc. All Rights Reserved.
+
+#include "MetasoundFrontendRegistryContainer.h"
+
+#include "HAL/PlatformTime.h"
+
+#include "MetasoundFrontendDataTypeRegistry.h"
+#include "MetasoundFrontendGraph.h"
+#include "MetasoundFrontendProxyDataCache.h"
+#include "MetasoundFrontendRegistryContainerImpl.h"
+#include "MetasoundFrontendRegistryTransaction.h"
+#include "MetasoundLog.h"
+#include "MetasoundRouter.h"
+#include "MetasoundTrace.h"
+
+
+namespace Metasound::Frontend
+{
+	FNodeRegistryTransaction::FNodeRegistryTransaction(ETransactionType InType, const FNodeClassInfo& InNodeClassInfo, FNodeRegistryTransaction::FTimeType InTimestamp)
+	: Type(InType)
+	, NodeClassInfo(InNodeClassInfo)
+	, Timestamp(InTimestamp)
+	{
+	}
+
+	FNodeRegistryTransaction::ETransactionType FNodeRegistryTransaction::GetTransactionType() const
+	{
+		return Type;
+	}
+
+	const FNodeClassInfo& FNodeRegistryTransaction::GetNodeClassInfo() const
+	{
+		return NodeClassInfo;
+	}
+
+	FNodeRegistryKey FNodeRegistryTransaction::GetNodeRegistryKey() const
+	{
+		return FNodeRegistryKey(NodeClassInfo);
+	}
+
+	FNodeRegistryTransaction::FTimeType FNodeRegistryTransaction::GetTimestamp() const
+	{
+		return Timestamp;
+	}
+
+	namespace NodeRegistryKey
+	{
+		FNodeRegistryKey CreateKey(EMetasoundFrontendClassType InType, const FString& InFullClassName, int32 InMajorVersion, int32 InMinorVersion)
+		{
+			if (InType == EMetasoundFrontendClassType::Graph)
+			{
+				// No graphs are registered. Any registered graph should be registered as an external node.
+				InType = EMetasoundFrontendClassType::External;
+			}
+
+			FMetasoundFrontendClassName ClassName;
+			FMetasoundFrontendClassName::Parse(InFullClassName, ClassName);
+			return FNodeRegistryKey(InType, ClassName, InMajorVersion, InMinorVersion);
+		}
+
+		const FNodeRegistryKey& GetInvalid()
+		{
+			return FNodeRegistryKey::GetInvalid();
+		}
+
+		bool IsValid(const FNodeRegistryKey& InKey)
+		{
+			return InKey.IsValid();
+		}
+
+		bool IsEqual(const FNodeRegistryKey& InLHS, const FNodeRegistryKey& InRHS)
+		{
+			return InLHS == InRHS;
+		}
+
+		bool IsEqual(const FMetasoundFrontendClassMetadata& InLHS, const FMetasoundFrontendClassMetadata& InRHS)
+		{
+			if (InLHS.GetClassName() == InRHS.GetClassName())
+			{
+				if (InLHS.GetType() == InRHS.GetType())
+				{
+					if (InLHS.GetVersion() == InRHS.GetVersion())
+					{
+						return true;
+					}
+				}
+			}
+			return false;
+		}
+
+		bool IsEqual(const FNodeClassInfo& InLHS, const FMetasoundFrontendClassMetadata& InRHS)
+		{
+			if (InLHS.ClassName == InRHS.GetClassName())
+			{
+				if (InLHS.Type == InRHS.GetType())
+				{
+					if (InLHS.Version == InRHS.GetVersion())
+					{
+						return true;
+					}
+				}
+			}
+			return false;
+		}
+
+		FNodeRegistryKey CreateKey(const FNodeClassMetadata& InNodeMetadata)
+		{
+			return FNodeRegistryKey(InNodeMetadata);
+		}
+
+		FNodeRegistryKey CreateKey(const FMetasoundFrontendClassMetadata& InNodeMetadata)
+		{
+			checkf(InNodeMetadata.GetType() != EMetasoundFrontendClassType::Graph, TEXT("Cannot create key from 'graph' type. Likely meant to use CreateKey overload that is provided FMetasoundFrontendGraphClass"));
+			return FNodeRegistryKey(InNodeMetadata);
+		}
+
+		FNodeRegistryKey CreateKey(const FMetasoundFrontendGraphClass& InGraphClass)
+		{
+			return FNodeRegistryKey(InGraphClass);
+		}
+
+		FNodeRegistryKey CreateKey(const FNodeClassInfo& InClassInfo)
+		{
+			return FNodeRegistryKey(InClassInfo);
+		}
+	} // namespace NodeRegistryKey
+} // namespace Metasound::Frontend
+
+
+FMetasoundFrontendRegistryContainer* FMetasoundFrontendRegistryContainer::Get()
+{
+	return &Metasound::Frontend::FRegistryContainerImpl::Get();
+}
+
+void FMetasoundFrontendRegistryContainer::ShutdownMetasoundFrontend()
+{
+	Metasound::Frontend::FRegistryContainerImpl::Shutdown();
+}
+
+Metasound::Frontend::FNodeRegistryKey FMetasoundFrontendRegistryContainer::GetRegistryKey(const FNodeClassMetadata& InNodeMetadata)
+{
+	return Metasound::Frontend::FNodeRegistryKey(InNodeMetadata);
+}
+
+Metasound::Frontend::FNodeRegistryKey FMetasoundFrontendRegistryContainer::GetRegistryKey(const FMetasoundFrontendClassMetadata& InNodeMetadata)
+{
+	return Metasound::Frontend::FNodeRegistryKey(InNodeMetadata);
+}
+
+Metasound::Frontend::FNodeRegistryKey FMetasoundFrontendRegistryContainer::GetRegistryKey(const FNodeClassInfo& InClassInfo)
+{
+	return Metasound::Frontend::FNodeRegistryKey(InClassInfo);
+}
+
+bool FMetasoundFrontendRegistryContainer::GetFrontendClassFromRegistered(const FNodeRegistryKey& InKey, FMetasoundFrontendClass& OutClass)
+{
+	FMetasoundFrontendRegistryContainer* Registry = FMetasoundFrontendRegistryContainer::Get();
+
+	if (ensure(nullptr != Registry))
+	{
+		return Registry->FindFrontendClassFromRegistered(InKey, OutClass);
+	}
+
+	return false;
+}
+
+bool FMetasoundFrontendRegistryContainer::GetNodeClassInfoFromRegistered(const FNodeRegistryKey& InKey, FNodeClassInfo& OutInfo)
+{
+	if (FMetasoundFrontendRegistryContainer* Registry = FMetasoundFrontendRegistryContainer::Get())
+	{
+		PRAGMA_DISABLE_DEPRECATION_WARNINGS
+		return Registry->FindNodeClassInfoFromRegistered(InKey, OutInfo);
+		PRAGMA_ENABLE_DEPRECATION_WARNINGS
+	}
+	return false;
+}
+
+bool FMetasoundFrontendRegistryContainer::GetInputNodeRegistryKeyForDataType(const FName& InDataTypeName, const EMetasoundFrontendVertexAccessType InAccessType, FNodeRegistryKey& OutKey)
+{
+	if (FMetasoundFrontendRegistryContainer* Registry = FMetasoundFrontendRegistryContainer::Get())
+	{
+		return Registry->FindInputNodeRegistryKeyForDataType(InDataTypeName, InAccessType, OutKey);
+	}
+	return false;
+}
+
+bool FMetasoundFrontendRegistryContainer::GetVariableNodeRegistryKeyForDataType(const FName& InDataTypeName, FNodeRegistryKey& OutKey)
+{
+	if (FMetasoundFrontendRegistryContainer* Registry = FMetasoundFrontendRegistryContainer::Get())
+	{
+		return Registry->FindVariableNodeRegistryKeyForDataType(InDataTypeName, OutKey);
+	}
+	return false;
+}
+
+bool FMetasoundFrontendRegistryContainer::GetOutputNodeRegistryKeyForDataType(const FName& InDataTypeName, const EMetasoundFrontendVertexAccessType InVertexAccessType, FNodeRegistryKey& OutKey)
+{
+	if (FMetasoundFrontendRegistryContainer* Registry = FMetasoundFrontendRegistryContainer::Get())
+	{
+		return Registry->FindOutputNodeRegistryKeyForDataType(InDataTypeName, InVertexAccessType, OutKey);
+	}
+	return false;
+}
diff --git a/Engine/Plugins/Runtime/Metasound/Source/MetasoundFrontend/Private/MetasoundFrontendRegistryContainerImpl.cpp b/Engine/Plugins/Runtime/Metasound/Source/MetasoundFrontend/Private/MetasoundFrontendRegistryContainerImpl.cpp
new file mode 100644
--- /dev/null
+++ b/Engine/Plugins/Runtime/Metasound/Source/MetasoundFrontend/Private/MetasoundFrontendRegistryContainerImpl.cpp
@@ -0,0 +1,1105 @@
+// Copyright Epic Games, Inc. All Rights Reserved.
+
+#include "MetasoundFrontendRegistryContainerImpl.h"
+
+#include "Algo/ForEach.h"
+#include "Algo/Transform.h"
+#include "Containers/UnrealString.h"
+#include "Containers/Array.h"
+#include "Containers/Set.h"
+#include "HAL/PlatformTime.h"
+#include "Misc/Guid.h"
+#include "Misc/ScopeLock.h"
+#include "Tasks/Pipe.h"
+#include "Tasks/Task.h"
+#include "Templates/SharedPointer.h"
+#include "Templates/UniquePtr.h"
+#include "UObject/NameTypes.h"
+#include "UObject/ScriptInterface.h"
+
+#include "MetasoundDocumentInterface.h"
+#include "MetasoundFrontendDataTypeRegistry.h"
+#include "MetasoundFrontendGraph.h"
+#include "MetasoundFrontendProxyDataCache.h"
+#include "MetasoundFrontendRegistryContainerImpl.h"
+#include "MetasoundFrontendRegistryTransaction.h"
+#include "MetasoundFrontendSearchEngine.h"
+#include "MetasoundLog.h"
+#include "MetasoundRouter.h"
+#include "MetasoundTrace.h"
+
+bool bBusyWaitOnAsyncRegistrationTasks = true;
+static FAutoConsoleVariableRef CVarAsyncRegistrationTasksBusyWait(
+	TEXT("au.MetaSound.BusyWaitOnAsyncRegistrationTasks"),
+	bBusyWaitOnAsyncRegistrationTasks,
+	TEXT("Use TaskGraph BusyWait instead of simple Wait. Required to avoid hangs on platforms with low number of cores."),
+	ECVF_Default);
+
+namespace Metasound::Frontend
+{
+	namespace RegistryPrivate
+	{
+		// FGraphNode is used to create unique INodes based off of a IGraph. 
+		//
+		// Individual nodes need to reflect their InstanceName and InstanceID, but otherwise
+		// they simply encapsulate a shared set of behavior. To minimize memory usage, a single
+		// shared IGraph is used for all nodes referring to the same IGraph.
+		class FGraphNode : public INode
+		{
+			// This adapter class forwards the correct FBuilderOperatorParams
+			// to the graph's operator creation method. Many operator creation
+			// methods downcast the supplied INode in `FBuilderOperatorParams`
+			// and so it is required that it point to the correct runtime instance
+			// when calling CreateOperator(...)
+			class FGraphOperatorFactoryAdapter : public IOperatorFactory 
+			{
+			public:
+				FGraphOperatorFactoryAdapter(const IGraph& InGraph)
+				: Graph(&InGraph)
+				, GraphFactory(InGraph.GetDefaultOperatorFactory())
+				{
+				}
+
+				virtual ~FGraphOperatorFactoryAdapter() = default;
+
+				virtual TUniquePtr<IOperator> CreateOperator(const FBuildOperatorParams& InParams, FBuildResults& OutResults) override
+				{
+					FBuildOperatorParams ForwardParams
+					{
+						*Graph,  // Point to correct INode instance
+						InParams.OperatorSettings,
+						InParams.InputData,
+						InParams.Environment,
+						InParams.Builder
+					};
+					return GraphFactory->CreateOperator(ForwardParams, OutResults);
+				}
+			private:
+				const IGraph* Graph; // Only store pointer because owning node keeps wrapped IGraph alive.  
+				FOperatorFactorySharedRef GraphFactory;
+			};
+
+		public:
+			FGraphNode(const FNodeInitData& InNodeInitData, TSharedRef<const IGraph> InGraphToWrap)
+			: InstanceID(InNodeInitData.InstanceID)
+			, Factory(MakeShared<FGraphOperatorFactoryAdapter>(*InGraphToWrap))
+			, Graph(MoveTemp(InGraphToWrap))
+			{
+			}
+
+			virtual const FName& GetInstanceName() const override
+			{
+				// Use the instance name of underlying graph because it refers
+				// to the actual asset name.
+				return Graph->GetInstanceName();
+			}
+
+			virtual const FGuid& GetInstanceID() const override
+			{
+				return InstanceID;
+			}
+
+			virtual const FNodeClassMetadata& GetMetadata() const override
+			{
+				return Graph->GetMetadata();
+			}
+
+			virtual const FVertexInterface& GetVertexInterface() const override
+			{
+				return Graph->GetVertexInterface();
+			}
+
+			virtual bool SetVertexInterface(const FVertexInterface& InInterface) override
+			{
+				// Cannot set vertex interface because of const reference to IGraph
+				// Return true if the supplied interface is the same as the existing interface.
+				return Graph->GetVertexInterface() == InInterface;
+			}
+
+			virtual bool IsVertexInterfaceSupported(const FVertexInterface& InInterface) const override
+			{
+				return Graph->IsVertexInterfaceSupported(InInterface);
+			}
+
+			virtual FOperatorFactorySharedRef GetDefaultOperatorFactory() const override
+			{
+				return Factory;
+			}
+
+		private:
+
+			FGuid InstanceID;
+			TSharedRef<FGraphOperatorFactoryAdapter> Factory;
+			TSharedRef<const IGraph> Graph;
+		};
+
+
+		// FDocumentNodeRegistryEntry encapsulates a node registry entry for a FGraph
+		class FDocumentNodeRegistryEntry : public INodeRegistryEntry
+		{
+		public:
+			FDocumentNodeRegistryEntry(const FMetasoundFrontendGraphClass& InGraphClass, const TSet<FMetasoundFrontendVersion>& InInterfaces, FNodeClassInfo&& InNodeClassInfo, TSharedPtr<const IGraph> InGraph)
+			: FrontendClass(InGraphClass)
+			, Interfaces(InInterfaces)
+			, ClassInfo(MoveTemp(InNodeClassInfo))
+			, Graph(InGraph)
+			{
+				FrontendClass.Metadata.SetType(EMetasoundFrontendClassType::External);
+			}
+
+			FDocumentNodeRegistryEntry(const FDocumentNodeRegistryEntry&) = default;
+
+			virtual ~FDocumentNodeRegistryEntry() = default;
+
+			virtual const FNodeClassInfo& GetClassInfo() const override
+			{
+				return ClassInfo;
+			}
+
+			virtual TUniquePtr<INode> CreateNode(const FNodeInitData& InNodeInitData) const override
+			{
+				if (Graph.IsValid())
+				{
+					return MakeUnique<FGraphNode>(InNodeInitData, Graph.ToSharedRef());
+				}
+				else
+				{
+					UE_LOG(LogMetaSound, Error, TEXT("Cannot create MetaSound node from asset %s due to prior failure to build graph"), *ClassInfo.AssetPath.ToString());
+					return TUniquePtr<INode>();
+				}
+			}
+
+			virtual TUniquePtr<INode> CreateNode(FDefaultLiteralNodeConstructorParams&&) const override { return nullptr; }
+			virtual TUniquePtr<INode> CreateNode(FDefaultNamedVertexNodeConstructorParams&&) const override { return nullptr; }
+			virtual TUniquePtr<INode> CreateNode(FDefaultNamedVertexWithLiteralNodeConstructorParams&&) const override { return nullptr; }
+
+			virtual const FMetasoundFrontendClass& GetFrontendClass() const override
+			{
+				return FrontendClass;
+			}
+
+			virtual TUniquePtr<INodeRegistryEntry> Clone() const override
+			{
+				return MakeUnique<FDocumentNodeRegistryEntry>(*this);
+			}
+
+			virtual const TSet<FMetasoundFrontendVersion>* GetImplementedInterfaces() const override
+			{
+				return &Interfaces;
+			}
+
+			virtual bool IsNative() const override
+			{
+				return false;
+			}
+
+		private:
+
+			FMetasoundFrontendClass FrontendClass;
+			TSet<FMetasoundFrontendVersion> Interfaces;
+			FNodeClassInfo ClassInfo;
+			TSharedPtr<const IGraph> Graph;
+		};
+	} // namespace RegistryPrivate
+
+	void FRegistryContainerImpl::BuildAndRegisterGraphFromDocument(TScriptInterface<IMetaSoundDocumentInterface> DocumentInterface, const FProxyDataCache& InProxyDataCache, FNodeClassInfo&& InNodeClassInfo)
+	{
+		using namespace RegistryPrivate;
+
+		METASOUND_TRACE_CPUPROFILER_EVENT_SCOPE(Metasound::FRegistryContainerImpl::BuildAndRegisterGraphFromDocument);
+		METASOUND_TRACE_CPUPROFILER_EVENT_SCOPE_TEXT(*FString::Printf(TEXT("Metasound::FRegistryContainerImpl::BuildAndRegisterGraphFromDocument asset %s"), *InNodeClassInfo.AssetPath.ToString()));
+
+		const FMetasoundFrontendDocument& Document = DocumentInterface->GetConstDocument();
+		TUniquePtr<FFrontendGraph> FrontendGraph = FFrontendGraphBuilder::CreateGraph(Document, InProxyDataCache, InNodeClassInfo.AssetPath.ToString());
+		if (!FrontendGraph.IsValid())
+		{
+			UE_LOG(LogMetaSound, Error, TEXT("Failed to build MetaSound graph in asset '%s'"), *InNodeClassInfo.AssetPath.ToString());
+		}
+
+		TSharedPtr<const FGraph> GraphToRegister = MakeShareable<const FGraph>(FrontendGraph.Release());
+
+		TUniquePtr<INodeRegistryEntry> RegistryEntry = MakeUnique<FDocumentNodeRegistryEntry>(Document.RootGraph, Document.Interfaces, MoveTemp(InNodeClassInfo), GraphToRegister);
+
+		const FNodeRegistryKey RegistryKey = RegisterNodeInternal(MoveTemp(RegistryEntry));
+		RegisterGraphInternal(RegistryKey, DocumentInterface->GetAssetPathChecked(), GraphToRegister);
+	}
+
+	FRegistryContainerImpl* FRegistryContainerImpl::LazySingleton = nullptr;
+
+	FRegistryContainerImpl& FRegistryContainerImpl::Get()
+	{
+		if (!LazySingleton)
+		{
+			LazySingleton = new Metasound::Frontend::FRegistryContainerImpl();
+		}
+
+		return *LazySingleton;
+	}
+
+	void FRegistryContainerImpl::Shutdown()
+	{
+		if (nullptr != LazySingleton)
+		{
+			delete LazySingleton;
+			LazySingleton = nullptr;
+		}
+	}
+
+	FRegistryContainerImpl::FRegistryContainerImpl()
+	: TransactionBuffer(MakeShared<FNodeRegistryTransactionBuffer>())
+	, AsyncRegistrationPipe( UE_SOURCE_LOCATION )
+	{
+	}
+
+	void FRegistryContainerImpl::RegisterPendingNodes()
+	{
+		METASOUND_LLM_SCOPE;
+		METASOUND_TRACE_CPUPROFILER_EVENT_SCOPE(metasound::FRegistryContainerImpl::RegisterPendingNodes);
+		{
+			FScopeLock ScopeLock(&LazyInitCommandCritSection);
+
+			for (TUniqueFunction<void()>& Command : LazyInitCommands)
+			{
+				Command();
+			}
+
+			LazyInitCommands.Empty();
+		}
+
+		// Prime search engine after bulk registration.
+		ISearchEngine::Get().Prime();
+	}
+
+	bool FRegistryContainerImpl::EnqueueInitCommand(TUniqueFunction<void()>&& InFunc)
+	{
+
+		FScopeLock ScopeLock(&LazyInitCommandCritSection);
+		if (LazyInitCommands.Num() >= MaxNumNodesAndDatatypesToInitialize)
+		{
+			UE_LOG(LogMetaSound, Warning, TEXT("Registering more that %d nodes and datatypes for metasounds! Consider increasing MetasoundFrontendRegistryContainer::MaxNumNodesAndDatatypesToInitialize."), MaxNumNodesAndDatatypesToInitialize);
+		}
+
+		LazyInitCommands.Add(MoveTemp(InFunc));
+		return true;
+	}
+
+	void FRegistryContainerImpl::SetObjectReferencer(TUniquePtr<IObjectReferencer> InReferencer)
+	{
+		FScopeLock LockActiveReg(&ActiveRegistrationTasksCriticalSection);
+		checkf(ActiveRegistrationTasks.IsEmpty(), TEXT("Object Referencer cannot be set while registry is actively being manipulated"));
+		ObjectReferencer = MoveTemp(InReferencer);
+	}
+
+	TUniquePtr<Metasound::INode> FRegistryContainerImpl::CreateNode(const FNodeRegistryKey& InKey, const Metasound::FNodeInitData& InInitData) const
+	{
+		TUniquePtr<INode> Node;
+		auto CreateNodeLambda = [&Node, &InInitData](const INodeRegistryEntry& Entry)
+		{ 
+			Node = Entry.CreateNode(InInitData); 
+		};
+
+		if (!AccessNodeEntryThreadSafe(InKey, CreateNodeLambda))
+		{
+			// Creation of external nodes can rely on assets being unavailable due to errors in loading order, asset(s)
+			// missing, etc. 
+			UE_LOG(LogMetaSound, Error, TEXT("Could not find node [RegistryKey:%s]"), *InKey.ToString());
+		}
+
+		return MoveTemp(Node);
+	}
+
+	TUniquePtr<Metasound::INode> FRegistryContainerImpl::CreateNode(const FNodeRegistryKey& InKey, FDefaultLiteralNodeConstructorParams&& InParams) const
+	{
+		TUniquePtr<INode> Node;
+		auto CreateNodeLambda = [&Node, &InParams](const INodeRegistryEntry& Entry)
+		{ 
+			Node = Entry.CreateNode(MoveTemp(InParams)); 
+		};
+
+		if (!AccessNodeEntryThreadSafe(InKey, CreateNodeLambda))
+		{
+			UE_LOG(LogMetaSound, Error, TEXT("Could not find node [RegistryKey:%s]"), *InKey.ToString());
+		}
+
+		return MoveTemp(Node);
+	}
+
+	TUniquePtr<Metasound::INode> FRegistryContainerImpl::CreateNode(const FNodeRegistryKey& InKey, FDefaultNamedVertexNodeConstructorParams&& InParams) const
+	{
+		TUniquePtr<INode> Node;
+		auto CreateNodeLambda = [&Node, &InParams](const INodeRegistryEntry& Entry) 
+		{ 
+			Node = Entry.CreateNode(MoveTemp(InParams)); 
+		};
+
+		if (!AccessNodeEntryThreadSafe(InKey, CreateNodeLambda))
+		{
+			UE_LOG(LogMetaSound, Error, TEXT("Could not find node [RegistryKey:%s]"), *InKey.ToString());
+		}
+
+		return MoveTemp(Node);
+	}
+
+	TUniquePtr<Metasound::INode> FRegistryContainerImpl::CreateNode(const FNodeRegistryKey& InKey, FDefaultNamedVertexWithLiteralNodeConstructorParams&& InParams) const
+	{
+		TUniquePtr<INode> Node;
+		auto CreateNodeLambda = [&Node, &InParams](const INodeRegistryEntry& Entry)
+		{ 
+			Node = Entry.CreateNode(MoveTemp(InParams));
+		};
+
+		if (!AccessNodeEntryThreadSafe(InKey, CreateNodeLambda))
+		{
+			UE_LOG(LogMetaSound, Error, TEXT("Could not find node [RegistryKey:%s]"), *InKey.ToString());
+		}
+
+		return MoveTemp(Node);
+	}
+
+	TArray<::Metasound::Frontend::FConverterNodeInfo> FRegistryContainerImpl::GetPossibleConverterNodes(const FName& FromDataType, const FName& ToDataType)
+	{
+		FConverterNodeRegistryKey InKey = { FromDataType, ToDataType };
+		if (!ConverterNodeRegistry.Contains(InKey))
+		{
+			return TArray<FConverterNodeInfo>();
+		}
+		else
+		{
+			return ConverterNodeRegistry[InKey].PotentialConverterNodes;
+		}
+	}
+
+	TUniquePtr<FNodeRegistryTransactionStream> FRegistryContainerImpl::CreateTransactionStream()
+	{
+		return MakeUnique<FNodeRegistryTransactionStream>(TransactionBuffer);
+	}
+
+	FNodeRegistryKey FRegistryContainerImpl::RegisterGraph(const TScriptInterface<IMetaSoundDocumentInterface>& InDocumentInterface, bool bAsync)
+	{
+		using namespace UE;
+
+		METASOUND_TRACE_CPUPROFILER_EVENT_SCOPE(Metasound::FRegistryContainerImpl::RegisterGraph);
+
+		check(InDocumentInterface);
+		check(IsInGameThread());
+
+		UObject* OwningObject = InDocumentInterface.GetObject();
+		check(OwningObject);
+		const FTopLevelAssetPath AssetPath = InDocumentInterface->GetAssetPathChecked();
+
+		const FMetasoundFrontendDocument& Document = InDocumentInterface->GetConstDocument();
+		const FNodeRegistryKey RegistryKey(Document.RootGraph);
+
+		if (!RegistryKey.IsValid())
+		{
+			// Do not attempt to build and register a MetaSound with an invalid registry key
+			UE_LOG(LogMetaSound, Warning, TEXT("Registry key is invalid when attemping to register graph for asset %s"), *AssetPath.ToString());
+			return RegistryKey;
+		}
+
+		METASOUND_TRACE_CPUPROFILER_EVENT_SCOPE_TEXT(*FString::Printf(TEXT("FRegistryContainerImpl::RegisterGraph key:%s, asset %s"), *RegistryKey.ToString(), *AssetPath.ToString()));
+
+		FNodeClassInfo NodeClassInfo(Document.RootGraph, AssetPath);
+
+		// Proxies are created synchronously to avoid creating proxies in async tasks. Proxies
+		// are created from UObjects which need to be protected from GC and non-GT access.
+		FProxyDataCache ProxyDataCache;
+		ProxyDataCache.CreateAndCacheProxies(Document);
+
+		// Store update to newly registered node in history so nodes
+		// can be queried by transaction ID
+		{
+			FNodeRegistryTransaction::FTimeType Timestamp = FPlatformTime::Cycles64();
+			TransactionBuffer->AddTransaction(FNodeRegistryTransaction(FNodeRegistryTransaction::ETransactionType::NodeRegistration, NodeClassInfo, Timestamp));
+		}
+
+		if (bAsync)
+		{
+			// Wait for any async tasks that are in flight which correspond to the same graph
+			WaitForAsyncRegistrationInternal(RegistryKey, &AssetPath);
+
+			Tasks::FTask BuildAndRegisterTask = AsyncRegistrationPipe.Launch(
+				UE_SOURCE_LOCATION,
+				[RegistryKey, ClassInfo = MoveTemp(NodeClassInfo), AssetPath, DocumentInterface = InDocumentInterface, ProxyDataCache = MoveTemp(ProxyDataCache)]() mutable
+				{
+					FRegistryContainerImpl& Registry = FRegistryContainerImpl::Get();
+					// Unregister the graph before re-registering
+					if (Registry.IsGraphRegistered(RegistryKey, AssetPath))
+					{
+						Registry.UnregisterGraphInternal(RegistryKey, AssetPath);
+					}
+
+					Registry.BuildAndRegisterGraphFromDocument(DocumentInterface, ProxyDataCache, MoveTemp(ClassInfo));
+					Registry.RemoveRegistrationTask(RegistryKey, FNodeRegistryTransaction::ETransactionType::NodeRegistration, DocumentInterface);
+				}
+			);
+
+			AddRegistrationTask(RegistryKey, *OwningObject, FActiveRegistrationTaskInfo
+			{
+				FNodeRegistryTransaction::ETransactionType::NodeRegistration,
+				BuildAndRegisterTask,
+				AssetPath
+			});
+		}
+		else
+		{
+			if (IsGraphRegistered(RegistryKey, AssetPath))
+			{
+				UnregisterGraphInternal(RegistryKey, AssetPath);
+			}
+
+			// Build and register graph synchronously
+			BuildAndRegisterGraphFromDocument(InDocumentInterface, ProxyDataCache, MoveTemp(NodeClassInfo));
+		}
+
+		return RegistryKey;
+	}
+
+	void FRegistryContainerImpl::AddRegistrationTask(const FNodeRegistryKey& InKey, UObject& InObject, FActiveRegistrationTaskInfo&& TaskInfo)
+	{
+		FScopeLock LockActiveReg(&ActiveRegistrationTasksCriticalSection);
+		if (ObjectReferencer)
+		{
+			ObjectReferencer->AddObject(&InObject);
+		}
+
+		ActiveRegistrationTasks.FindOrAdd(InKey).Add(MoveTemp(TaskInfo));
+	}
+
+	void FRegistryContainerImpl::RemoveRegistrationTask(const FNodeRegistryKey& InKey, FNodeRegistryTransaction::ETransactionType TransactionType, const TScriptInterface<IMetaSoundDocumentInterface>& DocumentInterface)
+	{
+		FScopeLock LockActiveReg(&ActiveRegistrationTasksCriticalSection);
+
+		if (TArray<FActiveRegistrationTaskInfo>* TaskInfos = ActiveRegistrationTasks.Find(InKey))
+		{
+			const FTopLevelAssetPath Path = DocumentInterface->GetAssetPathChecked();
+			constexpr bool bAllowShrinking = false;
+			const int32 NumRemoved = TaskInfos->RemoveAllSwap([&Path, &TransactionType](const FActiveRegistrationTaskInfo& Info)
+			{
+				return Info.AssetPath == Path && Info.TransactionType == TransactionType;
+			}, bAllowShrinking);
+
+			if (NumRemoved != 1)
+			{
+				UE_LOG(LogMetaSound, Warning,
+					TEXT("Multiple active %s tasks for the same asset may result in incorrect MetaSound graphs being instantiated during runtime. "
+						"Registry Key: (%s) "
+						"Asset Path: (%s)"),
+					*FNodeRegistryTransaction::LexToString(TransactionType),
+					*InKey.ToString(),
+					*Path.ToString());
+			}
+
+			if (TaskInfos->IsEmpty())
+			{
+				ActiveRegistrationTasks.Remove(InKey);
+			}
+		}
+
+		if (ObjectReferencer)
+		{
+			if (UObject* Object = DocumentInterface.GetObject())
+			{
+				ObjectReferencer->RemoveObject(Object);
+			}
+		}
+	}
+
+	void FRegistryContainerImpl::RegisterGraphInternal(const FNodeRegistryKey& InKey, const FTopLevelAssetPath& InAssetPath, TSharedPtr<const FGraph> InGraph)
+	{
+		FScopeLock Lock(&RegistryMapsCriticalSection);
+
+		FGraphRegistryKey GraphKey(InKey, InAssetPath);
+		if (RegisteredGraphs.Contains(GraphKey))
+		{
+			UE_LOG(LogMetaSound, Warning, TEXT("Graph is already registered with the same registry key ('%s, %s'). The existing registered graph will be replaced with the new graph."),  *InKey.ToString(), *InAssetPath.ToString());
+		}
+
+		RegisteredGraphs.Add(GraphKey, InGraph);
+	}
+
+	bool FRegistryContainerImpl::UnregisterGraphInternal(const FNodeRegistryKey& InKey, const FTopLevelAssetPath& InAssetPath)
+	{
+		METASOUND_TRACE_CPUPROFILER_EVENT_SCOPE_TEXT(*FString::Printf(TEXT("FRegistryContainerImpl::UnregisterGraphInternal key:%s, asset %s"), *InKey.ToString(), *InAssetPath.ToString()));
+
+		FScopeLock Lock(&RegistryMapsCriticalSection);
+		{
+			const int32 NumRemoved = RegisteredGraphs.Remove(FRegistryContainerImpl::FGraphRegistryKey(InKey, InAssetPath));
+			if (NumRemoved != 0)
+			{
+				UnregisterNodeInternal(InKey);
+				return true;
+			}
+		}
+
+		return false;
+	}
+
+	bool FRegistryContainerImpl::UnregisterGraph(const TScriptInterface<IMetaSoundDocumentInterface>& InDocumentInterface, bool bAsync)
+	{
+		using namespace UE;
+
+		METASOUND_TRACE_CPUPROFILER_EVENT_SCOPE(Metasound::FRegistryContainerImpl::UnregisterGraph);
+
+		UObject* OwningObject = InDocumentInterface.GetObject();
+		check(OwningObject);
+		check(InDocumentInterface);
+		check(IsInGameThread());
+
+		const FMetasoundFrontendDocument& Document = InDocumentInterface->GetConstDocument();
+		const FNodeRegistryKey RegistryKey(Document.RootGraph);
+		METASOUND_TRACE_CPUPROFILER_EVENT_SCOPE_TEXT(*FString::Printf(TEXT("FRegistryContainerImpl::UnregisterGraph key %s object %s"), *RegistryKey.ToString(), *OwningObject->GetPathName()));
+
+		// Do not attempt to unregister a MetaSound with an invalid registry key
+		if (!RegistryKey.IsValid())
+		{
+			UE_LOG(LogMetaSound, Warning, TEXT("Registry key is invalid when attempting to unregister graph for object %s"), *OwningObject->GetFullName());
+			return false;
+		}
+
+		// This is a hack to avoid requiring the asset path to be passed while unregistering.
+		// The asset path may be invalid by this point if the object being unregistered is being GC'ed.
+		// FNodeClassInfo needs to be deprecated in favor of more precise types as a key, editor data, etc.
+		// Its currently kind of a dumping ground.
+		FNodeClassInfo NodeClassInfo(Document.RootGraph.Metadata);
+		NodeClassInfo.Type = EMetasoundFrontendClassType::External;
+
+		// Store update to unregistered node in history so nodes can be queried by transaction ID
+		{
+			FNodeRegistryTransaction::FTimeType Timestamp = FPlatformTime::Cycles64();
+			TransactionBuffer->AddTransaction(FNodeRegistryTransaction(FNodeRegistryTransaction::ETransactionType::NodeUnregistration, NodeClassInfo, Timestamp));
+		}
+
+		if (bAsync)
+		{
+			Tasks::FTask UnregisterTask = AsyncRegistrationPipe.Launch(
+				UE_SOURCE_LOCATION,
+				[RegistryKey, AssetPath = InDocumentInterface->GetAssetPathChecked(), DocumentInterface = InDocumentInterface]()
+				{
+					FRegistryContainerImpl& Registry = FRegistryContainerImpl::Get();
+					Registry.UnregisterGraphInternal(RegistryKey, AssetPath);
+					Registry.RemoveRegistrationTask(RegistryKey, FNodeRegistryTransaction::ETransactionType::NodeUnregistration, DocumentInterface);
+				});
+
+			AddRegistrationTask(RegistryKey, *OwningObject, FActiveRegistrationTaskInfo
+			{
+				FNodeRegistryTransaction::ETransactionType::NodeUnregistration,
+				UnregisterTask,
+				InDocumentInterface->GetAssetPathChecked()
+			});
+		}
+		else
+		{
+			UnregisterGraphInternal(RegistryKey, InDocumentInterface->GetAssetPathChecked());
+		}
+
+		return true;
+	}
+
+	TSharedPtr<const Metasound::FGraph> FRegistryContainerImpl::GetGraph(const FNodeRegistryKey& InNodeRegistryKey, const FTopLevelAssetPath& InAssetPath) const
+	{
+		WaitForAsyncGraphRegistration(InNodeRegistryKey, InAssetPath);
+
+		TSharedPtr<const FGraph> Graph;
+		{
+			FScopeLock Lock(&RegistryMapsCriticalSection);
+			if (const TSharedPtr<const FGraph>* RegisteredGraph = RegisteredGraphs.Find(FGraphRegistryKey { InNodeRegistryKey, InAssetPath }))
+			{
+				Graph = *RegisteredGraph;
+			}
+		}
+
+		if (!Graph)
+		{
+			UE_LOG(LogMetaSound, Error, TEXT("Could not find graph with registry key (%s) and asset (%s)."),  *InNodeRegistryKey.ToString(), *InAssetPath.ToString());
+		}
+
+		return Graph;
+	}
+
+	FNodeRegistryKey FRegistryContainerImpl::RegisterNodeInternal(TUniquePtr<INodeRegistryEntry>&& InEntry)
+	{
+		METASOUND_LLM_SCOPE;
+
+		FNodeRegistryKey Key;
+
+		if (InEntry.IsValid())
+		{
+			TSharedRef<INodeRegistryEntry, ESPMode::ThreadSafe> Entry(InEntry.Release());
+
+			Key = FNodeRegistryKey(Entry->GetClassInfo());
+			METASOUND_TRACE_CPUPROFILER_EVENT_SCOPE_TEXT(*FString::Printf(TEXT("FRegistryContainerImpl::RegisterNodeInternal with key %s"), *Key.ToString()))
+			{
+				FScopeLock Lock(&RegistryMapsCriticalSection);
+
+				// check to see if an identical node was already registered, and log
+				if (const TSharedRef<INodeRegistryEntry>* ExistingEntry = RegisteredNodes.Find(Key))
+				{
+					const FNodeClassInfo& ClassInfo = (*ExistingEntry)->GetClassInfo();
+					UE_LOG(LogMetaSound, Error,
+						TEXT("Node with registry key '%s' already registered by asset '%s' encountered while registering node with asset %s. MetaSounds which depend on these assets may utilize incorrect asset."),
+						*Key.ToString(),
+						*ClassInfo.AssetPath.ToString(),
+						*Entry->GetClassInfo().AssetPath.ToString());
+				}
+
+				// Store registry elements in map so nodes can be queried using registry key.
+				RegisteredNodes.Add(Key, Entry);
+			}
+		}
+
+		return Key;
+	}
+
+	FNodeRegistryKey FRegistryContainerImpl::RegisterNode(TUniquePtr<INodeRegistryEntry>&& InEntry)
+	{
+		const FNodeClassInfo ClassInfo = InEntry->GetClassInfo();
+		const FNodeRegistryKey Key = RegisterNodeInternal(MoveTemp(InEntry));
+
+		if (Key.IsValid())
+		{
+			// Store update to newly registered node in history so nodes
+			// can be queried by transaction ID
+			const FNodeRegistryTransaction::FTimeType Timestamp = FPlatformTime::Cycles64();
+			TransactionBuffer->AddTransaction(FNodeRegistryTransaction(FNodeRegistryTransaction::ETransactionType::NodeRegistration, ClassInfo, Timestamp));
+		}
+
+		return Key;
+	}
+
+	FNodeRegistryKey FRegistryContainerImpl::RegisterNodeTemplate(TUniquePtr<INodeRegistryTemplateEntry>&& InEntry)
+	{
+		METASOUND_LLM_SCOPE;
+
+		FNodeRegistryKey Key;
+
+		if (InEntry.IsValid())
+		{
+			TSharedRef<INodeRegistryTemplateEntry, ESPMode::ThreadSafe> Entry(InEntry.Release());
+
+			FNodeRegistryTransaction::FTimeType Timestamp = FPlatformTime::Cycles64();
+
+			Key = FNodeRegistryKey(Entry->GetClassInfo());
+
+			{
+				FScopeLock Lock(&RegistryMapsCriticalSection);
+				// check to see if an identical node was already registered, and log
+				ensureAlwaysMsgf(
+					!RegisteredNodeTemplates.Contains(Key),
+					TEXT("Node template with registry key '%s' already registered. The previously registered node will be overwritten."),
+					*Key.ToString());
+
+				// Store registry elements in map so nodes can be queried using registry key.
+				RegisteredNodeTemplates.Add(Key, Entry);
+			}
+
+			// Store update to newly registered node in history so nodes
+			// can be queried by transaction ID
+
+			TransactionBuffer->AddTransaction(FNodeRegistryTransaction(FNodeRegistryTransaction::ETransactionType::NodeRegistration, Entry->GetClassInfo(), Timestamp));
+		}
+
+		return Key;
+	}
+
+	bool FRegistryContainerImpl::UnregisterNodeInternal(const FNodeRegistryKey& InKey, FNodeClassInfo* OutClassInfo)
+	{
+		METASOUND_LLM_SCOPE;
+
+		if (InKey.IsValid())
+		{
+			METASOUND_TRACE_CPUPROFILER_EVENT_SCOPE_TEXT(*FString::Printf(TEXT("FRegistryContainerImpl::UnregisterNodeInternal key %s"), *InKey.ToString()))
+
+			FScopeLock Lock(&RegistryMapsCriticalSection);
+			if (const TSharedRef<INodeRegistryEntry, ESPMode::ThreadSafe>* EntryPtr = RegisteredNodes.Find(InKey))
+			{
+				const TSharedRef<INodeRegistryEntry>& Entry = *EntryPtr;
+				if (OutClassInfo)
+				{
+					*OutClassInfo = Entry->GetClassInfo();
+				}
+				const uint32 NumRemoved = RegisteredNodes.RemoveSingle(InKey, Entry);
+				if (ensure(NumRemoved == 1))
+				{
+					return true;
+				}
+			}
+		}
+
+		if (OutClassInfo)
+		{
+			*OutClassInfo = { };
+		}
+		return false;
+	}
+
+	bool FRegistryContainerImpl::UnregisterNode(const FNodeRegistryKey& InKey)
+	{
+		FNodeClassInfo ClassInfo;
+		if (UnregisterNodeInternal(InKey, &ClassInfo))
+		{
+			const FNodeRegistryTransaction::FTimeType Timestamp = FPlatformTime::Cycles64();
+			TransactionBuffer->AddTransaction(FNodeRegistryTransaction(FNodeRegistryTransaction::ETransactionType::NodeUnregistration, ClassInfo, Timestamp));
+
+			return true;
+		}
+
+		return false;
+	}
+
+	bool FRegistryContainerImpl::UnregisterNodeTemplate(const FNodeRegistryKey& InKey)
+	{
+		METASOUND_LLM_SCOPE;
+
+		if (InKey.IsValid())
+		{
+			if (const INodeRegistryTemplateEntry* Entry = FindNodeTemplateEntry(InKey))
+			{
+				FNodeRegistryTransaction::FTimeType Timestamp = FPlatformTime::Cycles64();
+
+				TransactionBuffer->AddTransaction(FNodeRegistryTransaction(FNodeRegistryTransaction::ETransactionType::NodeUnregistration, Entry->GetClassInfo(), Timestamp));
+
+				{
+					FScopeLock Lock(&RegistryMapsCriticalSection);
+					RegisteredNodeTemplates.Remove(InKey);
+				}
+				return true;
+			}
+		}
+
+		return false;
+	}
+
+	bool FRegistryContainerImpl::RegisterConversionNode(const FConverterNodeRegistryKey& InNodeKey, const FConverterNodeInfo& InNodeInfo)
+	{
+		if (!ConverterNodeRegistry.Contains(InNodeKey))
+		{
+			ConverterNodeRegistry.Add(InNodeKey);
+		}
+
+		FConverterNodeRegistryValue& ConverterNodeList = ConverterNodeRegistry[InNodeKey];
+
+		if (ensureAlways(!ConverterNodeList.PotentialConverterNodes.Contains(InNodeInfo)))
+		{
+			ConverterNodeList.PotentialConverterNodes.Add(InNodeInfo);
+			return true;
+		}
+		else
+		{
+			// If we hit this, someone attempted to add the same converter node to our list multiple times.
+			return false;
+		}
+	}
+
+	bool FRegistryContainerImpl::IsNodeRegistered(const FNodeRegistryKey& InKey) const
+	{
+		auto IsNodeRegisteredInternal = [this, &InKey]() -> bool
+		{
+			FScopeLock Lock(&RegistryMapsCriticalSection);
+			return RegisteredNodes.Contains(InKey) || RegisteredNodeTemplates.Contains(InKey);
+		};
+
+		if (IsNodeRegisteredInternal())
+		{
+			return true;
+		}
+		else
+		{
+			WaitForAsyncRegistrationInternal(InKey, nullptr /* InAssetPath */);
+			return IsNodeRegisteredInternal();
+		}
+	}
+
+	bool FRegistryContainerImpl::IsGraphRegistered(const FNodeRegistryKey& InKey, const FTopLevelAssetPath& InAssetPath) const
+	{
+		WaitForAsyncGraphRegistration(InKey, InAssetPath);
+
+		{
+			FScopeLock Lock(&RegistryMapsCriticalSection);
+			FGraphRegistryKey GraphRegistryKey { InKey, InAssetPath };
+			return RegisteredGraphs.Contains(GraphRegistryKey);
+		}
+	}
+
+	bool FRegistryContainerImpl::IsNodeNative(const FNodeRegistryKey& InKey) const
+	{
+		bool bIsNative = false;
+		auto SetIsNative = [&bIsNative](const INodeRegistryEntry& Entry) 
+		{ 
+			bIsNative = Entry.IsNative();
+		};
+
+		if (AccessNodeEntryThreadSafe(InKey, SetIsNative))
+		{
+			return bIsNative;
+		}
+
+		if (const INodeRegistryTemplateEntry* TemplateEntry = FindNodeTemplateEntry(InKey))
+		{
+			return true;
+		}
+		return false;
+	}
+
+	bool FRegistryContainerImpl::FindFrontendClassFromRegistered(const FNodeRegistryKey& InKey, FMetasoundFrontendClass& OutClass)
+	{
+		auto SetFrontendClass = [&OutClass](const INodeRegistryEntry& Entry)
+		{
+			OutClass = Entry.GetFrontendClass();
+		};
+
+		if (AccessNodeEntryThreadSafe(InKey, SetFrontendClass))
+		{
+			return true;
+		}
+
+		if (const INodeRegistryTemplateEntry* Entry = FindNodeTemplateEntry(InKey))
+		{
+			OutClass = Entry->GetFrontendClass();
+			return true;
+		}
+
+		return false;
+	}
+
+	const TSet<FMetasoundFrontendVersion>* FRegistryContainerImpl::FindImplementedInterfacesFromRegistered(const Metasound::Frontend::FNodeRegistryKey& InKey) const
+	{
+		static bool bHasWarningBeenIssued = false;
+		if (!bHasWarningBeenIssued)
+		{
+			// This function is known to be thread-unsafe and should not longer be used.
+			// This implementation exists to support deprecated usage. 
+			UE_LOG(LogMetaSound, Warning, TEXT("Accessing non-thread-safe implementation of FindImplementedInterfacesFromRegistered(...) is known to cause crashes. Please update your code to use the non-deprecated version of this function with the same name"));
+			bHasWarningBeenIssued = true;
+		}
+
+		const TSet<FMetasoundFrontendVersion>* Interfaces = nullptr;
+		AccessNodeEntryThreadSafe(InKey, 
+			[&Interfaces](const INodeRegistryEntry& Entry)
+			{
+				Interfaces = Entry.GetImplementedInterfaces();
+			}
+		);
+
+		return Interfaces;
+	}
+
+	bool FRegistryContainerImpl::FindImplementedInterfacesFromRegistered(const Metasound::Frontend::FNodeRegistryKey& InKey, TSet<FMetasoundFrontendVersion>& OutInterfaceVersions) const 
+	{
+		bool bDidCopy = false;
+
+		auto CopyImplementedInterfaces = [&OutInterfaceVersions, &bDidCopy](const INodeRegistryEntry& Entry)
+		{
+			if (const TSet<FMetasoundFrontendVersion>* Interfaces = Entry.GetImplementedInterfaces())
+			{
+				OutInterfaceVersions = *Interfaces;
+				bDidCopy = true;
+			}
+		};
+
+		AccessNodeEntryThreadSafe(InKey, CopyImplementedInterfaces);
+
+		return bDidCopy;
+	}
+
+	bool FRegistryContainerImpl::FindNodeClassInfoFromRegistered(const Metasound::Frontend::FNodeRegistryKey& InKey, FNodeClassInfo& OutInfo)
+	{
+		auto CopyClassInfo = [&OutInfo](const INodeRegistryEntry& Entry)
+		{
+			OutInfo = Entry.GetClassInfo();
+		};
+		if (AccessNodeEntryThreadSafe(InKey, CopyClassInfo))
+		{
+			return true;
+		}
+
+		if (const INodeRegistryTemplateEntry* Entry = FindNodeTemplateEntry(InKey))
+		{
+			OutInfo = Entry->GetClassInfo();
+			return true;
+		}
+
+		return false;
+	}
+
+	bool FRegistryContainerImpl::FindInputNodeRegistryKeyForDataType(const FName& InDataTypeName, const EMetasoundFrontendVertexAccessType InAccessType, FNodeRegistryKey& OutKey)
+	{
+		FMetasoundFrontendClass Class;
+		switch (InAccessType)
+		{
+			case EMetasoundFrontendVertexAccessType::Reference:
+			{
+				if (IDataTypeRegistry::Get().GetFrontendInputClass(InDataTypeName, Class))
+				{
+					OutKey = FNodeRegistryKey(Class.Metadata);
+					return true;
+				}
+			}
+			break;
+
+			case EMetasoundFrontendVertexAccessType::Value:
+			{
+				if (IDataTypeRegistry::Get().GetFrontendConstructorInputClass(InDataTypeName, Class))
+				{
+					OutKey = FNodeRegistryKey(Class.Metadata);
+					return true;
+				}
+			}
+			break;
+
+			default:
+			case EMetasoundFrontendVertexAccessType::Unset:
+			{
+				return false;
+			}
+			break;
+		}
+
+		return false;
+	}
+
+	bool FRegistryContainerImpl::FindVariableNodeRegistryKeyForDataType(const FName& InDataTypeName, FNodeRegistryKey& OutKey)
+	{
+		FMetasoundFrontendClass Class;
+		if (IDataTypeRegistry::Get().GetFrontendLiteralClass(InDataTypeName, Class))
+		{
+			OutKey = FNodeRegistryKey(Class.Metadata);
+			return true;
+		}
+		return false;
+	}
+
+	bool FRegistryContainerImpl::FindOutputNodeRegistryKeyForDataType(const FName& InDataTypeName, const EMetasoundFrontendVertexAccessType InAccessType, FNodeRegistryKey& OutKey)
+	{
+		FMetasoundFrontendClass Class;
+		switch (InAccessType)
+		{
+			case EMetasoundFrontendVertexAccessType::Reference:
+			{
+				if (IDataTypeRegistry::Get().GetFrontendOutputClass(InDataTypeName, Class))
+				{
+					OutKey = FNodeRegistryKey(Class.Metadata);
+					return true;
+				}
+			}
+			break;
+
+			case EMetasoundFrontendVertexAccessType::Value:
+			{
+				if (IDataTypeRegistry::Get().GetFrontendConstructorOutputClass(InDataTypeName, Class))
+				{
+					OutKey = FNodeRegistryKey(Class.Metadata);
+					return true;
+				}
+			}
+			break;
+		}
+
+		return false;
+	}
+
+	void FRegistryContainerImpl::IterateRegistry(Metasound::FIterateMetasoundFrontendClassFunction InIterFunc, EMetasoundFrontendClassType InClassType) const
+	{
+		UE_LOG(LogMetaSound, Warning, TEXT("Calling FMetasoundRegistryContainer::IterateRegistry(...) is not threadsafe. Please use Metasound::Frontend::ISearchEngine instead"));
+		auto WrappedFunc = [&](const TPair<FNodeRegistryKey, TSharedPtr<INodeRegistryEntry, ESPMode::ThreadSafe>>& Pair)
+		{
+			InIterFunc(Pair.Value->GetFrontendClass());
+		};
+
+		if (EMetasoundFrontendClassType::Invalid == InClassType)
+		{
+			// Iterate through all classes. 
+			Algo::ForEach(RegisteredNodes, WrappedFunc);
+		}
+		else
+		{
+			// Only call function on classes of certain type.
+			auto IsMatchingClassType = [&](const TPair<FNodeRegistryKey, TSharedPtr<INodeRegistryEntry, ESPMode::ThreadSafe>>& Pair)
+			{
+				return Pair.Value->GetClassInfo().Type == InClassType;
+			};
+			Algo::ForEachIf(RegisteredNodes, IsMatchingClassType, WrappedFunc);
+		}
+	}
+
+	bool FRegistryContainerImpl::AccessNodeEntryThreadSafe(const FNodeRegistryKey& InKey, TFunctionRef<void(const INodeRegistryEntry&)> InFunc) const
+	{
+		auto TryAccessNodeEntry = [this, &InKey, &InFunc]() -> bool
+		{
+			FScopeLock Lock(&RegistryMapsCriticalSection);
+			if (const TSharedRef<INodeRegistryEntry, ESPMode::ThreadSafe>* Entry = RegisteredNodes.Find(InKey))
+			{
+				InFunc(*(*Entry));
+				return true;
+			}
+			return false;
+		};
+
+		if (TryAccessNodeEntry())
+		{
+			return true;
+		}
+		else
+		{
+			// Wait for any async registration tasks related to the registry key. 
+			WaitForAsyncRegistrationInternal(InKey, nullptr /* InAssetPath */);
+			return TryAccessNodeEntry();
+		}
+	}
+
+	const INodeRegistryTemplateEntry* FRegistryContainerImpl::FindNodeTemplateEntry(const FNodeRegistryKey& InKey) const
+	{
+		FScopeLock Lock(&RegistryMapsCriticalSection);
+		if (const TSharedRef<INodeRegistryTemplateEntry, ESPMode::ThreadSafe>* Entry = RegisteredNodeTemplates.Find(InKey))
+		{
+			return &Entry->Get();
+		}
+
+		return nullptr;
+	}
+
+	void FRegistryContainerImpl::WaitForAsyncGraphRegistration(const FNodeRegistryKey& InRegistryKey, const FTopLevelAssetPath& InAssetPath) const
+	{
+		WaitForAsyncRegistrationInternal(InRegistryKey, &InAssetPath);
+	}
+
+	void FRegistryContainerImpl::WaitForAsyncRegistrationInternal(const FNodeRegistryKey& InRegistryKey, const FTopLevelAssetPath* InAssetPath) const
+	{
+		using namespace UE::Tasks;
+
+		if (AsyncRegistrationPipe.IsInContext())
+		{
+			// It is not safe to wait for an async registration task from within the async registration pipe because it will result in a deadlock. 
+			UE_LOG(LogMetaSound, Verbose, TEXT("Async registration pipe is already in context for registering key %s. Task will not be waited for."), *InRegistryKey.ToString());
+			return;
+		}
+
+		TArray<FTask> TasksToWaitFor;
+		{
+			FScopeLock Lock(&ActiveRegistrationTasksCriticalSection);
+			if (const TArray<FActiveRegistrationTaskInfo>* FoundTasks = ActiveRegistrationTasks.Find(InRegistryKey))
+			{
+				// Filter by asset path or ignore if not provided
+				Algo::TransformIf(*FoundTasks, TasksToWaitFor,
+					[&InAssetPath](const FActiveRegistrationTaskInfo& TaskInfo) { return !InAssetPath || TaskInfo.AssetPath == *InAssetPath; },
+					[](const FActiveRegistrationTaskInfo& TaskInfo) { return TaskInfo.Task; });
+			}
+		}
+
+		for (const FTask& Task : TasksToWaitFor)
+		{
+			METASOUND_TRACE_CPUPROFILER_EVENT_SCOPE(Metasound::FRegistryContainerImpl::WaitForRegistrationTaskToComplete);
+			if (Task.IsValid())
+			{
+				if (bBusyWaitOnAsyncRegistrationTasks)
+				{
+					Task.BusyWait();
+				}
+				else
+				{
+					Task.Wait();
+				}
+			}	
+		}
+	}
+} // namespace Metasound::Frontend
diff --git a/Engine/Plugins/Runtime/Metasound/Source/MetasoundFrontend/Private/MetasoundFrontendRegistryContainerImpl.h b/Engine/Plugins/Runtime/Metasound/Source/MetasoundFrontend/Private/MetasoundFrontendRegistryContainerImpl.h
new file mode 100644
--- /dev/null
+++ b/Engine/Plugins/Runtime/Metasound/Source/MetasoundFrontend/Private/MetasoundFrontendRegistryContainerImpl.h
@@ -0,0 +1,209 @@
+// Copyright Epic Games, Inc. All Rights Reserved.
+#pragma once
+
+#include "Containers/Array.h"
+#include "Containers/Map.h"
+#include "HAL/CriticalSection.h"
+#include "Tasks/Pipe.h"
+
+#include "MetasoundFrontendDataTypeRegistry.h"
+#include "MetasoundFrontendRegistryContainer.h"
+#include "MetasoundFrontendRegistryTransaction.h"
+
+struct FMetasoundFrontendDocument; 
+
+namespace Metasound
+{
+	class FProxyDataCache;
+	class FGraph;
+
+	namespace Frontend
+	{
+		struct FNodeClassInfo;
+
+		using FNodeRegistryTransactionBuffer = TTransactionBuffer<FNodeRegistryTransaction>;
+		using FNodeRegistryTransactionStream = TTransactionStream<FNodeRegistryTransaction>; 
+
+		/** INodeRegistryTemplateEntry declares the interface for a node registry entry.
+		 * Each node class in the registry must satisfy this interface.
+		 */
+		class INodeRegistryTemplateEntry
+		{
+		public:
+			virtual ~INodeRegistryTemplateEntry() = default;
+
+			/** Return FNodeClassInfo for the node class.
+			 *
+			 * Implementations of method should avoid any expensive operations
+			 * (e.g. loading from disk, allocating memory) as this method is called
+			 * frequently when querying nodes.
+			 */
+			virtual const FNodeClassInfo& GetClassInfo() const = 0;
+
+			/** Return a FMetasoundFrontendClass which describes the node. */
+			virtual const FMetasoundFrontendClass& GetFrontendClass() const = 0;
+		};
+
+		// Registry container private implementation.
+		class FRegistryContainerImpl
+			: public FMetasoundFrontendRegistryContainer
+		{
+
+		public:
+			using FConverterNodeRegistryKey = ::Metasound::Frontend::FConverterNodeRegistryKey;
+			using FConverterNodeRegistryValue = ::Metasound::Frontend::FConverterNodeRegistryValue;
+			using FConverterNodeInfo = ::Metasound::Frontend::FConverterNodeInfo;
+
+			using FNodeRegistryKey = Metasound::Frontend::FNodeRegistryKey;
+
+			using FDataTypeRegistryInfo = Metasound::Frontend::FDataTypeRegistryInfo;
+			using FNodeClassMetadata = Metasound::FNodeClassMetadata;
+			using IEnumDataTypeInterface = Metasound::Frontend::IEnumDataTypeInterface;
+
+			FRegistryContainerImpl();
+
+			FRegistryContainerImpl(const FRegistryContainerImpl&) = delete;
+			FRegistryContainerImpl& operator=(const FRegistryContainerImpl&) = delete;
+
+			static FRegistryContainerImpl& Get();
+			static void Shutdown();
+
+			virtual ~FRegistryContainerImpl() = default;
+
+			// Add a function to the init command array.
+			virtual bool EnqueueInitCommand(TUniqueFunction<void()>&& InFunc) override;
+
+			virtual void SetObjectReferencer(TUniquePtr<IObjectReferencer> InReferencer) override;
+
+			// This is called on module startup. This invokes any registration commands enqueued by our registration macros.
+			virtual void RegisterPendingNodes() override;
+
+			// Register a graph from an IMetaSoundDocumentInterface
+			virtual FNodeRegistryKey RegisterGraph(const TScriptInterface<IMetaSoundDocumentInterface>& InDocumentInterface, bool bAsync = true) override;
+			
+			// Wait for async graph registration to complete for a specific graph
+			virtual void WaitForAsyncGraphRegistration(const FNodeRegistryKey& InRegistryKey, const FTopLevelAssetPath& InAssetPath) const override;
+
+			// Unregister a graph 
+			virtual bool UnregisterGraph(const TScriptInterface<IMetaSoundDocumentInterface>& InDocumentInterface, bool bAsync = true) override;
+
+			// Retrieve a registered graph. 
+			//
+			// If the graph is registered asynchronously, this will wait until the registration task has completed.
+			virtual TSharedPtr<const FGraph> GetGraph(const FNodeRegistryKey& InRegistryKey, const FTopLevelAssetPath& InAssetPath) const override;
+
+			/** Register external node with the frontend.
+			 *
+			 * @param InEntry - Entry to register.
+			 *
+			 * @return True on success.
+			 */
+			virtual FNodeRegistryKey RegisterNode(TUniquePtr<Metasound::Frontend::INodeRegistryEntry>&& InEntry) override;
+
+			virtual bool UnregisterNode(const FNodeRegistryKey& InKey) override;
+			virtual bool IsNodeRegistered(const FNodeRegistryKey& InKey) const override;
+			virtual bool IsGraphRegistered(const FNodeRegistryKey& InKey, const FTopLevelAssetPath& InAssetPath) const override;
+			virtual bool IsNodeNative(const FNodeRegistryKey& InKey) const override;
+
+			virtual bool RegisterConversionNode(const FConverterNodeRegistryKey& InNodeKey, const FConverterNodeInfo& InNodeInfo) override;
+
+			virtual void IterateRegistry(Metasound::FIterateMetasoundFrontendClassFunction InIterFunc, EMetasoundFrontendClassType InClassType = EMetasoundFrontendClassType::Invalid) const override;
+
+			// Find Frontend class data.
+			virtual bool FindFrontendClassFromRegistered(const FNodeRegistryKey& InKey, FMetasoundFrontendClass& OutClass) override;
+
+			UE_DEPRECATED(5.4, "This implementation of FindImplementedInterfacesFromRegistered(...). Please use bool FindImplementedInterfacesFromRegistered(const Metasound::Frontend::FNodeRegistryKey& InKey, TSet<FMetasoundFrontendVersion>& OutInterfacVersion)")
+			virtual const TSet<FMetasoundFrontendVersion>* FindImplementedInterfacesFromRegistered(const Metasound::Frontend::FNodeRegistryKey& InKey) const override;
+
+			virtual bool FindImplementedInterfacesFromRegistered(const Metasound::Frontend::FNodeRegistryKey& InKey, TSet<FMetasoundFrontendVersion>& OutInterfaceVersions) const override;
+			
+			UE_DEPRECATED(5.4, "Use FindFrontendClassFromRegistered instead")
+			virtual bool FindNodeClassInfoFromRegistered(const Metasound::Frontend::FNodeRegistryKey& InKey, FNodeClassInfo& OutInfo) override;
+			virtual bool FindInputNodeRegistryKeyForDataType(const FName& InDataTypeName, const EMetasoundFrontendVertexAccessType InAccessType, FNodeRegistryKey& OutKey) override;
+			virtual bool FindVariableNodeRegistryKeyForDataType(const FName& InDataTypeName, FNodeRegistryKey& OutKey) override;
+			virtual bool FindOutputNodeRegistryKeyForDataType(const FName& InDataTypeName, const EMetasoundFrontendVertexAccessType InAccessType, FNodeRegistryKey& OutKey) override;
+
+			// Create a new instance of a C++ implemented node from the registry.
+			virtual TUniquePtr<Metasound::INode> CreateNode(const FNodeRegistryKey& InKey, const Metasound::FNodeInitData& InInitData) const override;
+			virtual TUniquePtr<INode> CreateNode(const FNodeRegistryKey& InKey, FDefaultLiteralNodeConstructorParams&&) const override;
+			virtual TUniquePtr<INode> CreateNode(const FNodeRegistryKey& InKey, FDefaultNamedVertexNodeConstructorParams&&) const override;
+			virtual TUniquePtr<INode> CreateNode(const FNodeRegistryKey& InKey, FDefaultNamedVertexWithLiteralNodeConstructorParams&&) const override;
+
+			// Returns a list of possible nodes to use to convert from FromDataType to ToDataType.
+			// Returns an empty array if none are available.
+			virtual TArray<FConverterNodeInfo> GetPossibleConverterNodes(const FName& FromDataType, const FName& ToDataType) override;
+
+			// Private implementation until hardened and used for template nodes other than reroutes.
+			FNodeRegistryKey RegisterNodeTemplate(TUniquePtr<Metasound::Frontend::INodeRegistryTemplateEntry>&& InEntry);
+			bool UnregisterNodeTemplate(const FNodeRegistryKey& InKey);
+
+			// Create a transaction stream for any newly transactions
+			TUniquePtr<FNodeRegistryTransactionStream> CreateTransactionStream();
+
+		private:
+			using FGraphRegistryKey = TTuple<FNodeRegistryKey, FTopLevelAssetPath>;
+
+			struct FActiveRegistrationTaskInfo
+			{
+				FNodeRegistryTransaction::ETransactionType TransactionType = FNodeRegistryTransaction::ETransactionType::NodeRegistration;
+				UE::Tasks::FTask Task;
+				FTopLevelAssetPath AssetPath;
+			};
+
+			void BuildAndRegisterGraphFromDocument(TScriptInterface<IMetaSoundDocumentInterface> DocumentInterface, const FProxyDataCache& InProxyDataCache, FNodeClassInfo&& InNodeClassInfo);
+
+			void AddRegistrationTask(const FNodeRegistryKey& InKey, UObject& InObject, FActiveRegistrationTaskInfo&& TaskInfo);
+			void RemoveRegistrationTask(const FNodeRegistryKey& InKey, FNodeRegistryTransaction::ETransactionType TransactionType, const TScriptInterface<IMetaSoundDocumentInterface>& DocumentInterface);
+
+			bool UnregisterNodeInternal(const FNodeRegistryKey& InKey, FNodeClassInfo* OutClassInfo = nullptr);
+			FNodeRegistryKey RegisterNodeInternal(TUniquePtr<INodeRegistryEntry>&& InEntry);
+
+			static FRegistryContainerImpl* LazySingleton;
+
+			void WaitForAsyncRegistrationInternal(const FNodeRegistryKey& InRegistryKey, const FTopLevelAssetPath* InAssetPath) const;
+
+			void RegisterGraphInternal(const FNodeRegistryKey& InKey, const FTopLevelAssetPath& InAssetPath, TSharedPtr<const FGraph> InGraph);
+			bool UnregisterGraphInternal(const FNodeRegistryKey& InKey, const FTopLevelAssetPath& InAssetPath);
+
+			// Access a node entry safely. Node entries can be added/removed asynchronously. Functions passed to this method will be
+			// executed in a manner where access to the node registry entry is safe from threading issues. 
+			//
+			// @returns true if a node registry entry was found and the function executed. False if the entry was not 
+			//          found and the function not executed. 
+			bool AccessNodeEntryThreadSafe(const FNodeRegistryKey& InKey, TFunctionRef<void(const INodeRegistryEntry&)> InFunc) const;
+
+			const INodeRegistryTemplateEntry* FindNodeTemplateEntry(const FNodeRegistryKey& InKey) const;
+
+
+			// This buffer is used to enqueue nodes and datatypes to register when the module has been initialized,
+			// in order to avoid bad behavior with ensures, logs, etc. on static initialization.
+			// The bad news is that TInlineAllocator is the safest allocator to use on static init.
+			// The good news is that none of these lambdas typically have captures, so this should have low memory overhead.
+			static constexpr int32 MaxNumNodesAndDatatypesToInitialize = 2048;
+			TArray<TUniqueFunction<void()>, TInlineAllocator<MaxNumNodesAndDatatypesToInitialize>> LazyInitCommands;
+			
+			FCriticalSection LazyInitCommandCritSection;
+
+			// Registry in which we keep all information about nodes implemented in C++.
+			TMultiMap<FNodeRegistryKey, TSharedRef<INodeRegistryEntry, ESPMode::ThreadSafe>> RegisteredNodes;
+
+			// Registry in which we keep all information about dynamically-generated templated nodes via in C++.
+			TMap<FNodeRegistryKey, TSharedRef<INodeRegistryTemplateEntry, ESPMode::ThreadSafe>> RegisteredNodeTemplates;
+
+			// Map of all registered graphs
+			TMap<FGraphRegistryKey, TSharedPtr<const FGraph>> RegisteredGraphs;
+
+			// Registry in which we keep lists of possible nodes to use to convert between two datatypes
+			TMap<FConverterNodeRegistryKey, FConverterNodeRegistryValue> ConverterNodeRegistry;
+
+			TSharedRef<FNodeRegistryTransactionBuffer> TransactionBuffer;
+
+			mutable FCriticalSection RegistryMapsCriticalSection;
+			mutable FCriticalSection ActiveRegistrationTasksCriticalSection;
+			UE::Tasks::FPipe AsyncRegistrationPipe;
+			TMap<FNodeRegistryKey, TArray<FActiveRegistrationTaskInfo>> ActiveRegistrationTasks;
+			TUniquePtr<FMetasoundFrontendRegistryContainer::IObjectReferencer> ObjectReferencer;
+		};
+	}
+}
+
diff --git a/Engine/Plugins/Runtime/Metasound/Source/MetasoundFrontend/Private/MetasoundFrontendRegistryKey.cpp b/Engine/Plugins/Runtime/Metasound/Source/MetasoundFrontend/Private/MetasoundFrontendRegistryKey.cpp
--- a/Engine/Plugins/Runtime/Metasound/Source/MetasoundFrontend/Private/MetasoundFrontendRegistryKey.cpp
+++ b/Engine/Plugins/Runtime/Metasound/Source/MetasoundFrontend/Private/MetasoundFrontendRegistryKey.cpp
@@ -5,6 +5,11 @@
 
 namespace Metasound::Frontend
 {
+	namespace NodeClassInfoPrivate
+	{
+		auto GetVertexTypeName = [](const FMetasoundFrontendVertex& Vertex) { return Vertex.TypeName; };
+	}
+
 	FNodeClassInfo::FNodeClassInfo(const FMetasoundFrontendClassMetadata& InMetadata)
 		: ClassName(InMetadata.GetClassName())
 		, Type(InMetadata.GetType())
@@ -16,36 +21,49 @@ namespace Metasound::Frontend
 		: ClassName(InClass.Metadata.GetClassName())
 		, Type(EMetasoundFrontendClassType::External) // Overridden as it is considered the same as an external class in registries
 		, AssetClassID(FGuid(ClassName.Name.ToString()))
-		, AssetPath(InAssetPath)
 		, Version(InClass.Metadata.GetVersion())
 	{
-		ensure(!AssetPath.IsNull());
-#if WITH_EDITORONLY_DATA
-		for (const FMetasoundFrontendClassInput& Input : InClass.Interface.Inputs)
-		{
-			InputTypes.Add(Input.TypeName);
-		}
+		using namespace NodeClassInfoPrivate;
 
-		for (const FMetasoundFrontendClassOutput& Output : InClass.Interface.Outputs)
-		{
-			OutputTypes.Add(Output.TypeName);
-		}
+		ensure(AssetPath.TrySetPath(InAssetPath.ToString()));
+		ensure(!AssetPath.IsNull());
 
+#if WITH_EDITORONLY_DATA
+		Algo::Transform(InClass.Interface.Inputs, InputTypes, GetVertexTypeName);
+		Algo::Transform(InClass.Interface.Outputs, OutputTypes, GetVertexTypeName);
 		bIsPreset = InClass.PresetOptions.bIsPreset;
 #endif // WITH_EDITORONLY_DATA
 	}
 
-	UObject* FNodeClassInfo::LoadAsset() const
+	FNodeClassInfo::FNodeClassInfo(const FMetasoundFrontendGraphClass& InClass, const FTopLevelAssetPath& InAssetPath)
+		: ClassName(InClass.Metadata.GetClassName())
+		, Type(EMetasoundFrontendClassType::External) // Overridden as it is considered the same as an external class in registries
+		, AssetClassID(FGuid(ClassName.Name.ToString()))
+		, AssetPath(InAssetPath)
+		, Version(InClass.Metadata.GetVersion())
 	{
-		if (ensure(Type == EMetasoundFrontendClassType::External))
-		{
-			FSoftObjectPath SoftObjectPath(AssetPath);
-			return SoftObjectPath.TryLoad();
-		}
+		using namespace NodeClassInfoPrivate;
 
-		return nullptr;
+		ensure(!AssetPath.IsNull());
+
+#if WITH_EDITORONLY_DATA
+		Algo::Transform(InClass.Interface.Inputs, InputTypes, GetVertexTypeName);
+		Algo::Transform(InClass.Interface.Outputs, OutputTypes, GetVertexTypeName);
+		bIsPreset = InClass.PresetOptions.bIsPreset;
+#endif // WITH_EDITORONLY_DATA
 	}
 
+// 	UObject* FNodeClassInfo::LoadAsset() const
+// 	{
+// 		return nullptr;
+// 		if (ensure(Type == EMetasoundFrontendClassType::External))
+// 		{
+// 			return AssetPath.TryLoad();
+// 		}
+// 
+// 		return nullptr;
+// 	}
+
 	FNodeRegistryKey::FNodeRegistryKey(EMetasoundFrontendClassType InType, const FMetasoundFrontendClassName& InClassName, int32 InMajorVersion, int32 InMinorVersion)
 		: Type(InType)
 		, ClassName(InClassName)
diff --git a/Engine/Plugins/Runtime/Metasound/Source/MetasoundFrontend/Private/MetasoundFrontendTransform.cpp b/Engine/Plugins/Runtime/Metasound/Source/MetasoundFrontend/Private/MetasoundFrontendTransform.cpp
--- a/Engine/Plugins/Runtime/Metasound/Source/MetasoundFrontend/Private/MetasoundFrontendTransform.cpp
+++ b/Engine/Plugins/Runtime/Metasound/Source/MetasoundFrontend/Private/MetasoundFrontendTransform.cpp
@@ -617,7 +617,6 @@ namespace Metasound
 
 				if (FMetasoundAssetBase* ReferencedMetaSoundAsset = IMetaSoundAssetManager::GetChecked().TryLoadAssetFromKey(RegistryKey))
 				{
- 					ReferencedMetaSoundAsset->WaitForAsyncGraphRegistration();
 					if (bIsPreset)
 					{
 						PresetReferencedMetaSoundAsset = ReferencedMetaSoundAsset;
diff --git a/Engine/Plugins/Runtime/Metasound/Source/MetasoundFrontend/Public/MetasoundAssetBase.h b/Engine/Plugins/Runtime/Metasound/Source/MetasoundFrontend/Public/MetasoundAssetBase.h
--- a/Engine/Plugins/Runtime/Metasound/Source/MetasoundFrontend/Public/MetasoundAssetBase.h
+++ b/Engine/Plugins/Runtime/Metasound/Source/MetasoundFrontend/Public/MetasoundAssetBase.h
@@ -86,9 +86,6 @@ public:
 	// Registers the root graph of the given asset with the MetaSound Frontend.
 	virtual void RegisterGraphWithFrontend(Metasound::Frontend::FMetaSoundAssetRegistrationOptions InRegistrationOptions = Metasound::Frontend::FMetaSoundAssetRegistrationOptions());
 
-	// Return the graph registered for this asset. 
-	TSharedPtr<const Metasound::FGraph> GetRegisteredGraph() const;
-
 	// Unregisters the root graph of the given asset with the MetaSound Frontend.
 	void UnregisterGraphWithFrontend();
 
@@ -275,11 +272,7 @@ private:
 	// race conditions with an active builder.
 	virtual bool IsBuilderActive() const = 0;
 
-	// Returns the cached registry key.
-	Metasound::Frontend::FNodeRegistryKey CacheRuntimeData(const TScriptInterface<IMetaSoundDocumentInterface>& InDoc);
-
 	Metasound::Frontend::FNodeRegistryKey RegistryKey;
-	FSoftObjectPath RegisteredGraphAssetPath; // This is cached to avoid reparsing path strings to create FSoftObjectPath from UObject
 
 	TSharedPtr<Metasound::FGraph, ESPMode::ThreadSafe> BuildMetasoundDocument(const FMetasoundFrontendDocument& InPreprocessDoc, const Metasound::Frontend::FProxyDataCache& InProxies) const;
 };
diff --git a/Engine/Plugins/Runtime/Metasound/Source/MetasoundFrontend/Public/MetasoundDocumentInterface.h b/Engine/Plugins/Runtime/Metasound/Source/MetasoundFrontend/Public/MetasoundDocumentInterface.h
--- a/Engine/Plugins/Runtime/Metasound/Source/MetasoundFrontend/Public/MetasoundDocumentInterface.h
+++ b/Engine/Plugins/Runtime/Metasound/Source/MetasoundFrontend/Public/MetasoundDocumentInterface.h
@@ -5,6 +5,7 @@
 #include "MetasoundFrontendDocument.h"
 #include "MetasoundFrontendDocumentModifyDelegates.h"
 #include "Templates/Function.h"
+#include "UObject/TopLevelAssetPath.h"
 
 #include "MetasoundDocumentInterface.generated.h"
 
@@ -27,6 +28,8 @@ class METASOUNDFRONTEND_API IMetaSoundDocumentInterface : public IInterface
 	GENERATED_BODY()
 
 public:
+	virtual FTopLevelAssetPath GetAssetPathChecked() const = 0;
+
 	// Returns read-only reference to the the MetaSoundFrontendDocument
 	// containing all MetaSound runtime & editor data.
 	UE_DEPRECATED(5.4, "Use GetConstDocument instead")
@@ -42,7 +45,6 @@ public:
 	// Returns the parent class registered with the MetaSound UObject registry.
 	virtual const UClass& GetBaseMetaSoundUClass() const = 0;
 
-
 private:
 	virtual FMetasoundFrontendDocument& GetDocument() = 0;
 
diff --git a/Engine/Plugins/Runtime/Metasound/Source/MetasoundFrontend/Public/MetasoundFrontendDocumentBuilder.h b/Engine/Plugins/Runtime/Metasound/Source/MetasoundFrontend/Public/MetasoundFrontendDocumentBuilder.h
--- a/Engine/Plugins/Runtime/Metasound/Source/MetasoundFrontend/Public/MetasoundFrontendDocumentBuilder.h
+++ b/Engine/Plugins/Runtime/Metasound/Source/MetasoundFrontend/Public/MetasoundFrontendDocumentBuilder.h
@@ -104,6 +104,9 @@ public:
 	// Returns the document
 	virtual const FMetasoundFrontendDocument& GetConstDocument() const override;
 
+	// Returns temp path of builder document
+	virtual FTopLevelAssetPath GetAssetPathChecked() const override;
+
 	// Base MetaSoundClass that document is published to.
 	virtual const UClass& GetBaseMetaSoundUClass() const final override;
 
diff --git a/Engine/Plugins/Runtime/Metasound/Source/MetasoundFrontend/Public/MetasoundFrontendRegistries.h b/Engine/Plugins/Runtime/Metasound/Source/MetasoundFrontend/Public/MetasoundFrontendRegistries.h
--- a/Engine/Plugins/Runtime/Metasound/Source/MetasoundFrontend/Public/MetasoundFrontendRegistries.h
+++ b/Engine/Plugins/Runtime/Metasound/Source/MetasoundFrontend/Public/MetasoundFrontendRegistries.h
@@ -1,360 +1,5 @@
 // Copyright Epic Games, Inc. All Rights Reserved.
 #pragma once
+#include "MetasoundFrontendRegistryContainer.h"
 
-#include "IAudioProxyInitializer.h"
-#include "MetasoundDataReference.h"
-#include "MetasoundEnum.h"
-#include "MetasoundFrontendDocument.h"
-#include "MetasoundLiteral.h"
-#include "MetasoundNodeConstructorParams.h"
-#include "MetasoundNodeInterface.h"
-#include "MetasoundOperatorInterface.h"
-#include "MetasoundFrontendRegistryKey.h"
-#include "MetasoundRouter.h"
-#include "MetasoundVertex.h"
-#include "Templates/Function.h"
-
-
-// Forward Declarations
-class IMetaSoundDocumentInterface;
-template <typename InInterfaceType>
-class TScriptInterface;
-class UObject;
-
-
-namespace Metasound
-{
-	// Forward Declarations
-	using FIterateMetasoundFrontendClassFunction = TFunctionRef<void(const FMetasoundFrontendClass&)>;
-	class FGraph;
-
-	namespace Frontend
-	{
-		/** INodeRegistryEntry declares the interface for a node registry entry.
-		 * Each node class in the registry must satisfy this interface. 
-		 */
-		class METASOUNDFRONTEND_API INodeRegistryEntry
-		{
-		public:
-			virtual ~INodeRegistryEntry() = default;
-
-			/** Return FNodeClassInfo for the node class.
-			 *
-			 * Implementations of method should avoid any expensive operations 
-			 * (e.g. loading from disk, allocating memory) as this method is called
-			 * frequently when querying nodes.
-			 */
-			virtual const FNodeClassInfo& GetClassInfo() const = 0;
-
-			/** Create a node given FDefaultNamedVertexNodeConstructorParams.
-			 *
-			 * If a node can be created with FDefaultNamedVertexNodeConstructorParams, this function
-			 * should return a valid node pointer.
-			 */
-			virtual TUniquePtr<INode> CreateNode(FDefaultNamedVertexNodeConstructorParams&&) const = 0;
-
-			/** Create a node given FDefaultNamedVertexWithLiteralNodeConstructorParams.
-			 *
-			 * If a node can be created with FDefaultNamedVertexWithLiteralNodeConstructorParams, this function
-			 * should return a valid node pointer.
-			 */
-			virtual TUniquePtr<INode> CreateNode(FDefaultNamedVertexWithLiteralNodeConstructorParams&&) const = 0;
-
-			/** Create a node given FDefaultLiteralNodeConstructorParams.
-			 *
-			 * If a node can be created with FDefaultLiteralNodeConstructorParams, this function
-			 * should return a valid node pointer.
-			 */
-			virtual TUniquePtr<INode> CreateNode(FDefaultLiteralNodeConstructorParams&&) const = 0;
-
-			/** Create a node given FNodeInitData.
-			 *
-			 * If a node can be created with FNodeInitData, this function
-			 * should return a valid node pointer.
-			 */
-			virtual TUniquePtr<INode> CreateNode(const FNodeInitData&) const = 0;
-
-			/** Return a FMetasoundFrontendClass which describes the node. */
-			virtual const FMetasoundFrontendClass& GetFrontendClass() const = 0;
-
-			/** Clone this registry entry. */
-			virtual TUniquePtr<INodeRegistryEntry> Clone() const = 0;
-
-			/** Returns set of implemented interface versions.
-			  *
-			  * Returns nullptr if node class implementation does not support interface implementation.
-			  */
-			virtual const TSet<FMetasoundFrontendVersion>* GetImplementedInterfaces() const = 0;
-
-			/** Whether or not the node is natively defined */
-			virtual bool IsNative() const = 0;
-		};
-
-		struct METASOUNDFRONTEND_API FConverterNodeRegistryKey
-		{
-			// The datatype one would like to convert from.
-			FName FromDataType;
-
-			// The datatype one would like to convert to.
-			FName ToDataType;
-
-			FORCEINLINE bool operator==(const FConverterNodeRegistryKey& Other) const
-			{
-				return FromDataType == Other.FromDataType && ToDataType == Other.ToDataType;
-			}
-
-			friend uint32 GetTypeHash(const ::Metasound::Frontend::FConverterNodeRegistryKey& InKey)
-			{
-				return HashCombine(GetTypeHash(InKey.FromDataType), GetTypeHash(InKey.ToDataType));
-			}
-		};
-
-		struct METASOUNDFRONTEND_API FConverterNodeInfo
-		{
-			// If this node has multiple input pins, we use this to designate which pin should be used.
-			FVertexName PreferredConverterInputPin;
-
-			// If this node has multiple output pins, we use this to designate which pin should be used.
-			FVertexName PreferredConverterOutputPin;
-
-			// The key for this node in the node registry.
-			FNodeRegistryKey NodeKey;
-
-			FORCEINLINE bool operator==(const FConverterNodeInfo& Other) const
-			{
-				return NodeKey == Other.NodeKey;
-			}
-		};
-
-		struct METASOUNDFRONTEND_API FConverterNodeRegistryValue
-		{
-			// A list of nodes that can perform a conversion between the two datatypes described in the FConverterNodeRegistryKey for this map element.
-			TArray<FConverterNodeInfo> PotentialConverterNodes;
-		};
-
-		using FRegistryTransactionID = int32;
-
-		class METASOUNDFRONTEND_API FNodeRegistryTransaction 
-		{
-		public:
-			using FTimeType = uint64;
-
-			/** Describes the type of transaction. */
-			enum class ETransactionType : uint8
-			{
-				NodeRegistration,     //< Something was added to the registry.
-				NodeUnregistration,  //< Something was removed from the registry.
-				Invalid
-			};
-
-			static FString LexToString(ETransactionType InTransactionType)
-			{
-				using namespace Metasound;
-
-				switch (InTransactionType)
-				{
-					case ETransactionType::NodeRegistration:
-						return FString(TEXT("NodeRegistration"));
-
-					case ETransactionType::NodeUnregistration:
-						return FString(TEXT("NodeUnregistration"));
-
-					case ETransactionType::Invalid:
-						return FString(TEXT("Invalid"));
-
-					default:
-					{
-						checkNoEntry();
-					}
-				}
-				return FString(TEXT(""));
-			}
-
-			FNodeRegistryTransaction(ETransactionType InType, const FNodeClassInfo& InNodeClassInfo, FTimeType InTimestamp);
-
-			ETransactionType GetTransactionType() const;
-			const FNodeClassInfo& GetNodeClassInfo() const;
-			FNodeRegistryKey GetNodeRegistryKey() const;
-			FTimeType GetTimestamp() const;
-
-		private:
-
-			ETransactionType Type;
-			FNodeClassInfo NodeClassInfo;
-			FTimeType Timestamp;
-		};
-
-		namespace NodeRegistryKey
-		{
-			UE_DEPRECATED(5.4, "Use FNodeRegistryKey::GetInvalid() instead")
-			METASOUNDFRONTEND_API const FNodeRegistryKey& GetInvalid();
-
-			UE_DEPRECATED(5.4, "Use FNodeRegistryKey::IsValid() instead")
-			METASOUNDFRONTEND_API bool IsValid(const FNodeRegistryKey& InKey);
-
-			UE_DEPRECATED(5.4, "Use FNodeRegistryKey equality operator instead")
-			METASOUNDFRONTEND_API bool IsEqual(const FNodeRegistryKey& InLHS, const FNodeRegistryKey& InRHS);
-
-			// Returns true if the class metadata represent the same entry in the node registry.
-			METASOUNDFRONTEND_API bool IsEqual(const FMetasoundFrontendClassMetadata& InLHS, const FMetasoundFrontendClassMetadata& InRHS);
-
-			// Returns true if the class info and class metadata represent the same entry in the node registry.
-			METASOUNDFRONTEND_API bool IsEqual(const FNodeClassInfo& InLHS, const FMetasoundFrontendClassMetadata& InRHS);
-
-			UE_DEPRECATED(5.4, "Use FNodeRegistryKey constructors instead")
-			METASOUNDFRONTEND_API FNodeRegistryKey CreateKey(EMetasoundFrontendClassType InType, const FString& InFullClassName, int32 InMajorVersion, int32 InMinorVersion);
-
-			UE_DEPRECATED(5.4, "Use FNodeRegistryKey constructors instead")
-			METASOUNDFRONTEND_API FNodeRegistryKey CreateKey(const FNodeClassMetadata& InNodeMetadata);
-
-			UE_DEPRECATED(5.4, "Use FNodeRegistryKey constructor instead")
-			METASOUNDFRONTEND_API FNodeRegistryKey CreateKey(const FMetasoundFrontendClassMetadata& InNodeMetadata);
-
-			UE_DEPRECATED(5.4, "Use FNodeRegistryKey constructor instead")
-			METASOUNDFRONTEND_API FNodeRegistryKey CreateKey(const FMetasoundFrontendGraphClass& InNodeMetadata);
-
-			UE_DEPRECATED(5.4, "Use FNodeRegistryKey constructor instead")
-			METASOUNDFRONTEND_API FNodeRegistryKey CreateKey(const FNodeClassInfo& ClassInfo);
-		}
-	} // namespace Frontend
-} // namespace Metasound
-
-/**
- * Singleton registry for all types and nodes.
- */
-class METASOUNDFRONTEND_API FMetasoundFrontendRegistryContainer
-{
-
-public:
-	// The MetaSound frontend does not rely on the Engine module and therefore
-	// does not have the ability to provide GC protection. This interface allows
-	// external modules to provide GC protection so that async tasks can safely
-	// use UObjects
-	class IObjectReferencer
-	{
-	public:
-		virtual ~IObjectReferencer() = default;
-		// Called when an object should be referenced.
-		virtual void AddObject(UObject* InObject) = 0;
-		// Called when an object no longer needs to be referenced. 
-		virtual void RemoveObject(UObject* InObject) = 0;
-	};
-
-	using FNodeClassInfo = Metasound::Frontend::FNodeClassInfo;
-	using FConverterNodeRegistryKey = ::Metasound::Frontend::FConverterNodeRegistryKey;
-	using FConverterNodeRegistryValue = ::Metasound::Frontend::FConverterNodeRegistryValue;
-	using FConverterNodeInfo = ::Metasound::Frontend::FConverterNodeInfo;
-
-	using FNodeRegistryKey = Metasound::Frontend::FNodeRegistryKey;
-	using FNodeClassMetadata = Metasound::FNodeClassMetadata;
-
-	static FMetasoundFrontendRegistryContainer* Get();
-	static void ShutdownMetasoundFrontend();
-
-	UE_DEPRECATED(5.4, "Use FNodeRegistryKey constructor that takes FNodeClassMetadata instead")
-	static FNodeRegistryKey GetRegistryKey(const FNodeClassMetadata& InNodeMetadata);
-
-	UE_DEPRECATED(5.4, "Use FNodeRegistryKey constructor that takes FMetasoundFrontendClassMetadata instead")
-	static FNodeRegistryKey GetRegistryKey(const FMetasoundFrontendClassMetadata& InNodeMetadata);
-
-	UE_DEPRECATED(5.4, "Use FNodeRegistryKey constructor that takes FNodeClassInfo instead")
-	static FNodeRegistryKey GetRegistryKey(const FNodeClassInfo& ClassInfo);
-
-	static bool GetFrontendClassFromRegistered(const FNodeRegistryKey& InKey, FMetasoundFrontendClass& OutClass);
-
-	UE_DEPRECATED(5.4, "Use GetFrontendClassFromRegistered instead")
-	static bool GetNodeClassInfoFromRegistered(const FNodeRegistryKey& InKey, FNodeClassInfo& OutInfo);
-
-	static bool GetInputNodeRegistryKeyForDataType(const FName& InDataTypeName, const EMetasoundFrontendVertexAccessType InAccessType, FNodeRegistryKey& OutKey);
-	static bool GetVariableNodeRegistryKeyForDataType(const FName& InDataTypeName, FNodeRegistryKey& OutKey);
-	static bool GetOutputNodeRegistryKeyForDataType(const FName& InDataTypeName, const EMetasoundFrontendVertexAccessType InAccessType, FNodeRegistryKey& OutKey);
-
-
-	FMetasoundFrontendRegistryContainer() = default;
-	virtual ~FMetasoundFrontendRegistryContainer() = default;
-
-	FMetasoundFrontendRegistryContainer(const FMetasoundFrontendRegistryContainer&) = delete;
-	FMetasoundFrontendRegistryContainer& operator=(const FMetasoundFrontendRegistryContainer&) = delete;
-
-	// Enqueu and command for registering a node or data type.
-	// The command queue will be processed on module init or when calling `RegisterPendingNodes()`
-	virtual bool EnqueueInitCommand(TUniqueFunction<void()>&& InFunc) = 0;
-
-	virtual void SetObjectReferencer(TUniquePtr<IObjectReferencer> InReferencer) = 0;
-
-	// This is called on module startup. This invokes any registration commands enqueued by our registration macros.
-	virtual void RegisterPendingNodes() = 0;
-
-	// Register a graph from an IMetaSoundDocumentInterface
-	virtual FNodeRegistryKey RegisterGraph(const FSoftObjectPath& InAssetPath, const TScriptInterface<IMetaSoundDocumentInterface>& InDocument, bool bAsync=true) = 0;
-
-	// Wait for async graph registration to complete for a specific graph
-	virtual void WaitForAsyncGraphRegistration(const FNodeRegistryKey& InRegistryKey, const FSoftObjectPath& InAssetPath) const = 0;
-
-	// Unregister an asset based graph from the frontend. 
-	virtual bool UnregisterGraph(const FNodeRegistryKey& InNodeRegistryKey, const FSoftObjectPath& InAssetPath) = 0;
-
-	// Retrieve a registered graph. 
-	//
-	// If the graph is registered asynchronously, this will wait until the registration task has completed.
-	virtual TSharedPtr<const Metasound::FGraph> GetGraph(const FNodeRegistryKey& InRegistryKey, const FSoftObjectPath& InAssetPath) const = 0;
-
-	/** Register an external node with the frontend.
-	 *
-	 * @param InCreateNode - Function for creating node from FNodeInitData.
-	 * @param InCreateDescription - Function for creating a FMetasoundFrontendClass.
-	 *
-	 * @return A node registration key. If the registration failed, then the registry 
-	 *         key will be invalid.
-	 */
-	virtual FNodeRegistryKey RegisterNode(TUniquePtr<Metasound::Frontend::INodeRegistryEntry>&& InEntry) = 0;
-
-	/** Unregister an external node from the frontend.
-	 *
-	 * @param InKey - The registration key for the node.
-	 *
-	 * @return True on success, false on failure.
-	 */
-	virtual bool UnregisterNode(const FNodeRegistryKey& InKey) = 0;
-
-	/** Returns true if the provided registry key corresponds to a valid registered node. */
-	virtual bool IsNodeRegistered(const FNodeRegistryKey& InKey) const = 0;
-
-	/** Returns true if the provided registry key and asset path corresponds to a valid registered graph. */
-	virtual bool IsGraphRegistered(const FNodeRegistryKey& InKey, const FSoftObjectPath& InAssetPath) const = 0;
-
-	/** Returns true if the provided registry key corresponds to a valid registered node that is natively defined. */
-	virtual bool IsNodeNative(const FNodeRegistryKey& InKey) const = 0;
-
-	// Iterates class types in registry.  If InClassType is set to a valid class type (optional), only iterates classes of the given type
-	virtual void IterateRegistry(Metasound::FIterateMetasoundFrontendClassFunction InIterFunc, EMetasoundFrontendClassType InClassType = EMetasoundFrontendClassType::Invalid) const = 0;
-
-	// Query for MetaSound Frontend document objects.
-	virtual bool FindFrontendClassFromRegistered(const Metasound::Frontend::FNodeRegistryKey& InKey, FMetasoundFrontendClass& OutClass) = 0;
-
-	UE_DEPRECATED(5.4, "This implementation of FindImplementedInterfacesFromRegistered(...). Please use bool FindImplementedInterfacesFromRegistered(const Metasound::Frontend::FNodeRegistryKey& InKey, TSet<FMetasoundFrontendVersion>& OutInterfacVersion)")
-	virtual const TSet<FMetasoundFrontendVersion>* FindImplementedInterfacesFromRegistered(const Metasound::Frontend::FNodeRegistryKey& InKey) const = 0;
-
-	virtual bool FindImplementedInterfacesFromRegistered(const Metasound::Frontend::FNodeRegistryKey& InKey, TSet<FMetasoundFrontendVersion>& OutInterfaceVersions) const = 0;
-	
-	UE_DEPRECATED(5.4, "Use FindFrontendClassFromRegistered instead")
-	virtual bool FindNodeClassInfoFromRegistered(const Metasound::Frontend::FNodeRegistryKey& InKey, FNodeClassInfo& OutInfo) = 0;
-	virtual bool FindInputNodeRegistryKeyForDataType(const FName& InDataTypeName, const EMetasoundFrontendVertexAccessType InAccessType, FNodeRegistryKey& OutKey) = 0;
-	virtual bool FindVariableNodeRegistryKeyForDataType(const FName& InDataTypeName, FNodeRegistryKey& OutKey) = 0;
-	virtual bool FindOutputNodeRegistryKeyForDataType(const FName& InDataTypeName, const EMetasoundFrontendVertexAccessType InAccessType, FNodeRegistryKey& OutKey) = 0;
-
-	virtual TUniquePtr<Metasound::INode> CreateNode(const FNodeRegistryKey& InKey, const Metasound::FNodeInitData&) const = 0;
-	virtual TUniquePtr<Metasound::INode> CreateNode(const FNodeRegistryKey& InKey, Metasound::FDefaultLiteralNodeConstructorParams&&) const = 0;
-	virtual TUniquePtr<Metasound::INode> CreateNode(const FNodeRegistryKey& InKey, Metasound::FDefaultNamedVertexNodeConstructorParams&&) const = 0;
-	virtual TUniquePtr<Metasound::INode> CreateNode(const FNodeRegistryKey& InKey, Metasound::FDefaultNamedVertexWithLiteralNodeConstructorParams&&) const = 0;
-
-	virtual bool RegisterConversionNode(const FConverterNodeRegistryKey& InNodeKey, const FConverterNodeInfo& InNodeInfo) = 0;
-
-	// Returns a list of possible nodes to use to convert from FromDataType to ToDataType.
-	// Returns an empty array if none are available.
-	virtual TArray<FConverterNodeInfo> GetPossibleConverterNodes(const FName& FromDataType, const FName& ToDataType) = 0;
-
-private:
-};
-
-
+// TODO: Repurpose or deprecate this file
diff --git a/Engine/Plugins/Runtime/Metasound/Source/MetasoundFrontend/Public/MetasoundFrontendRegistryContainer.h b/Engine/Plugins/Runtime/Metasound/Source/MetasoundFrontend/Public/MetasoundFrontendRegistryContainer.h
new file mode 100644
--- /dev/null
+++ b/Engine/Plugins/Runtime/Metasound/Source/MetasoundFrontend/Public/MetasoundFrontendRegistryContainer.h
@@ -0,0 +1,359 @@
+// Copyright Epic Games, Inc. All Rights Reserved.
+#pragma once
+
+#include "IAudioProxyInitializer.h"
+#include "MetasoundDataReference.h"
+#include "MetasoundEnum.h"
+#include "MetasoundFrontendDocument.h"
+#include "MetasoundLiteral.h"
+#include "MetasoundNodeConstructorParams.h"
+#include "MetasoundNodeInterface.h"
+#include "MetasoundOperatorInterface.h"
+#include "MetasoundFrontendRegistryKey.h"
+#include "MetasoundRouter.h"
+#include "MetasoundVertex.h"
+#include "Templates/Function.h"
+
+
+// Forward Declarations
+class IMetaSoundDocumentInterface;
+template <typename InInterfaceType>
+class TScriptInterface;
+class UObject;
+
+
+namespace Metasound
+{
+	// Forward Declarations
+	using FIterateMetasoundFrontendClassFunction = TFunctionRef<void(const FMetasoundFrontendClass&)>;
+	class FGraph;
+} // namespace Metasound
+
+namespace Metasound::Frontend
+{
+	/** INodeRegistryEntry declares the interface for a node registry entry.
+		* Each node class in the registry must satisfy this interface. 
+		*/
+	class METASOUNDFRONTEND_API INodeRegistryEntry
+	{
+	public:
+		virtual ~INodeRegistryEntry() = default;
+
+		/** Return FNodeClassInfo for the node class.
+			*
+			* Implementations of method should avoid any expensive operations 
+			* (e.g. loading from disk, allocating memory) as this method is called
+			* frequently when querying nodes.
+			*/
+		virtual const FNodeClassInfo& GetClassInfo() const = 0;
+
+		/** Create a node given FDefaultNamedVertexNodeConstructorParams.
+			*
+			* If a node can be created with FDefaultNamedVertexNodeConstructorParams, this function
+			* should return a valid node pointer.
+			*/
+		virtual TUniquePtr<INode> CreateNode(FDefaultNamedVertexNodeConstructorParams&&) const = 0;
+
+		/** Create a node given FDefaultNamedVertexWithLiteralNodeConstructorParams.
+			*
+			* If a node can be created with FDefaultNamedVertexWithLiteralNodeConstructorParams, this function
+			* should return a valid node pointer.
+			*/
+		virtual TUniquePtr<INode> CreateNode(FDefaultNamedVertexWithLiteralNodeConstructorParams&&) const = 0;
+
+		/** Create a node given FDefaultLiteralNodeConstructorParams.
+			*
+			* If a node can be created with FDefaultLiteralNodeConstructorParams, this function
+			* should return a valid node pointer.
+			*/
+		virtual TUniquePtr<INode> CreateNode(FDefaultLiteralNodeConstructorParams&&) const = 0;
+
+		/** Create a node given FNodeInitData.
+			*
+			* If a node can be created with FNodeInitData, this function
+			* should return a valid node pointer.
+			*/
+		virtual TUniquePtr<INode> CreateNode(const FNodeInitData&) const = 0;
+
+		/** Return a FMetasoundFrontendClass which describes the node. */
+		virtual const FMetasoundFrontendClass& GetFrontendClass() const = 0;
+
+		/** Clone this registry entry. */
+		virtual TUniquePtr<INodeRegistryEntry> Clone() const = 0;
+
+		/** Returns set of implemented interface versions.
+			*
+			* Returns nullptr if node class implementation does not support interface implementation.
+			*/
+		virtual const TSet<FMetasoundFrontendVersion>* GetImplementedInterfaces() const = 0;
+
+		/** Whether or not the node is natively defined */
+		virtual bool IsNative() const = 0;
+	};
+
+	struct METASOUNDFRONTEND_API FConverterNodeRegistryKey
+	{
+		// The datatype one would like to convert from.
+		FName FromDataType;
+
+		// The datatype one would like to convert to.
+		FName ToDataType;
+
+		FORCEINLINE bool operator==(const FConverterNodeRegistryKey& Other) const
+		{
+			return FromDataType == Other.FromDataType && ToDataType == Other.ToDataType;
+		}
+
+		friend uint32 GetTypeHash(const ::Metasound::Frontend::FConverterNodeRegistryKey& InKey)
+		{
+			return HashCombine(GetTypeHash(InKey.FromDataType), GetTypeHash(InKey.ToDataType));
+		}
+	};
+
+	struct METASOUNDFRONTEND_API FConverterNodeInfo
+	{
+		// If this node has multiple input pins, we use this to designate which pin should be used.
+		FVertexName PreferredConverterInputPin;
+
+		// If this node has multiple output pins, we use this to designate which pin should be used.
+		FVertexName PreferredConverterOutputPin;
+
+		// The key for this node in the node registry.
+		FNodeRegistryKey NodeKey;
+
+		FORCEINLINE bool operator==(const FConverterNodeInfo& Other) const
+		{
+			return NodeKey == Other.NodeKey;
+		}
+	};
+
+	struct METASOUNDFRONTEND_API FConverterNodeRegistryValue
+	{
+		// A list of nodes that can perform a conversion between the two datatypes described in the FConverterNodeRegistryKey for this map element.
+		TArray<FConverterNodeInfo> PotentialConverterNodes;
+	};
+
+	using FRegistryTransactionID = int32;
+
+	class METASOUNDFRONTEND_API FNodeRegistryTransaction 
+	{
+	public:
+		using FTimeType = uint64;
+
+		/** Describes the type of transaction. */
+		enum class ETransactionType : uint8
+		{
+			NodeRegistration,     //< Something was added to the registry.
+			NodeUnregistration,  //< Something was removed from the registry.
+			Invalid
+		};
+
+		static FString LexToString(ETransactionType InTransactionType)
+		{
+			using namespace Metasound;
+
+			switch (InTransactionType)
+			{
+				case ETransactionType::NodeRegistration:
+					return FString(TEXT("NodeRegistration"));
+
+				case ETransactionType::NodeUnregistration:
+					return FString(TEXT("NodeUnregistration"));
+
+				case ETransactionType::Invalid:
+					return FString(TEXT("Invalid"));
+
+				default:
+				{
+					checkNoEntry();
+				}
+			}
+			return FString(TEXT(""));
+		}
+
+		FNodeRegistryTransaction(ETransactionType InType, const FNodeClassInfo& InNodeClassInfo, FTimeType InTimestamp);
+
+		ETransactionType GetTransactionType() const;
+		const FNodeClassInfo& GetNodeClassInfo() const;
+		FNodeRegistryKey GetNodeRegistryKey() const;
+		FTimeType GetTimestamp() const;
+
+	private:
+		ETransactionType Type;
+		FNodeClassInfo NodeClassInfo;
+		FTimeType Timestamp;
+	};
+
+	namespace NodeRegistryKey
+	{
+		UE_DEPRECATED(5.4, "Use FNodeRegistryKey::GetInvalid() instead")
+		METASOUNDFRONTEND_API const FNodeRegistryKey& GetInvalid();
+
+		UE_DEPRECATED(5.4, "Use FNodeRegistryKey::IsValid() instead")
+		METASOUNDFRONTEND_API bool IsValid(const FNodeRegistryKey& InKey);
+
+		UE_DEPRECATED(5.4, "Use FNodeRegistryKey equality operator instead")
+		METASOUNDFRONTEND_API bool IsEqual(const FNodeRegistryKey& InLHS, const FNodeRegistryKey& InRHS);
+
+		// Returns true if the class metadata represent the same entry in the node registry.
+		METASOUNDFRONTEND_API bool IsEqual(const FMetasoundFrontendClassMetadata& InLHS, const FMetasoundFrontendClassMetadata& InRHS);
+
+		// Returns true if the class info and class metadata represent the same entry in the node registry.
+		METASOUNDFRONTEND_API bool IsEqual(const FNodeClassInfo& InLHS, const FMetasoundFrontendClassMetadata& InRHS);
+
+		UE_DEPRECATED(5.4, "Use FNodeRegistryKey constructors instead")
+		METASOUNDFRONTEND_API FNodeRegistryKey CreateKey(EMetasoundFrontendClassType InType, const FString& InFullClassName, int32 InMajorVersion, int32 InMinorVersion);
+
+		UE_DEPRECATED(5.4, "Use FNodeRegistryKey constructors instead")
+		METASOUNDFRONTEND_API FNodeRegistryKey CreateKey(const FNodeClassMetadata& InNodeMetadata);
+
+		UE_DEPRECATED(5.4, "Use FNodeRegistryKey constructor instead")
+		METASOUNDFRONTEND_API FNodeRegistryKey CreateKey(const FMetasoundFrontendClassMetadata& InNodeMetadata);
+
+		UE_DEPRECATED(5.4, "Use FNodeRegistryKey constructor instead")
+		METASOUNDFRONTEND_API FNodeRegistryKey CreateKey(const FMetasoundFrontendGraphClass& InNodeMetadata);
+
+		UE_DEPRECATED(5.4, "Use FNodeRegistryKey constructor instead")
+		METASOUNDFRONTEND_API FNodeRegistryKey CreateKey(const FNodeClassInfo& ClassInfo);
+	}
+} // namespace Metasound::Frontend
+
+/**
+ * Singleton registry for all types and nodes.
+ */
+class METASOUNDFRONTEND_API FMetasoundFrontendRegistryContainer
+{
+
+public:
+	// The MetaSound frontend does not rely on the Engine module and therefore
+	// does not have the ability to provide GC protection. This interface allows
+	// external modules to provide GC protection so that async tasks can safely
+	// use UObjects
+	class IObjectReferencer
+	{
+	public:
+		virtual ~IObjectReferencer() = default;
+
+		// Called when an object should be referenced.
+		virtual void AddObject(UObject* InObject) = 0;
+
+		// Called when an object no longer needs to be referenced. 
+		virtual void RemoveObject(UObject* InObject) = 0;
+	};
+
+	using FNodeClassInfo = Metasound::Frontend::FNodeClassInfo;
+	using FConverterNodeRegistryKey = ::Metasound::Frontend::FConverterNodeRegistryKey;
+	using FConverterNodeRegistryValue = ::Metasound::Frontend::FConverterNodeRegistryValue;
+	using FConverterNodeInfo = ::Metasound::Frontend::FConverterNodeInfo;
+
+	using FNodeRegistryKey = Metasound::Frontend::FNodeRegistryKey;
+	using FNodeClassMetadata = Metasound::FNodeClassMetadata;
+
+	static FMetasoundFrontendRegistryContainer* Get();
+	static void ShutdownMetasoundFrontend();
+
+	UE_DEPRECATED(5.4, "Use FNodeRegistryKey constructor that takes FNodeClassMetadata instead")
+	static FNodeRegistryKey GetRegistryKey(const FNodeClassMetadata& InNodeMetadata);
+
+	UE_DEPRECATED(5.4, "Use FNodeRegistryKey constructor that takes FMetasoundFrontendClassMetadata instead")
+	static FNodeRegistryKey GetRegistryKey(const FMetasoundFrontendClassMetadata& InNodeMetadata);
+
+	UE_DEPRECATED(5.4, "Use FNodeRegistryKey constructor that takes FNodeClassInfo instead")
+	static FNodeRegistryKey GetRegistryKey(const FNodeClassInfo& ClassInfo);
+
+	static bool GetFrontendClassFromRegistered(const FNodeRegistryKey& InKey, FMetasoundFrontendClass& OutClass);
+
+	UE_DEPRECATED(5.4, "Use GetFrontendClassFromRegistered instead")
+	static bool GetNodeClassInfoFromRegistered(const FNodeRegistryKey& InKey, FNodeClassInfo& OutInfo);
+
+	static bool GetInputNodeRegistryKeyForDataType(const FName& InDataTypeName, const EMetasoundFrontendVertexAccessType InAccessType, FNodeRegistryKey& OutKey);
+	static bool GetVariableNodeRegistryKeyForDataType(const FName& InDataTypeName, FNodeRegistryKey& OutKey);
+	static bool GetOutputNodeRegistryKeyForDataType(const FName& InDataTypeName, const EMetasoundFrontendVertexAccessType InAccessType, FNodeRegistryKey& OutKey);
+
+
+	FMetasoundFrontendRegistryContainer() = default;
+	virtual ~FMetasoundFrontendRegistryContainer() = default;
+
+	FMetasoundFrontendRegistryContainer(const FMetasoundFrontendRegistryContainer&) = delete;
+	FMetasoundFrontendRegistryContainer& operator=(const FMetasoundFrontendRegistryContainer&) = delete;
+
+	// Enqueue and init command for registering a node or data type.
+	// The command queue will be processed on module init or when calling `RegisterPendingNodes()`
+	virtual bool EnqueueInitCommand(TUniqueFunction<void()>&& InFunc) = 0;
+
+	virtual void SetObjectReferencer(TUniquePtr<IObjectReferencer> InReferencer) = 0;
+
+	// This is called on module startup. This invokes any registration commands enqueued by our registration macros.
+	virtual void RegisterPendingNodes() = 0;
+
+	// Register a graph from an IMetaSoundDocumentInterface
+	virtual FNodeRegistryKey RegisterGraph(const TScriptInterface<IMetaSoundDocumentInterface>& InDocument, bool bAsync = true) = 0;
+
+	// Wait for async graph registration to complete for a specific graph
+	virtual void WaitForAsyncGraphRegistration(const FNodeRegistryKey& InRegistryKey, const FTopLevelAssetPath& InAssetPath) const = 0;
+
+	// Unregister an asset based graph from the frontend.
+	virtual bool UnregisterGraph(const TScriptInterface<IMetaSoundDocumentInterface>& InDocument, bool bAsync = true) = 0;
+
+	// Retrieve a registered graph. 
+	//
+	// If the graph is registered asynchronously, this will wait until the registration task has completed.
+	virtual TSharedPtr<const Metasound::FGraph> GetGraph(const FNodeRegistryKey& InRegistryKey, const FTopLevelAssetPath& InAssetPath) const = 0;
+
+	/** Register an external node with the frontend.
+	 *
+	 * @param InCreateNode - Function for creating node from FNodeInitData.
+	 * @param InCreateDescription - Function for creating a FMetasoundFrontendClass.
+	 *
+	 * @return A node registration key. If the registration failed, then the registry 
+	 *         key will be invalid.
+	 */
+	virtual FNodeRegistryKey RegisterNode(TUniquePtr<Metasound::Frontend::INodeRegistryEntry>&& InEntry) = 0;
+
+	/** Unregister an external node from the frontend.
+	 *
+	 * @param InKey - The registration key for the node.
+	 *
+	 * @return True on success, false on failure.
+	 */
+	virtual bool UnregisterNode(const FNodeRegistryKey& InKey) = 0;
+
+	/** Returns true if the provided registry key corresponds to a valid registered node. */
+	virtual bool IsNodeRegistered(const FNodeRegistryKey& InKey) const = 0;
+
+	/** Returns true if the provided registry key and asset path corresponds to a valid registered graph. */
+	virtual bool IsGraphRegistered(const FNodeRegistryKey& InKey, const FTopLevelAssetPath& InAssetPath) const = 0;
+
+	/** Returns true if the provided registry key corresponds to a valid registered node that is natively defined. */
+	virtual bool IsNodeNative(const FNodeRegistryKey& InKey) const = 0;
+
+	// Iterates class types in registry.  If InClassType is set to a valid class type (optional), only iterates classes of the given type
+	virtual void IterateRegistry(Metasound::FIterateMetasoundFrontendClassFunction InIterFunc, EMetasoundFrontendClassType InClassType = EMetasoundFrontendClassType::Invalid) const = 0;
+
+	// Query for MetaSound Frontend document objects.
+	virtual bool FindFrontendClassFromRegistered(const Metasound::Frontend::FNodeRegistryKey& InKey, FMetasoundFrontendClass& OutClass) = 0;
+
+	UE_DEPRECATED(5.4, "This implementation of FindImplementedInterfacesFromRegistered(...). Please use bool FindImplementedInterfacesFromRegistered(const Metasound::Frontend::FNodeRegistryKey& InKey, TSet<FMetasoundFrontendVersion>& OutInterfacVersion)")
+	virtual const TSet<FMetasoundFrontendVersion>* FindImplementedInterfacesFromRegistered(const Metasound::Frontend::FNodeRegistryKey& InKey) const = 0;
+
+	virtual bool FindImplementedInterfacesFromRegistered(const Metasound::Frontend::FNodeRegistryKey& InKey, TSet<FMetasoundFrontendVersion>& OutInterfaceVersions) const = 0;
+	
+	UE_DEPRECATED(5.4, "Use FindFrontendClassFromRegistered instead")
+	virtual bool FindNodeClassInfoFromRegistered(const Metasound::Frontend::FNodeRegistryKey& InKey, FNodeClassInfo& OutInfo) = 0;
+	virtual bool FindInputNodeRegistryKeyForDataType(const FName& InDataTypeName, const EMetasoundFrontendVertexAccessType InAccessType, FNodeRegistryKey& OutKey) = 0;
+	virtual bool FindVariableNodeRegistryKeyForDataType(const FName& InDataTypeName, FNodeRegistryKey& OutKey) = 0;
+	virtual bool FindOutputNodeRegistryKeyForDataType(const FName& InDataTypeName, const EMetasoundFrontendVertexAccessType InAccessType, FNodeRegistryKey& OutKey) = 0;
+
+	virtual TUniquePtr<Metasound::INode> CreateNode(const FNodeRegistryKey& InKey, const Metasound::FNodeInitData&) const = 0;
+	virtual TUniquePtr<Metasound::INode> CreateNode(const FNodeRegistryKey& InKey, Metasound::FDefaultLiteralNodeConstructorParams&&) const = 0;
+	virtual TUniquePtr<Metasound::INode> CreateNode(const FNodeRegistryKey& InKey, Metasound::FDefaultNamedVertexNodeConstructorParams&&) const = 0;
+	virtual TUniquePtr<Metasound::INode> CreateNode(const FNodeRegistryKey& InKey, Metasound::FDefaultNamedVertexWithLiteralNodeConstructorParams&&) const = 0;
+
+	virtual bool RegisterConversionNode(const FConverterNodeRegistryKey& InNodeKey, const FConverterNodeInfo& InNodeInfo) = 0;
+
+	// Returns a list of possible nodes to use to convert from FromDataType to ToDataType.
+	// Returns an empty array if none are available.
+	virtual TArray<FConverterNodeInfo> GetPossibleConverterNodes(const FName& FromDataType, const FName& ToDataType) = 0;
+};
+
+
diff --git a/Engine/Plugins/Runtime/Metasound/Source/MetasoundFrontend/Public/MetasoundFrontendRegistryKey.h b/Engine/Plugins/Runtime/Metasound/Source/MetasoundFrontend/Public/MetasoundFrontendRegistryKey.h
--- a/Engine/Plugins/Runtime/Metasound/Source/MetasoundFrontend/Public/MetasoundFrontendRegistryKey.h
+++ b/Engine/Plugins/Runtime/Metasound/Source/MetasoundFrontend/Public/MetasoundFrontendRegistryKey.h
@@ -9,9 +9,8 @@
 
 namespace Metasound::Frontend
 {
-	/** FNodeClassInfo contains a minimal set of information needed to find
-	 * and query asset node classes without loading the asset.
-	 */
+	/** FNodeClassInfo contains a minimal set of information needed without having to load an asset.
+	  */
 	struct METASOUNDFRONTEND_API FNodeClassInfo
 	{
 		// ClassName of the given class
@@ -24,7 +23,7 @@ namespace Metasound::Frontend
 		FGuid AssetClassID;
 
 		// Path to asset containing graph if external type and references asset class.
-		FSoftObjectPath AssetPath;
+		FTopLevelAssetPath AssetPath;
 
 		// Version of the registered class
 		FMetasoundFrontendVersionNumber Version;
@@ -42,12 +41,17 @@ namespace Metasound::Frontend
 
 		FNodeClassInfo() = default;
 
-		// Constructor used to generate NodeClassInfo from a native class' Metadata.
+		// Constructor used to generate NodeClassInfo from a class' Metadata.
+		// (Does not cache AssetPath and thus may not support loading asset
+		// should the class originate from one).
 		FNodeClassInfo(const FMetasoundFrontendClassMetadata& InMetadata);
 
-		// Constructor used to generate NodeClassInfo from an asset
+		UE_DEPRECATED(5.4, "Use constructor that takes an FTopLevelAssetPath instead.")
 		FNodeClassInfo(const FMetasoundFrontendGraphClass& InClass, const FSoftObjectPath& InAssetPath);
 
+		// Constructor used to generate NodeClassInfo from an asset
+		FNodeClassInfo(const FMetasoundFrontendGraphClass& InClass, const FTopLevelAssetPath& InAssetPath);
+
 		// Loads the asset from the provided path, ensuring that the class is of type graph.
 		UObject* LoadAsset() const;
 	};
diff --git a/Engine/Plugins/Runtime/Metasound/Source/MetasoundGenerator/Private/MetasoundOperatorCache.cpp b/Engine/Plugins/Runtime/Metasound/Source/MetasoundGenerator/Private/MetasoundOperatorCache.cpp
--- a/Engine/Plugins/Runtime/Metasound/Source/MetasoundGenerator/Private/MetasoundOperatorCache.cpp
+++ b/Engine/Plugins/Runtime/Metasound/Source/MetasoundGenerator/Private/MetasoundOperatorCache.cpp
@@ -163,7 +163,7 @@ namespace Metasound
 	FOperatorBuildData::FOperatorBuildData(
 		  FMetasoundGeneratorInitParams&& InInitParams
 		, Frontend::FNodeRegistryKey InRegistryKey
-		, FSoftObjectPath InAssetPath
+		, FTopLevelAssetPath InAssetPath
 		, FGuid InAssetClassID
 		, int32 InNumInstances
 	)
@@ -172,7 +172,8 @@ namespace Metasound
 	, AssetPath(InAssetPath)
 	, AssetClassID(InAssetClassID)
 	, NumInstances(InNumInstances)
-	{}
+	{
+	}
 
 
 	FOperatorPool::FOperatorPool(const FOperatorPoolSettings& InSettings)
diff --git a/Engine/Plugins/Runtime/Metasound/Source/MetasoundGenerator/Public/MetasoundOperatorCache.h b/Engine/Plugins/Runtime/Metasound/Source/MetasoundGenerator/Public/MetasoundOperatorCache.h
--- a/Engine/Plugins/Runtime/Metasound/Source/MetasoundGenerator/Public/MetasoundOperatorCache.h
+++ b/Engine/Plugins/Runtime/Metasound/Source/MetasoundGenerator/Public/MetasoundOperatorCache.h
@@ -67,7 +67,7 @@ namespace Metasound
 	{
 		FMetasoundGeneratorInitParams InitParams;
 		Frontend::FNodeRegistryKey RegistryKey;
-		FSoftObjectPath AssetPath;
+		FTopLevelAssetPath AssetPath;
 		FGuid AssetClassID;
 		int32 NumInstances;
 
@@ -75,7 +75,7 @@ namespace Metasound
 		FOperatorBuildData(
 			  FMetasoundGeneratorInitParams&& InInitParams
 			, Frontend::FNodeRegistryKey InRegistryKey
-			, FSoftObjectPath InAssetPath
+			, FTopLevelAssetPath InAssetPath
 			, FGuid InAssetID
 			, int32 InNumInstances = 1
 		);
