Fix quality settings named "None" causing issues with Defaults.

diff --git a/Engine/Plugins/Runtime/Metasound/Source/MetasoundEngine/Private/MetasoundSettings.cpp b/Engine/Plugins/Runtime/Metasound/Source/MetasoundEngine/Private/MetasoundSettings.cpp
--- a/Engine/Plugins/Runtime/Metasound/Source/MetasoundEngine/Private/MetasoundSettings.cpp
+++ b/Engine/Plugins/Runtime/Metasound/Source/MetasoundEngine/Private/MetasoundSettings.cpp
@@ -41,8 +41,14 @@ namespace MetaSoundSettingsPrivate
 	}
 	static void OnRenameQualitySettings(FMetaSoundQualitySettings& InRenamed)
 	{
+		// Prevent 'None' as an option.
+		if (InRenamed.Name.IsNone())
+		{
+			InRenamed.Name = GenerateNextName(TEXT("New Quality"));
+		}
+		
 		// More than one?
-		if (Algo::Count(UMetaSoundQualityHelper::GetQualityList(), InRenamed.Name) > 1)
+		else if (Algo::Count(UMetaSoundQualityHelper::GetQualityList(), InRenamed.Name) > 1)
 		{
 			// add something unique to the name.
 			InRenamed.Name = GenerateNextName(*InRenamed.Name.ToString());
