Fix MetaSound warning for external node classes that should only apply for asset classes

diff --git a/Engine/Plugins/Runtime/Metasound/Source/MetasoundEngine/Private/MetasoundAssetSubsystem.cpp b/Engine/Plugins/Runtime/Metasound/Source/MetasoundEngine/Private/MetasoundAssetSubsystem.cpp
--- a/Engine/Plugins/Runtime/Metasound/Source/MetasoundEngine/Private/MetasoundAssetSubsystem.cpp
+++ b/Engine/Plugins/Runtime/Metasound/Source/MetasoundEngine/Private/MetasoundAssetSubsystem.cpp
@@ -11,6 +11,7 @@
 #include "MetasoundBuilderSubsystem.h"
 #include "MetasoundFrontendDocument.h"
 #include "MetasoundFrontendRegistries.h"
+#include "MetasoundFrontendSearchEngine.h"
 #include "MetasoundFrontendTransform.h"
 #include "MetasoundSettings.h"
 #include "MetasoundSource.h"
@@ -356,7 +357,15 @@ TSet<UMetaSoundAssetSubsystem::FAssetInfo> UMetaSoundAssetSubsystem::GetReferenc
 			}
 			else
 			{
-				bReportFail = true;
+				// Don't report failure if a matching class with a matching major version and higher minor version exists (it will be autoupdated) 
+				FMetasoundFrontendClass FrontendClass;
+				const bool bDidFindClassWithName = ISearchEngine::Get().FindClassWithHighestVersion(Key.ClassName.ToNodeClassName(), FrontendClass);
+				if (!(bDidFindClassWithName && 
+					Key.Version.Major == FrontendClass.Metadata.GetVersion().Major && 
+					Key.Version.Minor < FrontendClass.Metadata.GetVersion().Minor))
+				{
+					bReportFail = true;
+				}
 			}
 
 			if (bReportFail)
