Don't warn for missing unregistration tasks in cook, where graphs are not registered
#rnx

diff --git a/Engine/Plugins/Runtime/Metasound/Source/MetasoundFrontend/Private/MetasoundFrontendRegistryContainerImpl.cpp b/Engine/Plugins/Runtime/Metasound/Source/MetasoundFrontend/Private/MetasoundFrontendRegistryContainerImpl.cpp
--- a/Engine/Plugins/Runtime/Metasound/Source/MetasoundFrontend/Private/MetasoundFrontendRegistryContainerImpl.cpp
+++ b/Engine/Plugins/Runtime/Metasound/Source/MetasoundFrontend/Private/MetasoundFrontendRegistryContainerImpl.cpp
@@ -17,6 +17,7 @@
 #include "UObject/NameTypes.h"
 #include "UObject/ScriptInterface.h"
 
+#include "MetasoundAssetManager.h"
 #include "MetasoundDocumentInterface.h"
 #include "MetasoundFrontendDataTypeRegistry.h"
 #include "MetasoundFrontendGraph.h"
@@ -27,7 +28,6 @@
 #include "MetasoundLog.h"
 #include "MetasoundRouter.h"
 #include "MetasoundTrace.h"
-#include "../Public/MetasoundAssetManager.h"
 
 bool bBusyWaitOnAsyncRegistrationTasks = true;
 static FAutoConsoleVariableRef CVarAsyncRegistrationTasksBusyWait(
@@ -491,7 +491,8 @@ namespace Metasound::Frontend
 			}
 		}
 
-		if (NumRemoved == 0)
+
+		if (NumRemoved == 0 && !IsRunningCookCommandlet())
 		{
 			UE_LOG(LogMetaSound, Warning,
 				TEXT("Failed to find active %s tasks for the graph '%s' "),
