Change MetaSound trace macros to use scoped events for platform specific captures

diff --git a/Engine/Plugins/Runtime/Metasound/Source/MetasoundGraphCore/Public/MetasoundTrace.h b/Engine/Plugins/Runtime/Metasound/Source/MetasoundGraphCore/Public/MetasoundTrace.h
--- a/Engine/Plugins/Runtime/Metasound/Source/MetasoundGraphCore/Public/MetasoundTrace.h
+++ b/Engine/Plugins/Runtime/Metasound/Source/MetasoundGraphCore/Public/MetasoundTrace.h
@@ -16,12 +16,14 @@
 #if METASOUND_CPUPROFILERTRACE_ENABLED
 // Metasound CPU profiler trace enabled
 
-#define METASOUND_TRACE_CPUPROFILER_EVENT_SCOPE(Name) \
-	TRACE_CPUPROFILER_EVENT_SCOPE(Name)
+// Copied from SCOPED_NAMED_EVENT but modified
+// to accommodate event names containing ::
+#define METASOUND_TRACE_CPUPROFILER_EVENT_SCOPE(Name)\
+	FScopedNamedEventStatic PREPROCESSOR_JOIN(MetaSound_NamedEvent_,__LINE__)(FColor::Green, NAMED_EVENT_STR(#Name));\
+	TRACE_CPUPROFILER_EVENT_SCOPE(Name);
 
 #define METASOUND_TRACE_CPUPROFILER_EVENT_SCOPE_TEXT(Name) \
-	TRACE_CPUPROFILER_EVENT_SCOPE_TEXT(Name)
-
+    SCOPED_NAMED_EVENT_TCHAR(Name, FColor::Green)
 #else
 // Metasound CPU profiler trace *not* enabled
 
