Remove and assert against registration/unregisteration calls of MetaSound graphs when cooking, and make all node unregistration calls synchronous when cooking
#rnx
#tests AudioUnitTests, AudioQA -run=Cook -TargetPlatform=WindowsClient, rename asset, run Engine 'Audio' Automated Tests

diff --git a/Engine/Plugins/Runtime/Metasound/Source/MetasoundEngine/Private/Metasound.cpp b/Engine/Plugins/Runtime/Metasound/Source/MetasoundEngine/Private/Metasound.cpp
--- a/Engine/Plugins/Runtime/Metasound/Source/MetasoundEngine/Private/Metasound.cpp
+++ b/Engine/Plugins/Runtime/Metasound/Source/MetasoundEngine/Private/Metasound.cpp
@@ -118,7 +118,7 @@ void UMetaSoundPatch::PostEditUndo()
 
 void UMetaSoundPatch::BeginDestroy()
 {
-	UnregisterGraphWithFrontend();
+	OnNotifyBeginDestroy();
 	Super::BeginDestroy();
 }
 
diff --git a/Engine/Plugins/Runtime/Metasound/Source/MetasoundEngine/Private/MetasoundAssetSubsystem.cpp b/Engine/Plugins/Runtime/Metasound/Source/MetasoundEngine/Private/MetasoundAssetSubsystem.cpp
--- a/Engine/Plugins/Runtime/Metasound/Source/MetasoundEngine/Private/MetasoundAssetSubsystem.cpp
+++ b/Engine/Plugins/Runtime/Metasound/Source/MetasoundEngine/Private/MetasoundAssetSubsystem.cpp
@@ -646,19 +646,15 @@ void UMetaSoundAssetSubsystem::RenameAsset(const FAssetData& InAssetData, bool b
 		AddOrUpdateAsset(InAssetData);
 	};
 
+	FMetasoundAssetBase* MetaSoundAsset = Metasound::IMetasoundUObjectRegistry::Get().GetObjectAsAssetBase(InAssetData.GetAsset());
+	check(MetaSoundAsset);
+	MetaSoundAsset->UnregisterGraphWithFrontend();
+	PerformRename();
+
 	if (bInReregisterWithFrontend)
 	{
-		FMetasoundAssetBase* MetaSoundAsset = Metasound::IMetasoundUObjectRegistry::Get().GetObjectAsAssetBase(InAssetData.GetAsset());
-		check(MetaSoundAsset);
-
-		MetaSoundAsset->UnregisterGraphWithFrontend();
-		PerformRename();
 		MetaSoundAsset->RegisterGraphWithFrontend();
 	}
-	else
-	{
-		PerformRename();
-	}
 }
 
 void UMetaSoundAssetSubsystem::ResetAssetClassDisplayName(const FAssetData& InAssetData)
diff --git a/Engine/Plugins/Runtime/Metasound/Source/MetasoundEngine/Private/MetasoundEngineAsset.h b/Engine/Plugins/Runtime/Metasound/Source/MetasoundEngine/Private/MetasoundEngineAsset.h
--- a/Engine/Plugins/Runtime/Metasound/Source/MetasoundEngine/Private/MetasoundEngineAsset.h
+++ b/Engine/Plugins/Runtime/Metasound/Source/MetasoundEngine/Private/MetasoundEngineAsset.h
@@ -6,6 +6,7 @@
 #include "MetasoundAssetManager.h"
 #include "MetasoundBuilderSubsystem.h"
 #include "MetasoundFrontendDocumentIdGenerator.h"
+#include "MetasoundFrontendRegistryKey.h"
 #include "MetasoundUObjectRegistry.h"
 #include "Serialization/Archive.h"
 
@@ -40,9 +41,9 @@ namespace Metasound
 		}
 
 		template<typename TMetaSoundObject>
-		static void SetReferencedAssetClasses(TMetaSoundObject& InMetaSound, TSet<Metasound::Frontend::IMetaSoundAssetManager::FAssetInfo>&& InAssetClasses)
+		static void SetReferencedAssetClasses(TMetaSoundObject& InMetaSound, TSet<Frontend::IMetaSoundAssetManager::FAssetInfo>&& InAssetClasses)
 		{
-			using namespace Metasound::Frontend;
+			using namespace Frontend;
 			
 			InMetaSound.ReferencedAssetClassKeys.Reset();
 			InMetaSound.ReferencedAssetClassObjects.Reset();
@@ -97,7 +98,7 @@ namespace Metasound
 		static void PreSaveAsset(TMetaSoundObject& InMetaSound, FObjectPreSaveContext InSaveContext)
 		{
 #if WITH_EDITORONLY_DATA
-			using namespace Metasound::Frontend;
+			using namespace Frontend;
 
 			// Do not call asset manager on CDO objects which may be loaded before asset 
 			// manager is set.
@@ -187,10 +188,9 @@ namespace Metasound
 
 #if WITH_EDITORONLY_DATA
 		template <typename TMetaSoundObject>
-		static void SetMetaSoundRegistryAssetClassInfo(TMetaSoundObject& InMetaSound, const Metasound::Frontend::FNodeClassInfo& InClassInfo)
+		static void SetMetaSoundRegistryAssetClassInfo(TMetaSoundObject& InMetaSound, const Frontend::FNodeClassInfo& InClassInfo)
 		{
-			using namespace Metasound;
-			using namespace Metasound::Frontend;
+			using namespace Frontend;
 
 			check(AssetTags::AssetClassID == GET_MEMBER_NAME_CHECKED(TMetaSoundObject, AssetClassID));
 			check(AssetTags::IsPreset == GET_MEMBER_NAME_CHECKED(TMetaSoundObject, bIsPreset));
diff --git a/Engine/Plugins/Runtime/Metasound/Source/MetasoundEngine/Private/MetasoundSource.cpp b/Engine/Plugins/Runtime/Metasound/Source/MetasoundEngine/Private/MetasoundSource.cpp
--- a/Engine/Plugins/Runtime/Metasound/Source/MetasoundEngine/Private/MetasoundSource.cpp
+++ b/Engine/Plugins/Runtime/Metasound/Source/MetasoundEngine/Private/MetasoundSource.cpp
@@ -497,7 +497,7 @@ FTopLevelAssetPath UMetaSoundSource::GetAssetPathChecked() const
 
 void UMetaSoundSource::BeginDestroy()
 {
-	UnregisterGraphWithFrontend();
+	OnNotifyBeginDestroy();
 	Super::BeginDestroy();
 }
 
diff --git a/Engine/Plugins/Runtime/Metasound/Source/MetasoundFrontend/Private/MetasoundAssetBase.cpp b/Engine/Plugins/Runtime/Metasound/Source/MetasoundFrontend/Private/MetasoundAssetBase.cpp
--- a/Engine/Plugins/Runtime/Metasound/Source/MetasoundFrontend/Private/MetasoundAssetBase.cpp
+++ b/Engine/Plugins/Runtime/Metasound/Source/MetasoundFrontend/Private/MetasoundAssetBase.cpp
@@ -265,6 +265,7 @@ void FMetasoundAssetBase::RegisterGraphWithFrontend(Metasound::Frontend::FMetaSo
 
 	// Graph registration must only happen on one thread to avoid race conditions on graph registration.
 	checkf(IsInGameThread(), TEXT("MetaSound %s graph can only be registered on the GameThread"), *GetOwningAssetName());
+	checkf(!IsRunningCookCommandlet(), TEXT("Cook of asset must call RegisterNode directly providing FDocumentNodeRegistryEntryForCook to avoid proxy/runtime graph generation."));
 
 	METASOUND_TRACE_CPUPROFILER_EVENT_SCOPE(MetaSoundAssetBase::RegisterGraphWithFrontend);
 	METASOUND_TRACE_CPUPROFILER_EVENT_SCOPE_TEXT(*FString::Printf(TEXT("MetaSoundAssetBase::RegisterGraphWithFrontend asset %s"), *this->GetOwningAssetName()));
@@ -321,9 +322,11 @@ void FMetasoundAssetBase::RegisterGraphWithFrontend(Metasound::Frontend::FMetaSo
 	check(Owner);
 	const FString AssetName = Owner->GetName();
 
-	// If the IMetaSoundDocumentInterface actively modified by a builder, then
-	// we must build synchronously to avoid a race condition on reading/writing
-	// the IMetaSoundDocumentInterface
+	// Async registration is only available if:
+	// 1. The IMetaSoundDocumentInterface is not actively modified by a builder
+	//    (built graph must be released synchronously to avoid a race condition on
+	//    reading/writing the IMetaSoundDocumentInterface on the Game Thread)
+	// 2. Async registration is globally disabled via console variable.
 	const bool bAsync = !(IsBuilderActive() || ConsoleVariables::bDisableAsyncGraphRegistration);
 	const TScriptInterface<IMetaSoundDocumentInterface> RegistryDocInterface = AssetBasePrivate::BuildRegistryDocument(Owner);
 	GraphRegistryKey = FRegistryContainerImpl::Get().RegisterGraph(RegistryDocInterface, bAsync);
@@ -384,16 +387,21 @@ void FMetasoundAssetBase::CookMetaSound()
 			DocBuilder.TransformTemplateNodes();
 		}
 
+		if (GraphRegistryKey.IsValid())
+		{
+			FRegistryContainerImpl::Get().UnregisterNode(GraphRegistryKey.NodeKey);
+			GraphRegistryKey = { };
+		}
+
 		// During cook, we need to register the node so that it is available for other graphs, but we need to avoid
 		// creating proxies. To do so, we use a special node registration object which reflects the necessary information
 		// for the node registry, but does not create INodes.
-		UnregisterGraphWithFrontend();
 		TScriptInterface<IMetaSoundDocumentInterface> DocInterface(Owner);
 		const FMetasoundFrontendDocument& Document = DocInterface->GetConstDocument();
 		const FTopLevelAssetPath AssetPath = DocInterface->GetAssetPathChecked();
 		TUniquePtr<INodeRegistryEntry> RegistryEntry = MakeUnique<AssetBasePrivate::FDocumentNodeRegistryEntryForCook>(Document, AssetPath);
 
-		const FNodeRegistryKey NodeKey = FMetasoundFrontendRegistryContainer::Get()->RegisterNode(MoveTemp(RegistryEntry));
+		const FNodeRegistryKey NodeKey = FRegistryContainerImpl::Get().RegisterNode(MoveTemp(RegistryEntry));
 		GraphRegistryKey = FGraphRegistryKey { NodeKey, AssetPath };
 	}
 
@@ -412,20 +420,44 @@ void FMetasoundAssetBase::CookMetaSound()
 	}
 }
 
+void FMetasoundAssetBase::OnNotifyBeginDestroy()
+{
+	using namespace Metasound::Frontend;
+
+	// Unregistration of graph is not necessary when cooking as deserialized objects are not mutable and, should they be reloaded,
+	// omitting unregistration avoids potentially kicking off an invalid asynchronous task to unregister a non-existent runtime graph.
+	if (IsRunningCookCommandlet())
+	{
+		if (GraphRegistryKey.IsValid())
+		{
+			FRegistryContainerImpl::Get().UnregisterNode(GraphRegistryKey.NodeKey);
+			GraphRegistryKey = { };
+		}
+	}
+	else
+	{
+		UnregisterGraphWithFrontend();
+	}
+}
+
 void FMetasoundAssetBase::UnregisterGraphWithFrontend()
 {
 	using namespace Metasound::Frontend;
 	METASOUND_TRACE_CPUPROFILER_EVENT_SCOPE(MetaSoundAssetBase::UnregisterGraphWithFrontend);
 
 	check(IsInGameThread());
+	checkf(!IsRunningCookCommandlet(), TEXT("Cook of asset must call UnregisterNode directly providing FDocumentNodeRegistryEntryForCook to avoid proxy/runtime graph generation."));
+
 	if (GraphRegistryKey.IsValid())
 	{
 		UObject* OwningAsset = GetOwningAsset();
 		if (ensureAlways(OwningAsset))
 		{
-			// If the IMetaSoundDocumentInterface is actively modified by a builder, then
-			// we must release built graph synchronously to avoid a race condition on
-			// reading/writing the IMetaSoundDocumentInterface
+			// Async registration is only available if:
+			// 1. The IMetaSoundDocumentInterface is not actively modified by a builder
+			//    (built graph must be released synchronously to avoid a race condition on
+			//    reading/writing the IMetaSoundDocumentInterface on the Game Thread)
+			// 2. Async registration is globally disabled via console variable.
 			const bool bAsync = !(IsBuilderActive() || ConsoleVariables::bDisableAsyncGraphRegistration);
 			const bool bSuccess = FRegistryContainerImpl::Get().UnregisterGraph(GraphRegistryKey, OwningAsset, bAsync);
 			if (!bSuccess)
diff --git a/Engine/Plugins/Runtime/Metasound/Source/MetasoundFrontend/Private/MetasoundFrontendRegistryContainerImpl.cpp b/Engine/Plugins/Runtime/Metasound/Source/MetasoundFrontend/Private/MetasoundFrontendRegistryContainerImpl.cpp
--- a/Engine/Plugins/Runtime/Metasound/Source/MetasoundFrontend/Private/MetasoundFrontendRegistryContainerImpl.cpp
+++ b/Engine/Plugins/Runtime/Metasound/Source/MetasoundFrontend/Private/MetasoundFrontendRegistryContainerImpl.cpp
@@ -492,12 +492,19 @@ namespace Metasound::Frontend
 		}
 
 
-		if (NumRemoved == 0 && !IsRunningCookCommandlet())
+		if (NumRemoved == 0)
 		{
-			UE_LOG(LogMetaSound, Warning,
-				TEXT("Failed to find active %s tasks for the graph '%s' "),
+			const bool bIsCooking = IsRunningCookCommandlet();
+			if (ensureMsgf(!bIsCooking,
+				TEXT("Failed to find active %s tasks for the graph '%s': Async registration is not supported while cooking"),
 				*FNodeRegistryTransaction::LexToString(TransactionType),
-				*InKey.ToString());
+				*InKey.ToString()))
+			{
+				UE_LOG(LogMetaSound, Warning,
+					TEXT("Failed to find active %s tasks for the graph '%s'."),
+					*FNodeRegistryTransaction::LexToString(TransactionType),
+					*InKey.ToString());
+			}
 		}
 	}
 
@@ -556,8 +563,8 @@ namespace Metasound::Frontend
 			}
 			else
 			{
-				// Avoid warning if in cook as we always expect a graph to not get registered/unregistered
-				// while cooking (as its unnecessary for serialization)
+				// Avoid warning if in cook as we always expect a graph to not get registered/
+				// unregistered while cooking (as its unnecessary for serialization).
 				if (bNodeUnregistered && !IsRunningCookCommandlet())
 				{
 					UE_LOG(LogMetaSound, Warning, TEXT("Graph '%s' was not found, but analogous registered node class was when unregistering."), *InKey.ToString());
diff --git a/Engine/Plugins/Runtime/Metasound/Source/MetasoundFrontend/Public/MetasoundAssetBase.h b/Engine/Plugins/Runtime/Metasound/Source/MetasoundFrontend/Public/MetasoundAssetBase.h
--- a/Engine/Plugins/Runtime/Metasound/Source/MetasoundFrontend/Public/MetasoundAssetBase.h
+++ b/Engine/Plugins/Runtime/Metasound/Source/MetasoundFrontend/Public/MetasoundAssetBase.h
@@ -197,6 +197,8 @@ public:
 	FString GetOwningAssetName() const;
 
 protected:
+	void OnNotifyBeginDestroy();
+
 #if WITH_EDITOR
 	virtual void SetReferencedAssetClasses(TSet<Metasound::Frontend::IMetaSoundAssetManager::FAssetInfo>&& InAssetClasses) = 0;
 #endif
