Adding audio csv stats around tracking potential perf culprits.

diff --git a/Engine/Plugins/Runtime/Metasound/Source/MetasoundGenerator/Private/MetasoundGenerator.cpp b/Engine/Plugins/Runtime/Metasound/Source/MetasoundGenerator/Private/MetasoundGenerator.cpp
--- a/Engine/Plugins/Runtime/Metasound/Source/MetasoundGenerator/Private/MetasoundGenerator.cpp
+++ b/Engine/Plugins/Runtime/Metasound/Source/MetasoundGenerator/Private/MetasoundGenerator.cpp
@@ -26,6 +26,8 @@
 #define ENABLE_METASOUNDGENERATOR_INVALID_SAMPLE_VALUE_LOGGING !UE_BUILD_SHIPPING
 #endif
 
+CSV_DECLARE_CATEGORY_MODULE_EXTERN(METASOUNDGRAPHCORE_API, Audio_Metasound);
+
 namespace Metasound
 {
 	namespace ConsoleVariables
@@ -554,6 +556,7 @@ PRAGMA_ENABLE_DEPRECATION_WARNINGS
 	{
 		METASOUND_LLM_SCOPE;
 		METASOUND_TRACE_CPUPROFILER_EVENT_SCOPE_TEXT(*FString::Printf(TEXT("MetasoundGenerator::OnGenerateAudio %s"), *MetasoundName));
+		CSV_SCOPED_TIMING_STAT(Audio_Metasound, OnGenerateAudio);
 
 		// Defer finishing the metasound generator one block
 		if (bIsFinishTriggered)
@@ -574,6 +577,7 @@ PRAGMA_ENABLE_DEPRECATION_WARNINGS
 		if (bIsWaitingForFirstGraph)
 		{
 			METASOUND_TRACE_CPUPROFILER_EVENT_SCOPE_TEXT(*FString::Printf(TEXT("MetasoundGenerator::OnGenerateAudio::MissedRenderDeadline %s"), *MetasoundName));
+			CSV_CUSTOM_STAT(Audio_Metasound, WaitingForGraphBuildOnPlaybackCount, 1, ECsvCustomStatOp::Accumulate);
 			FMemory::Memset(OutAudio, 0, sizeof(float)* NumSamplesRemaining);
 			return NumSamplesRemaining;
 		}
@@ -836,8 +840,12 @@ PRAGMA_ENABLE_DEPRECATION_WARNINGS
 	{
 	}
 
+	static FThreadSafeCounter NumActiveConstGraphs;
 	void FMetasoundConstGraphGenerator::Init(FMetasoundGeneratorInitParams&& InParams)
 	{
+		NumActiveConstGraphs.Increment();
+		CSV_CUSTOM_STAT(Audio_Metasound, NumActiveConstGraphs, NumActiveConstGraphs.GetValue(), ECsvCustomStatOp::Set);
+
 		InitBase(InParams);
 		// attempt to use operator cache instead of building a new operator.
 		bool bDidUseCachedOperator = false;
@@ -861,6 +869,9 @@ PRAGMA_ENABLE_DEPRECATION_WARNINGS
 
 	FMetasoundConstGraphGenerator::~FMetasoundConstGraphGenerator()
 	{
+		NumActiveConstGraphs.Decrement();
+		CSV_CUSTOM_STAT(Audio_Metasound, NumActiveConstGraphs, NumActiveConstGraphs.GetValue(), ECsvCustomStatOp::Set);
+
 		if (BuilderTask.IsValid())
 		{
 			BuilderTask->EnsureCompletion();
@@ -973,17 +984,24 @@ PRAGMA_ENABLE_DEPRECATION_WARNINGS
 	{
 	}
 
+	static FThreadSafeCounter NumActiveDynamicGraphs;
 	FMetasoundDynamicGraphGenerator::~FMetasoundDynamicGraphGenerator()
 	{
+		NumActiveDynamicGraphs.Decrement();
+		CSV_CUSTOM_STAT(Audio_Metasound, NumActiveDynamicGraphs, NumActiveDynamicGraphs.GetValue(), ECsvCustomStatOp::Set);
+
 		if (BuilderTask.IsValid())
 		{
 			BuilderTask->EnsureCompletion();
 			BuilderTask = nullptr;
 		}
 	}
-
+	
 	void FMetasoundDynamicGraphGenerator::Init(FMetasoundDynamicGraphGeneratorInitParams&& InParams)
 	{
+		NumActiveDynamicGraphs.Increment();
+		CSV_CUSTOM_STAT(Audio_Metasound, NumActiveDynamicGraphs, NumActiveDynamicGraphs.GetValue(), ECsvCustomStatOp::Set);
+
 		InitBase(InParams);
 		AudioOutputNames = InParams.AudioOutputNames;
 		BuildGraph(MoveTemp(InParams));
diff --git a/Engine/Plugins/Runtime/Metasound/Source/MetasoundGraphCore/Private/MetasoundTrace.cpp b/Engine/Plugins/Runtime/Metasound/Source/MetasoundGraphCore/Private/MetasoundTrace.cpp
--- a/Engine/Plugins/Runtime/Metasound/Source/MetasoundGraphCore/Private/MetasoundTrace.cpp
+++ b/Engine/Plugins/Runtime/Metasound/Source/MetasoundGraphCore/Private/MetasoundTrace.cpp
@@ -5,6 +5,14 @@
 #include "HAL/LowLevelMemStats.h"
 #include "HAL/LowLevelMemTracker.h"
 
+#include "ProfilingDebugging/CsvProfiler.h"
+
 DECLARE_LLM_MEMORY_STAT(TEXT("MetaSound"), STAT_MetaSoundLLM, STATGROUP_LLMFULL);
 
 LLM_DEFINE_TAG(Audio_MetaSound, TEXT("MetaSound"), TEXT("Audio"), GET_STATFNAME(STAT_MetaSoundLLM), GET_STATFNAME(STAT_AudioSummaryLLM));
+
+// Defines the "Audio_Metasound" category in the CSV profiler.
+// This should only be defined here. Modules who wish to use this category should contain the line
+// 		CSV_DECLARE_CATEGORY_MODULE_EXTERN(METASOUNDGRAPHCORE_API, Audio_Metasound);
+//
+CSV_DEFINE_CATEGORY_MODULE(METASOUNDGRAPHCORE_API, Audio_Metasound, false);
diff --git a/Engine/Source/Runtime/AudioMixer/Private/AudioMixerSourceManager.cpp b/Engine/Source/Runtime/AudioMixer/Private/AudioMixerSourceManager.cpp
--- a/Engine/Source/Runtime/AudioMixer/Private/AudioMixerSourceManager.cpp
+++ b/Engine/Source/Runtime/AudioMixer/Private/AudioMixerSourceManager.cpp
@@ -3265,6 +3265,7 @@ namespace Audio
 
 		CSV_SCOPED_TIMING_STAT(Audio, SourceManagerUpdate);
 		SCOPE_CYCLE_COUNTER(STAT_AudioMixerSourceManagerUpdate);
+		CSV_CUSTOM_STAT(Audio, NumActiveSources, NumActiveSources, ECsvCustomStatOp::Set);
 
 		RenderThreadPhase = ESourceManagerRenderThreadPhase::Begin;
 
@@ -3428,6 +3429,7 @@ namespace Audio
 		}
 
 		// update trace values
+		CSV_CUSTOM_STAT(Audio, AudioMixerThreadCommands, NewNum, ECsvCustomStatOp::Set);
 		TRACE_INT_VALUE(TEXT("AudioMixerThreadCommands::NumCommands"), NewNum);
 		TRACE_INT_VALUE(TEXT("AudioMixerThreadCommands::CurrentBufferSizeInKb"), CurrentBufferSizeInBytes >> 10);
 	}
