MetaSounds: Clamp RandomGet node's NoRepeats value immediately instead of waiting until it's saved to the RandomGets map. This fixes an issue where Shared State wouldn't work if No Repeat value was being clamped since the comparison would always indicate it needs to be stomped.

diff --git a/Engine/Plugins/Runtime/Metasound/Source/MetasoundFrontend/Public/MetasoundArrayRandomNode.h b/Engine/Plugins/Runtime/Metasound/Source/MetasoundFrontend/Public/MetasoundArrayRandomNode.h
--- a/Engine/Plugins/Runtime/Metasound/Source/MetasoundFrontend/Public/MetasoundArrayRandomNode.h
+++ b/Engine/Plugins/Runtime/Metasound/Source/MetasoundFrontend/Public/MetasoundArrayRandomNode.h
@@ -241,12 +241,12 @@ namespace Metasound
 			// Check to see if this is a global shuffler or a local one. 
 			// Global shuffler will use a namespace to opt into it.
 			PrevSeedValue = *SeedValue;
-			PrevNoRepeatOrder = FMath::Max(*NoRepeatOrder, 0);
-
+			
 			WeightsArray = *InputWeightsArray;
 
 			const ArrayType& InputArrayRef = *InputArray;
 			PrevArraySize = InputArrayRef.Num();
+			PrevNoRepeatOrder = FMath::Clamp(*NoRepeatOrder, 0, PrevArraySize - 1);
 
 			InitializeState(PrevArraySize);
 
