Fix incorrect transient asset path being registered at times due to template node processing generating new transient object representing provided document interface object. (Added bonus with BuildRegistryDoc now actually living in RegistryImpl private implementation where it should've to begin with rather in asset base now that registration headers have been cleaned up).
#rnx

diff --git a/Engine/Plugins/Runtime/Metasound/Source/MetasoundFrontend/Private/MetasoundAssetBase.cpp b/Engine/Plugins/Runtime/Metasound/Source/MetasoundFrontend/Private/MetasoundAssetBase.cpp
--- a/Engine/Plugins/Runtime/Metasound/Source/MetasoundFrontend/Private/MetasoundAssetBase.cpp
+++ b/Engine/Plugins/Runtime/Metasound/Source/MetasoundFrontend/Private/MetasoundAssetBase.cpp
@@ -52,57 +52,6 @@ namespace Metasound
 			static float BlockRateOverride = 0;
 			static int32 SampleRateOverride = 0;
 
-			TScriptInterface<IMetaSoundDocumentInterface> BuildRegistryDocument(TScriptInterface<IMetaSoundDocumentInterface> DocumentInterface)
-			{
-				using namespace Metasound::Frontend;
-
-				METASOUND_TRACE_CPUPROFILER_EVENT_SCOPE(BuildRegistryDocument);
-
-#if WITH_EDITOR
-				// Node template transform is performed on copy of local document to avoid overwriting editable data
-				constexpr bool bTransformDocumentBeforeRegistering = true;
-#else // !WITH_EDITOR
-				// Node template transform is performed on local document only if cook determinism ID generation
-				// is enabled to avoid transforms potentially creating new edges with non-deterministic IDs.
-				const bool bTransformDocumentBeforeRegistering = MetaSoundEnableCookDeterministicIDGeneration == 0;
-#endif // WITH_EDITOR
-
-				const FMetasoundFrontendDocument& Document = DocumentInterface->GetConstDocument();
-				const TArray<FMetasoundFrontendClass>& Dependencies = Document.Dependencies;
-
-				if (bTransformDocumentBeforeRegistering)
-				{
-					// 1. Find template dependencies to build prior to making new document/builder as an optimization
-					// (no sense in creating new document/builder if no templates need processing)
-					FMetaSoundFrontendDocumentBuilder OriginalDocBuilder(DocumentInterface);
-					const bool bContainsTemplateDependency = OriginalDocBuilder.ContainsDependencyOfType(EMetasoundFrontendClassType::Template);
-					if (bContainsTemplateDependency)
-					{
-						UMetaSoundBuilderDocument& RegistryDocObject = UMetaSoundBuilderDocument::Create(*DocumentInterface.GetInterface());
-						FMetaSoundFrontendDocumentBuilder RegistryDocBuilder(&RegistryDocObject);
-
-						RegistryDocBuilder.TransformTemplateNodes();
-
-						return &RegistryDocObject;
-					}
-				}
-#if !NO_LOGGING
-				else
-				{
-					FMetaSoundFrontendDocumentBuilder OriginalDocBuilder(DocumentInterface);
-					const bool bContainsTemplateDependency = OriginalDocBuilder.ContainsDependencyOfType(EMetasoundFrontendClassType::Template);
-					if (bContainsTemplateDependency)
-					{
-						UE_LOG(LogMetaSound, Error,
-							TEXT("Template node processing disabled but provided asset class at '%s' to register contains template nodes. Runtime graph will fail to build."),
-							*OriginalDocBuilder.GetDebugName());
-					}
-				}
-#endif // !NO_LOGGING
-
-				return DocumentInterface;
-			}
-
 			void DepthFirstTraversal(const FMetasoundAssetBase& InInitAsset, TFunctionRef<TSet<const FMetasoundAssetBase*>(const FMetasoundAssetBase&)> InVisitFunction)
 			{
 				// Non recursive depth first traversal.
@@ -122,8 +71,6 @@ namespace Metasound
 				}
 			}
 
-
-
 			// Registers node by copying document. Updates to document require re-registration.
 			// This registry entry does not support node creation as it is only intended to be
 			// used when cooking MetaSounds. 
@@ -328,8 +275,7 @@ void FMetasoundAssetBase::RegisterGraphWithFrontend(Metasound::Frontend::FMetaSo
 	//    reading/writing the IMetaSoundDocumentInterface on the Game Thread)
 	// 2. Async registration is globally disabled via console variable.
 	const bool bAsync = !(IsBuilderActive() || ConsoleVariables::bDisableAsyncGraphRegistration);
-	const TScriptInterface<IMetaSoundDocumentInterface> RegistryDocInterface = AssetBasePrivate::BuildRegistryDocument(Owner);
-	GraphRegistryKey = FRegistryContainerImpl::Get().RegisterGraph(RegistryDocInterface, bAsync);
+	GraphRegistryKey = FRegistryContainerImpl::Get().RegisterGraph(Owner, bAsync);
 
 	if (GraphRegistryKey.IsValid())
 	{
diff --git a/Engine/Plugins/Runtime/Metasound/Source/MetasoundFrontend/Private/MetasoundFrontendRegistryContainerImpl.cpp b/Engine/Plugins/Runtime/Metasound/Source/MetasoundFrontend/Private/MetasoundFrontendRegistryContainerImpl.cpp
--- a/Engine/Plugins/Runtime/Metasound/Source/MetasoundFrontend/Private/MetasoundFrontendRegistryContainerImpl.cpp
+++ b/Engine/Plugins/Runtime/Metasound/Source/MetasoundFrontend/Private/MetasoundFrontendRegistryContainerImpl.cpp
@@ -20,6 +20,8 @@
 #include "MetasoundAssetManager.h"
 #include "MetasoundDocumentInterface.h"
 #include "MetasoundFrontendDataTypeRegistry.h"
+#include "MetasoundFrontendDocumentBuilder.h"
+#include "MetasoundFrontendDocumentIdGenerator.h"
 #include "MetasoundFrontendGraph.h"
 #include "MetasoundFrontendProxyDataCache.h"
 #include "MetasoundFrontendRegistryContainerImpl.h"
@@ -40,7 +42,55 @@ namespace Metasound::Frontend
 {
 	namespace RegistryPrivate
 	{
-		// FGraphNode is used to create unique INodes based off of a IGraph. 
+		TScriptInterface<IMetaSoundDocumentInterface> BuildRegistryDocument(TScriptInterface<IMetaSoundDocumentInterface> DocumentInterface)
+		{
+			using namespace Metasound::Frontend;
+
+			METASOUND_TRACE_CPUPROFILER_EVENT_SCOPE(Metasound::Frontend::BuildRegistryDocument);
+
+#if WITH_EDITOR
+			// Node template transform is performed on copy of local document to avoid overwriting editable data
+			constexpr bool bTransformDocumentBeforeRegistering = true;
+#else // !WITH_EDITOR
+			// Node template transform is performed on local document only if cook determinism ID generation
+			// is enabled to avoid transforms potentially creating new edges with non-deterministic IDs.
+			const bool bTransformDocumentBeforeRegistering = MetaSoundEnableCookDeterministicIDGeneration == 0;
+#endif // !WITH_EDITOR
+
+			if (bTransformDocumentBeforeRegistering)
+			{
+				// 1. Find template dependencies to build prior to making new document/builder as an optimization
+				// (no sense in creating new document/builder if no templates need processing)
+				FMetaSoundFrontendDocumentBuilder OriginalDocBuilder(DocumentInterface);
+				const bool bContainsTemplateDependency = OriginalDocBuilder.ContainsDependencyOfType(EMetasoundFrontendClassType::Template);
+				if (bContainsTemplateDependency)
+				{
+					UMetaSoundBuilderDocument& RegistryDocObject = UMetaSoundBuilderDocument::Create(*DocumentInterface.GetInterface());
+					FMetaSoundFrontendDocumentBuilder RegistryDocBuilder(&RegistryDocObject);
+
+					RegistryDocBuilder.TransformTemplateNodes();
+
+					return &RegistryDocObject;
+				}
+			}
+#if !NO_LOGGING
+			else
+			{
+				FMetaSoundFrontendDocumentBuilder OriginalDocBuilder(DocumentInterface);
+				const bool bContainsTemplateDependency = OriginalDocBuilder.ContainsDependencyOfType(EMetasoundFrontendClassType::Template);
+				if (bContainsTemplateDependency)
+				{
+					UE_LOG(LogMetaSound, Error,
+						TEXT("Template node processing disabled but provided asset class at '%s' to register contains template nodes. Runtime graph will fail to build."),
+						*OriginalDocBuilder.GetDebugName());
+				}
+			}
+#endif // !NO_LOGGING
+
+			return DocumentInterface;
+		}
+
+		// FGraphNode is used to create unique INodes based off of a IGraph.
 		//
 		// Individual nodes need to reflect their InstanceName and InstanceID, but otherwise
 		// they simply encapsulate a shared set of behavior. To minimize memory usage, a single
@@ -222,7 +272,11 @@ namespace Metasound::Frontend
 		TUniquePtr<INodeRegistryEntry> RegistryEntry = MakeUnique<FDocumentNodeRegistryEntry>(Document.RootGraph, Document.Interfaces, MoveTemp(InNodeClassInfo), GraphToRegister);
 
 		const FNodeRegistryKey RegistryKey = RegisterNodeInternal(MoveTemp(RegistryEntry));
-		const FGraphRegistryKey GraphKey { RegistryKey, DocumentInterface->GetAssetPathChecked() };
+
+		// Key must use the asset path provided to class info and *NOT* that of the
+		// provided DocumentInterface object, as that may be a built transient object
+		// with a different transient asset path.
+		const FGraphRegistryKey GraphKey { RegistryKey, InNodeClassInfo.AssetPath };
 		RegisterGraphInternal(GraphKey, GraphToRegister);
 	}
 
@@ -385,11 +439,15 @@ namespace Metasound::Frontend
 		check(InDocumentInterface);
 		check(IsInGameThread());
 
-		UObject* OwningObject = InDocumentInterface.GetObject();
-		check(OwningObject);
+		// Use the asset path of the provided document interface object for identification, *NOT* the
+		// built version as the build process may in fact create a new object with a transient path.
 		const FTopLevelAssetPath AssetPath = InDocumentInterface->GetAssetPathChecked();
+		const TScriptInterface<IMetaSoundDocumentInterface> RegistryDocInterface = RegistryPrivate::BuildRegistryDocument(InDocumentInterface);
 
-		const FMetasoundFrontendDocument& Document = InDocumentInterface->GetConstDocument();
+		UObject* OwningObject = RegistryDocInterface.GetObject();
+		check(OwningObject);
+
+		const FMetasoundFrontendDocument& Document = RegistryDocInterface->GetConstDocument();
 		const FGraphRegistryKey RegistryKey { FNodeRegistryKey(Document.RootGraph), AssetPath };
 
 		if (!RegistryKey.IsValid())
@@ -422,7 +480,7 @@ namespace Metasound::Frontend
 
 			Tasks::FTask BuildAndRegisterTask = AsyncRegistrationPipe.Launch(
 				UE_SOURCE_LOCATION,
-				[RegistryKey, ClassInfo = MoveTemp(NodeClassInfo), DocumentInterface = InDocumentInterface, ProxyDataCache = MoveTemp(ProxyDataCache)]() mutable
+				[RegistryKey, ClassInfo = MoveTemp(NodeClassInfo), RegistryDocInterface, ProxyDataCache = MoveTemp(ProxyDataCache)]() mutable
 				{
 					FRegistryContainerImpl& Registry = FRegistryContainerImpl::Get();
 					// Unregister the graph before re-registering
@@ -431,13 +489,13 @@ namespace Metasound::Frontend
 						Registry.UnregisterGraphInternal(RegistryKey);
 					}
 
-					Registry.BuildAndRegisterGraphFromDocument(DocumentInterface, ProxyDataCache, MoveTemp(ClassInfo));
+					Registry.BuildAndRegisterGraphFromDocument(RegistryDocInterface, ProxyDataCache, MoveTemp(ClassInfo));
 					Registry.RemoveRegistrationTask(RegistryKey, FNodeRegistryTransaction::ETransactionType::NodeRegistration);
-					Registry.RemoveDocumentReference(DocumentInterface);
+					Registry.RemoveDocumentReference(RegistryDocInterface);
 				}
 			);
 
-			AddDocumentReference(InDocumentInterface);
+			AddDocumentReference(RegistryDocInterface);
 			AddRegistrationTask(RegistryKey, FActiveRegistrationTaskInfo
 			{
 				FNodeRegistryTransaction::ETransactionType::NodeRegistration,
@@ -453,7 +511,7 @@ namespace Metasound::Frontend
 			}
 
 			// Build and register graph synchronously
-			BuildAndRegisterGraphFromDocument(InDocumentInterface, ProxyDataCache, MoveTemp(NodeClassInfo));
+			BuildAndRegisterGraphFromDocument(RegistryDocInterface, ProxyDataCache, MoveTemp(NodeClassInfo));
 		}
 
 		return RegistryKey;
@@ -538,10 +596,12 @@ namespace Metasound::Frontend
 
 		FScopeLock Lock(&RegistryMapsCriticalSection);
 
+#if !NO_LOGGING
 		if (RegisteredGraphs.Contains(InKey))
 		{
 			UE_LOG(LogMetaSound, Warning, TEXT("Graph is already registered with the same registry key '%s'. The existing registered graph will be replaced with the new graph."), *InKey.ToString());
 		}
+#endif // !NO_LOGGING
 
 		RegisteredGraphs.Add(InKey, InGraph);
 	}
@@ -557,6 +617,7 @@ namespace Metasound::Frontend
 			const int32 GraphUnregistered = RegisteredGraphs.Remove(InKey) > 0;
 			const bool bNodeUnregistered = UnregisterNodeInternal(InKey.NodeKey);
 
+#if !NO_LOGGING
 			if (GraphUnregistered)
 			{
 				UE_LOG(LogMetaSound, VeryVerbose, TEXT("Unregistered graph with key '%s'"), *InKey.ToString());
@@ -570,6 +631,7 @@ namespace Metasound::Frontend
 					UE_LOG(LogMetaSound, Warning, TEXT("Graph '%s' was not found, but analogous registered node class was when unregistering."), *InKey.ToString());
 				}
 			}
+#endif // !NO_LOGGING
 
 			return bNodeUnregistered;
 		}
diff --git a/Engine/Plugins/Runtime/Metasound/Source/MetasoundFrontend/Public/MetasoundFrontendDocumentBuilder.h b/Engine/Plugins/Runtime/Metasound/Source/MetasoundFrontend/Public/MetasoundFrontendDocumentBuilder.h
--- a/Engine/Plugins/Runtime/Metasound/Source/MetasoundFrontend/Public/MetasoundFrontendDocumentBuilder.h
+++ b/Engine/Plugins/Runtime/Metasound/Source/MetasoundFrontend/Public/MetasoundFrontendDocumentBuilder.h
@@ -139,11 +139,10 @@ public:
 	FMetaSoundFrontendDocumentBuilder(TScriptInterface<IMetaSoundDocumentInterface> InDocumentInterface);
 	FMetaSoundFrontendDocumentBuilder(TScriptInterface<IMetaSoundDocumentInterface> InDocumentInterface, TSharedRef<Metasound::Frontend::FDocumentModifyDelegates> InDocumentDelegates);
 	virtual ~FMetaSoundFrontendDocumentBuilder();
-		
+
 	// Call when the builder will no longer modify the IMetaSoundDocumentInterface
 	void FinishBuilding();
 
-
 	const FMetasoundFrontendClass* AddDependency(const FMetasoundFrontendClass& InClass);
 	void AddEdge(FMetasoundFrontendEdge&& InNewEdge);
 	bool AddNamedEdges(const TSet<Metasound::Frontend::FNamedEdge>& ConnectionsToMake, TArray<const FMetasoundFrontendEdge*>* OutEdgesCreated = nullptr, bool bReplaceExistingConnections = true);
@@ -281,11 +280,6 @@ public:
 	bool SwapGraphOutput(const FMetasoundFrontendClassVertex& InExistingOutputVertex, const FMetasoundFrontendClassVertex& NewOutputVertex);
 	bool UpdateDependencyClassNames(const TMap<FMetasoundFrontendClassName, FMetasoundFrontendClassName>& OldToNewReferencedClassNames);
 
-	// Determines if the provided interface requires further processing or if its data  is ready as is for insertion into the FrontendRegistry.
-	// If it requires further processing, returns a transient document interface to be used by the registry with an optimized, transformed version
-	// of the given interface's document.
-	static TScriptInterface<IMetaSoundDocumentInterface> BuildRegistryDocument(TScriptInterface<IMetaSoundDocumentInterface> DocumentInterface);
-
 	// Transforms template nodes within the given builder's document, which can include swapping associated edges and/or
 	// replacing nodes with other, registry-defined concrete node class instances. Returns true if any template nodes were processed.
 	bool TransformTemplateNodes();
