Fix for assert when cooking by book due to invalid graph registration call from asset manager.
#rnx

diff --git a/Engine/Plugins/Runtime/Metasound/Source/MetasoundEditor/Private/MetasoundEditorModule.cpp b/Engine/Plugins/Runtime/Metasound/Source/MetasoundEditor/Private/MetasoundEditorModule.cpp
--- a/Engine/Plugins/Runtime/Metasound/Source/MetasoundEditor/Private/MetasoundEditorModule.cpp
+++ b/Engine/Plugins/Runtime/Metasound/Source/MetasoundEditor/Private/MetasoundEditorModule.cpp
@@ -231,6 +231,13 @@ namespace Metasound
 		{
 			void LoadAndRegisterAsset(const FAssetData& InAssetData)
 			{
+				// Ignore requests if running cook, as presave will register and cook using the
+				// proper node register call which avoids async registration/igraph generation.
+				if (IsRunningCookCommandlet())
+				{
+					return;
+				}
+
 				Frontend::FMetaSoundAssetRegistrationOptions RegOptions;
 				RegOptions.bForceReregister = false;
 				if (const UMetaSoundSettings* Settings = GetDefault<UMetaSoundSettings>())
