Fix for compiled out path setter, le sigh

diff --git a/Engine/Plugins/Runtime/Metasound/Source/MetasoundEngine/Private/MetasoundEngineAsset.h b/Engine/Plugins/Runtime/Metasound/Source/MetasoundEngine/Private/MetasoundEngineAsset.h
--- a/Engine/Plugins/Runtime/Metasound/Source/MetasoundEngine/Private/MetasoundEngineAsset.h
+++ b/Engine/Plugins/Runtime/Metasound/Source/MetasoundEngine/Private/MetasoundEngineAsset.h
@@ -67,8 +67,8 @@ namespace Metasound
 		static FTopLevelAssetPath GetAssetPathChecked(TMetaSoundObject& InMetaSound)
 		{
 			FTopLevelAssetPath Path;
-			checkf(Path.TrySetPath(&InMetaSound), TEXT("Failed to set TopLevelAssetPath from MetaSound '%s'. MetaSound must be highest level object in package."), *InMetaSound.GetPathName());
-			checkf(Path.IsValid(), TEXT("Failed to set TopLevelAssetPath from MetaSound '%s'. This may be caused by calling this function when the asset is being destroyed."), *InMetaSound.GetPathName());
+			ensureAlwaysMsgf(Path.TrySetPath(&InMetaSound), TEXT("Failed to set TopLevelAssetPath from MetaSound '%s'. MetaSound must be highest level object in package."), *InMetaSound.GetPathName());
+			ensureAlwaysMsgf(Path.IsValid(), TEXT("Failed to set TopLevelAssetPath from MetaSound '%s'. This may be caused by calling this function when the asset is being destroyed."), *InMetaSound.GetPathName());
 			return Path;
 		}
 
diff --git a/Engine/Plugins/Runtime/Metasound/Source/MetasoundFrontend/Private/MetasoundFrontendDocumentBuilder.cpp b/Engine/Plugins/Runtime/Metasound/Source/MetasoundFrontend/Private/MetasoundFrontendDocumentBuilder.cpp
--- a/Engine/Plugins/Runtime/Metasound/Source/MetasoundFrontend/Private/MetasoundFrontendDocumentBuilder.cpp
+++ b/Engine/Plugins/Runtime/Metasound/Source/MetasoundFrontend/Private/MetasoundFrontendDocumentBuilder.cpp
@@ -591,8 +591,8 @@ UMetaSoundBuilderDocument& UMetaSoundBuilderDocument::Create(const IMetaSoundDoc
 FTopLevelAssetPath UMetaSoundBuilderDocument::GetAssetPathChecked() const
 {
 	FTopLevelAssetPath Path;
-	checkf(Path.TrySetPath(this), TEXT("Failed to set TopLevelAssetPath from transient MetaSound '%s'. MetaSound must be highest level object in package."), *GetPathName());
-	checkf(Path.IsValid(), TEXT("Failed to set TopLevelAssetPath from MetaSound '%s'. This may be caused by calling this function when the asset is being destroyed."), *GetPathName());
+	ensureAlwaysMsgf(Path.TrySetPath(this), TEXT("Failed to set TopLevelAssetPath from transient MetaSound '%s'. MetaSound must be highest level object in package."), *GetPathName());
+	ensureAlwaysMsgf(Path.IsValid(), TEXT("Failed to set TopLevelAssetPath from MetaSound '%s'. This may be caused by calling this function when the asset is being destroyed."), *GetPathName());
 	return Path;
 }
 
