Support for custom pin icons for registered data types.

diff --git a/Engine/Plugins/Runtime/Metasound/Source/MetasoundEditor/Private/MetasoundEditorModule.cpp b/Engine/Plugins/Runtime/Metasound/Source/MetasoundEditor/Private/MetasoundEditorModule.cpp
--- a/Engine/Plugins/Runtime/Metasound/Source/MetasoundEditor/Private/MetasoundEditorModule.cpp
+++ b/Engine/Plugins/Runtime/Metasound/Source/MetasoundEditor/Private/MetasoundEditorModule.cpp
@@ -227,6 +227,14 @@ namespace Metasound
 			}
 		}
 
+		// A structure that contains information about registered custom pin types. 
+		struct FGraphPinConfiguration
+		{
+			FEdGraphPinType PinType;
+			const FSlateBrush* PinConnectedIcon = nullptr;
+			const FSlateBrush* PinDisconnectedIcon = nullptr;
+		};
+
 		class FModule : public IMetasoundEditorModule
 		{
 			void LoadAndRegisterAsset(const FAssetData& InAssetData)
@@ -601,7 +609,7 @@ namespace Metasound
 				}
 			}
 
-			void RegisterPinType(FName InDataTypeName, FName InPinCategory, FName InPinSubCategory)
+			virtual void RegisterPinType(FName InDataTypeName, FName InPinCategory, FName InPinSubCategory, const FSlateBrush* InPinConnectedIcon = nullptr, const FSlateBrush* InPinDisconnectedIcon = nullptr) override
 			{
 				using namespace Frontend;
 
@@ -612,11 +620,15 @@ namespace Metasound
 				const FName PinCategory = InPinCategory.IsNone() ? FGraphBuilder::PinCategoryObject : InPinCategory;
 
 				const EPinContainerType ContainerType = DataTypeInfo.bIsArrayType ? EPinContainerType::Array : EPinContainerType::None;
-				FEdGraphPinType PinType(PinCategory, InPinSubCategory, nullptr, ContainerType, false, FEdGraphTerminalType());
+				FGraphPinConfiguration PinConfiguration;
+				PinConfiguration.PinType.PinCategory = PinCategory;
+				PinConfiguration.PinType.PinSubCategory = InPinSubCategory;
+				PinConfiguration.PinType.ContainerType = ContainerType;
 				UClass* ClassToUse = IDataTypeRegistry::Get().GetUClassForDataType(InDataTypeName);
-				PinType.PinSubCategoryObject = Cast<UObject>(ClassToUse);
-
-				PinTypes.Emplace(InDataTypeName, MoveTemp(PinType));
+				PinConfiguration.PinType.PinSubCategoryObject = Cast<UObject>(ClassToUse);
+				PinConfiguration.PinConnectedIcon = InPinConnectedIcon;
+				PinConfiguration.PinDisconnectedIcon = InPinDisconnectedIcon;
+				PinTypes.Emplace(InDataTypeName, MoveTemp(PinConfiguration));
 			}
 
 			void ShutdownAssetClassRegistry()
@@ -715,10 +727,40 @@ namespace Metasound
 					return bIsConstructorType ? &Style::GetSlateBrushSafe("MetasoundEditor.Graph.ConstructorPin") : FAppStyle::GetBrush("Icons.BulletPoint");
 				}
 			}
+			
+			virtual bool GetCustomPinIcons(UEdGraphPin* InPin, const FSlateBrush*& PinConnectedIcon, const FSlateBrush*& PinDisconnectedIcon) const override
+			{
+				if (const UEdGraphNode* Node = InPin->GetOwningNode())
+				{
+					if (const UMetasoundEditorGraphNode* MetaSoundNode = Cast<UMetasoundEditorGraphNode>(InPin->GetOwningNode()))
+					{
+						Metasound::Frontend::FDataTypeRegistryInfo RegistryInfo = MetaSoundNode->GetPinDataTypeInfo(*InPin);
+						return GetCustomPinIcons(RegistryInfo.DataTypeName, PinConnectedIcon, PinDisconnectedIcon);
+					}
+				}
+				return false;
+			}
+
+			virtual bool GetCustomPinIcons(FName InDataType, const FSlateBrush*& PinConnectedIcon, const FSlateBrush*& PinDisconnectedIcon) const override
+			{
+				const FGraphPinConfiguration* PinConfiguration = PinTypes.Find(InDataType);
+				if (!PinConfiguration || (!PinConfiguration->PinConnectedIcon && !PinConfiguration->PinDisconnectedIcon))
+				{
+					return false;
+				}
+				PinConnectedIcon = PinConfiguration->PinConnectedIcon;
+				PinDisconnectedIcon = PinConfiguration->PinDisconnectedIcon ? PinConfiguration->PinDisconnectedIcon : PinConfiguration->PinConnectedIcon;
+				return true;
+			}
 
 			virtual const FEdGraphPinType* FindPinType(FName InDataTypeName) const
 			{
-				return PinTypes.Find(InDataTypeName);
+				const FGraphPinConfiguration* PinConfiguration = PinTypes.Find(InDataTypeName);
+				if (PinConfiguration)
+				{
+					return &PinConfiguration->PinType;
+				}
+				return nullptr;
 			}
 
 			virtual bool IsMetaSoundAssetClass(const FTopLevelAssetPath& InClassName) const override
@@ -925,7 +967,7 @@ namespace Metasound
 			
 			TArray<TSharedPtr<FAssetTypeActions_Base>> AssetActions;
 			TMap<EMetasoundFrontendLiteralType, const TSubclassOf<UMetasoundEditorGraphMemberDefaultLiteral>> InputDefaultLiteralClassRegistry;
-			TMap<FName, FEdGraphPinType> PinTypes;
+			TMap<FName, FGraphPinConfiguration> PinTypes;
 
 			TMap<UClass*, TUniquePtr<IMemberDefaultLiteralCustomizationFactory>> LiteralCustomizationFactories;
 
diff --git a/Engine/Plugins/Runtime/Metasound/Source/MetasoundEditor/Private/SMetasoundGraphNode.cpp b/Engine/Plugins/Runtime/Metasound/Source/MetasoundEditor/Private/SMetasoundGraphNode.cpp
--- a/Engine/Plugins/Runtime/Metasound/Source/MetasoundEditor/Private/SMetasoundGraphNode.cpp
+++ b/Engine/Plugins/Runtime/Metasound/Source/MetasoundEditor/Private/SMetasoundGraphNode.cpp
@@ -343,6 +343,14 @@ namespace Metasound
 				PinWidget = SNew(SMetasoundGraphPin, InPin);
 			}
 
+			IMetasoundEditorModule& EditorModule = FModuleManager::GetModuleChecked<IMetasoundEditorModule>("MetaSoundEditor");
+			const FSlateBrush* PinConnectedIcon = nullptr;
+			const FSlateBrush* PinDisconnectedIcon = nullptr;
+			if (EditorModule.GetCustomPinIcons(InPin, PinConnectedIcon, PinDisconnectedIcon))
+			{
+				PinWidget->SetCustomPinIcon(PinConnectedIcon, PinDisconnectedIcon);
+			}
+
 			return PinWidget;
 		}
 
diff --git a/Engine/Plugins/Runtime/Metasound/Source/MetasoundEditor/Public/MetasoundEditorModule.h b/Engine/Plugins/Runtime/Metasound/Source/MetasoundEditor/Public/MetasoundEditorModule.h
--- a/Engine/Plugins/Runtime/Metasound/Source/MetasoundEditor/Public/MetasoundEditorModule.h
+++ b/Engine/Plugins/Runtime/Metasound/Source/MetasoundEditor/Public/MetasoundEditorModule.h
@@ -15,6 +15,7 @@
 // Forward Declarations
 class IDetailLayoutBuilder;
 class IDetailPropertyRow;
+class UEdGraphPin;
 class UMetasoundEditorGraph;
 class UMetasoundEditorGraphMemberDefaultLiteral;
 
@@ -101,7 +102,11 @@ namespace Metasound
 
 			virtual const FSlateBrush* GetIconBrush(FName InDataType, const bool bIsConstructorType) const = 0;
 
-			virtual void RegisterPinType(FName InDataTypeName, FName InPinCategory = { }, FName InPinSubCategory = { }) = 0;
+			virtual bool GetCustomPinIcons(UEdGraphPin* InPin, const FSlateBrush*& PinConnectedIcon, const FSlateBrush*& PinDisconnectedIcon) const = 0;
+			virtual bool GetCustomPinIcons(FName InDataType, const FSlateBrush*& PinConnectedIcon, const FSlateBrush*& PinDisconnectedIcon) const = 0;
+			
+			virtual void RegisterPinType(FName InDataTypeName, FName InPinCategory = { }, FName InPinSubCategory = { },
+				const FSlateBrush* InPinConnectedIcon = nullptr, const FSlateBrush* InPinDisconnectedIcon = nullptr) = 0;
 		};
 	} // namespace Editor
 } // namespace Metasound
