[Metasound] (speculative) Guard against Audio Bus Reader crash when AudioBusProxy reports unreasonable channel count.

diff --git a/Engine/Plugins/Runtime/Metasound/Source/MetasoundEngine/Private/MetasoundAudioBusReaderNode.cpp b/Engine/Plugins/Runtime/Metasound/Source/MetasoundEngine/Private/MetasoundAudioBusReaderNode.cpp
--- a/Engine/Plugins/Runtime/Metasound/Source/MetasoundEngine/Private/MetasoundAudioBusReaderNode.cpp
+++ b/Engine/Plugins/Runtime/Metasound/Source/MetasoundEngine/Private/MetasoundAudioBusReaderNode.cpp
@@ -1,6 +1,7 @@
 // Copyright Epic Games, Inc. All Rights Reserved.
 
 #include "AudioDevice.h"
+#include "Sound/AudioBus.h"
 #include "AudioBusSubsystem.h"
 #include "DSP/ConvertDeinterleave.h"
 #include "Internationalization/Text.h"
@@ -234,7 +235,7 @@ namespace Metasound
 					if (FAudioDevice* AudioDevice = ADM->GetAudioDeviceRaw(AudioDeviceId))
 					{
 						// Start the audio bus in case it's not already started
-						AudioBusChannels = AudioBusProxy->NumChannels;
+						AudioBusChannels = static_cast<uint32>(FMath::Min(AudioBusProxy->NumChannels, static_cast<int32>(EAudioBusChannels::SevenPointOne)));
 						AudioBusId = AudioBusProxy->AudioBusId;
 						
 						const Audio::FAudioBusKey AudioBusKey = Audio::FAudioBusKey(AudioBusId);
