Delay MetaSoundSource calling RegisterGraphWithFrontend if the MetaSoundAssetManager has not been set yet - to facilitate server preloading assets before the manager has been set.

diff --git a/Engine/Plugins/Runtime/Metasound/Source/MetasoundEngine/Private/MetasoundSource.cpp b/Engine/Plugins/Runtime/Metasound/Source/MetasoundEngine/Private/MetasoundSource.cpp
--- a/Engine/Plugins/Runtime/Metasound/Source/MetasoundEngine/Private/MetasoundSource.cpp
+++ b/Engine/Plugins/Runtime/Metasound/Source/MetasoundEngine/Private/MetasoundSource.cpp
@@ -497,6 +497,10 @@ FTopLevelAssetPath UMetaSoundSource::GetAssetPathChecked() const
 
 void UMetaSoundSource::BeginDestroy()
 {
+#if WITH_SERVER_CODE
+	Metasound::Frontend::IMetaSoundAssetManager::OnManagerSet.RemoveAll(this);
+#endif //WITH_SERVER_CODE
+
 	OnNotifyBeginDestroy();
 	Super::BeginDestroy();
 }
@@ -720,7 +724,20 @@ void UMetaSoundSource::InitResources()
 
 	if (IsInGameThread())
 	{
+#if WITH_SERVER_CODE
+		
+		if (IMetaSoundAssetManager::Get())
+		{
+			RegisterGraphWithFrontend(GetInitRegistrationOptions());
+		}
+		else
+		{
+			//to allow the server to preload UMetaSoundSource objects before the manager has been set up
+			IMetaSoundAssetManager::OnManagerSet.AddUObject(this, &UMetaSoundSource::InitResources);
+		}
+#else
 		RegisterGraphWithFrontend(GetInitRegistrationOptions());
+#endif //WITH_SERVER_CODE
 	}
 	else
 	{
diff --git a/Engine/Plugins/Runtime/Metasound/Source/MetasoundFrontend/Private/MetasoundAssetManager.cpp b/Engine/Plugins/Runtime/Metasound/Source/MetasoundFrontend/Private/MetasoundAssetManager.cpp
--- a/Engine/Plugins/Runtime/Metasound/Source/MetasoundFrontend/Private/MetasoundAssetManager.cpp
+++ b/Engine/Plugins/Runtime/Metasound/Source/MetasoundFrontend/Private/MetasoundAssetManager.cpp
@@ -26,5 +26,9 @@ namespace Metasound
 		} // namespace AssetTags
 
 		IMetaSoundAssetManager* IMetaSoundAssetManager::Instance = nullptr;
+
+#if WITH_SERVER_CODE
+		TMulticastDelegate<void()> IMetaSoundAssetManager::OnManagerSet;
+#endif //WITH_SERVER_CODE
 	} // namespace Frontend
 } // namespace Metasound
diff --git a/Engine/Plugins/Runtime/Metasound/Source/MetasoundFrontend/Public/MetasoundAssetManager.h b/Engine/Plugins/Runtime/Metasound/Source/MetasoundFrontend/Public/MetasoundAssetManager.h
--- a/Engine/Plugins/Runtime/Metasound/Source/MetasoundFrontend/Public/MetasoundAssetManager.h
+++ b/Engine/Plugins/Runtime/Metasound/Source/MetasoundFrontend/Public/MetasoundAssetManager.h
@@ -75,6 +75,10 @@ namespace Metasound
 					check(!Instance);
 				}
 				Instance = &InInterface;
+
+#if WITH_SERVER_CODE
+				OnManagerSet.Broadcast();
+#endif //WITH_SERVER_CODE
 			}
 
 			static IMetaSoundAssetManager* Get()
@@ -88,6 +92,11 @@ namespace Metasound
 				return *Instance;
 			}
 
+#if WITH_SERVER_CODE
+			//to allow the server to preload UMetaSoundSource objects before the manager has been set up
+			static TMulticastDelegate<void()> OnManagerSet;
+#endif //WITH_SERVER_CODE
+
 			struct FAssetInfo
 			{
 				const Metasound::Frontend::FNodeRegistryKey RegistryKey;
