Add Find tab for searching within a MetaSound
- Search by node/pin name string, node type, data type, display name, value, or comment
- View default literal values (including marking overridden inputs, array num), extended values on tooltip
- Jump to result in graph on click

diff --git a/Engine/Plugins/Runtime/Metasound/Source/MetasoundEditor/Private/FindInMetasound.cpp b/Engine/Plugins/Runtime/Metasound/Source/MetasoundEditor/Private/FindInMetasound.cpp
new file mode 100644
--- /dev/null
+++ b/Engine/Plugins/Runtime/Metasound/Source/MetasoundEditor/Private/FindInMetasound.cpp
@@ -0,0 +1,688 @@
+// Copyright Epic Games, Inc. All Rights Reserved.
+
+#include "FindInMetasound.h"
+#include "Framework/Application/SlateApplication.h"
+#include "Layout/WidgetPath.h"
+#include "MetasoundEditorGraph.h"
+#include "MetasoundEditorGraphBuilder.h"
+#include "MetasoundEditorGraphInputNode.h"
+#include "MetasoundEditorModule.h"
+#include "Styling/AppStyle.h"
+#include "Widgets/Images/SImage.h"
+#include "Widgets/Input/SSearchBox.h"
+
+#define LOCTEXT_NAMESPACE "FindInMetasound"
+
+namespace Metasound::Editor
+{
+	FFindInMetasoundResult::FFindInMetasoundResult(const FString& InResultName)
+		:Value(InResultName), DuplicationIndex(0), Class(nullptr), Pin(), GraphNode(nullptr)
+	{
+	}
+
+	FFindInMetasoundResult::FFindInMetasoundResult(const FString& InResultName, TSharedPtr<FFindInMetasoundResult>& InParent, UClass* InClass, int InDuplicationIndex)
+	: Parent(InParent), Value(InResultName), DuplicationIndex(InDuplicationIndex), Class(InClass), Pin(), GraphNode(nullptr)
+	{
+		ValueText = GetValueText();
+	}
+
+	FFindInMetasoundResult::FFindInMetasoundResult(const FString& InResultName, TSharedPtr<FFindInMetasoundResult>& InParent, UEdGraphPin* InPin)
+	: Parent(InParent), Value(InResultName), DuplicationIndex(0), Class(nullptr), Pin(InPin), GraphNode(nullptr)
+	{
+		ValueText = GetValueText();
+	}
+
+	FFindInMetasoundResult::FFindInMetasoundResult(const FString& InResultName, TSharedPtr<FFindInMetasoundResult>& InParent, UEdGraphNode* InNode)
+	: Parent(InParent), Value(InResultName), DuplicationIndex(0), Class(InNode->GetClass()), Pin(), GraphNode(InNode)
+	{
+		if (GraphNode.IsValid())
+		{
+			CommentText = GraphNode->NodeComment;
+		}
+		ValueText = GetValueText();
+	}
+
+	FReply FFindInMetasoundResult::OnClick(TWeakPtr<class FEditor> MetaSoundEditor)
+	{
+		if (GraphNode.IsValid())
+		{
+			MetaSoundEditor.Pin()->GetGraphEditor()->JumpToNode(GraphNode.Get());
+		}
+		else if (UEdGraphPin* ResolvedPin = Pin.Get())
+		{
+			MetaSoundEditor.Pin()->GetGraphEditor()->JumpToPin(ResolvedPin);
+		}
+		return FReply::Handled();
+	}
+
+	FText FFindInMetasoundResult::GetCategory() const
+	{
+		if (Class == nullptr && Pin.Get())
+		{
+			return LOCTEXT("PinCategory", "Pin");
+		}
+		else
+		{
+			return LOCTEXT("NodeCategory", "Node");
+		}
+		return FText::GetEmpty();
+	}
+
+	TSharedRef<SWidget> FFindInMetasoundResult::CreateIcon() const
+	{
+		FSlateColor IconColor = FSlateColor::UseForeground();
+		const FSlateBrush* Brush = nullptr;
+		bool bIsPin = false;
+		if (const UEdGraphPin* ResolvedPin = Pin.Get())
+		{
+			bIsPin = true;
+			if (const UMetasoundEditorGraphNode* Node = Cast<UMetasoundEditorGraphNode>(ResolvedPin->GetOwningNode()))
+			{
+				// Get data type and constructor pin status
+				FName DataTypeName;
+				bool bIsConstructorPin = false;
+				if (const UMetasoundEditorGraphMember* Member = GetMetaSoundGraphMember(Node))
+				{
+					DataTypeName = Member->GetDataType();
+					if (const UMetasoundEditorGraphVertex* Vertex = Cast<UMetasoundEditorGraphVertex>(Member))
+					{
+						bIsConstructorPin = Vertex->GetVertexAccessType() == EMetasoundFrontendVertexAccessType::Value;
+					}
+				}
+				else if (const UMetasoundEditorGraphExternalNode* ExternalNode = Cast<UMetasoundEditorGraphExternalNode>(Node))
+				{
+					if (ResolvedPin->Direction == EGPD_Input)
+					{
+						Frontend::FConstInputHandle InputHandle = FGraphBuilder::GetConstInputHandleFromPin(ResolvedPin);
+						DataTypeName = InputHandle->GetDataType();
+						bIsConstructorPin = InputHandle->GetVertexAccessType() == EMetasoundFrontendVertexAccessType::Value;
+					}
+					else if (ResolvedPin->Direction == EGPD_Output)
+					{
+						Frontend::FConstOutputHandle OutputHandle = FGraphBuilder::GetConstOutputHandleFromPin(ResolvedPin);
+						DataTypeName = OutputHandle->GetDataType();
+						bIsConstructorPin = OutputHandle->GetVertexAccessType() == EMetasoundFrontendVertexAccessType::Value;
+					}
+				}
+
+				// Get brush 
+				const IMetasoundEditorModule& EditorModule = FModuleManager::GetModuleChecked<IMetasoundEditorModule>("MetaSoundEditor");
+				Brush = EditorModule.GetIconBrush(DataTypeName, bIsConstructorPin);
+			}
+			else
+			{
+				Brush = FAppStyle::GetBrush(TEXT("GraphEditor.PinIcon"));
+			}
+
+			// Get color
+			const UEdGraphSchema* Schema = ResolvedPin->GetSchema();
+			IconColor = Schema->GetPinTypeColor(ResolvedPin->PinType);
+		}
+		else if (GraphNode.IsValid())
+		{
+			if (UMetasoundEditorGraphNode* MetaSoundGraphNode = Cast<UMetasoundEditorGraphNode>(GraphNode))
+			{
+				// Variable nodes do not have a node title icon or color, 
+				// so use generic one and corresponding pin type color
+				if (UMetasoundEditorGraphVariableNode* MetaSoundGraphVariableNode = Cast<UMetasoundEditorGraphVariableNode>(MetaSoundGraphNode))
+				{
+					Brush = FAppStyle::GetBrush(TEXT("Kismet.AllClasses.VariableIcon"));
+					UEdGraphPin* FirstPin = MetaSoundGraphVariableNode->GetPinAt(0);
+					if (FirstPin)
+					{
+						const UEdGraphSchema* Schema = FirstPin->GetSchema();
+						IconColor = Schema->GetPinTypeColor(FirstPin->PinType);
+					}
+				}
+				else
+				{
+					Brush = MetaSoundGraphNode->GetNodeTitleIcon().GetIcon();
+					IconColor = GraphNode->GetNodeTitleColor();
+				}
+			}
+			else
+			{
+				Brush = FAppStyle::GetBrush(TEXT("GraphEditor.NodeGlyph"));
+			}
+		}
+
+		static const FVector2D IconSize16 = FVector2D(16.0f, 16.0f);
+		TOptional<FVector2D> BrushSize = bIsPin ? IconSize16 : TOptional<FVector2D>();
+
+		return SNew(SImage)
+			.Image(Brush)
+			.ColorAndOpacity(IconColor)
+			.ToolTipText(GetCategory())
+			.DesiredSizeOverride(BrushSize);
+	}
+
+	FString FFindInMetasoundResult::GetCommentText() const
+	{
+		return CommentText;
+	}
+
+	FString FFindInMetasoundResult::GetValueText()
+	{
+		if (!ValueText.IsEmpty())
+		{
+			return ValueText;
+		}
+
+		bool bIsOverriddenLiteral = false;
+
+		// Try to get corresponding member
+		const UEdGraphPin* ResolvedPin = Pin.Get();
+		const UMetasoundEditorGraphMemberNode* MemberNode = Cast<UMetasoundEditorGraphMemberNode>(GraphNode);
+		if (!MemberNode && ResolvedPin)
+		{
+			MemberNode = Cast<UMetasoundEditorGraphMemberNode>(ResolvedPin->GetOwningNode());
+		}
+
+		// Get value from member literal
+		if (const UMetasoundEditorGraphMember* Member = GetMetaSoundGraphMember(MemberNode))
+		{
+			// Don't show boolean value for triggers 
+			if (Member->GetDataType() != GetMetasoundDataTypeName<FTrigger>())
+			{
+				// Only list number of items for arrays to avoid string issues 
+				const FMetasoundFrontendLiteral& DefaultLiteral = Member->GetLiteral()->GetDefault();
+				ValueText = DefaultLiteral.ToString();
+				if (DefaultLiteral.IsArray())
+				{
+					ValueText = FString::Printf(TEXT("(%u element array)"), DefaultLiteral.GetArrayNum());
+				}
+				// Get a shorter name than the ToString one (which is LexToString)
+				else if (DefaultLiteral.GetType() == EMetasoundFrontendLiteralType::UObject)
+				{
+					UObject* Object = nullptr;
+					DefaultLiteral.TryGet(Object);
+					if (Object)
+					{
+						Object->GetName(ValueText);
+					}
+				}
+
+				// Check if input is an overridden default value (for presets) 
+				if (const UMetasoundEditorGraphInput* Input = Cast<UMetasoundEditorGraphInput>(Member))
+				{
+					Frontend::FConstGraphHandle RootGraph = Input->GetOwningGraph()->GetDocumentHandle()->GetRootGraph();
+					if (RootGraph->GetGraphPresetOptions().bIsPreset)
+					{
+						const TSet<FName>& InputsInheritingDefault = RootGraph->GetInputsInheritingDefault();
+						FName NodeName = Input->GetConstNodeHandle()->GetNodeName();
+						if (!InputsInheritingDefault.Contains(NodeName))
+						{
+							bIsOverriddenLiteral = true;
+						}
+					}
+				}
+			}
+		}
+		// Get value information directly from pin 
+		else if (ResolvedPin)
+		{
+			const bool bHasCachedValue = !ResolvedPin->DefaultValue.IsEmpty() || !ResolvedPin->AutogeneratedDefaultValue.IsEmpty() || ResolvedPin->DefaultObject || !ResolvedPin->DefaultTextValue.IsEmpty();
+			const bool bIsUnconnectedInputPin = ResolvedPin->Direction == EGPD_Input && ResolvedPin->LinkedTo.Num() == 0;
+			if (bHasCachedValue && bIsUnconnectedInputPin)
+			{
+				if (ResolvedPin->DefaultObject)
+				{
+					const AActor* DefaultActor = Cast<AActor>(ResolvedPin->DefaultObject);
+					ValueText = FString::Printf(TEXT("(%s)"), DefaultActor ? *DefaultActor->GetActorLabel() : *ResolvedPin->DefaultObject->GetName());
+				}
+				else if (!ResolvedPin->DefaultValue.IsEmpty())
+				{
+					// Don't show default for trigger type
+					if (ResolvedPin->PinType.PinCategory != GetMetasoundDataTypeName<FTrigger>())
+					{
+						ValueText = FString::Printf(TEXT("%s"), *ResolvedPin->DefaultValue);
+					}
+				}
+				else if (!ResolvedPin->AutogeneratedDefaultValue.IsEmpty())
+				{
+					ValueText = FString::Printf(TEXT("%s"), *ResolvedPin->AutogeneratedDefaultValue);
+				}
+				else if (!ResolvedPin->DefaultTextValue.IsEmpty())
+				{
+					ValueText = FString::Printf(TEXT("%s"), *ResolvedPin->DefaultTextValue.ToString());
+				}
+
+				// Get value from input literal 
+				if (const UMetasoundEditorGraphExternalNode* ExternalNode = Cast<UMetasoundEditorGraphExternalNode>(ResolvedPin->GetOwningNode()))
+				{
+					if (ResolvedPin->Direction == EGPD_Input)
+					{
+						Frontend::FConstInputHandle InputHandle = FGraphBuilder::GetConstInputHandleFromPin(ResolvedPin);
+						if (const FMetasoundFrontendLiteral* DefaultLiteral = InputHandle->GetLiteral())
+						{
+							// Get input literal class default override status 
+							const bool bIsDefaultConstructed = DefaultLiteral->GetType() == EMetasoundFrontendLiteralType::None;
+							if (!bIsDefaultConstructed)
+							{
+								bIsOverriddenLiteral = true;
+							}
+
+							// Get shorter object names 
+							if (DefaultLiteral->GetType() == EMetasoundFrontendLiteralType::UObject)
+							{
+								UObject* Object = nullptr;
+								DefaultLiteral->TryGet(Object);
+								if (Object)
+								{
+									Object->GetName(ValueText);
+								}
+							}
+						}
+					}
+				}
+			}
+		}
+			
+		if (bIsOverriddenLiteral)
+		{
+			ValueText = FString::Printf(TEXT("(overridden) %s"), *ValueText);
+		}
+		return ValueText;
+	}
+
+	FString FFindInMetasoundResult::GetValueTooltipText()
+	{
+		// Get extended value text for tooltip
+			
+		// Try to get corresponding member
+		const UEdGraphPin* ResolvedPin = Pin.Get();
+		UMetasoundEditorGraphMemberNode* MemberNode = Cast<UMetasoundEditorGraphMemberNode>(GraphNode);
+		if (!MemberNode && ResolvedPin)
+		{
+			MemberNode = Cast<UMetasoundEditorGraphMemberNode>(ResolvedPin->GetOwningNode());
+		}
+
+		// Get value from member literal
+		if (const UMetasoundEditorGraphMember* Member = GetMetaSoundGraphMember(MemberNode))
+		{
+			// Get full object names and array contents (GetValueText shortens these)
+			const FMetasoundFrontendLiteral& DefaultLiteral = Member->GetLiteral()->GetDefault();
+			return DefaultLiteral.ToString();
+		}
+		// Get value from pin's external node 
+		else if (ResolvedPin)
+		{
+			// Get full object name (GetValueText shortens these)
+			if (const UMetasoundEditorGraphExternalNode* ExternalNode = Cast<UMetasoundEditorGraphExternalNode>(ResolvedPin->GetOwningNode()))
+			{
+				if (ResolvedPin->Direction == EGPD_Input)
+				{
+					Frontend::FConstInputHandle InputHandle = FGraphBuilder::GetConstInputHandleFromPin(ResolvedPin);
+					if (const FMetasoundFrontendLiteral* DefaultLiteral = InputHandle->GetLiteral())
+					{
+						return DefaultLiteral->ToString();
+					}
+				}
+			}
+		}
+
+		// Default to same as value text 
+		return GetValueText();
+	}
+
+	const UMetasoundEditorGraphMember* FFindInMetasoundResult::GetMetaSoundGraphMember(const UEdGraphNode* EdGraphNode)
+	{
+		if (const UMetasoundEditorGraphMemberNode* MemberNode = Cast<UMetasoundEditorGraphMemberNode>(EdGraphNode))
+		{
+			return Cast<UMetasoundEditorGraphMember>(MemberNode->GetMember());
+		}
+		return nullptr;
+	}
+
+	void SFindInMetasound::Construct(const FArguments& InArgs, TSharedPtr<FEditor> InMetaSoundEditor)
+	{
+		MetaSoundEditorPtr = InMetaSoundEditor;
+
+		this->ChildSlot
+		[
+			SNew(SVerticalBox)
+			+ SVerticalBox::Slot()
+			.AutoHeight()
+			[
+				SNew(SHorizontalBox)
+				+ SHorizontalBox::Slot()
+				.FillWidth(1)
+				[
+					SAssignNew(SearchTextField, SSearchBox)
+					.HintText(LOCTEXT("GraphSearchHint", "Search"))
+					.OnTextChanged(this, &SFindInMetasound::OnSearchTextChanged)
+					.OnTextCommitted(this, &SFindInMetasound::OnSearchTextCommitted)
+				]
+			]
+			+ SVerticalBox::Slot()
+			.FillHeight(1.0f)
+			.Padding(0.f, 4.f, 0.f, 0.f)
+			[
+				SNew(SBorder)
+				.BorderImage(FAppStyle::GetBrush("Menu.Background"))
+				[
+					SAssignNew(TreeView, STreeViewType)
+					.ItemHeight(24)
+					.TreeItemsSource(&ItemsFound)
+					.OnGenerateRow(this, &SFindInMetasound::OnGenerateRow)
+					.OnGetChildren(this, &SFindInMetasound::OnGetChildren)
+					.OnSelectionChanged(this, &SFindInMetasound::OnTreeSelectionChanged)
+					.OnMouseButtonDoubleClick(this, &SFindInMetasound::OnTreeSelectionDoubleClick)
+					.SelectionMode(ESelectionMode::Single)
+				]
+			]
+		];
+	}
+
+	void SFindInMetasound::FocusForUse()
+	{
+		// NOTE: Careful, GeneratePathToWidget can be reentrant in that it can call visibility delegates and such
+		FWidgetPath FilterTextBoxWidgetPath;
+		FSlateApplication::Get().GeneratePathToWidgetUnchecked(SearchTextField.ToSharedRef(), FilterTextBoxWidgetPath);
+
+		// Set keyboard focus directly
+		FSlateApplication::Get().SetKeyboardFocus(FilterTextBoxWidgetPath, EFocusCause::SetDirectly);
+	}
+
+	void SFindInMetasound::OnSearchTextChanged(const FText& Text)
+	{
+		SearchValue = Text.ToString();
+	}
+
+	void SFindInMetasound::OnSearchTextCommitted(const FText& Text, ETextCommit::Type CommitType)
+	{
+		if (CommitType == ETextCommit::OnEnter)
+		{
+			InitiateSearch();
+		}
+	}
+
+	void SFindInMetasound::InitiateSearch()
+	{
+		TArray<FString> Tokens;
+		if (SearchValue.Contains("\"") && SearchValue.ParseIntoArray(Tokens, TEXT("\""), true) > 0)
+		{
+			for (auto &TokenIt : Tokens)
+			{
+				// we have the token, we don't need the quotes anymore, they'll just confused the comparison later on
+				TokenIt = TokenIt.TrimQuotes();
+				// We remove the spaces as all later comparison strings will also be de-spaced
+				TokenIt = TokenIt.Replace(TEXT(" "), TEXT(""));
+			}
+
+			// due to being able to handle multiple quoted blocks like ("Make Epic" "Game Now") we can end up with
+			// and empty string between (" ") blocks so this simply removes them
+			struct FRemoveMatchingStrings
+			{
+				bool operator()(const FString& RemovalCandidate) const
+				{
+					return RemovalCandidate.IsEmpty();
+				}
+			};
+			Tokens.RemoveAll(FRemoveMatchingStrings());
+		}
+		else
+		{
+			// unquoted search equivalent to a match-any-of search
+			SearchValue.ParseIntoArray(Tokens, TEXT(" "), true);
+		}
+
+		for (auto It(ItemsFound.CreateIterator()); It; ++It)
+		{
+			TreeView->SetItemExpansion(*It, false);
+		}
+		ItemsFound.Empty();
+		if (Tokens.Num() > 0)
+		{
+			HighlightText = FText::FromString(SearchValue);
+			MatchTokens(Tokens);
+		}
+
+		// Insert a fake result to inform user if none found
+		if (ItemsFound.Num() == 0)
+		{
+			ItemsFound.Add(FSearchResult(MakeShared<FFindInMetasoundResult>(LOCTEXT("FindInMetaSoundNoResults", "No Results found").ToString())));
+		}
+		else
+		{
+			// Insert a fake result for stat tracking
+			FText ResultsStats = FText::Format(LOCTEXT("FindInMetaSoundNumResultsFmt", "{0} Result(s): {1} Matching Node(s), {2} Matching Pin(s)"), FoundNodeCount + FoundPinCount, FoundNodeCount, FoundPinCount);
+			ItemsFound.Insert(FSearchResult(MakeShared<FFindInMetasoundResult>(ResultsStats.ToString())), 0);
+		}
+
+		TreeView->RequestTreeRefresh();
+
+		for (auto It(ItemsFound.CreateIterator()); It; ++It)
+		{
+			TreeView->SetItemExpansion(*It, true);
+		}
+	}
+
+	void SFindInMetasound::MatchTokens(const TArray<FString> &Tokens)
+	{
+		RootSearchResult.Reset();
+
+		UEdGraph* Graph = MetaSoundEditorPtr.Pin()->GetGraphEditor()->GetCurrentGraph();
+		MatchTokensInGraph(Graph, Tokens);
+	}
+
+	void SFindInMetasound::MatchTokensInGraph(const UEdGraph* Graph, const TArray<FString>& Tokens)
+	{
+		if (!Graph)
+		{
+			return;
+		}
+		RootSearchResult = MakeShared<FFindInMetasoundResult>(FString("MetasoundRootResult"));
+		FoundNodeCount = 0;
+		FoundPinCount = 0;
+
+		for (auto It(Graph->Nodes.CreateConstIterator()); It; ++It)
+		{
+			UEdGraphNode* Node = *It;
+
+			FString DisplayName;
+			FString NodeName = Node->GetNodeTitle(ENodeTitleType::FullTitle).ToString();
+			FString NodeType = Node->GetNodeTitle(ENodeTitleType::ListView).ToString();
+			FString DataTypeName;
+			bool bIsMemberNode = false;
+			if (const UMetasoundEditorGraphNode* MetaSoundNode = Cast<UMetasoundEditorGraphNode>(Node))
+			{
+				FMetasoundFrontendClassName ClassName = MetaSoundNode->GetClassName();
+				DisplayName = MetaSoundNode->GetDisplayName().ToString();
+
+				// Additional information for member nodes
+				if (const UMetasoundEditorGraphMember* MetaSoundMember = FFindInMetasoundResult::GetMetaSoundGraphMember(Node))
+				{
+					bIsMemberNode = true;
+					DataTypeName = MetaSoundMember->GetDataType().ToString();
+					// Get specific node type and display name for variables
+					if (const UMetasoundEditorGraphVariableNode* VariableNode = Cast<UMetasoundEditorGraphVariableNode>(MetaSoundNode))
+					{
+						NodeName = VariableNode->GetClassName().ToString();
+						switch (VariableNode->GetClassType())
+						{
+						case EMetasoundFrontendClassType::VariableMutator:
+						{
+							NodeType = TEXT("Variable (Set)");
+							break;
+						}
+						case EMetasoundFrontendClassType::VariableAccessor:
+						{
+							NodeType = TEXT("Variable (Get)");
+							break;
+						}
+						case EMetasoundFrontendClassType::VariableDeferredAccessor:
+						{
+							NodeType = TEXT("Variable (Get Delayed)");
+							break;
+						}
+
+						default:
+							break;
+						}
+						if (const UMetasoundEditorGraphVariable* MetaSoundVariable = Cast<UMetasoundEditorGraphVariable>(MetaSoundMember))
+						{
+							DisplayName = MetaSoundVariable->GetDisplayName().ToString();
+						}
+					}
+				}
+			}
+
+			FString NodeResultName;
+			if (!DisplayName.IsEmpty())
+			{
+				// Show node type (input/output/variable) for members, and only display name for others
+				if (bIsMemberNode)
+				{
+					NodeResultName = DisplayName + " - " + NodeType;
+				}
+				else
+				{
+					NodeResultName = DisplayName;
+				}
+			}
+			else
+			{
+				NodeResultName = NodeName + " - " + NodeType;
+			}
+
+			FSearchResult NodeResult;
+
+			FString NodeSearchString = NodeName + NodeType + Node->NodeComment + DisplayName + DataTypeName;
+			NodeSearchString = NodeSearchString.Replace(TEXT(" "), TEXT(""));
+
+			bool bNodeMatchesSearch = StringMatchesSearchTokens(Tokens, NodeSearchString);
+			if (bNodeMatchesSearch)
+			{
+				NodeResult = MakeShared<FFindInMetasoundResult>(NodeResultName, RootSearchResult, Node);
+				FoundNodeCount++;
+			}
+
+			for (TArray<UEdGraphPin*>::TIterator PinIt(Node->Pins); PinIt; ++PinIt)
+			{
+				UEdGraphPin* Pin = *PinIt;
+				if (Pin && Pin->PinFriendlyName.CompareTo(FText::FromString(TEXT(" "))) != 0)
+				{
+					FText PinDisplayName = Pin->GetSchema()->GetPinDisplayName(Pin);
+					FString PinDataTypeName;
+					if (const UMetasoundEditorGraphMember* MetaSoundMember = FFindInMetasoundResult::GetMetaSoundGraphMember(Node))
+					{
+						PinDataTypeName = MetaSoundMember->GetDataType().ToString();
+					}
+
+					FString PinSearchString = PinDisplayName.ToString() + Pin->PinName.ToString() + Pin->PinFriendlyName.ToString() + Pin->DefaultValue + Pin->PinType.PinCategory.ToString() + Pin->PinType.PinSubCategory.ToString() + PinDataTypeName + (Pin->PinType.PinSubCategoryObject.IsValid() ? Pin->PinType.PinSubCategoryObject.Get()->GetFullName() : TEXT(""));
+					PinSearchString = PinSearchString.Replace(TEXT(" "), TEXT(""));
+					if (StringMatchesSearchTokens(Tokens, PinSearchString))
+					{
+						if (!NodeResult.IsValid())
+						{
+							NodeResult = MakeShared<FFindInMetasoundResult>(NodeResultName, RootSearchResult, Node);
+						}
+						FSearchResult PinResult(MakeShared<FFindInMetasoundResult>(PinDisplayName.ToString(), NodeResult, Pin));
+						NodeResult->Children.Add(PinResult);
+						FoundPinCount++;
+					}
+				}
+			}
+
+			// Node or pin matches search
+			if (bNodeMatchesSearch || (NodeResult.IsValid() && NodeResult->Children.Num() > 0))
+			{
+				ItemsFound.Add(NodeResult);
+			}
+		}
+
+		for (const UEdGraph* Subgraph : Graph->SubGraphs)
+		{
+			MatchTokensInGraph(Subgraph, Tokens);
+		}
+	}
+
+	TSharedRef<ITableRow> SFindInMetasound::OnGenerateRow(FSearchResult InItem, const TSharedRef<STableViewBase>& OwnerTable)
+	{
+		FString CommentText = InItem->GetCommentText();
+
+		return SNew(STableRow<TSharedPtr<FFindInMetasoundResult>>, OwnerTable)
+			[
+				SNew(SHorizontalBox)
+				+ SHorizontalBox::Slot()
+				.AutoWidth()
+				.VAlign(VAlign_Center)
+				.Padding(4.f)
+				[
+					InItem->CreateIcon()
+				]
+				+ SHorizontalBox::Slot()
+					.FillWidth(1)
+					.HAlign(HAlign_Left)
+					.VAlign(VAlign_Center)
+					.Padding(2, 0)
+					[
+						SNew(STextBlock)
+						.Text(FText::FromString(InItem->Value))
+						.HighlightText(HighlightText)
+						.ToolTipText(FText::Format(LOCTEXT("MetasoundResultSearchToolTipFmt", "{0}: {1}"), InItem->GetCategory(), FText::FromString(InItem->Value)))
+					]
+				+ SHorizontalBox::Slot()
+					.FillWidth(1)
+					.HAlign(HAlign_Right)
+					.VAlign(VAlign_Center)
+					.Padding(2, 0)
+					[
+						SNew(STextBlock)
+						.Text(FText::FromString(InItem->GetValueText()))
+						.HighlightText(HighlightText)
+						.ToolTipText(FText::FromString(InItem->GetValueTooltipText()))
+					]
+				+ SHorizontalBox::Slot()
+					.AutoWidth()
+					.HAlign(HAlign_Right)
+					.VAlign(VAlign_Center)
+					.Padding(2, 0)
+					[
+						SNew(STextBlock)
+						.Text(CommentText.IsEmpty() ? FText::GetEmpty() : FText::Format(LOCTEXT("NodeCommentFmt", "[Comment: {0}]"), FText::FromString(CommentText)))
+						.ColorAndOpacity(FSlateColor::UseForeground())
+						.HighlightText(HighlightText)
+					]
+			];
+	}
+
+	void SFindInMetasound::OnGetChildren(FSearchResult InItem, TArray< FSearchResult >& OutChildren)
+	{
+		OutChildren += InItem->Children;
+	}
+
+	void SFindInMetasound::OnTreeSelectionChanged(FSearchResult Item, ESelectInfo::Type)
+	{
+		if (Item.IsValid())
+		{
+			Item->OnClick(MetaSoundEditorPtr);
+		}
+	}
+
+	void SFindInMetasound::OnTreeSelectionDoubleClick(FSearchResult Item)
+	{
+		if (Item.IsValid())
+		{
+			Item->OnClick(MetaSoundEditorPtr);
+		}
+	}
+
+	bool SFindInMetasound::StringMatchesSearchTokens(const TArray<FString>& Tokens, const FString& ComparisonString)
+	{
+		bool bFoundAllTokens = true;
+		//search the entry for each token, it must have all of them to pass
+		for (auto TokIT(Tokens.CreateConstIterator()); TokIT; ++TokIT)
+		{
+			const FString& Token = *TokIT;
+			if (!ComparisonString.Contains(Token))
+			{
+				bFoundAllTokens = false;
+				break;
+			}
+		}
+		return bFoundAllTokens;
+	}
+} // namespace Metasound::Editor
+#undef LOCTEXT_NAMESPACE
diff --git a/Engine/Plugins/Runtime/Metasound/Source/MetasoundEditor/Private/MetasoundEditor.cpp b/Engine/Plugins/Runtime/Metasound/Source/MetasoundEditor/Private/MetasoundEditor.cpp
--- a/Engine/Plugins/Runtime/Metasound/Source/MetasoundEditor/Private/MetasoundEditor.cpp
+++ b/Engine/Plugins/Runtime/Metasound/Source/MetasoundEditor/Private/MetasoundEditor.cpp
@@ -666,6 +666,14 @@ namespace Metasound
 			.SetDisplayName(LOCTEXT("InterfacesTab", "Interfaces"))
 			.SetGroup(WorkspaceMenuCategoryRef)
 			.SetIcon(FSlateIcon(FAppStyle::GetAppStyleSetName(), "ClassIcon.Interface"));
+
+			InTabManager->RegisterTabSpawner(TabFactory::Names::Find, FOnSpawnTab::CreateLambda([InFindWidget = FindWidget](const FSpawnTabArgs& Args)
+			{
+				return TabFactory::CreateFindTab(InFindWidget, Args);
+			}))
+			.SetDisplayName(LOCTEXT("FindTab", "Find in MetaSound"))
+			.SetGroup(WorkspaceMenuCategoryRef)
+			.SetIcon(FSlateIcon(FAppStyle::GetAppStyleSetName(), "Kismet.Tabs.FindResults"));
 		}
 
 		void FEditor::UnregisterTabSpawners(const TSharedRef<FTabManager>& InTabManager)
@@ -679,6 +687,7 @@ namespace Metasound
 			InTabManager->UnregisterTabSpawner(TabFactory::Names::Details);
 			InTabManager->UnregisterTabSpawner(TabFactory::Names::Members);
 			InTabManager->UnregisterTabSpawner(TabFactory::Names::Interfaces);
+			InTabManager->UnregisterTabSpawner(TabFactory::Names::Find);
 		}
 
 		TSharedPtr<SWidget> FEditor::BuildAnalyzerWidget() const
@@ -811,7 +820,7 @@ namespace Metasound
 				NotifyAssetPrimeInProgress();
 			}
 
-			const TSharedRef<FTabManager::FLayout> StandaloneDefaultLayout = FTabManager::NewLayout("Standalone_MetasoundEditor_Layout_v10")
+			const TSharedRef<FTabManager::FLayout> StandaloneDefaultLayout = FTabManager::NewLayout("Standalone_MetasoundEditor_Layout_v11")
 				->AddArea
 				(
 					FTabManager::NewPrimaryArea()
@@ -847,11 +856,25 @@ namespace Metasound
 						)
 						->Split
 						(
-							FTabManager::NewStack()
+							FTabManager::NewSplitter()
 							->SetSizeCoefficient(0.77f)
-							->SetHideTabWell(true)
-							->AddTab(TabFactory::Names::GraphCanvas, ETabState::OpenedTab)
+							->SetOrientation(Orient_Vertical)
+							->Split
+							(
+								FTabManager::NewStack()
+								->SetSizeCoefficient(0.8f)
+								->SetHideTabWell(true)
+								->AddTab(TabFactory::Names::GraphCanvas, ETabState::OpenedTab)
+							)
+							->Split
+							(
+								FTabManager::NewStack()
+								->SetSizeCoefficient(0.2f)
+								->SetHideTabWell(true)
+								->AddTab(TabFactory::Names::Find, ETabState::OpenedTab)
+							)
 						)
+
 						->Split
 						(
 							FTabManager::NewStack()
@@ -1156,6 +1179,8 @@ namespace Metasound
 			}
 
 			Palette = SNew(SMetasoundPalette);
+
+			FindWidget = SNew(SFindInMetasound, SharedThis(this));
 		}
 
 		// TODO: Tie in rename on GraphActionMenu.  For now, just renameable via field in details
@@ -1430,6 +1455,10 @@ namespace Metasound
 			ToolkitCommands->MapAction(
 				FEditorCommands::Get().UpdateNodeClass,
 				FExecuteAction::CreateSP(this, &FEditor::UpdateSelectedNodeClasses));
+
+			ToolkitCommands->MapAction(
+				FEditorCommands::Get().FindInMetaSound,
+				FExecuteAction::CreateSP(this, &FEditor::ShowFindInMetaSound));
 		}
 
 		void FEditor::Import()
@@ -3691,6 +3720,15 @@ namespace Metasound
 					.ColorAndOpacity(FSlateColor::UseForeground())
 				];
 		}
+
+		void FEditor::ShowFindInMetaSound()
+		{
+			TabManager->TryInvokeTab(TabFactory::Names::Find);
+			if (FindWidget.IsValid())
+			{
+				FindWidget->FocusForUse();
+			}
+		}
 	}
 }
 #undef LOCTEXT_NAMESPACE
diff --git a/Engine/Plugins/Runtime/Metasound/Source/MetasoundEditor/Private/MetasoundEditor.h b/Engine/Plugins/Runtime/Metasound/Source/MetasoundEditor/Private/MetasoundEditor.h
--- a/Engine/Plugins/Runtime/Metasound/Source/MetasoundEditor/Private/MetasoundEditor.h
+++ b/Engine/Plugins/Runtime/Metasound/Source/MetasoundEditor/Private/MetasoundEditor.h
@@ -50,6 +50,10 @@ class SMetasoundPalette;
 class FSlateRect;
 class IDetailsView;
 class IToolkitHost;
+namespace Metasound::Editor
+{
+	class SFindInMetasound;
+}
 class SVerticalBox;
 class UAudioComponent;
 class UEdGraphNode;
@@ -335,6 +339,9 @@ namespace Metasound
 
 			void RefreshEditorContext();
 
+			/** Show and focus the Find in MetaSound tab. */
+			void ShowFindInMetaSound();
+
 		private:
 			void RemoveInvalidSelection();
 
@@ -436,6 +443,9 @@ namespace Metasound
 			/** Spectrum analyzer used in the analyzer tab for auditioning preview output. */
 			TSharedPtr<AudioWidgets::FAudioSpectrumAnalyzer> OutputSpectrumAnalyzer;
 
+			/** Find in MetaSound widget. */
+			TSharedPtr<SFindInMetasound> FindWidget;
+
 			/** Palette of Node types */
 			TSharedPtr<SMetasoundPalette> Palette;
 
diff --git a/Engine/Plugins/Runtime/Metasound/Source/MetasoundEditor/Private/MetasoundEditorCommands.cpp b/Engine/Plugins/Runtime/Metasound/Source/MetasoundEditor/Private/MetasoundEditorCommands.cpp
--- a/Engine/Plugins/Runtime/Metasound/Source/MetasoundEditor/Private/MetasoundEditorCommands.cpp
+++ b/Engine/Plugins/Runtime/Metasound/Source/MetasoundEditor/Private/MetasoundEditorCommands.cpp
@@ -26,6 +26,7 @@ namespace Metasound
 
 			UI_COMMAND(ConvertFromPreset, "Convert From Preset", "Converts this preset to a fully accessible MetaSound.", EUserInterfaceActionType::Button, FInputChord());
 			UI_COMMAND(Delete, "Delete Selected", "Delete selected items.", EUserInterfaceActionType::None, FInputChord(EKeys::Delete));
+			UI_COMMAND(FindInMetaSound, "Find in MetaSound", "Find a node or pin within this MetaSound.", EUserInterfaceActionType::Button, FInputChord(EModifierKey::Control, EKeys::F));
 		}
 	} // namespace Editor
 } // namespace Metasound
diff --git a/Engine/Plugins/Runtime/Metasound/Source/MetasoundEditor/Private/MetasoundEditorTabFactory.h b/Engine/Plugins/Runtime/Metasound/Source/MetasoundEditor/Private/MetasoundEditorTabFactory.h
--- a/Engine/Plugins/Runtime/Metasound/Source/MetasoundEditor/Private/MetasoundEditorTabFactory.h
+++ b/Engine/Plugins/Runtime/Metasound/Source/MetasoundEditor/Private/MetasoundEditorTabFactory.h
@@ -1,19 +1,19 @@
 // Copyright Epic Games, Inc. All Rights Reserved.
 #pragma once
 
-#include "CoreMinimal.h"
+#include "FindInMetasound.h"
 #include "IDetailsView.h"
+#include "SAudioMeter.h"
 #include "SGraphActionMenu.h"
 #include "SMetasoundPalette.h"
-#include "Styling/AppStyle.h"
 #include "Framework/Docking/TabManager.h"
 #include "Internationalization/Internationalization.h"
-#include "Widgets/Docking/SDockTab.h"
-#include "SAudioMeter.h"
+#include "Styling/AppStyle.h"
 #include "Styling/ISlateStyle.h"
 #include "Styling/SlateStyleRegistry.h"
 #include "Templates/SharedPointer.h"
 #include "Widgets/Colors/SColorBlock.h"
+#include "Widgets/Docking/SDockTab.h"
 
 #define LOCTEXT_NAMESPACE "MetaSoundEditor"
 
@@ -32,6 +32,7 @@ namespace Metasound
 				const FName Members = "MetasoundEditor_Members";
 				const FName Palette = "MetasoundEditor_Palette";
 				const FName Interfaces = "MetasoundEditor_Interfaces";
+				const FName Find = "MetasoundEditor_Find";
 			}
 
 			TSharedRef<SDockTab> CreateAnalyzersTab(TSharedPtr<SWidget> InAnalyzerWidget, const FSpawnTabArgs& Args)
@@ -110,6 +111,17 @@ namespace Metasound
 						InInterfacesDetails.ToSharedRef()
 					];
 			}
+
+			TSharedRef<SDockTab> CreateFindTab(TSharedPtr<SFindInMetasound> InFindWidget, const FSpawnTabArgs& Args)
+			{
+				check(Args.GetTabId() == Names::Find);
+
+				return SNew(SDockTab)
+					.Label(LOCTEXT("MetasoundFindTitle", "Find Results"))
+					[
+						InFindWidget.ToSharedRef()
+					];
+			}
 		}
 	} // namespace Editor
 } // namespace Metasound
diff --git a/Engine/Plugins/Runtime/Metasound/Source/MetasoundEditor/Public/FindInMetasound.h b/Engine/Plugins/Runtime/Metasound/Source/MetasoundEditor/Public/FindInMetasound.h
new file mode 100644
--- /dev/null
+++ b/Engine/Plugins/Runtime/Metasound/Source/MetasoundEditor/Public/FindInMetasound.h
@@ -0,0 +1,155 @@
+// Copyright Epic Games, Inc. All Rights Reserved.
+
+#pragma once
+
+#include "EdGraph/EdGraph.h"
+#include "EdGraph/EdGraphPin.h"
+#include "Input/Reply.h"
+#include "MetasoundEditor.h"
+#include "Widgets/DeclarativeSyntaxSupport.h"
+#include "Widgets/SCompoundWidget.h"
+#include "Widgets/Views/STableRow.h"
+#include "Widgets/Views/STableViewBase.h"
+#include "Widgets/Views/STreeView.h"
+
+// FindInMetaSound is largely taken from FindInMaterial
+namespace Metasound::Editor
+{
+	/** Item that matched the search results */
+	class FFindInMetasoundResult
+	{
+	public:
+		/** Create a root (or only text) result */
+		FFindInMetasoundResult(const FString& InResultName);
+
+		/** Create a listing for a search result*/
+		FFindInMetasoundResult(const FString& InResultName, TSharedPtr<FFindInMetasoundResult>& InParent, UClass* InClass, int InDuplicationIndex);
+
+		/** Create a listing for a pin result */
+		FFindInMetasoundResult(const FString& InResultName, TSharedPtr<FFindInMetasoundResult>& InParent, UEdGraphPin* InPin);
+
+		/** Create a listing for a node result */
+		FFindInMetasoundResult(const FString& InResultName, TSharedPtr<FFindInMetasoundResult>& InParent, UEdGraphNode* InNode);
+
+		/** Called when user clicks on the search item */
+		FReply OnClick(TWeakPtr<class FEditor> MetasoundEditor);
+
+		/* Get Category for this search result */
+		FText GetCategory() const;
+
+		/** Create an icon to represent the result */
+		TSharedRef<SWidget> CreateIcon() const;
+
+		/** Gets the comment on this node if any */
+		FString GetCommentText() const;
+
+		/** Gets the value of the pin or node if any */
+		FString GetValueText();
+			
+		/** Gets the value tooltip (longer version of value text) of the pin or node if any */
+		FString GetValueTooltipText();
+
+		/** Any children listed under this category */
+		TArray<TSharedPtr<FFindInMetasoundResult>> Children;
+
+		/** Search result Parent */
+		TWeakPtr<FFindInMetasoundResult> Parent;
+
+		/*The meta string that was stored in the asset registry for this item */
+		FString Value;
+
+		/*The graph may have multiple instances of whatever we are looking for, this tells us which instance # we refer to*/
+		int	DuplicationIndex;
+
+		/*The class this item refers to */
+		UClass* Class;
+
+		/** The pin that this search result refers to */
+		FEdGraphPinReference Pin;
+
+		/** The graph node that this search result refers to (if not by asset registry or UK2Node) */
+		TWeakObjectPtr<UEdGraphNode> GraphNode;
+
+		/** Display text for comment information */
+		FString CommentText;
+
+		/** Cache value text */
+		FString ValueText;
+
+		/** Helper function to get MetaSound graph member from an ed graph node */
+		static const UMetasoundEditorGraphMember* GetMetaSoundGraphMember(const UEdGraphNode* EdGraphNode);
+	};
+
+
+	/** Widget for searching for items that are part of a UEdGraph */
+	class SFindInMetasound : public SCompoundWidget
+	{
+	public:
+		SLATE_BEGIN_ARGS(SFindInMetasound) {}
+		SLATE_END_ARGS()
+
+			void Construct(const FArguments& InArgs, TSharedPtr<class FEditor> InMetaSoundEditor);
+
+		/** Focuses this widget's search box */
+		void FocusForUse();
+
+	protected:
+		typedef TSharedPtr<FFindInMetasoundResult> FSearchResult;
+		typedef STreeView<FSearchResult> STreeViewType;
+
+		/** Called when user changes the text they are searching for */
+		void OnSearchTextChanged(const FText& Text);
+
+		/** Called when user changes commits text to the search box */
+		void OnSearchTextCommitted(const FText& Text, ETextCommit::Type CommitType);
+
+		/** Get the children of a row */
+		void OnGetChildren(FSearchResult InItem, TArray<FSearchResult>& OutChildren);
+
+		/** Called when user clicks on a new result */
+		void OnTreeSelectionChanged(FSearchResult Item, ESelectInfo::Type SelectInfo);
+
+		/** Called when user double clicks on a new result */
+		void OnTreeSelectionDoubleClick(FSearchResult Item);
+
+		/** Called when a new row is being generated */
+		TSharedRef<ITableRow> OnGenerateRow(FSearchResult InItem, const TSharedRef<STableViewBase>& OwnerTable);
+
+		/** Begins the search based on the SearchValue */
+		void InitiateSearch();
+
+		/** Find any results that contain all of the tokens */
+		void MatchTokens(const TArray<FString>& Tokens);
+
+		/** Find any results that contain all of the tokens in provided graph and subgraphs */
+		void MatchTokensInGraph(const UEdGraph* Graph, const TArray<FString>& Tokens);
+
+		/** Determines if a string matches the search tokens */
+		static bool StringMatchesSearchTokens(const TArray<FString>& Tokens, const FString& ComparisonString);
+
+		/** Pointer back to the MetaSound editor that owns us */
+		TWeakPtr<class FEditor> MetaSoundEditorPtr;
+
+		/** The tree view displays the results */
+		TSharedPtr<STreeViewType> TreeView;
+
+		/** The search text box */
+		TSharedPtr<class SSearchBox> SearchTextField;
+
+		/** This buffer stores the currently displayed results */
+		TArray<FSearchResult> ItemsFound;
+
+		/** we need to keep a handle on the root result, because it won't show up in the tree */
+		FSearchResult RootSearchResult;
+
+		/** The string to highlight in the results */
+		FText HighlightText;
+
+		/** The string to search for */
+		FString SearchValue;
+
+		/** Counters for number of results of each type */
+		uint32 FoundNodeCount = 0;
+		uint32 FoundPinCount = 0;
+	};
+} // namespace Metasound
\ No newline at end of file
diff --git a/Engine/Plugins/Runtime/Metasound/Source/MetasoundEditor/Public/MetasoundEditorCommands.h b/Engine/Plugins/Runtime/Metasound/Source/MetasoundEditor/Public/MetasoundEditorCommands.h
--- a/Engine/Plugins/Runtime/Metasound/Source/MetasoundEditor/Public/MetasoundEditorCommands.h
+++ b/Engine/Plugins/Runtime/Metasound/Source/MetasoundEditor/Public/MetasoundEditorCommands.h
@@ -58,6 +58,9 @@ namespace Metasound
 			// Delete selected items. 
 			TSharedPtr<FUICommandInfo> Delete;
 
+			// Show the Find in MetaSound window. 
+			TSharedPtr<FUICommandInfo> FindInMetaSound;
+
 			/** Initialize commands */
 			virtual void RegisterCommands() override;
 		};
diff --git a/Engine/Plugins/Runtime/Metasound/Source/MetasoundFrontend/Private/MetasoundFrontendLiteral.cpp b/Engine/Plugins/Runtime/Metasound/Source/MetasoundFrontend/Private/MetasoundFrontendLiteral.cpp
--- a/Engine/Plugins/Runtime/Metasound/Source/MetasoundFrontend/Private/MetasoundFrontendLiteral.cpp
+++ b/Engine/Plugins/Runtime/Metasound/Source/MetasoundFrontend/Private/MetasoundFrontendLiteral.cpp
@@ -67,6 +67,8 @@ namespace MetasoundFrontendLiteralPrivate
 				Builder << TLiteralValueToStringHelper<ElementType>::Convert(InArray[i]);
 			}
 
+			Builder << TEXT("]");
+
 			return FString(Builder);
 		}
 	};
@@ -828,6 +830,43 @@ FString FMetasoundFrontendLiteral::ToString() const
 	}
 }
 
+int32 FMetasoundFrontendLiteral::GetArrayNum() const
+{
+	switch (Type)
+	{
+	case EMetasoundFrontendLiteralType::None:
+	case EMetasoundFrontendLiteralType::Boolean:
+	case EMetasoundFrontendLiteralType::Integer:
+	case EMetasoundFrontendLiteralType::Float:
+	case EMetasoundFrontendLiteralType::String:
+	case EMetasoundFrontendLiteralType::UObject:
+		return INDEX_NONE;
+
+	case EMetasoundFrontendLiteralType::NoneArray:
+		return AsNumDefault;
+
+	case EMetasoundFrontendLiteralType::BooleanArray:
+		return AsBoolean.Num();
+
+	case EMetasoundFrontendLiteralType::IntegerArray:
+		return AsInteger.Num();
+
+	case EMetasoundFrontendLiteralType::FloatArray:
+		return AsFloat.Num();
+
+	case EMetasoundFrontendLiteralType::StringArray:
+		return AsString.Num();
+
+	case EMetasoundFrontendLiteralType::UObjectArray:
+		return AsUObject.Num();
+
+	case EMetasoundFrontendLiteralType::Invalid:
+	default:
+		static_assert(static_cast<int32>(EMetasoundFrontendLiteralType::Invalid) == 12, "Possible missing literal type switch coverage");
+		return INDEX_NONE;
+	}
+}
+
 EMetasoundFrontendLiteralType FMetasoundFrontendLiteral::GetType() const
 {
 	return Type;
diff --git a/Engine/Plugins/Runtime/Metasound/Source/MetasoundFrontend/Public/MetasoundFrontendLiteral.h b/Engine/Plugins/Runtime/Metasound/Source/MetasoundFrontend/Public/MetasoundFrontendLiteral.h
--- a/Engine/Plugins/Runtime/Metasound/Source/MetasoundFrontend/Public/MetasoundFrontendLiteral.h
+++ b/Engine/Plugins/Runtime/Metasound/Source/MetasoundFrontend/Public/MetasoundFrontendLiteral.h
@@ -139,6 +139,9 @@ public:
 	void SetFromLiteral(const Metasound::FLiteral& InLiteral);
 
 	EMetasoundFrontendLiteralType GetType() const;
+	
+	// Get the number of array elements if this literal is an array type, otherwise return 0
+	int32 GetArrayNum() const;
 
 	// Return a Metasound::FLiteral representation of this object. 
 	//
