Fix value text for variables with connected setters for Find in MetaSound and change value text to FText for loc

diff --git a/Engine/Plugins/Runtime/Metasound/Source/MetasoundEditor/Private/FindInMetasound.cpp b/Engine/Plugins/Runtime/Metasound/Source/MetasoundEditor/Private/FindInMetasound.cpp
--- a/Engine/Plugins/Runtime/Metasound/Source/MetasoundEditor/Private/FindInMetasound.cpp
+++ b/Engine/Plugins/Runtime/Metasound/Source/MetasoundEditor/Private/FindInMetasound.cpp
@@ -161,7 +161,7 @@ namespace Metasound::Editor
 		return CommentText;
 	}
 
-	FString FFindInMetasoundResult::GetValueText()
+	FText FFindInMetasoundResult::GetValueText()
 	{
 		if (!ValueText.IsEmpty())
 		{
@@ -186,10 +186,10 @@ namespace Metasound::Editor
 			{
 				// Only list number of items for arrays to avoid string issues 
 				const FMetasoundFrontendLiteral& DefaultLiteral = Member->GetLiteral()->GetDefault();
-				ValueText = DefaultLiteral.ToString();
+				ValueText = FText::FromString(DefaultLiteral.ToString());
 				if (DefaultLiteral.IsArray())
 				{
-					ValueText = FString::Printf(TEXT("(%u element array)"), DefaultLiteral.GetArrayNum());
+					ValueText = FText::Format(LOCTEXT("FindMetasound_OverriddenValueDescriptionFormat", "({0} element array)"), DefaultLiteral.GetArrayNum());
 				}
 				// Get a shorter name than the ToString one (which is LexToString)
 				else if (DefaultLiteral.GetType() == EMetasoundFrontendLiteralType::UObject)
@@ -198,7 +198,9 @@ namespace Metasound::Editor
 					DefaultLiteral.TryGet(Object);
 					if (Object)
 					{
-						Object->GetName(ValueText);
+						FString ObjectName;
+						Object->GetName(ObjectName);
+						ValueText = FText::FromString(ObjectName);
 					}
 				}
 
@@ -216,6 +218,26 @@ namespace Metasound::Editor
 						}
 					}
 				}
+				// Don't show default value if variable mutator node is connected 
+				if (const UMetasoundEditorGraphVariable* Variable = Cast<UMetasoundEditorGraphVariable>(Member))
+				{
+					Frontend::FConstNodeHandle MutatorNodeHandle = Variable->GetConstVariableHandle()->FindMutatorNode();
+					if (MutatorNodeHandle->IsValid())
+					{
+						auto IsMutatorNode = [&MutatorNodeHandle](const UMetasoundEditorGraphMemberNode* Node)
+						{
+							return Node->GetNodeID() == MutatorNodeHandle->GetID();
+						};
+						TArray<UMetasoundEditorGraphMemberNode*> Nodes = Variable->GetNodes();
+						if ( UMetasoundEditorGraphMemberNode** MutatorNode = Nodes.FindByPredicate(IsMutatorNode))
+						{
+							if (*MutatorNode && (*MutatorNode)->Pins[0]->HasAnyConnections())
+							{
+								ValueText = LOCTEXT("FindMetasound_VariableSetterConnectedDescription", "(See value from variable setter node connection)");
+							}
+						}
+					}
+				}
 			}
 		}
 		// Get value information directly from pin 
@@ -228,23 +250,23 @@ namespace Metasound::Editor
 				if (ResolvedPin->DefaultObject)
 				{
 					const AActor* DefaultActor = Cast<AActor>(ResolvedPin->DefaultObject);
-					ValueText = FString::Printf(TEXT("(%s)"), DefaultActor ? *DefaultActor->GetActorLabel() : *ResolvedPin->DefaultObject->GetName());
+					ValueText = FText::FromString(DefaultActor ? *DefaultActor->GetActorLabel() : *ResolvedPin->DefaultObject->GetName());
 				}
 				else if (!ResolvedPin->DefaultValue.IsEmpty())
 				{
 					// Don't show default for trigger type
 					if (ResolvedPin->PinType.PinCategory != GetMetasoundDataTypeName<FTrigger>())
 					{
-						ValueText = FString::Printf(TEXT("%s"), *ResolvedPin->DefaultValue);
+						ValueText = FText::FromString(ResolvedPin->DefaultValue);
 					}
 				}
 				else if (!ResolvedPin->AutogeneratedDefaultValue.IsEmpty())
 				{
-					ValueText = FString::Printf(TEXT("%s"), *ResolvedPin->AutogeneratedDefaultValue);
+					ValueText = FText::FromString(ResolvedPin->AutogeneratedDefaultValue);
 				}
 				else if (!ResolvedPin->DefaultTextValue.IsEmpty())
 				{
-					ValueText = FString::Printf(TEXT("%s"), *ResolvedPin->DefaultTextValue.ToString());
+					ValueText = FText::FromString(ResolvedPin->DefaultTextValue.ToString());
 				}
 
 				// Get value from input literal 
@@ -269,7 +291,9 @@ namespace Metasound::Editor
 								DefaultLiteral->TryGet(Object);
 								if (Object)
 								{
-									Object->GetName(ValueText);
+									FString ObjectName;
+									Object->GetName(ObjectName);
+									ValueText = FText::FromString(ObjectName);
 								}
 							}
 						}
@@ -280,12 +304,12 @@ namespace Metasound::Editor
 			
 		if (bIsOverriddenLiteral)
 		{
-			ValueText = FString::Printf(TEXT("(overridden) %s"), *ValueText);
+			ValueText = FText::Format(LOCTEXT("FindMetasound_OverriddenValueDescriptionFormat", "(overridden) {0}"), ValueText);
 		}
 		return ValueText;
 	}
 
-	FString FFindInMetasoundResult::GetValueTooltipText()
+	FText FFindInMetasoundResult::GetValueTooltipText()
 	{
 		// Get extended value text for tooltip
 			
@@ -302,7 +326,7 @@ namespace Metasound::Editor
 		{
 			// Get full object names and array contents (GetValueText shortens these)
 			const FMetasoundFrontendLiteral& DefaultLiteral = Member->GetLiteral()->GetDefault();
-			return DefaultLiteral.ToString();
+			return FText::FromString(DefaultLiteral.ToString());
 		}
 		// Get value from pin's external node 
 		else if (ResolvedPin)
@@ -315,7 +339,7 @@ namespace Metasound::Editor
 					Frontend::FConstInputHandle InputHandle = FGraphBuilder::GetConstInputHandleFromPin(ResolvedPin);
 					if (const FMetasoundFrontendLiteral* DefaultLiteral = InputHandle->GetLiteral())
 					{
-						return DefaultLiteral->ToString();
+						return FText::FromString(DefaultLiteral->ToString());
 					}
 				}
 			}
@@ -630,9 +654,9 @@ namespace Metasound::Editor
 					.Padding(2, 0)
 					[
 						SNew(STextBlock)
-						.Text(FText::FromString(InItem->GetValueText()))
+						.Text(InItem->GetValueText())
 						.HighlightText(HighlightText)
-						.ToolTipText(FText::FromString(InItem->GetValueTooltipText()))
+						.ToolTipText(InItem->GetValueTooltipText())
 					]
 				+ SHorizontalBox::Slot()
 					.AutoWidth()
diff --git a/Engine/Plugins/Runtime/Metasound/Source/MetasoundEditor/Public/FindInMetasound.h b/Engine/Plugins/Runtime/Metasound/Source/MetasoundEditor/Public/FindInMetasound.h
--- a/Engine/Plugins/Runtime/Metasound/Source/MetasoundEditor/Public/FindInMetasound.h
+++ b/Engine/Plugins/Runtime/Metasound/Source/MetasoundEditor/Public/FindInMetasound.h
@@ -44,10 +44,10 @@ namespace Metasound::Editor
 		FString GetCommentText() const;
 
 		/** Gets the value of the pin or node if any */
-		FString GetValueText();
+		FText GetValueText();
 			
 		/** Gets the value tooltip (longer version of value text) of the pin or node if any */
-		FString GetValueTooltipText();
+		FText GetValueTooltipText();
 
 		/** Any children listed under this category */
 		TArray<TSharedPtr<FFindInMetasoundResult>> Children;
@@ -74,7 +74,7 @@ namespace Metasound::Editor
 		FString CommentText;
 
 		/** Cache value text */
-		FString ValueText;
+		FText ValueText;
 
 		/** Helper function to get MetaSound graph member from an ed graph node */
 		static const UMetasoundEditorGraphMember* GetMetaSoundGraphMember(const UEdGraphNode* EdGraphNode);
@@ -88,7 +88,7 @@ namespace Metasound::Editor
 		SLATE_BEGIN_ARGS(SFindInMetasound) {}
 		SLATE_END_ARGS()
 
-			void Construct(const FArguments& InArgs, TSharedPtr<class FEditor> InMetaSoundEditor);
+		void Construct(const FArguments& InArgs, TSharedPtr<class FEditor> InMetaSoundEditor);
 
 		/** Focuses this widget's search box */
 		void FocusForUse();
