Fix remaining failures in metasound automated node tests

diff --git a/Engine/Plugins/Runtime/Metasound/Source/MetasoundFrontend/Private/MetasoundFrontendDocument.cpp b/Engine/Plugins/Runtime/Metasound/Source/MetasoundFrontend/Private/MetasoundFrontendDocument.cpp
--- a/Engine/Plugins/Runtime/Metasound/Source/MetasoundFrontend/Private/MetasoundFrontendDocument.cpp
+++ b/Engine/Plugins/Runtime/Metasound/Source/MetasoundFrontend/Private/MetasoundFrontendDocument.cpp
@@ -50,6 +50,21 @@ namespace Metasound
 			static const FString VariableMutator = TEXT("Variable (Mutator)");
 			static const FString Template = TEXT("Template");
 			static const FString Invalid = TEXT("Invalid");
+
+			static const TSortedMap<FString, EMetasoundFrontendClassType, TInlineAllocator<32>> ClassTypeCStringToEnum = 
+			{
+				{External, EMetasoundFrontendClassType::External},
+				{Graph, EMetasoundFrontendClassType::Graph},
+				{Input, EMetasoundFrontendClassType::Input},
+				{Output, EMetasoundFrontendClassType::Output},
+				{Literal, EMetasoundFrontendClassType::Literal},
+				{Variable, EMetasoundFrontendClassType::Variable},
+				{VariableDeferredAccessor, EMetasoundFrontendClassType::VariableDeferredAccessor},
+				{VariableAccessor, EMetasoundFrontendClassType::VariableAccessor},
+				{VariableMutator, EMetasoundFrontendClassType::VariableMutator},
+				{Template, EMetasoundFrontendClassType::Template},
+				{Invalid, EMetasoundFrontendClassType::Invalid}
+			};
 		}
 	} // namespace Frontend
 
@@ -921,62 +936,16 @@ namespace Metasound::Frontend
 {
 	bool StringToClassType(const FString& InString, EMetasoundFrontendClassType& OutClassType)
 	{
-		if (InString == ClassTypePrivate::External)
-		{
-			OutClassType = EMetasoundFrontendClassType::External;
-			return true;
-		}
-
-		if (InString == ClassTypePrivate::Graph)
-		{
-			OutClassType = EMetasoundFrontendClassType::Graph;
-			return true;
-		}
-
-		if (InString == ClassTypePrivate::Input)
+		if (const EMetasoundFrontendClassType* FoundClassType = ClassTypePrivate::ClassTypeCStringToEnum.Find(*InString))
 		{
-			OutClassType = EMetasoundFrontendClassType::Input;
-			return true;
+			OutClassType = *FoundClassType;
 		}
-
-		if (InString == ClassTypePrivate::Invalid)
+		else
 		{
 			OutClassType = EMetasoundFrontendClassType::Invalid;
-			return true;
-		}
-
-		if (InString == ClassTypePrivate::Literal)
-		{
-			OutClassType = EMetasoundFrontendClassType::Literal;
-			return true;
 		}
-
-		if (InString == ClassTypePrivate::Output)
-		{
-			OutClassType = EMetasoundFrontendClassType::Output;
-			return true;
-		}
-
-		if (InString == ClassTypePrivate::Template)
-		{
-			OutClassType = EMetasoundFrontendClassType::Template;
-			return true;
-		}
-
-		if (InString == ClassTypePrivate::Variable)
-		{
-			OutClassType = EMetasoundFrontendClassType::Variable;
-			return true;
-		}
-
-		if (InString == ClassTypePrivate::VariableDeferredAccessor)
-		{
-			OutClassType = EMetasoundFrontendClassType::VariableDeferredAccessor;
-			return true;
-		}
-
-		OutClassType = EMetasoundFrontendClassType::Invalid;
-		return false;
+		
+		return OutClassType != EMetasoundFrontendClassType::Invalid;
 	}
 
 	void ForEachLiteral(const FMetasoundFrontendDocument& InDoc, FForEachLiteralFunctionRef OnLiteral)
diff --git a/Engine/Source/Runtime/SignalProcessing/Private/DynamicsProcessor.cpp b/Engine/Source/Runtime/SignalProcessing/Private/DynamicsProcessor.cpp
--- a/Engine/Source/Runtime/SignalProcessing/Private/DynamicsProcessor.cpp
+++ b/Engine/Source/Runtime/SignalProcessing/Private/DynamicsProcessor.cpp
@@ -77,7 +77,7 @@ namespace Audio
 	
 	void FDynamicsProcessor::SetLookaheadMsec(const float InLookAheadMsec)
 	{
-		LookaheadDelayMsec = InLookAheadMsec;
+		LookaheadDelayMsec = FMath::Min(InLookAheadMsec, MaxLookaheadMsec);
 		const int32 NumDelayFrames = GetNumDelayFrames();
 		for (int32 Channel = 0; Channel < LookaheadDelay.Num(); ++Channel)
 		{
@@ -699,6 +699,7 @@ namespace Audio
 
 	int32 FDynamicsProcessor::GetNumDelayFrames() const
 	{
+		checkf(LookaheadDelayMsec <= MaxLookaheadMsec, TEXT("An lookahead delay of %fms exceeds maximum lookahead delay of %fms"), LookaheadDelayMsec, MaxLookaheadMsec)
 		return FMath::Max(0, FMath::CeilToInt(LookaheadDelayMsec * SampleRate / 1000.0f));
 	}
 
diff --git a/Engine/Source/Runtime/SignalProcessing/Private/GrainDelay.cpp b/Engine/Source/Runtime/SignalProcessing/Private/GrainDelay.cpp
--- a/Engine/Source/Runtime/SignalProcessing/Private/GrainDelay.cpp
+++ b/Engine/Source/Runtime/SignalProcessing/Private/GrainDelay.cpp
@@ -14,7 +14,14 @@ namespace Audio
 		FGrainDelay::FGrainDelay(const float InSampleRate)
 			: SampleRate(InSampleRate)
 		{
-			GrainDelayLine.Init(SampleRate, MaxDelaySeconds);
+			// DelayLine must accommodate the delay of the grain and the pitch shifter.
+			// The pitch shifter has a maximum delay and requires an extra millisecond 
+			// for phase offsets. An additional sample is added to account for floating 
+			// point rounding.
+			const float RoundingErrorBufferInSeconds = 1.f / InSampleRate;
+			const float MaxTapDelayInSeconds = (FTapDelayPitchShifter::MaxDelayLength + 1.f) / 1000.f;
+			const float DelayLineSizeInSeconds = MaxDelaySeconds + MaxTapDelayInSeconds + RoundingErrorBufferInSeconds;
+			GrainDelayLine.Init(SampleRate, DelayLineSizeInSeconds);
 			
 			// Generate the grain data
 			GenerateEnvelopeData(GrainEnvelope, GrainEnvelopeSampleCount, GrainEnvelopeType);
