CIS loctext fix for find in MetaSound

diff --git a/Engine/Plugins/Runtime/Metasound/Source/MetasoundEditor/Private/FindInMetasound.cpp b/Engine/Plugins/Runtime/Metasound/Source/MetasoundEditor/Private/FindInMetasound.cpp
--- a/Engine/Plugins/Runtime/Metasound/Source/MetasoundEditor/Private/FindInMetasound.cpp
+++ b/Engine/Plugins/Runtime/Metasound/Source/MetasoundEditor/Private/FindInMetasound.cpp
@@ -11,7 +11,7 @@
 #include "Widgets/Images/SImage.h"
 #include "Widgets/Input/SSearchBox.h"
 
-#define LOCTEXT_NAMESPACE "FindInMetasound"
+#define LOCTEXT_NAMESPACE "MetasoundEditor"
 
 namespace Metasound::Editor
 {
@@ -59,11 +59,11 @@ namespace Metasound::Editor
 	{
 		if (Class == nullptr && Pin.Get())
 		{
-			return LOCTEXT("PinCategory", "Pin");
+			return LOCTEXT("FindMetasound_PinCategory", "Pin");
 		}
 		else
 		{
-			return LOCTEXT("NodeCategory", "Node");
+			return LOCTEXT("FindMetasound_NodeCategory", "Node");
 		}
 		return FText::GetEmpty();
 	}
@@ -189,7 +189,7 @@ namespace Metasound::Editor
 				ValueText = FText::FromString(DefaultLiteral.ToString());
 				if (DefaultLiteral.IsArray())
 				{
-					ValueText = FText::Format(LOCTEXT("FindMetasound_OverriddenValueDescriptionFormat", "({0} element array)"), DefaultLiteral.GetArrayNum());
+					ValueText = FText::Format(LOCTEXT("FindMetasound_ArrayNumDescriptionFormat", "({0} element array)"), DefaultLiteral.GetArrayNum());
 				}
 				// Get a shorter name than the ToString one (which is LexToString)
 				else if (DefaultLiteral.GetType() == EMetasoundFrontendLiteralType::UObject)
@@ -373,7 +373,7 @@ namespace Metasound::Editor
 				.FillWidth(1)
 				[
 					SAssignNew(SearchTextField, SSearchBox)
-					.HintText(LOCTEXT("GraphSearchHint", "Search"))
+					.HintText(LOCTEXT("FindMetasound_GraphSearchHint", "Search"))
 					.OnTextChanged(this, &SFindInMetasound::OnSearchTextChanged)
 					.OnTextCommitted(this, &SFindInMetasound::OnSearchTextCommitted)
 				]
@@ -465,12 +465,12 @@ namespace Metasound::Editor
 		// Insert a fake result to inform user if none found
 		if (ItemsFound.Num() == 0)
 		{
-			ItemsFound.Add(FSearchResult(MakeShared<FFindInMetasoundResult>(LOCTEXT("FindInMetaSoundNoResults", "No Results found").ToString())));
+			ItemsFound.Add(FSearchResult(MakeShared<FFindInMetasoundResult>(LOCTEXT("FindMetaSound_NoResults", "No Results found").ToString())));
 		}
 		else
 		{
 			// Insert a fake result for stat tracking
-			FText ResultsStats = FText::Format(LOCTEXT("FindInMetaSoundNumResultsFmt", "{0} Result(s): {1} Matching Node(s), {2} Matching Pin(s)"), FoundNodeCount + FoundPinCount, FoundNodeCount, FoundPinCount);
+			FText ResultsStats = FText::Format(LOCTEXT("FindMetaSound_NumResultsFmt", "{0} Result(s): {1} Matching Node(s), {2} Matching Pin(s)"), FoundNodeCount + FoundPinCount, FoundNodeCount, FoundPinCount);
 			ItemsFound.Insert(FSearchResult(MakeShared<FFindInMetasoundResult>(ResultsStats.ToString())), 0);
 		}
 
@@ -645,7 +645,7 @@ namespace Metasound::Editor
 						SNew(STextBlock)
 						.Text(FText::FromString(InItem->Value))
 						.HighlightText(HighlightText)
-						.ToolTipText(FText::Format(LOCTEXT("MetasoundResultSearchToolTipFmt", "{0}: {1}"), InItem->GetCategory(), FText::FromString(InItem->Value)))
+						.ToolTipText(FText::Format(LOCTEXT("FindMetasound_ResultSearchToolTipFmt", "{0}: {1}"), InItem->GetCategory(), FText::FromString(InItem->Value)))
 					]
 				+ SHorizontalBox::Slot()
 					.FillWidth(1)
@@ -665,7 +665,7 @@ namespace Metasound::Editor
 					.Padding(2, 0)
 					[
 						SNew(STextBlock)
-						.Text(CommentText.IsEmpty() ? FText::GetEmpty() : FText::Format(LOCTEXT("NodeCommentFmt", "[Comment: {0}]"), FText::FromString(CommentText)))
+						.Text(CommentText.IsEmpty() ? FText::GetEmpty() : FText::Format(LOCTEXT("FindMetasound_NodeCommentFmt", "[Comment: {0}]"), FText::FromString(CommentText)))
 						.ColorAndOpacity(FSlateColor::UseForeground())
 						.HighlightText(HighlightText)
 					]
