Add MetaSound Editor Toolbar extension registry to MetaSoundEditorSubsystem

diff --git a/Engine/Plugins/Runtime/Metasound/Source/MetasoundEditor/Private/MetasoundEditor.cpp b/Engine/Plugins/Runtime/Metasound/Source/MetasoundEditor/Private/MetasoundEditor.cpp
--- a/Engine/Plugins/Runtime/Metasound/Source/MetasoundEditor/Private/MetasoundEditor.cpp
+++ b/Engine/Plugins/Runtime/Metasound/Source/MetasoundEditor/Private/MetasoundEditor.cpp
@@ -44,6 +44,7 @@
 #include "MetasoundEditorGraphValidation.h"
 #include "MetasoundEditorModule.h"
 #include "MetasoundEditorSettings.h"
+#include "MetasoundEditorSubsystem.h"
 #include "MetasoundEditorTabFactory.h"
 #include "MetasoundFrontend.h"
 #include "MetasoundFrontendDocument.h"
@@ -891,7 +892,7 @@ namespace Metasound
 			constexpr bool bUseSmallToolbarIcons = true;
 			FAssetEditorToolkit::InitAssetEditor(Mode, InitToolkitHost, TEXT("MetasoundEditorApp"), StandaloneDefaultLayout, bCreateDefaultStandaloneMenu, bCreateDefaultToolbar, ObjectToEdit, bToolbarFocusable, bUseSmallToolbarIcons);
 
-			ExtendToolbar();
+			ExtendToolbarInternal();
 			RegenerateMenusAndToolbars();
 
 			NotifyDocumentVersioned();
@@ -1268,7 +1269,7 @@ namespace Metasound
 			OutputSpectrumAnalyzer.Reset();
 		}
 
-		void FEditor::ExtendToolbar()
+		void FEditor::ExtendToolbarInternal()
 		{
 			TSharedPtr<FExtender> ToolbarExtender = MakeShared<FExtender>();
 			ToolbarExtender->AddToolBarExtension
@@ -1375,6 +1376,17 @@ namespace Metasound
 			);
 
 			AddToolbarExtender(ToolbarExtender);
+
+			if (GEditor)
+			{
+				if (UMetaSoundEditorSubsystem* Subsystem = GEditor->GetEditorSubsystem<UMetaSoundEditorSubsystem>())
+				{
+					for (const TSharedRef<FExtender>& Extender : Subsystem->GetToolbarExtenders())
+					{
+						AddToolbarExtender(Extender);
+					}
+				}
+			}
 		}
 
 		FSlateIcon FEditor::GetImportStatusImage() const
diff --git a/Engine/Plugins/Runtime/Metasound/Source/MetasoundEditor/Private/MetasoundEditor.h b/Engine/Plugins/Runtime/Metasound/Source/MetasoundEditor/Private/MetasoundEditor.h
--- a/Engine/Plugins/Runtime/Metasound/Source/MetasoundEditor/Private/MetasoundEditor.h
+++ b/Engine/Plugins/Runtime/Metasound/Source/MetasoundEditor/Private/MetasoundEditor.h
@@ -202,8 +202,6 @@ namespace Metasound
 			/** Creates analyzers */
 			void CreateAnalyzers();
 
-			void BuildTransport(FToolBarBuilder& InToolBarBuilder);
-
 			/** Destroys analyzers */
 			void DestroyAnalyzers();
 
@@ -354,7 +352,7 @@ namespace Metasound
 			void CreateInternalWidgets();
 
 			/** Builds the toolbar widget for the Metasound editor */
-			void ExtendToolbar();
+			void ExtendToolbarInternal();
 
 			/** Binds new graph commands to delegates */
 			void BindGraphCommands();
diff --git a/Engine/Plugins/Runtime/Metasound/Source/MetasoundEditor/Private/MetasoundEditorSubsystem.cpp b/Engine/Plugins/Runtime/Metasound/Source/MetasoundEditor/Private/MetasoundEditorSubsystem.cpp
--- a/Engine/Plugins/Runtime/Metasound/Source/MetasoundEditor/Private/MetasoundEditorSubsystem.cpp
+++ b/Engine/Plugins/Runtime/Metasound/Source/MetasoundEditor/Private/MetasoundEditorSubsystem.cpp
@@ -134,6 +134,11 @@ const FString UMetaSoundEditorSubsystem::GetDefaultAuthor()
 	return Author;
 }
 
+const TArray<TSharedRef<FExtender>>& UMetaSoundEditorSubsystem::GetToolbarExtenders() const
+{
+	return EditorToolbarExtenders;
+}
+
 void UMetaSoundEditorSubsystem::InitAsset(UObject& InNewMetaSound, UObject* InReferencedMetaSound)
 {
 	using namespace Metasound;
@@ -203,6 +208,17 @@ void UMetaSoundEditorSubsystem::RegisterGraphWithFrontend(UObject& InMetaSound,
 	Metasound::Editor::FGraphBuilder::RegisterGraphWithFrontend(InMetaSound, bInForceViewSynchronization);
 }
 
+void UMetaSoundEditorSubsystem::RegisterToolbarExtender(TSharedRef<FExtender> InExtender)
+{
+	EditorToolbarExtenders.AddUnique(InExtender);
+}
+
+bool UMetaSoundEditorSubsystem::UnregisterToolbarExtender(TSharedRef<FExtender> InExtender)
+{
+	const int32 NumRemoved = EditorToolbarExtenders.RemoveAllSwap([&InExtender](const TSharedRef<FExtender>& Extender) { return Extender == InExtender; });
+	return NumRemoved > 0;
+}
+
 void UMetaSoundEditorSubsystem::SetNodeLocation(
 	UMetaSoundBuilderBase* InBuilder,
 	const FMetaSoundNodeHandle& InNode,
diff --git a/Engine/Plugins/Runtime/Metasound/Source/MetasoundEditor/Public/MetasoundEditorSubsystem.h b/Engine/Plugins/Runtime/Metasound/Source/MetasoundEditor/Public/MetasoundEditorSubsystem.h
--- a/Engine/Plugins/Runtime/Metasound/Source/MetasoundEditor/Public/MetasoundEditorSubsystem.h
+++ b/Engine/Plugins/Runtime/Metasound/Source/MetasoundEditor/Public/MetasoundEditorSubsystem.h
@@ -3,6 +3,7 @@
 
 #include "AssetRegistry/AssetData.h"
 #include "EditorSubsystem.h"
+#include "Framework/MultiBox/MultiBoxExtender.h"
 #include "MetasoundDocumentInterface.h"
 #include "UObject/ScriptInterface.h"
 
@@ -54,13 +55,25 @@ public:
 	// referencing graphs open in editors)
 	void RegisterGraphWithFrontend(UObject& InMetaSound, bool bInForceViewSynchronization = false);
 
+	// Register toolbar extender that will be displayed in the MetaSound Asset Editor.
+	void RegisterToolbarExtender(TSharedRef<FExtender> InExtender);
+
+	// Unregisters toolbar extender that is displayed in the MetaSound Asset Editor.
+	bool UnregisterToolbarExtender(TSharedRef<FExtender> InExtender);
+
 	// Get the default author for a MetaSound asset
 	const FString GetDefaultAuthor();
 
+	// Returns all currently toolbar extenders registered to be displayed within the MetaSound Asset Editor.
+	const TArray<TSharedRef<FExtender>>& GetToolbarExtenders() const;
+
 	static UMetaSoundEditorSubsystem& GetChecked();
-	static const UMetaSoundEditorSubsystem& GetConstChecked(); 
+	static const UMetaSoundEditorSubsystem& GetConstChecked();
 
 private:
-	// Copy over sound wave settings such as attenuation, modulation, and sound class from the template sound wave to the MetaSound 
+	// Copy over sound wave settings such as attenuation, modulation, and sound class from the template sound wave to the MetaSound
 	void SetSoundWaveSettingsFromTemplate(USoundWave& NewMetasound, const USoundWave& TemplateSoundWave) const;
+
+	// Editor Toolbar Extenders
+	TArray<TSharedRef<FExtender>> EditorToolbarExtenders;
 };
