Fixed issue where FloatADSREnvelope node would drop to 0 for one cycle on trigger release if the trigger is not at the top of a block

diff --git a/Engine/Plugins/Runtime/Metasound/Source/MetasoundStandardNodes/Private/MetasoundADSREnvelopeNode.cpp b/Engine/Plugins/Runtime/Metasound/Source/MetasoundStandardNodes/Private/MetasoundADSREnvelopeNode.cpp
--- a/Engine/Plugins/Runtime/Metasound/Source/MetasoundStandardNodes/Private/MetasoundADSREnvelopeNode.cpp
+++ b/Engine/Plugins/Runtime/Metasound/Source/MetasoundStandardNodes/Private/MetasoundADSREnvelopeNode.cpp
@@ -125,7 +125,12 @@ namespace Metasound
 			static void GetNextEnvelopeOutput(FEnvState& InState, int32 StartFrame, int32 EndFrame, TArray<int32>& OutOnDecayFrames, TArray<int32>& OutOnSustainFrames, TArray<int32>& OutOnDoneFrames, float& OutEnvelopeValue)
 			{
 				// Don't need to do anything if we're not generating the envelope at the top of the block since this is a block-rate envelope
-				if (StartFrame > 0 || InState.CurrentSampleIndex == INDEX_NONE)
+				if (StartFrame > 0)
+				{
+					return;
+				}
+
+				if (InState.CurrentSampleIndex == INDEX_NONE)
 				{
 					OutEnvelopeValue = 0.0f;
 					return;
