Fixed TriggerDelay node to only fire after the most recent execution of the input trigger

diff --git a/Engine/Plugins/Runtime/Metasound/Source/MetasoundStandardNodes/Private/MetasoundTriggerDelayNode.cpp b/Engine/Plugins/Runtime/Metasound/Source/MetasoundStandardNodes/Private/MetasoundTriggerDelayNode.cpp
--- a/Engine/Plugins/Runtime/Metasound/Source/MetasoundStandardNodes/Private/MetasoundTriggerDelayNode.cpp
+++ b/Engine/Plugins/Runtime/Metasound/Source/MetasoundStandardNodes/Private/MetasoundTriggerDelayNode.cpp
@@ -103,6 +103,7 @@ namespace Metasound
 			[this](int32 StartFrame, int32 EndFrame)
 			{
 				const int32 FrameToTrigger = FMath::Max(0, FMath::RoundToInt(DelayTime->GetSeconds() * SampleRate)) + StartFrame;
+				TriggerOut->RemoveAfter(StartFrame);
 				TriggerOut->TriggerFrame(FrameToTrigger);
 			}
 		);
