Avoid registering unused send/receive nodes

diff --git a/Engine/Plugins/Runtime/Metasound/Source/MetasoundEngine/Private/MetasoundEngineModule.cpp b/Engine/Plugins/Runtime/Metasound/Source/MetasoundEngine/Private/MetasoundEngineModule.cpp
--- a/Engine/Plugins/Runtime/Metasound/Source/MetasoundEngine/Private/MetasoundEngineModule.cpp
+++ b/Engine/Plugins/Runtime/Metasound/Source/MetasoundEngine/Private/MetasoundEngineModule.cpp
@@ -26,6 +26,16 @@
 
 DEFINE_LOG_CATEGORY(LogMetasoundEngine);
 
+namespace Metasound
+{
+	// Enable send/receive node registration for data types which existed before
+	// send/receive were deprecated in order to support old UMetaSound assets. 
+	template<>
+	struct TEnableTransmissionNodeRegistration<FWaveAsset>
+	{
+		static constexpr bool Value = true;
+	};
+}
 
 REGISTER_METASOUND_DATATYPE(Metasound::FAudioBusAsset, "AudioBusAsset", Metasound::ELiteralType::UObjectProxy, UAudioBus);
 REGISTER_METASOUND_DATATYPE(Metasound::FWaveAsset, "WaveAsset", Metasound::ELiteralType::UObjectProxy, USoundWave);
diff --git a/Engine/Plugins/Runtime/Metasound/Source/MetasoundEngineTest/Private/EngineTestMetaSoundAutomatedNodeTest.cpp b/Engine/Plugins/Runtime/Metasound/Source/MetasoundEngineTest/Private/EngineTestMetaSoundAutomatedNodeTest.cpp
--- a/Engine/Plugins/Runtime/Metasound/Source/MetasoundEngineTest/Private/EngineTestMetaSoundAutomatedNodeTest.cpp
+++ b/Engine/Plugins/Runtime/Metasound/Source/MetasoundEngineTest/Private/EngineTestMetaSoundAutomatedNodeTest.cpp
@@ -742,6 +742,30 @@ namespace Metasound::EngineTest{
 		Frontend::ISearchEngine& NodeSearchEngine = Frontend::ISearchEngine::Get();
 		TArray<FMetasoundFrontendClass> AllClasses = NodeSearchEngine.FindAllClasses(true /* IncludeAllVersions */);
 
+		for (const FMetasoundFrontendClass& NodeClass : AllClasses)
+		{
+			// Exclude template classes because they cannot be created directly from the node registry
+			if (NodeClass.Metadata.GetType() == EMetasoundFrontendClassType::Template)
+			{
+				continue;
+			}
+
+			OutBeautifiedNames.Add(FString::Printf(TEXT("%s %s"), *NodeClass.Metadata.GetClassName().ToString(), *NodeClass.Metadata.GetVersion().ToString()));
+
+			// Test commands are node registry keys
+			Frontend::FNodeRegistryKey NodeRegistryKey(NodeClass.Metadata);
+			OutTestCommands.Add(NodeRegistryKey.ToString());
+		}
+	}
+
+	void GetAllRegisteredNativeNodes(TArray<FString>& OutBeautifiedNames, TArray <FString>& OutTestCommands)
+	{
+		using namespace Metasound;
+
+		// Get all the classes that have been registered
+		Frontend::ISearchEngine& NodeSearchEngine = Frontend::ISearchEngine::Get();
+		TArray<FMetasoundFrontendClass> AllClasses = NodeSearchEngine.FindAllClasses(true /* IncludeAllVersions */);
+
 		FMetasoundFrontendRegistryContainer* NodeRegistry = FMetasoundFrontendRegistryContainer::Get();
 		check(nullptr != NodeRegistry);
 
@@ -753,15 +777,21 @@ namespace Metasound::EngineTest{
 				continue;
 			}
 
+			Frontend::FNodeRegistryKey NodeRegistryKey(NodeClass.Metadata);
+
+			// Exclude non-native nodes (Nodes defined by assets instead of C++)
+			if (!NodeRegistry->IsNodeNative(NodeRegistryKey))
+			{
+				continue;
+			}
+
 			OutBeautifiedNames.Add(FString::Printf(TEXT("%s %s"), *NodeClass.Metadata.GetClassName().ToString(), *NodeClass.Metadata.GetVersion().ToString()));
 
 			// Test commands are node registry keys
-			Frontend::FNodeRegistryKey NodeRegistryKey(NodeClass.Metadata);
 			OutTestCommands.Add(NodeRegistryKey.ToString());
 		}
 	}
 
-
 	void CreateVariables(const FOperatorSettings& InOperatorSettings, FInputVertexInterfaceData& OutVertexData)
 	{
 		using namespace MetasoundVertexDataPrivate;
diff --git a/Engine/Plugins/Runtime/Metasound/Source/MetasoundEngineTest/Private/EngineTestMetaSoundAutomatedNodeTest.h b/Engine/Plugins/Runtime/Metasound/Source/MetasoundEngineTest/Private/EngineTestMetaSoundAutomatedNodeTest.h
--- a/Engine/Plugins/Runtime/Metasound/Source/MetasoundEngineTest/Private/EngineTestMetaSoundAutomatedNodeTest.h
+++ b/Engine/Plugins/Runtime/Metasound/Source/MetasoundEngineTest/Private/EngineTestMetaSoundAutomatedNodeTest.h
@@ -141,6 +141,9 @@ namespace Metasound::EngineTest
 	// Return an array of all currently registered MetaSound nodes. 
 	void GetAllRegisteredNodes(TArray<FString>& OutBeautifiedNames, TArray<FString>& OutNodeRegistryKeys);
 
+	// Return an array of all currently registered native MetaSound nodes. 
+	void GetAllRegisteredNativeNodes(TArray<FString>& OutBeautifiedNames, TArray<FString>& OutNodeRegistryKeys);
+
 	// Create any variables that exist on the vertex interface.
 	void CreateVariables(const FOperatorSettings& InOperatorSettings, FInputVertexInterfaceData& OutVertexData);
 
diff --git a/Engine/Plugins/Runtime/Metasound/Source/MetasoundEngineTest/Private/EngineTestMetaSoundAutomatedNodeTest_Bind.cpp b/Engine/Plugins/Runtime/Metasound/Source/MetasoundEngineTest/Private/EngineTestMetaSoundAutomatedNodeTest_Bind.cpp
--- a/Engine/Plugins/Runtime/Metasound/Source/MetasoundEngineTest/Private/EngineTestMetaSoundAutomatedNodeTest_Bind.cpp
+++ b/Engine/Plugins/Runtime/Metasound/Source/MetasoundEngineTest/Private/EngineTestMetaSoundAutomatedNodeTest_Bind.cpp
@@ -23,7 +23,7 @@ IMPLEMENT_COMPLEX_AUTOMATION_TEST(FMetasoundAutomatedNodeTest_Bind, "Audio.Metas
 void FMetasoundAutomatedNodeTest_Bind::GetTests(TArray<FString>& OutBeautifiedNames, TArray<FString>& OutTestCommands) const
 {
 	using namespace Metasound::EngineTest;
-	GetAllRegisteredNodes(OutBeautifiedNames, OutTestCommands);
+	GetAllRegisteredNativeNodes(OutBeautifiedNames, OutTestCommands);
 
 	UE_LOG(LogMetaSound, Verbose, TEXT("Found %d metasound nodes to test"), OutTestCommands.Num());
 }
diff --git a/Engine/Plugins/Runtime/Metasound/Source/MetasoundEngineTest/Private/EngineTestMetaSoundAutomatedNodeTest_Reset.cpp b/Engine/Plugins/Runtime/Metasound/Source/MetasoundEngineTest/Private/EngineTestMetaSoundAutomatedNodeTest_Reset.cpp
--- a/Engine/Plugins/Runtime/Metasound/Source/MetasoundEngineTest/Private/EngineTestMetaSoundAutomatedNodeTest_Reset.cpp
+++ b/Engine/Plugins/Runtime/Metasound/Source/MetasoundEngineTest/Private/EngineTestMetaSoundAutomatedNodeTest_Reset.cpp
@@ -22,7 +22,7 @@ IMPLEMENT_COMPLEX_AUTOMATION_TEST(FMetasoundAutomatedNodeTest_Reset, "Audio.Meta
 void FMetasoundAutomatedNodeTest_Reset::GetTests(TArray<FString>& OutBeautifiedNames, TArray<FString>& OutTestCommands) const
 {
 	using namespace Metasound::EngineTest;
-	GetAllRegisteredNodes(OutBeautifiedNames, OutTestCommands);
+	GetAllRegisteredNativeNodes(OutBeautifiedNames, OutTestCommands);
 
 	UE_LOG(LogMetaSound, Verbose, TEXT("Found %d metasound nodes to test"), OutTestCommands.Num());
 }
diff --git a/Engine/Plugins/Runtime/Metasound/Source/MetasoundEngineTest/Private/EngineTestMetaSoundAutomatedNodeTest_ValidOutputs.cpp b/Engine/Plugins/Runtime/Metasound/Source/MetasoundEngineTest/Private/EngineTestMetaSoundAutomatedNodeTest_ValidOutputs.cpp
--- a/Engine/Plugins/Runtime/Metasound/Source/MetasoundEngineTest/Private/EngineTestMetaSoundAutomatedNodeTest_ValidOutputs.cpp
+++ b/Engine/Plugins/Runtime/Metasound/Source/MetasoundEngineTest/Private/EngineTestMetaSoundAutomatedNodeTest_ValidOutputs.cpp
@@ -22,7 +22,7 @@ IMPLEMENT_COMPLEX_AUTOMATION_TEST(FMetasoundAutomatedNodeTest_ValidOutputs, "Aud
 void FMetasoundAutomatedNodeTest_ValidOutputs::GetTests(TArray<FString>& OutBeautifiedNames, TArray<FString>& OutTestCommands) const
 {
 	using namespace Metasound::EngineTest;
-	GetAllRegisteredNodes(OutBeautifiedNames, OutTestCommands);
+	GetAllRegisteredNativeNodes(OutBeautifiedNames, OutTestCommands);
 
 	UE_LOG(LogMetaSound, Verbose, TEXT("Found %d metasound nodes to test"), OutTestCommands.Num());
 }
diff --git a/Engine/Plugins/Runtime/Metasound/Source/MetasoundFrontend/Private/MetasoundPrimitives.cpp b/Engine/Plugins/Runtime/Metasound/Source/MetasoundFrontend/Private/MetasoundPrimitives.cpp
--- a/Engine/Plugins/Runtime/Metasound/Source/MetasoundFrontend/Private/MetasoundPrimitives.cpp
+++ b/Engine/Plugins/Runtime/Metasound/Source/MetasoundFrontend/Private/MetasoundPrimitives.cpp
@@ -53,6 +53,62 @@ namespace Metasound
 	{
 		static constexpr bool Value = !std::is_arithmetic<ToDataType>::value;
 	};
+
+	// Enable send/receive node registration for data types which existed before
+	// send/receive were deprecated in order to support old UMetaSound assets. 
+	template<>
+	struct TEnableTransmissionNodeRegistration<bool>
+	{
+		static constexpr bool Value = true;
+	};
+	
+	// Enable send/receive node registration for data types which existed before
+	// send/receive were deprecated in order to support old UMetaSound assets. 
+	template<>
+	struct TEnableTransmissionNodeRegistration<int32>
+	{
+		static constexpr bool Value = true;
+	};
+	
+	// Enable send/receive node registration for data types which existed before
+	// send/receive were deprecated in order to support old UMetaSound assets. 
+	template<>
+	struct TEnableTransmissionNodeRegistration<float>
+	{
+		static constexpr bool Value = true;
+	};
+	
+	// Enable send/receive node registration for data types which existed before
+	// send/receive were deprecated in order to support old UMetaSound assets. 
+	template<>
+	struct TEnableTransmissionNodeRegistration<FString>
+	{
+		static constexpr bool Value = true;
+	};
+	
+	// Enable send/receive node registration for data types which existed before
+	// send/receive were deprecated in order to support old UMetaSound assets. 
+	template<>
+	struct TEnableTransmissionNodeRegistration<FTrigger>
+	{
+		static constexpr bool Value = true;
+	};
+	
+	// Enable send/receive node registration for data types which existed before
+	// send/receive were deprecated in order to support old UMetaSound assets. 
+	template<>
+	struct TEnableTransmissionNodeRegistration<FTime>
+	{
+		static constexpr bool Value = true;
+	};
+
+	// Enable send/receive node registration for data types which existed before
+	// send/receive were deprecated in order to support old UMetaSound assets. 
+	template<>
+	struct TEnableTransmissionNodeRegistration<FAudioBuffer>
+	{
+		static constexpr bool Value = true;
+	};
 }
 
 REGISTER_METASOUND_DATATYPE(bool, "Bool", ::Metasound::ELiteralType::Boolean)
diff --git a/Engine/Plugins/Runtime/Metasound/Source/MetasoundFrontend/Public/MetasoundFrontendDataTypeTraits.h b/Engine/Plugins/Runtime/Metasound/Source/MetasoundFrontend/Public/MetasoundFrontendDataTypeTraits.h
--- a/Engine/Plugins/Runtime/Metasound/Source/MetasoundFrontend/Public/MetasoundFrontendDataTypeTraits.h
+++ b/Engine/Plugins/Runtime/Metasound/Source/MetasoundFrontend/Public/MetasoundFrontendDataTypeTraits.h
@@ -2,6 +2,7 @@
 
 #pragma once
 
+#include "Containers/Array.h"
 #include "MetasoundDataFactory.h"
 
 #include <type_traits>
@@ -28,12 +29,17 @@ namespace Metasound
 	};
 
 	/** Enables or disables send and receive node registration for a given MetaSound
-	 * data type. By default this is true and all data types supported by the transmission
-	 * system will have associated send and receive nodes. */
+	 * data type. By default this is false because send/receive nodes are deprecated. */
 	template<typename ... Type>
 	struct TEnableTransmissionNodeRegistration
 	{
-		static constexpr bool Value = true;
+		static constexpr bool Value = false;
+	};
+
+	template<typename ElementType>
+	struct TEnableTransmissionNodeRegistration<TArray<ElementType>>
+	{
+		static constexpr bool Value = TEnableTransmissionNodeRegistration<ElementType>::Value;
 	};
 
 	/** Enables or disables using a data type in constructor vertices. 
