Metasound engine tests call "PostExecute" after "Execute".

#tests EngineTests

diff --git a/Engine/Plugins/Runtime/Metasound/Source/MetasoundEngineTest/Private/EngineTestMetaSoundAutomatedNodeTest_Reset.cpp b/Engine/Plugins/Runtime/Metasound/Source/MetasoundEngineTest/Private/EngineTestMetaSoundAutomatedNodeTest_Reset.cpp
--- a/Engine/Plugins/Runtime/Metasound/Source/MetasoundEngineTest/Private/EngineTestMetaSoundAutomatedNodeTest_Reset.cpp
+++ b/Engine/Plugins/Runtime/Metasound/Source/MetasoundEngineTest/Private/EngineTestMetaSoundAutomatedNodeTest_Reset.cpp
@@ -103,22 +103,43 @@ bool FMetasoundAutomatedNodeTest_Reset::RunTest(const FString& InRegistryKeyStri
 		OutputTester.CaptureCurrentOutputValues();
 
 		IOperator::FExecuteFunction OpExecFunc = Operator->GetExecuteFunction();
+		auto ConditionallyCallExecFunc = [&OpExecFunc] (IOperator* Operator)
+		{
+			if (OpExecFunc)
+			{
+				OpExecFunc(Operator);
+			}
+		};
+		IOperator::FPostExecuteFunction OpPostExecFunc = Operator->GetPostExecuteFunction();
+		auto ConditionallyCallPostExecFunc = [&OpPostExecFunc] (IOperator* Operator)
+        {
+        	if (OpPostExecFunc)
+        	{
+        		OpPostExecFunc(Operator);
+        	}
+        };
 		IOperator::FResetFunction OpResetFunc = Operator->GetResetFunction();
+		
 		// Test execute function with input variations
-		if (OpExecFunc)
+		if (OpExecFunc || OpPostExecFunc)
 		{
-			OpExecFunc(Operator.Get());
-
+			ConditionallyCallExecFunc(Operator.Get());
+			ConditionallyCallPostExecFunc(Operator.Get());
+			
 			if (InputTester.GetNumMutableInputs() > 0)
 			{
 				InputTester.SetMutableInputsToDefault();
-				OpExecFunc(Operator.Get());
+				ConditionallyCallExecFunc(Operator.Get());
+				ConditionallyCallPostExecFunc(Operator.Get());
 				InputTester.SetMutableInputsToMin();
-				OpExecFunc(Operator.Get());
+				ConditionallyCallExecFunc(Operator.Get());
+				ConditionallyCallPostExecFunc(Operator.Get());
 				InputTester.SetMutableInputsToMax();
-				OpExecFunc(Operator.Get());
+				ConditionallyCallExecFunc(Operator.Get());
+				ConditionallyCallPostExecFunc(Operator.Get());
 				InputTester.SetMutableInputsToRandom();
-				OpExecFunc(Operator.Get());
+				ConditionallyCallExecFunc(Operator.Get());
+				ConditionallyCallPostExecFunc(Operator.Get());
 			}
 		}
 
@@ -129,7 +150,7 @@ bool FMetasoundAutomatedNodeTest_Reset::RunTest(const FString& InRegistryKeyStri
 		{
 			OpResetFunc(Operator.Get(), ResetParams);
 		}
-		else if (OpExecFunc)
+		else if (OpExecFunc || OpPostExecFunc)
 		{
 			AddError(FString::Printf(TEXT("Missing initialize function when execute function exists for node %s - %s"), *InRegistryKeyString, *GetPrettyName(RegistryKey)));
 		}
