Fix remove swap logic that wasn't covering case where both separate last index and swap index were in same referencing map entry.
#tests opened all existing MetaSound assets in a large project by opening MetaSound Editor, confirmed reroute repro was fixed

diff --git a/Engine/Plugins/Runtime/Metasound/Source/MetasoundFrontend/Private/MetasoundFrontendDocumentCache.cpp b/Engine/Plugins/Runtime/Metasound/Source/MetasoundFrontend/Private/MetasoundFrontendDocumentCache.cpp
--- a/Engine/Plugins/Runtime/Metasound/Source/MetasoundFrontend/Private/MetasoundFrontendDocumentCache.cpp
+++ b/Engine/Plugins/Runtime/Metasound/Source/MetasoundFrontend/Private/MetasoundFrontendDocumentCache.cpp
@@ -43,27 +43,46 @@ namespace Metasound::Frontend
 			TMapArrayType& OutMap
 		)
 		{
-			// Remove swap item
+			auto IndexIsLast = [&LastIndex](const int32& Index) { return Index == LastIndex; };
+
+			if (SwapIndex == LastIndex)
 			{
-				TArray<int32>& Indices = OutMap.FindChecked(InSwapKey);
-				Indices.RemoveAllSwap([&SwapIndex](const int32& Index)
-				{
-					return Index == SwapIndex;
-				}, EAllowShrinking::No);
+				check(InSwapKey == InLastKey);
+				TArray<int32>& Indices = OutMap.FindChecked(InLastKey);
+				Indices.RemoveAllSwap(IndexIsLast, EAllowShrinking::No);
 				if (Indices.IsEmpty())
 				{
-					OutMap.Remove(InSwapKey);
+					OutMap.Remove(InLastKey);
 				}
 			}
-
-			// Swap last item index if not the same entry
-			if (InSwapKey != InLastKey && SwapIndex != LastIndex)
+			else
 			{
-				TArray<int32>& Indices = OutMap.FindChecked(InLastKey);
-				Algo::ForEachIf(Indices,
-					[&LastIndex](const int32& Index) { return Index == LastIndex; },
-					[&SwapIndex](int32& Index) { Index = SwapIndex; }
-				);
+				if (InSwapKey == InLastKey)
+				{
+					TArray<int32>& Indices = OutMap.FindChecked(InLastKey);
+					Indices.RemoveAllSwap(IndexIsLast, EAllowShrinking::No);
+					check(!Indices.IsEmpty()); // has to still contain the key that now points to *just* the swap index
+				}
+				else
+				{
+					{
+						auto IndexIsSwap = [&SwapIndex](const int32& Index) { return Index == SwapIndex; };
+						TArray<int32>& Indices = OutMap.FindChecked(InSwapKey);
+						Indices.RemoveAllSwap(IndexIsSwap, EAllowShrinking::No);
+						if (Indices.IsEmpty())
+						{
+							OutMap.Remove(InSwapKey);
+						}
+					}
+
+					{
+						TArray<int32>& Indices = OutMap.FindChecked(InLastKey);
+						Algo::ForEachIf(Indices,
+							IndexIsLast,
+							[&SwapIndex](int32& Index) { Index = SwapIndex; }
+						);
+					}
+				}
 			}
 		}
 	}
