Minor clean-up/missing successful return statement
#rnx

diff --git a/Engine/Plugins/Runtime/Metasound/Source/MetasoundFrontend/Private/MetasoundFrontendDocumentBuilder.cpp b/Engine/Plugins/Runtime/Metasound/Source/MetasoundFrontend/Private/MetasoundFrontendDocumentBuilder.cpp
--- a/Engine/Plugins/Runtime/Metasound/Source/MetasoundFrontend/Private/MetasoundFrontendDocumentBuilder.cpp
+++ b/Engine/Plugins/Runtime/Metasound/Source/MetasoundFrontend/Private/MetasoundFrontendDocumentBuilder.cpp
@@ -1350,6 +1350,7 @@ bool FMetaSoundFrontendDocumentBuilder::ConvertToPreset(const FMetasoundFrontend
 		InitCacheInternal();
 		return true;
 	}
+
 	return false;
 }
 
@@ -1576,6 +1577,11 @@ TArray<const FMetasoundFrontendVertex*> FMetaSoundFrontendDocumentBuilder::FindN
 
 	const FMetasoundFrontendDocument& Document = GetDocument();
 
+	if (ConnectedInputNodes)
+	{
+		ConnectedInputNodes->Reset();
+	}
+
 	TArray<const FMetasoundFrontendVertex*> Inputs;
 	const TArrayView<const int32> Indices = EdgeCache.FindEdgeIndicesFromNodeOutput(InOutputNodeID, InOutputVertexID);
 	Algo::Transform(Indices, Inputs, [&Document, &NodeCache, &ConnectedInputNodes](const int32& Index)
@@ -2232,6 +2238,7 @@ bool FMetaSoundFrontendDocumentBuilder::RemoveEdge(const FMetasoundFrontendEdge&
 		const int32 LastIndex = Edges.Num() - 1;
 		DocumentDelegates->EdgeDelegates.OnRemoveSwappingEdge.Broadcast(Index, LastIndex);
 		Edges.RemoveAtSwap(Index, 1, EAllowShrinking::No);
+		return true;
 	}
 
 	return false;
