Scaffolding for custom MetaSound comment nodes that serialize to Frontend Document as initial part of making ed graph data transient
#rnx

diff --git a/Engine/Plugins/Runtime/Metasound/Source/MetasoundEditor/Private/MetasoundEditorGraphCommentNode.cpp b/Engine/Plugins/Runtime/Metasound/Source/MetasoundEditor/Private/MetasoundEditorGraphCommentNode.cpp
new file mode 100644
--- /dev/null
+++ b/Engine/Plugins/Runtime/Metasound/Source/MetasoundEditor/Private/MetasoundEditorGraphCommentNode.cpp
@@ -0,0 +1,121 @@
+// Copyright Epic Games, Inc. All Rights Reserved.
+
+#include "MetasoundEditorGraphCommentNode.h"
+
+#include "Internationalization/Internationalization.h"
+#include "Kismet2/Kismet2NameValidators.h"
+#include "Layout/SlateRect.h"
+#include "MetasoundBuilderSubsystem.h"
+#include "MetasoundUObjectRegistry.h"
+#include "Styling/AppStyle.h"
+#include "Templates/Casts.h"
+#include "UObject/UnrealType.h"
+
+class UEdGraphPin;
+
+#define LOCTEXT_NAMESPACE "MetasoundEditor"
+
+
+FMetasoundAssetBase& UMetasoundEditorGraphCommentNode::GetAssetChecked()
+{
+	using namespace Metasound;
+	using namespace Metasound::Frontend;
+
+	UObject* Outermost = GetOutermostObject();
+	check(Outermost);
+
+	FMetasoundAssetBase* MetaSound = IMetasoundUObjectRegistry::Get().GetObjectAsAssetBase(Outermost);
+	check(MetaSound);
+	return *MetaSound;
+}
+
+void UMetasoundEditorGraphCommentNode::ConvertToFrontendComment(const UEdGraphNode_Comment& InEdNode, FMetaSoundFrontendGraphComment& OutComment)
+{
+	OutComment.bColorBubble = InEdNode.bColorCommentBubble != 0;
+	OutComment.Color = InEdNode.CommentColor;
+	OutComment.Comment = InEdNode.NodeComment;
+	OutComment.Depth = InEdNode.CommentDepth;
+	OutComment.FontSize = InEdNode.FontSize;
+	OutComment.MoveMode = InEdNode.MoveMode == ECommentBoxMode::GroupMovement ? EMetaSoundFrontendGraphCommentMoveMode::GroupMovement : EMetaSoundFrontendGraphCommentMoveMode::NoGroupMovement;
+	OutComment.Position = FVector2D(InEdNode.NodePosX, InEdNode.NodePosY);
+	OutComment.Size = FVector2D(InEdNode.NodeWidth, InEdNode.NodeHeight);
+}
+
+void UMetasoundEditorGraphCommentNode::ConvertFromFrontendComment(const FMetaSoundFrontendGraphComment& InComment, UEdGraphNode_Comment& OutEdNode)
+{
+	OutEdNode.bColorCommentBubble = (uint32)InComment.bColorBubble;
+	OutEdNode.CommentColor = InComment.Color;
+	OutEdNode.NodeComment = InComment.Comment;
+	OutEdNode.CommentDepth = InComment.Depth;
+	OutEdNode.FontSize = InComment.FontSize;
+	OutEdNode.MoveMode = InComment.MoveMode == EMetaSoundFrontendGraphCommentMoveMode::GroupMovement ? ECommentBoxMode::GroupMovement : ECommentBoxMode::NoGroupMovement;
+	OutEdNode.NodePosX = InComment.Position.X;
+	OutEdNode.NodePosY = InComment.Position.Y;
+	OutEdNode.NodeWidth = InComment.Size.X;
+	OutEdNode.NodeHeight = InComment.Size.Y;
+}
+
+void UMetasoundEditorGraphCommentNode::PostEditChangeProperty(FPropertyChangedEvent& PropertyChangedEvent)
+{
+	Super::PostEditChangeProperty(PropertyChangedEvent);
+
+	// TODO: Check property name and only change that which changed
+	FMetaSoundFrontendGraphComment Comment;
+	ConvertToFrontendComment(*this, Comment);
+
+	UMetaSoundBuilderBase& Builder = UMetaSoundBuilderSubsystem::GetChecked().AttachBuilderToAssetChecked(*GetOutermostObject());
+	Builder.FindOrAddGraphComment(CommentID) = MoveTemp(Comment);
+}
+
+void UMetasoundEditorGraphCommentNode::PostPlacedNewNode()
+{
+	Super::PostPlacedNewNode();
+
+	UMetaSoundBuilderBase& Builder = UMetaSoundBuilderSubsystem::GetChecked().AttachBuilderToAssetChecked(*GetOutermostObject());
+
+	// If the frontend version predates this new node, sync this node's comment to it...
+	if (const FMetaSoundFrontendGraphComment* Comment = Builder.FindGraphComment(CommentID))
+	{
+		NodeComment = Comment->Comment;
+	}
+
+	// otherwise, initialize the frontend version with the ed graph's default comment.
+	else
+	{
+		Builder.FindOrAddGraphComment(CommentID).Comment = NodeComment;
+	}
+}
+
+void UMetasoundEditorGraphCommentNode::ResizeNode(const FVector2D& NewSize)
+{
+	Super::ResizeNode(NewSize);
+	if (bCanResizeNode) 
+	{
+		UMetaSoundBuilderBase& Builder = UMetaSoundBuilderSubsystem::GetChecked().AttachBuilderToAssetChecked(*GetOutermostObject());
+		Builder.FindOrAddGraphComment(CommentID).Size = FVector2D(NodeWidth, NodeHeight);
+	}
+}
+
+FGuid UMetasoundEditorGraphCommentNode::GetCommentID() const
+{
+	return CommentID;
+}
+
+void UMetasoundEditorGraphCommentNode::SetBounds(const class FSlateRect& Rect)
+{
+	Super::SetBounds(Rect);
+
+	UMetaSoundBuilderBase& Builder = UMetaSoundBuilderSubsystem::GetChecked().AttachBuilderToAssetChecked(*GetOutermostObject());
+	FMetaSoundFrontendGraphComment& FrontendComment = Builder.FindOrAddGraphComment(CommentID);
+	FrontendComment.Position = FVector2D(NodePosX, NodePosY);
+	FrontendComment.Size = FVector2D(NodeWidth, NodeHeight);
+}
+
+void UMetasoundEditorGraphCommentNode::OnRenameNode(const FString& NewName)
+{
+	Super::OnRenameNode(NewName);
+
+	UMetaSoundBuilderBase& Builder = UMetaSoundBuilderSubsystem::GetChecked().AttachBuilderToAssetChecked(*GetOutermostObject());
+	Builder.FindOrAddGraphComment(CommentID).Comment = NewName;
+}
+#undef LOCTEXT_NAMESPACE
diff --git a/Engine/Plugins/Runtime/Metasound/Source/MetasoundEditor/Public/MetasoundEditorGraphCommentNode.h b/Engine/Plugins/Runtime/Metasound/Source/MetasoundEditor/Public/MetasoundEditorGraphCommentNode.h
new file mode 100644
--- /dev/null
+++ b/Engine/Plugins/Runtime/Metasound/Source/MetasoundEditor/Public/MetasoundEditorGraphCommentNode.h
@@ -0,0 +1,61 @@
+// Copyright Epic Games, Inc. All Rights Reserved.
+#pragma once
+
+#include "EdGraphNode_Comment.h"
+#include "MetasoundFrontend.h"
+#include "MetasoundFrontendDocument.h"
+#include "MetasoundFrontendRegistries.h"
+#include "Misc/Guid.h"
+#include "Textures/SlateIcon.h"
+#include "UObject/ObjectMacros.h"
+#include "UObject/UObjectGlobals.h"
+
+#include "MetasoundEditorGraphCommentNode.generated.h"
+
+struct FPropertyChangedEvent;
+class FMetasoundAssetBase;
+
+
+// Forward Declarations
+namespace Metasound::Editor
+{
+	class FGraphBuilder;
+}
+
+
+UCLASS(MinimalAPI)
+class UMetasoundEditorGraphCommentNode : public UEdGraphNode_Comment
+{
+	GENERATED_BODY()
+
+public:
+	// Convenience function to convert underlying ed-graph properties to a frontend graph comment
+	static void ConvertToFrontendComment(const UEdGraphNode_Comment& InEdNode, FMetaSoundFrontendGraphComment& OutComment);
+
+	// Convenience function to convert underlying ed-graph properties from a frontend graph comment
+	static void ConvertFromFrontendComment(const FMetaSoundFrontendGraphComment& InComment, UEdGraphNode_Comment& OutEdNode);
+
+	//~ Begin UObject Interface
+	virtual void PostEditChangeProperty(FPropertyChangedEvent& PropertyChangedEvent) override;
+	//~ End UObject Interface
+
+	//~ Begin UEdGraphNode Interface
+	METASOUNDEDITOR_API virtual void ResizeNode(const FVector2D& NewSize) override;
+	METASOUNDEDITOR_API virtual void PostPlacedNewNode() override;
+	METASOUNDEDITOR_API virtual void OnRenameNode(const FString& NewName) override;
+	//~ End UEdGraphNode Interface
+
+	/** Set the Bounds for the comment node */
+	METASOUNDEDITOR_API void SetBounds(const class FSlateRect& Rect);
+
+	// Returns frontend comment ID
+	FGuid GetCommentID() const;
+
+private:
+	FMetasoundAssetBase& GetAssetChecked();
+
+	UPROPERTY()
+	FGuid CommentID;
+
+	friend class Metasound::Editor::FGraphBuilder;
+};
diff --git a/Engine/Plugins/Runtime/Metasound/Source/MetasoundEngine/Private/MetasoundBuilderSubsystem.cpp b/Engine/Plugins/Runtime/Metasound/Source/MetasoundEngine/Private/MetasoundBuilderSubsystem.cpp
--- a/Engine/Plugins/Runtime/Metasound/Source/MetasoundEngine/Private/MetasoundBuilderSubsystem.cpp
+++ b/Engine/Plugins/Runtime/Metasound/Source/MetasoundEngine/Private/MetasoundBuilderSubsystem.cpp
@@ -605,6 +605,23 @@ TArray<FMetaSoundBuilderNodeOutputHandle> UMetaSoundBuilderBase::FindNodeOutputs
 	return FoundVertices;
 }
 
+#if WITH_EDITOR
+const FMetaSoundFrontendGraphComment* UMetaSoundBuilderBase::FindGraphComment(const FGuid& InCommentID) const
+{
+	return Builder.FindGraphComment(InCommentID);
+}
+
+FMetaSoundFrontendGraphComment* UMetaSoundBuilderBase::FindGraphComment(const FGuid& InCommentID)
+{
+	return Builder.FindGraphComment(InCommentID);
+}
+
+FMetaSoundFrontendGraphComment& UMetaSoundBuilderBase::FindOrAddGraphComment(const FGuid& InCommentID)
+{
+	return Builder.FindOrAddGraphComment(InCommentID);
+}
+#endif // WITH_EDITOR
+
 TArray<FMetaSoundNodeHandle> UMetaSoundBuilderBase::FindInterfaceInputNodes(FName InterfaceName, EMetaSoundBuilderResult& OutResult)
 {
 	TArray<FMetaSoundNodeHandle> NodeHandles;
@@ -802,6 +819,13 @@ void UMetaSoundBuilderBase::ReloadCache(bool bPrimeCache)
 	Builder.ReloadCache();
 }
 
+#if WITH_EDITOR
+bool UMetaSoundBuilderBase::RemoveGraphComment(const FGuid& InCommentID)
+{
+	return Builder.RemoveGraphComment(InCommentID);
+}
+#endif // WITH_EDITOR
+
 void UMetaSoundBuilderBase::RemoveGraphInput(FName Name, EMetaSoundBuilderResult& OutResult)
 {
 	const bool bRemoved = Builder.RemoveGraphInput(Name);
@@ -857,11 +881,30 @@ void UMetaSoundBuilderBase::SetNodeInputDefault(const FMetaSoundBuilderNodeInput
 }
 
 #if WITH_EDITOR
+void UMetaSoundBuilderBase::SetNodeComment(const FMetaSoundNodeHandle& InNodeHandle, const FString& InNewComment, EMetaSoundBuilderResult& OutResult)
+{
+	FString NewComment = InNewComment;
+	const bool bCommentSet = Builder.SetNodeComment(InNodeHandle.NodeID, MoveTemp(NewComment));
+	OutResult = bCommentSet ? EMetaSoundBuilderResult::Succeeded : EMetaSoundBuilderResult::Failed;
+}
+
+void UMetaSoundBuilderBase::SetNodeCommentVisible(const FMetaSoundNodeHandle& InNodeHandle, bool bIsVisible, EMetaSoundBuilderResult& OutResult)
+{
+	const bool bCommentSet = Builder.SetNodeCommentVisible(InNodeHandle.NodeID, bIsVisible);
+	OutResult = bCommentSet ? EMetaSoundBuilderResult::Succeeded : EMetaSoundBuilderResult::Failed;
+}
+
 void UMetaSoundBuilderBase::SetNodeLocation(const FMetaSoundNodeHandle& InNodeHandle, const FVector2D& InLocation, EMetaSoundBuilderResult& OutResult)
 {
 	const bool bLocationSet = Builder.SetNodeLocation(InNodeHandle.NodeID, InLocation);
 	OutResult = bLocationSet ? EMetaSoundBuilderResult::Succeeded : EMetaSoundBuilderResult::Failed;
 }
+
+void UMetaSoundBuilderBase::SetNodeLocation(const FMetaSoundNodeHandle& InNodeHandle, const FVector2D& InLocation, const FGuid& InLocationGuid, EMetaSoundBuilderResult& OutResult)
+{
+	const bool bLocationSet = Builder.SetNodeLocation(InNodeHandle.NodeID, InLocation, &InLocationGuid);
+	OutResult = bLocationSet ? EMetaSoundBuilderResult::Succeeded : EMetaSoundBuilderResult::Failed;
+}
 #endif // WITH_EDITOR
 
 FMetaSoundNodeHandle UMetaSoundBuilderBase::FindNodeInputParent(const FMetaSoundBuilderNodeInputHandle& InputHandle, EMetaSoundBuilderResult& OutResult)
diff --git a/Engine/Plugins/Runtime/Metasound/Source/MetasoundEngine/Public/MetasoundBuilderSubsystem.h b/Engine/Plugins/Runtime/Metasound/Source/MetasoundEngine/Public/MetasoundBuilderSubsystem.h
--- a/Engine/Plugins/Runtime/Metasound/Source/MetasoundEngine/Public/MetasoundBuilderSubsystem.h
+++ b/Engine/Plugins/Runtime/Metasound/Source/MetasoundEngine/Public/MetasoundBuilderSubsystem.h
@@ -385,10 +385,18 @@ public:
 	virtual void InitFrontendBuilder();
 
 	// Initializes and ensures all nodes have a position (required prior to exporting to an asset if expected to be viewed in the editor).
-	void InitNodeLocations();
+	virtual void InitNodeLocations();
 
 #if WITH_EDITOR
+	const FMetaSoundFrontendGraphComment* FindGraphComment(const FGuid& InCommentID) const;
+	FMetaSoundFrontendGraphComment* FindGraphComment(const FGuid & InCommentID);
+	FMetaSoundFrontendGraphComment& FindOrAddGraphComment(const FGuid& InCommentID);
+	bool RemoveGraphComment(const FGuid& InCommentID);
+
+	void SetNodeComment(const FMetaSoundNodeHandle& InNodeHandle, const FString& InNewComment, EMetaSoundBuilderResult& OutResult);
+	void SetNodeCommentVisible(const FMetaSoundNodeHandle& InNodeHandle, bool bIsVisible, EMetaSoundBuilderResult& OutResult);
 	void SetNodeLocation(const FMetaSoundNodeHandle& InNodeHandle, const FVector2D& InLocation, EMetaSoundBuilderResult& OutResult);
+	void SetNodeLocation(const FMetaSoundNodeHandle & InNodeHandle, const FVector2D& InLocation, const FGuid& InLocationGuid, EMetaSoundBuilderResult& OutResult);
 #endif // WITH_EDITOR
 
 protected:
diff --git a/Engine/Plugins/Runtime/Metasound/Source/MetasoundFrontend/Private/MetasoundFrontendDocument.cpp b/Engine/Plugins/Runtime/Metasound/Source/MetasoundFrontend/Private/MetasoundFrontendDocument.cpp
--- a/Engine/Plugins/Runtime/Metasound/Source/MetasoundFrontend/Private/MetasoundFrontendDocument.cpp
+++ b/Engine/Plugins/Runtime/Metasound/Source/MetasoundFrontend/Private/MetasoundFrontendDocument.cpp
@@ -33,7 +33,7 @@ namespace Metasound
 			namespace NodeLayout
 			{
 				const FVector2D DefaultOffsetX { 300.0f, 0.0f };
-				const FVector2D DefaultOffsetY { 0.0f, 80.0f };
+				const FVector2D DefaultOffsetY { 0.0f, 120.0f };
 			} // namespace NodeLayout
 		} // namespace DisplayStyle
 
diff --git a/Engine/Plugins/Runtime/Metasound/Source/MetasoundFrontend/Private/MetasoundFrontendDocumentBuilder.cpp b/Engine/Plugins/Runtime/Metasound/Source/MetasoundFrontend/Private/MetasoundFrontendDocumentBuilder.cpp
--- a/Engine/Plugins/Runtime/Metasound/Source/MetasoundFrontend/Private/MetasoundFrontendDocumentBuilder.cpp
+++ b/Engine/Plugins/Runtime/Metasound/Source/MetasoundFrontend/Private/MetasoundFrontendDocumentBuilder.cpp
@@ -1405,6 +1405,20 @@ TArray<const FMetasoundFrontendEdge*> FMetaSoundFrontendDocumentBuilder::FindEdg
 	return EdgeCache.FindEdges(InNodeID, InVertexID);
 }
 
+#if WITH_EDITOR
+const FMetaSoundFrontendGraphComment* FMetaSoundFrontendDocumentBuilder::FindGraphComment(const FGuid& InCommentID) const
+{
+	const TMap<FGuid, FMetaSoundFrontendGraphComment>& Comments = GetDocument().RootGraph.Graph.Style.Comments;
+	return Comments.Find(InCommentID);
+}
+
+FMetaSoundFrontendGraphComment* FMetaSoundFrontendDocumentBuilder::FindGraphComment(const FGuid& InCommentID)
+{
+	TMap<FGuid, FMetaSoundFrontendGraphComment>& Comments = GetDocument().RootGraph.Graph.Style.Comments;
+	return Comments.Find(InCommentID);
+}
+#endif // WITH_EDITOR
+
 bool FMetaSoundFrontendDocumentBuilder::FindInterfaceInputNodes(FName InterfaceName, TArray<const FMetasoundFrontendNode*>& OutInputs) const
 {
 	using namespace Metasound::Frontend;
@@ -1596,6 +1610,22 @@ TArray<const FMetasoundFrontendVertex*> FMetaSoundFrontendDocumentBuilder::FindN
 	return Inputs;
 }
 
+FMetasoundFrontendNode* FMetaSoundFrontendDocumentBuilder::FindNodeInternal(const FGuid& InNodeID)
+{
+	using namespace Metasound::Frontend;
+
+	using namespace Metasound;
+	using namespace Metasound::Frontend;
+
+	const IDocumentGraphNodeCache& NodeCache = DocumentCache->GetNodeCache();
+	if (const int32* NodeIndex = NodeCache.FindNodeIndex(InNodeID))
+	{
+		return &GetDocument().RootGraph.Graph.Nodes[*NodeIndex];
+	}
+
+	return nullptr;
+}
+
 const FMetasoundFrontendVertex* FMetaSoundFrontendDocumentBuilder::FindNodeOutput(const FGuid& InNodeID, const FGuid& InVertexID) const
 {
 	using namespace Metasound::Frontend;
@@ -1639,6 +1669,15 @@ const FMetasoundFrontendVertex* FMetaSoundFrontendDocumentBuilder::FindNodeOutpu
 	return nullptr;
 }
 
+#if WITH_EDITOR
+FMetaSoundFrontendGraphComment& FMetaSoundFrontendDocumentBuilder::FindOrAddGraphComment(const FGuid& InCommentID)
+{
+	check(InCommentID.IsValid());
+	TMap<FGuid, FMetaSoundFrontendGraphComment>& Comments = GetDocument().RootGraph.Graph.Style.Comments;
+	return Comments.FindOrAdd(InCommentID);
+}
+#endif // WITH_EDITOR
+
 const FTopLevelAssetPath FMetaSoundFrontendDocumentBuilder::GetBuilderClassPath() const
 {
 	IMetaSoundDocumentInterface* Interface = DocumentInterface.GetInterface();
@@ -1767,6 +1806,7 @@ const FMetasoundFrontendLiteral* FMetaSoundFrontendDocumentBuilder::GetNodeInput
 					}
 					break;
 				}
+				static_assert(static_cast<uint32>(EMetasoundFrontendClassType::Invalid) == 10, "Potential missing case coverage for EMetasoundFrontendClassType");
 			}
 		}
 	}
@@ -2406,6 +2446,14 @@ bool FMetaSoundFrontendDocumentBuilder::RemoveEdgeToNodeInput(const FGuid& InNod
 	return false;
 }
 
+#if WITH_EDITOR
+bool FMetaSoundFrontendDocumentBuilder::RemoveGraphComment(const FGuid& InCommentID)
+{
+	TMap<FGuid, FMetaSoundFrontendGraphComment>& Comments = GetDocument().RootGraph.Graph.Style.Comments;
+	return Comments.Remove(InCommentID) > 0;
+}
+#endif // WITH_EDITOR
+
 bool FMetaSoundFrontendDocumentBuilder::RemoveGraphInput(FName InInputName)
 {
 	if (const FMetasoundFrontendNode* Node = FindGraphInputNode(InInputName))
@@ -2567,6 +2615,33 @@ bool FMetaSoundFrontendDocumentBuilder::RemoveNode(const FGuid& InNodeID)
 	return false;
 }
 
+#if WITH_EDITOR
+int32 FMetaSoundFrontendDocumentBuilder::RemoveNodeLocation(const FGuid& InNodeID, const FGuid* InLocationGuid)
+{
+	using namespace Metasound;
+	using namespace Metasound::Frontend;
+
+	const IDocumentGraphNodeCache& NodeCache = DocumentCache->GetNodeCache();
+	if (const int32* NodeIndex = NodeCache.FindNodeIndex(InNodeID))
+	{
+		FMetasoundFrontendNode& Node = GetDocument().RootGraph.Graph.Nodes[*NodeIndex];
+		FMetasoundFrontendNodeStyle& Style = Node.Style;
+		if (InLocationGuid)
+		{
+			return Style.Display.Locations.Remove(*InLocationGuid);
+		}
+		else
+		{
+			const int32 NumLocationsRemoved = Style.Display.Locations.Num();
+			Style.Display.Locations.Reset();
+			return NumLocationsRemoved;
+		}
+	}
+
+	return 0;
+}
+#endif // WITH_EDITOR
+
 bool FMetaSoundFrontendDocumentBuilder::RenameRootGraphClass(const FMetasoundFrontendClassName& InName)
 {
 	FGuid NameGuid;
@@ -2687,17 +2762,38 @@ bool FMetaSoundFrontendDocumentBuilder::SetNodeInputDefault(const FGuid& InNodeI
 }
 
 #if WITH_EDITOR
-bool FMetaSoundFrontendDocumentBuilder::SetNodeLocation(const FGuid& InNodeID, const FVector2D& InLocation)
+bool FMetaSoundFrontendDocumentBuilder::SetNodeComment(const FGuid& InNodeID, FString&& InNewComment)
 {
-	using namespace Metasound;
-	using namespace Metasound::Frontend;
+	if (FMetasoundFrontendNode* Node = FindNodeInternal(InNodeID))
+	{
+		Node->Style.Display.Comment = MoveTemp(InNewComment);
+		return true;
+	}
 
-	const IDocumentGraphNodeCache& NodeCache = DocumentCache->GetNodeCache();
-	if (const int32* NodeIndex = NodeCache.FindNodeIndex(InNodeID))
+	return false;
+}
+
+bool FMetaSoundFrontendDocumentBuilder::SetNodeCommentVisible(const FGuid& InNodeID, bool bIsVisible)
+{
+	if (FMetasoundFrontendNode* Node = FindNodeInternal(InNodeID))
 	{
-		FMetasoundFrontendNode& Node = GetDocument().RootGraph.Graph.Nodes[*NodeIndex];
-		FMetasoundFrontendNodeStyle& Style = Node.Style;
-		if (Style.Display.Locations.IsEmpty())
+		Node->Style.Display.bCommentVisible = bIsVisible;
+		return true;
+	}
+
+	return false;
+}
+
+bool FMetaSoundFrontendDocumentBuilder::SetNodeLocation(const FGuid& InNodeID, const FVector2D& InLocation, const FGuid* InLocationGuid)
+{
+	if (FMetasoundFrontendNode* Node = FindNodeInternal(InNodeID))
+	{
+		FMetasoundFrontendNodeStyle& Style = Node->Style;
+		if (InLocationGuid)
+		{
+			Style.Display.Locations.FindOrAdd(*InLocationGuid) = InLocation;
+		}
+		else if (Style.Display.Locations.IsEmpty())
 		{
 			Style.Display.Locations = { { FGuid::NewGuid(), InLocation } };
 		}
@@ -2716,6 +2812,11 @@ bool FMetaSoundFrontendDocumentBuilder::SetNodeLocation(const FGuid& InNodeID, c
 }
 #endif // WITH_EDITOR
 
+void FMetaSoundFrontendDocumentBuilder::SetVersionNumber(const FMetasoundFrontendVersionNumber& InDocumentVersionNumber)
+{
+	GetDocument().Metadata.Version.Number = InDocumentVersionNumber;
+}
+
 bool FMetaSoundFrontendDocumentBuilder::SwapGraphInput(const FMetasoundFrontendClassVertex& InExistingInputVertex, const FMetasoundFrontendClassVertex& InNewInputVertex)
 {
 	using namespace Metasound::Frontend;
diff --git a/Engine/Plugins/Runtime/Metasound/Source/MetasoundFrontend/Public/MetasoundFrontendDocument.h b/Engine/Plugins/Runtime/Metasound/Source/MetasoundFrontend/Public/MetasoundFrontendDocument.h
--- a/Engine/Plugins/Runtime/Metasound/Source/MetasoundFrontend/Public/MetasoundFrontendDocument.h
+++ b/Engine/Plugins/Runtime/Metasound/Source/MetasoundFrontend/Public/MetasoundFrontendDocument.h
@@ -156,6 +156,47 @@ enum class EMetasoundFrontendClassType : uint8
 	Invalid UMETA(Hidden)
 };
 
+UENUM()
+enum class EMetaSoundFrontendGraphCommentMoveMode : uint8
+{
+	/** This comment box will move any fully contained nodes when it moves. */
+	GroupMovement UMETA(DisplayName = "Group Movement"),
+
+	/** This comment box has no effect on nodes contained inside it. */
+	NoGroupMovement UMETA(DisplayName = "Comment")
+};
+
+USTRUCT()
+struct METASOUNDFRONTEND_API FMetaSoundFrontendGraphComment
+{
+	GENERATED_BODY()
+
+#if WITH_EDITORONLY_DATA
+	UPROPERTY()
+	FLinearColor Color;
+
+	UPROPERTY()
+	FString Comment;
+
+	UPROPERTY()
+	int32 Depth = 0;
+
+	UPROPERTY()
+	int32 FontSize = 0;
+
+	UPROPERTY()
+	FVector2D Position;
+
+	UPROPERTY()
+	FVector2D Size;
+
+	UPROPERTY()
+	EMetaSoundFrontendGraphCommentMoveMode MoveMode = EMetaSoundFrontendGraphCommentMoveMode::GroupMovement;
+
+	UPROPERTY()
+	uint8 bColorBubble : 1;
+#endif // WITH_EDITORONLY_DATA
+};
 
 // General purpose version number for Metasound Frontend objects.
 USTRUCT(BlueprintType)
@@ -493,6 +534,14 @@ struct METASOUNDFRONTEND_API FMetasoundFrontendNodeStyleDisplay
 	// more than one place on the graph (Only functionally relevant for nodes that cannot contain inputs.)
 	UPROPERTY()
 	TMap<FGuid, FVector2D> Locations;
+
+	// Comment to display about the given instance's usage
+	UPROPERTY()
+	FString Comment;
+
+	// Whether or not the comment is visible or not
+	UPROPERTY()
+	bool bCommentVisible = false;
 #endif // WITH_EDITORONLY_DATA
 };
 
@@ -659,6 +708,7 @@ struct METASOUNDFRONTEND_API FMetasoundFrontendGraphStyle
 {
 	GENERATED_BODY()
 
+#if WITH_EDITORONLY_DATA
 	// Whether or not the graph is editable by a user
 	UPROPERTY()
 	bool bIsGraphEditable = true;
@@ -666,6 +716,11 @@ struct METASOUNDFRONTEND_API FMetasoundFrontendGraphStyle
 	// Styles for graph edges.
 	UPROPERTY()
 	TArray<FMetasoundFrontendEdgeStyle> EdgeStyles;
+
+	// Map of comment id to comment data
+	UPROPERTY()
+	TMap<FGuid, FMetaSoundFrontendGraphComment> Comments;
+#endif // WITH_EDITORONLY_DATA
 };
 
 USTRUCT()
diff --git a/Engine/Plugins/Runtime/Metasound/Source/MetasoundFrontend/Public/MetasoundFrontendDocumentBuilder.h b/Engine/Plugins/Runtime/Metasound/Source/MetasoundFrontend/Public/MetasoundFrontendDocumentBuilder.h
--- a/Engine/Plugins/Runtime/Metasound/Source/MetasoundFrontend/Public/MetasoundFrontendDocumentBuilder.h
+++ b/Engine/Plugins/Runtime/Metasound/Source/MetasoundFrontend/Public/MetasoundFrontendDocumentBuilder.h
@@ -169,6 +169,12 @@ public:
 	bool ConvertFromPreset();
 	bool ConvertToPreset(const FMetasoundFrontendDocument& InReferencedDocument);
 
+#if WITH_EDITOR
+	const FMetaSoundFrontendGraphComment* FindGraphComment(const FGuid& InCommentID) const;
+	FMetaSoundFrontendGraphComment* FindGraphComment(const FGuid& InCommentID);
+	FMetaSoundFrontendGraphComment& FindOrAddGraphComment(const FGuid& InCommentID);
+#endif // WITH_EDITOR
+
 	static bool FindDeclaredInterfaces(const FMetasoundFrontendDocument& InDocument, TArray<const Metasound::Frontend::IInterfaceRegistryEntry*>& OutInterfaces);
 	bool FindDeclaredInterfaces(TArray<const Metasound::Frontend::IInterfaceRegistryEntry*>& OutInterfaces) const;
 
@@ -255,27 +261,47 @@ public:
 	bool RemoveEdgesByNodeClassInterfaceBindings(const FGuid& InOutputNodeID, const FGuid& InInputNodeID);
 	bool RemoveEdgesFromNodeOutput(const FGuid& InNodeID, const FGuid& InVertexID);
 	bool RemoveEdgeToNodeInput(const FGuid& InNodeID, const FGuid& InVertexID);
+
+#if WITH_EDITOR
+	bool RemoveGraphComment(const FGuid& InCommentID);
+#endif // WITH_EDITOR
+
 	bool RemoveGraphInput(FName InInputName);
 	bool RemoveGraphOutput(FName InOutputName);
 	bool RemoveInterface(FName InName);
 	bool RemoveNamedEdges(const TSet<Metasound::Frontend::FNamedEdge>& InNamedEdgesToRemove, TArray<FMetasoundFrontendEdge>* OutRemovedEdges = nullptr);
 	bool RemoveNode(const FGuid& InNodeID);
+
+#if WITH_EDITOR
+	int32 RemoveNodeLocation(const FGuid& InNodeID, const FGuid* InLocationGuid = nullptr);
+#endif // WITH_EDITOR
+
 	bool RemoveNodeInputDefault(const FGuid& InNodeID, const FGuid& InVertexID);
 	bool RemoveUnusedDependencies();
 	bool RenameRootGraphClass(const FMetasoundFrontendClassName& InName);
 
 #if WITH_EDITOR
-
 	void SetAuthor(const FString& InAuthor);
 
+	// Sets the editor-only comment to the provided value.
+	// Returns true if the node was found and the comment was updated, false if not.
+	bool SetNodeComment(const FGuid& InNodeID, FString&& InNewComment);
+
+	// Sets the editor-only comment visibility.
+	// Returns true if the node was found and the visibility was set, false if not.
+	bool SetNodeCommentVisible(const FGuid& InNodeID, bool bIsVisible);
+
 	// Sets the editor-only node location of a node with the given ID to the provided location.
 	// Returns true if the node was found and the location was updated, false if not.
-	bool SetNodeLocation(const FGuid& InNodeID, const FVector2D& InLocation);
+	bool SetNodeLocation(const FGuid& InNodeID, const FVector2D& InLocation, const FGuid* InLocationGuid = nullptr);
 #endif // WITH_EDITOR
 
 	bool SetGraphInputDefault(FName InputName, const FMetasoundFrontendLiteral& InDefaultLiteral);
 	bool SetNodeInputDefault(const FGuid& InNodeID, const FGuid& InVertexID, const FMetasoundFrontendLiteral& InLiteral);
 
+	// Sets the document's version number.  Should only be called by document versioning.
+	void SetVersionNumber(const FMetasoundFrontendVersionNumber& InDocumentVersionNumber);
+
 	bool SwapGraphInput(const FMetasoundFrontendClassVertex& InExistingInputVertex, const FMetasoundFrontendClassVertex& NewInputVertex);
 	bool SwapGraphOutput(const FMetasoundFrontendClassVertex& InExistingOutputVertex, const FMetasoundFrontendClassVertex& NewOutputVertex);
 	bool UpdateDependencyClassNames(const TMap<FMetasoundFrontendClassName, FMetasoundFrontendClassName>& OldToNewReferencedClassNames);
@@ -287,11 +313,12 @@ public:
 private:
 	using FFinalizeNodeFunctionRef = TFunctionRef<void(FMetasoundFrontendNode&, const Metasound::Frontend::FNodeRegistryKey&)>;
 
+	FMetasoundFrontendNode* AddNodeInternal(const FMetasoundFrontendClassMetadata& InClassMetadata, Metasound::Frontend::FFinalizeNodeFunctionRef FinalizeNode, FGuid InNodeID = FGuid::NewGuid());
+	FMetasoundFrontendNode* FindNodeInternal(const FGuid& InNodeID);
+
 	const FTopLevelAssetPath GetBuilderClassPath() const;
 	FMetasoundFrontendDocument& GetDocument();
 
-	FMetasoundFrontendNode* AddNodeInternal(const FMetasoundFrontendClassMetadata& InClassMetadata, Metasound::Frontend::FFinalizeNodeFunctionRef FinalizeNode, FGuid InNodeID = FGuid::NewGuid());
-
 	bool FindNodeClassInterfaces(const FGuid& InNodeID, TSet<FMetasoundFrontendVersion>& OutInterfaces) const;
 
 	void InitCacheInternal(bool bPrimeCache = false);
diff --git a/Engine/Source/Editor/UnrealEd/Public/EdGraphNode_Comment.h b/Engine/Source/Editor/UnrealEd/Public/EdGraphNode_Comment.h
--- a/Engine/Source/Editor/UnrealEd/Public/EdGraphNode_Comment.h
+++ b/Engine/Source/Editor/UnrealEd/Public/EdGraphNode_Comment.h
@@ -48,7 +48,7 @@ public:
 	GENERATED_BODY()
 
 public:
-	UEdGraphNode_Comment(const FObjectInitializer& ObjectInitializer = FObjectInitializer::Get());
+	UNREALED_API UEdGraphNode_Comment(const FObjectInitializer& ObjectInitializer = FObjectInitializer::Get());
 
 	/** Color to style comment with */
 	UPROPERTY(EditAnywhere, Category=Comment)
@@ -78,16 +78,16 @@ public:
 
 	//~ Begin UObject Interface
 	UNREALED_API static void AddReferencedObjects(UObject* InThis, FReferenceCollector& Collector);
-	virtual void PostEditChangeProperty(FPropertyChangedEvent& PropertyChangedEvent) override;
-	virtual bool IsSelectedInEditor() const override;
+	UNREALED_API virtual void PostEditChangeProperty(FPropertyChangedEvent& PropertyChangedEvent) override;
+	UNREALED_API virtual bool IsSelectedInEditor() const override;
 	//~ End UObject Interface
 
 	//~ Begin UEdGraphNode Interface
-	virtual void AllocateDefaultPins() override {}
+	UNREALED_API virtual void AllocateDefaultPins() override {}
 	UNREALED_API virtual FText GetTooltipText() const override;
 	UNREALED_API virtual FLinearColor GetNodeCommentColor() const override;
 	UNREALED_API virtual FText GetNodeTitle(ENodeTitleType::Type TitleType) const override;
-	virtual bool ShouldOverridePinNames() const override { return true; }
+	UNREALED_API virtual bool ShouldOverridePinNames() const override { return true; }
 	UNREALED_API virtual FText GetPinNameOverride(const UEdGraphPin& Pin) const override;
 	UNREALED_API virtual void ResizeNode(const FVector2D& NewSize) override;
 	UNREALED_API virtual void PostPlacedNewNode() override;
@@ -111,7 +111,7 @@ public:
 	UNREALED_API const FCommentNodeSet& GetNodesUnderComment() const;
 
 	/** Return the font size of the comment */
-	virtual int32 GetFontSize() const { return FontSize; }
+	UNREALED_API virtual int32 GetFontSize() const { return FontSize; }
 
 	/** Override the default selection state of this graph node */
 	enum class ESelectionState : uint8 { Inherited, Selected, Deselected };
