Remove unused ClassName (field will be deprecated in subsequent CL)
#rnx

diff --git a/Engine/Plugins/Runtime/Metasound/Source/MetasoundEditor/Private/FindInMetasound.cpp b/Engine/Plugins/Runtime/Metasound/Source/MetasoundEditor/Private/FindInMetasound.cpp
--- a/Engine/Plugins/Runtime/Metasound/Source/MetasoundEditor/Private/FindInMetasound.cpp
+++ b/Engine/Plugins/Runtime/Metasound/Source/MetasoundEditor/Private/FindInMetasound.cpp
@@ -511,7 +511,6 @@ namespace Metasound::Editor
 			bool bIsMemberNode = false;
 			if (const UMetasoundEditorGraphNode* MetaSoundNode = Cast<UMetasoundEditorGraphNode>(Node))
 			{
-				FMetasoundFrontendClassName ClassName = MetaSoundNode->GetClassName();
 				DisplayName = MetaSoundNode->GetDisplayName().ToString();
 
 				// Additional information for member nodes
