- Enable translating UObject to FMetaSoundAssetBase from Frontend module via AssetManager
- Expose input descriptor so it can be used with future template notes.
- Minor clean-up/nullptr protection in MetaSound Editor

diff --git a/Engine/Plugins/Runtime/Metasound/Source/MetasoundEditor/Private/MetasoundEditor.cpp b/Engine/Plugins/Runtime/Metasound/Source/MetasoundEditor/Private/MetasoundEditor.cpp
--- a/Engine/Plugins/Runtime/Metasound/Source/MetasoundEditor/Private/MetasoundEditor.cpp
+++ b/Engine/Plugins/Runtime/Metasound/Source/MetasoundEditor/Private/MetasoundEditor.cpp
@@ -100,9 +100,15 @@ namespace Metasound
 
 		class FMetasoundGraphMemberSchemaAction : public FEdGraphSchemaAction
 		{
+			FGuid MemberID;
+
 		public:
 			UEdGraph* Graph = nullptr;
-			FGuid MemberID;
+
+			void SetMemberID(const FGuid& InID)
+			{
+				MemberID = InID;
+			}
 
 			FMetasoundGraphMemberSchemaAction()
 				: FEdGraphSchemaAction()
@@ -530,24 +536,26 @@ namespace Metasound
 				if (TSharedPtr<FMetasoundGraphMemberSchemaAction> GraphMemberAction = StaticCastSharedPtr<FMetasoundGraphMemberSchemaAction>(InCreateData->Action))
 				{
 
-					UMetasoundEditorGraphMember* GraphMember = GraphMemberAction->GetGraphMember();
-					if (const UMetasoundEditorGraphVertex* Vertex = Cast<UMetasoundEditorGraphVertex>(GraphMember))
+					if (UMetasoundEditorGraphMember* GraphMember = GraphMemberAction->GetGraphMember())
 					{
-						EMetasoundFrontendVertexAccessType AccessType = Vertex->GetVertexAccessType();
-						bIsConstructorPin = AccessType == EMetasoundFrontendVertexAccessType::Value;
-					}
-					FName DataTypeName = GraphMember->GetDataType(); 
+						if (const UMetasoundEditorGraphVertex* Vertex = Cast<UMetasoundEditorGraphVertex>(GraphMember))
+						{
+							EMetasoundFrontendVertexAccessType AccessType = Vertex->GetVertexAccessType();
+							bIsConstructorPin = AccessType == EMetasoundFrontendVertexAccessType::Value;
+						}
+						FName DataTypeName = GraphMember->GetDataType();
 
-					const IMetasoundEditorModule& EditorModule = FModuleManager::GetModuleChecked<IMetasoundEditorModule>("MetaSoundEditor");
-					if (const FEdGraphPinType* PinType = EditorModule.FindPinType(DataTypeName))
-					{
-						if (const UMetasoundEditorGraphSchema* Schema = GetDefault<UMetasoundEditorGraphSchema>())
+						const IMetasoundEditorModule& EditorModule = FModuleManager::GetModuleChecked<IMetasoundEditorModule>("MetaSoundEditor");
+						if (const FEdGraphPinType* PinType = EditorModule.FindPinType(DataTypeName))
 						{
-							IconColor = Schema->GetPinTypeColor(*PinType);
+							if (const UMetasoundEditorGraphSchema* Schema = GetDefault<UMetasoundEditorGraphSchema>())
+							{
+								IconColor = Schema->GetPinTypeColor(*PinType);
+							}
 						}
-					}
 
-					IconBrush = EditorModule.GetIconBrush(DataTypeName, bIsConstructorPin);
+						IconBrush = EditorModule.GetIconBrush(DataTypeName, bIsConstructorPin);
+					}
 				}
 
 				TSharedRef<SHorizontalBox> LayoutWidget = SNew(SHorizontalBox);
@@ -3078,7 +3086,7 @@ namespace Metasound
 
 				TSharedPtr<FMetasoundGraphMemberSchemaAction> NewFuncAction = MakeShared<FMetasoundGraphMemberSchemaAction>(Category, MenuDesc, Tooltip, 1, ENodeSection::Inputs);
 				NewFuncAction->Graph = &EdGraph;
-				NewFuncAction->MemberID = NodeID;
+				NewFuncAction->SetMemberID(NodeID);
 
 				ActionList->AddAction(NewFuncAction);
 			}, EMetasoundFrontendClassType::Input);
@@ -3094,7 +3102,7 @@ namespace Metasound
 
 				TSharedPtr<FMetasoundGraphMemberSchemaAction> NewFuncAction = MakeShared<FMetasoundGraphMemberSchemaAction>(Category, MenuDesc, Tooltip, 1, ENodeSection::Outputs);
 				NewFuncAction->Graph = &EdGraph;
-				NewFuncAction->MemberID = Output->GetID();
+				NewFuncAction->SetMemberID(NodeID);
 				ActionList->AddAction(NewFuncAction);
 			}, EMetasoundFrontendClassType::Output);
 
@@ -3107,7 +3115,7 @@ namespace Metasound
 
 				TSharedPtr<FMetasoundGraphMemberSchemaAction> NewFuncAction = MakeShared<FMetasoundGraphMemberSchemaAction>(Category, MenuDesc, FText::GetEmpty(), 1, ENodeSection::Variables);
 				NewFuncAction->Graph = &EdGraph;
-				NewFuncAction->MemberID = VariableID; 
+				NewFuncAction->SetMemberID(VariableID);
 				OutAllActions.AddAction(NewFuncAction);
 			}
 
@@ -3463,40 +3471,6 @@ namespace Metasound
 							RefreshDetails();
 						}
 					}
-					
-					// TODO: Because input editor nodes are not one-to-one, this can cause multi-selection
-					// when not desired.  Once input alias templates are complete, this can be switched on.
-					// For now, callsites are just directly setting selection via ClearSelectionAndSelectNode.
-// 					MetasoundGraphEditor->ClearSelectionSet();
-// 
-// 					const UMetasoundEditorGraph& Graph = GetMetaSoundGraphChecked();
-// 
-// 					TArray<UObject*> NodeSelection;
-// 					for (UEdGraphNode* Node : Graph.Nodes)
-// 					{
-// 						if (UMetasoundEditorGraphNode* MetaSoundNode = Cast<UMetasoundEditorGraphNode>(Node))
-// 						{
-// 							if (NodesModified.Contains(MetaSoundNode->GetNodeID()))
-// 							{
-// 								if (UMetasoundEditorGraphMemberNode* MemberNode = Cast<UMetasoundEditorGraphMemberNode>(Node))
-// 								{
-// 									NodeSelection.Add(MemberNode->GetMember());
-// 								}
-// 								else
-// 								{
-// 									NodeSelection.Add(Node);
-// 								}
-// 
-// 								MetasoundGraphEditor->SetNodeSelection(Node, /*bSelect=*/true);
-// 
-// 							}
-// 						}
-// 					}
-// 
-// 					if (!NodeSelection.IsEmpty())
-// 					{
-// 						Selection = NodeSelection;
-// 					}
 				}
 
 				if (!InterfacesModified.IsEmpty() || bForceRefreshViews)
diff --git a/Engine/Plugins/Runtime/Metasound/Source/MetasoundEngine/Private/MetasoundAssetSubsystem.cpp b/Engine/Plugins/Runtime/Metasound/Source/MetasoundEngine/Private/MetasoundAssetSubsystem.cpp
--- a/Engine/Plugins/Runtime/Metasound/Source/MetasoundEngine/Private/MetasoundAssetSubsystem.cpp
+++ b/Engine/Plugins/Runtime/Metasound/Source/MetasoundEngine/Private/MetasoundAssetSubsystem.cpp
@@ -10,6 +10,7 @@
 #include "MetasoundAssetBase.h"
 #include "MetasoundBuilderSubsystem.h"
 #include "MetasoundFrontendDocument.h"
+#include "MetasoundFrontendDocumentBuilder.h"
 #include "MetasoundFrontendRegistries.h"
 #include "MetasoundFrontendSearchEngine.h"
 #include "MetasoundFrontendTransform.h"
@@ -290,6 +291,16 @@ bool UMetaSoundAssetSubsystem::ContainsKey(const Metasound::Frontend::FNodeRegis
 	return PathMap.Contains(InRegistryKey);
 }
 
+FMetasoundAssetBase* UMetaSoundAssetSubsystem::GetAsAsset(UObject& InObject) const
+{
+	return Metasound::IMetasoundUObjectRegistry::Get().GetObjectAsAssetBase(&InObject);
+}
+
+const FMetasoundAssetBase* UMetaSoundAssetSubsystem::GetAsAsset(const UObject& InObject) const
+{
+	return Metasound::IMetasoundUObjectRegistry::Get().GetObjectAsAssetBase(&InObject);
+}
+
 void UMetaSoundAssetSubsystem::RebuildDenyListCache(const UAssetManager& InAssetManager)
 {
 	using namespace Metasound::Frontend;
@@ -383,7 +394,7 @@ TSet<UMetaSoundAssetSubsystem::FAssetInfo> UMetaSoundAssetSubsystem::GetReferenc
 	}
 	return MoveTemp(OutAssetInfos);
 }
-#endif
+#endif // WITH_EDITOR
 
 void UMetaSoundAssetSubsystem::RescanAutoUpdateDenyList()
 {
diff --git a/Engine/Plugins/Runtime/Metasound/Source/MetasoundEngine/Public/MetasoundAssetSubsystem.h b/Engine/Plugins/Runtime/Metasound/Source/MetasoundEngine/Public/MetasoundAssetSubsystem.h
--- a/Engine/Plugins/Runtime/Metasound/Source/MetasoundEngine/Public/MetasoundAssetSubsystem.h
+++ b/Engine/Plugins/Runtime/Metasound/Source/MetasoundEngine/Public/MetasoundAssetSubsystem.h
@@ -15,6 +15,7 @@ class UAssetManager;
 class FMetasoundAssetBase;
 
 struct FDirectoryPath;
+struct FMetaSoundFrontendDocumentBuilder;
 
 
 USTRUCT(BlueprintType)
@@ -69,6 +70,10 @@ public:
 	virtual bool CanAutoUpdate(const FMetasoundFrontendClassName& InClassName) const override;
 	virtual bool ContainsKey(const Metasound::Frontend::FNodeRegistryKey& InRegistryKey) const override;
 	virtual const FSoftObjectPath* FindObjectPathFromKey(const Metasound::Frontend::FNodeRegistryKey& RegistryKey) const override;
+
+	virtual FMetasoundAssetBase* GetAsAsset(UObject& InObject) const override;
+	virtual const FMetasoundAssetBase* GetAsAsset(const UObject& InObject) const override;
+
 #if WITH_EDITOR
 	virtual TSet<FAssetInfo> GetReferencedAssetClasses(const FMetasoundAssetBase& InAssetBase) const override;
 #endif
diff --git a/Engine/Plugins/Runtime/Metasound/Source/MetasoundFrontend/Private/MetasoundInputNode.cpp b/Engine/Plugins/Runtime/Metasound/Source/MetasoundFrontend/Private/MetasoundInputNode.cpp
--- a/Engine/Plugins/Runtime/Metasound/Source/MetasoundFrontend/Private/MetasoundInputNode.cpp
+++ b/Engine/Plugins/Runtime/Metasound/Source/MetasoundFrontend/Private/MetasoundInputNode.cpp
@@ -92,6 +92,10 @@ namespace Metasound
 		return CreateVertexInterface(InVertexName, InDataTypeName, InVertexAccess, FLiteral());
 	}
 
+	FText FInputNode::GetInputDescription()
+	{
+		return METASOUND_LOCTEXT("Metasound_InputNodeDescription", "Input into the parent MetaSound graph.");
+	}
 
 	FNodeClassMetadata FInputNode::GetNodeMetadata(const FVertexName& InVertexName, const FName& InDataTypeName, EVertexAccessType InVertexAccess)
 	{
@@ -100,7 +104,7 @@ namespace Metasound
 		Info.ClassName = { "Input", InDataTypeName, GetVariantName(InVertexAccess) };
 		Info.MajorVersion = 1;
 		Info.MinorVersion = 0;
-		Info.Description = METASOUND_LOCTEXT("Metasound_InputNodeDescription", "Input into the parent Metasound graph.");
+		Info.Description = GetInputDescription();
 		Info.Author = PluginAuthor;
 		Info.PromptIfMissing = PluginNodeMissingPrompt;
 		Info.DefaultInterface = CreateDefaultVertexInterface(InVertexName, InDataTypeName, InVertexAccess);
diff --git a/Engine/Plugins/Runtime/Metasound/Source/MetasoundFrontend/Public/MetasoundAssetManager.h b/Engine/Plugins/Runtime/Metasound/Source/MetasoundFrontend/Public/MetasoundAssetManager.h
--- a/Engine/Plugins/Runtime/Metasound/Source/MetasoundFrontend/Public/MetasoundAssetManager.h
+++ b/Engine/Plugins/Runtime/Metasound/Source/MetasoundFrontend/Public/MetasoundAssetManager.h
@@ -1,6 +1,7 @@
 // Copyright Epic Games, Inc. All Rights Reserved.
 #pragma once
 
+#include "MetasoundFrontendDocument.h"
 #include "MetasoundVertex.h"
 #include "MetasoundFrontendRegistryKey.h"
 #include "UObject/SoftObjectPath.h"
@@ -9,6 +10,7 @@
 // Forward Declarations
 class FMetasoundAssetBase;
 class UEdGraph;
+
 struct FMetasoundFrontendClassName;
 
 namespace Metasound
@@ -134,6 +136,10 @@ namespace Metasound
 			// Returns path associated with the given key (null if key is not registered with the AssetManager or was not loaded from asset)
 			virtual const FSoftObjectPath* FindObjectPathFromKey(const Metasound::Frontend::FNodeRegistryKey& InRegistryKey) const = 0;
 
+			// Converts an object to an AssetBase if its a registered asset
+			virtual FMetasoundAssetBase* GetAsAsset(UObject& InObject) const = 0;
+			virtual const FMetasoundAssetBase* GetAsAsset(const UObject& InObject) const = 0;
+
 #if WITH_EDITOR
 			// Generates all asset info associated with registered assets that are referenced by the provided asset's graph.
 			virtual TSet<FAssetInfo> GetReferencedAssetClasses(const FMetasoundAssetBase& InAssetBase) const = 0;
diff --git a/Engine/Plugins/Runtime/Metasound/Source/MetasoundFrontend/Public/MetasoundInputNode.h b/Engine/Plugins/Runtime/Metasound/Source/MetasoundFrontend/Public/MetasoundInputNode.h
--- a/Engine/Plugins/Runtime/Metasound/Source/MetasoundFrontend/Public/MetasoundInputNode.h
+++ b/Engine/Plugins/Runtime/Metasound/Source/MetasoundFrontend/Public/MetasoundInputNode.h
@@ -319,7 +319,7 @@ namespace Metasound
 		static FVertexInterface CreateDefaultVertexInterface(const FVertexName& InVertexName, const FName& InDataTypeName, EVertexAccessType InVertexAccess);
 
 	public:
-
+		static FText GetInputDescription();
 		static FNodeClassMetadata GetNodeMetadata(const FVertexName& InVertexName, const FName& InDataTypeName, EVertexAccessType InVertexAccess);
 
 		/* Construct a TInputNode using the TInputOperatorLiteralFactory<> and moving
