Amortize dynamic operator updates

diff --git a/Engine/Plugins/Runtime/Metasound/Source/MetasoundGraphCore/Private/MetasoundDynamicGraphAlgo.cpp b/Engine/Plugins/Runtime/Metasound/Source/MetasoundGraphCore/Private/MetasoundDynamicGraphAlgo.cpp
--- a/Engine/Plugins/Runtime/Metasound/Source/MetasoundGraphCore/Private/MetasoundDynamicGraphAlgo.cpp
+++ b/Engine/Plugins/Runtime/Metasound/Source/MetasoundGraphCore/Private/MetasoundDynamicGraphAlgo.cpp
@@ -7,6 +7,7 @@
 #include "MetasoundDataReference.h"
 #include "MetasoundGraphAlgoPrivate.h"
 #include "MetasoundLog.h"
+#include "MetasoundTrace.h"
 #include "MetasoundVertex.h"
 #include "MetasoundVertexData.h"
 
@@ -66,6 +67,7 @@ namespace Metasound
 		void UpdateGraphRuntimeTableEntries(const FOperatorID& InOperatorID, IOperator* InOperator, FDynamicGraphOperatorData& InOutGraphOperatorData)
 		{
 			using namespace DynamicGraphAlgoPrivate;
+			METASOUND_TRACE_CPUPROFILER_EVENT_SCOPE(Metasound::DynamicGraphAlgo::UpdateOperatorRuntimeTableEntries)
 
 			UpdateTableEntry(InOutGraphOperatorData.OperatorOrder, InOperatorID, InOperator, InOutGraphOperatorData.ExecuteTable, InOperator ? InOperator->GetExecuteFunction() : nullptr);
 			UpdateTableEntry(InOutGraphOperatorData.OperatorOrder, InOperatorID, InOperator, InOutGraphOperatorData.PostExecuteTable, InOperator ? InOperator->GetPostExecuteFunction() : nullptr);
@@ -78,6 +80,8 @@ namespace Metasound
 			using namespace DirectedGraphAlgo;
 			using namespace DynamicGraphAlgoPrivate;
 
+			METASOUND_TRACE_CPUPROFILER_EVENT_SCOPE(Metasound::DynamicGraphAlgo::PropagateBindUpdate)
+
 			struct FInputToUpdate
 			{
 				FOperatorID OperatorID;
@@ -92,6 +96,7 @@ namespace Metasound
 			TSortedVertexNameMap<FAnyDataReference> OutputUpdates;
 			while (PropagateStack.Num())
 			{
+				METASOUND_TRACE_CPUPROFILER_EVENT_SCOPE(Metasound::DynamicGraphAlgo::PropagateBindUpdate_Iteration)
 				FInputToUpdate Current = PropagateStack.Pop();
 				if (FOperatorInfo* OpInfo = InOutGraphOperatorData.OperatorMap.Find(Current.OperatorID))
 				{
@@ -206,6 +211,7 @@ namespace Metasound
 		void UpdateOutputVertexData(FDynamicGraphOperatorData& InOutGraphOperatorData)
 		{
 			using namespace DirectedGraphAlgo;
+			METASOUND_TRACE_CPUPROFILER_EVENT_SCOPE(Metasound::DynamicGraphAlgo::UpdateOutputVertexData)
 
 			// Iterate through the output operators and force their output data references
 			// to be reflected in the graph's FOutputVertexInterfaceData
@@ -243,6 +249,7 @@ namespace Metasound
 			using namespace DynamicGraphAlgoPrivate;
 
 			check(InOperatorInfo.Operator);
+			METASOUND_TRACE_CPUPROFILER_EVENT_SCOPE(Metasound::DynamicGraphAlgo::RebindWrappedOperator)
 
 			/* Bind and diff the graph's interface to determine if there is an update to any vertices */
 
@@ -289,6 +296,7 @@ namespace Metasound
 		void RebindGraphInputs(FInputVertexInterfaceData& InOutVertexData, FDynamicGraphOperatorData& InOutGraphOperatorData)
 		{
 			using namespace DirectedGraphAlgo;
+			METASOUND_TRACE_CPUPROFILER_EVENT_SCOPE(Metasound::DynamicGraphAlgo::RebindGraphInputs)
 
 			// Bind and diff the graph's interface to determine if there is an update to any vertices
 			FInputVertexInterfaceData& InputVertexData = InOutGraphOperatorData.VertexData.GetInputs();
@@ -329,6 +337,7 @@ namespace Metasound
 
 		void RebindGraphOutputs(FOutputVertexInterfaceData& InOutVertexData, FDynamicGraphOperatorData& InOutGraphOperatorData)
 		{
+			METASOUND_TRACE_CPUPROFILER_EVENT_SCOPE(Metasound::DynamicGraphAlgo::RebindGraphOutputs)
 			// Output rebinding does not alter data references in an operator. Here we can get away with
 			// simply reading the latest values.
 			InOutVertexData.Bind(InOutGraphOperatorData.VertexData.GetOutputs());
diff --git a/Engine/Plugins/Runtime/Metasound/Source/MetasoundGraphCore/Private/MetasoundDynamicOperator.cpp b/Engine/Plugins/Runtime/Metasound/Source/MetasoundGraphCore/Private/MetasoundDynamicOperator.cpp
--- a/Engine/Plugins/Runtime/Metasound/Source/MetasoundGraphCore/Private/MetasoundDynamicOperator.cpp
+++ b/Engine/Plugins/Runtime/Metasound/Source/MetasoundGraphCore/Private/MetasoundDynamicOperator.cpp
@@ -3,6 +3,7 @@
 #include "MetasoundDynamicOperator.h"
 
 #include "Containers/Array.h"
+#include "HAL/IConsoleManager.h"
 #include "MetasoundDynamicOperatorAudioFade.h"
 #include "MetasoundNodeInterface.h"
 #include "MetasoundOperatorInterface.h"
@@ -35,6 +36,14 @@ namespace Metasound
 
 		namespace DynamicOperatorPrivate
 		{
+			float MetaSoundExperimentalTransformTimeoutInSeconds = 0.010f;
+			FAutoConsoleVariableRef CVarMetaSoundExperimentalTransformTimeoutInSeconds(
+				TEXT("au.MetaSound.Experimental.DynamicOperatorTransformTimeoutInSeconds"),
+				MetaSoundExperimentalTransformTimeoutInSeconds,
+				TEXT("Sets the number of seconds allowed to process pending dynamic graph transformations for a single MetaSound render cycle .\n")
+				TEXT("[Less than zero]: Disabled, [Greater than zero]: Enabled, 0.010s (default)"),
+				ECVF_Default);
+			
 			// Table sorter helper so we don't rewrite this algorithm for each differe
 			// stack type (Execute/PostExecute/Reset)
 			struct FTableSorter
@@ -171,9 +180,56 @@ namespace Metasound
 			}
 		}
 
+		void FDynamicOperator::ApplyTransformsUntilFenceOrTimeout(double InTimeoutInSeconds)
+		{
+			METASOUND_TRACE_CPUPROFILER_EVENT_SCOPE(Metasound::FDynamicOperator::ApplyTransformsUntilFenceOrTimeout)
+
+			if (bExecuteFenceIsSet)
+			{
+				// Execute fence needs to be cleared before applying any transforms.
+				return;
+			}
+
+			TOptional<TUniquePtr<IDynamicOperatorTransform>> Transform = TransformQueue->Dequeue();
+
+			if (Transform)
+			{
+				check(FPlatformTime::GetSecondsPerCycle() > 0);
+				const uint64 BreakTimeInCycles = FPlatformTime::Cycles64() + static_cast<uint64>(InTimeoutInSeconds / FPlatformTime::GetSecondsPerCycle());
+				do
+				{
+					if (Transform.IsSet() && Transform->IsValid())
+					{
+						EDynamicOperatorTransformQueueAction Result = (*Transform)->Transform(DynamicOperatorData);
+						if (EDynamicOperatorTransformQueueAction::Fence == Result)
+						{
+							bExecuteFenceIsSet = true;
+							break;
+						}
+					}
+
+					if (FPlatformTime::Cycles64() >= BreakTimeInCycles)
+					{
+						UE_LOG(LogMetaSound, Verbose, TEXT("Transforms exceeded duration."));
+						break;
+					}
+				}
+				while((Transform = TransformQueue->Dequeue()));
+			}
+		}
+
 		void FDynamicOperator::Execute()
 		{
-			ApplyTransformsUntilFence();
+			using namespace DynamicOperatorPrivate;
+
+			if (MetaSoundExperimentalTransformTimeoutInSeconds > 0)
+			{
+				ApplyTransformsUntilFenceOrTimeout(MetaSoundExperimentalTransformTimeoutInSeconds);
+			}
+			else
+			{
+				ApplyTransformsUntilFence();
+			}
 
 			for (FExecuteEntry& Entry : DynamicOperatorData.ExecuteTable)
 			{
@@ -637,6 +693,8 @@ namespace Metasound
 
 		EDynamicOperatorTransformQueueAction FBeginAudioFadeTransform::FBeginAudioFadeTransform::Transform(FDynamicGraphOperatorData& InGraphOperatorData)
 		{
+			METASOUND_TRACE_CPUPROFILER_EVENT_SCOPE(Metasound::DynamicOperator::BeginAudioFadeTransform)
+
 		  	if (FOperatorInfo* OperatorInfo = InGraphOperatorData.OperatorMap.Find(OperatorIDToFade))
 			{
 				// Make wrapped operator
@@ -660,6 +718,8 @@ namespace Metasound
 
 		EDynamicOperatorTransformQueueAction FEndAudioFadeTransform::Transform(FDynamicGraphOperatorData& InGraphOperatorData)
 		{
+			METASOUND_TRACE_CPUPROFILER_EVENT_SCOPE(Metasound::DynamicOperator::EndAudioFadeTransform)
+
 		  	if (FOperatorInfo* OperatorInfo = InGraphOperatorData.OperatorMap.Find(OperatorIDToFade))
 			{
 				FAudioFadeOperatorWrapper* Wrapper = static_cast<FAudioFadeOperatorWrapper*>(OperatorInfo->Operator.Get());
diff --git a/Engine/Plugins/Runtime/Metasound/Source/MetasoundGraphCore/Private/MetasoundDynamicOperator.h b/Engine/Plugins/Runtime/Metasound/Source/MetasoundGraphCore/Private/MetasoundDynamicOperator.h
--- a/Engine/Plugins/Runtime/Metasound/Source/MetasoundGraphCore/Private/MetasoundDynamicOperator.h
+++ b/Engine/Plugins/Runtime/Metasound/Source/MetasoundGraphCore/Private/MetasoundDynamicOperator.h
@@ -60,6 +60,7 @@ namespace Metasound
 
 		private:
 			void ApplyTransformsUntilFence();
+			void ApplyTransformsUntilFenceOrTimeout(double InTimeoutInSeconds);
 			void Execute();
 			void PostExecute();
 			void Reset(const IOperator::FResetParams& InParams);
