Fix for uninitialized members
#rnx

diff --git a/Engine/Plugins/Runtime/Metasound/Source/MetasoundFrontend/Public/MetasoundFrontendDocument.h b/Engine/Plugins/Runtime/Metasound/Source/MetasoundFrontend/Public/MetasoundFrontendDocument.h
--- a/Engine/Plugins/Runtime/Metasound/Source/MetasoundFrontend/Public/MetasoundFrontendDocument.h
+++ b/Engine/Plugins/Runtime/Metasound/Source/MetasoundFrontend/Public/MetasoundFrontendDocument.h
@@ -173,7 +173,7 @@ struct METASOUNDFRONTEND_API FMetaSoundFrontendGraphComment
 
 #if WITH_EDITORONLY_DATA
 	UPROPERTY()
-	FLinearColor Color;
+	FLinearColor Color = FLinearColor::Black;
 
 	UPROPERTY()
 	FString Comment;
@@ -185,10 +185,10 @@ struct METASOUNDFRONTEND_API FMetaSoundFrontendGraphComment
 	int32 FontSize = 0;
 
 	UPROPERTY()
-	FVector2D Position;
+	FVector2D Position = FVector2D::Zero();
 
 	UPROPERTY()
-	FVector2D Size;
+	FVector2D Size = FVector2D::Zero();
 
 	UPROPERTY()
 	EMetaSoundFrontendGraphCommentMoveMode MoveMode = EMetaSoundFrontendGraphCommentMoveMode::GroupMovement;
