Disable adding comment nodes in non editable MetaSound graphs like presets

diff --git a/Engine/Plugins/Runtime/Metasound/Source/MetasoundEditor/Private/MetasoundEditor.cpp b/Engine/Plugins/Runtime/Metasound/Source/MetasoundEditor/Private/MetasoundEditor.cpp
--- a/Engine/Plugins/Runtime/Metasound/Source/MetasoundEditor/Private/MetasoundEditor.cpp
+++ b/Engine/Plugins/Runtime/Metasound/Source/MetasoundEditor/Private/MetasoundEditor.cpp
@@ -1985,7 +1985,8 @@ namespace Metasound
 		{
 			if (MetasoundGraphEditor.IsValid())
 			{
-				if (UEdGraph* Graph = MetasoundGraphEditor->GetCurrentGraph())
+				UEdGraph* Graph = MetasoundGraphEditor->GetCurrentGraph();
+				if (Graph && IsGraphEditable())
 				{
 					FMetasoundGraphSchemaAction_NewComment CommentAction;
 					CommentAction.PerformAction(Graph, nullptr, MetasoundGraphEditor->GetPasteLocation());
@@ -2398,11 +2399,6 @@ namespace Metasound
 
 		bool FEditor::CanDeleteNodes() const
 		{
-			if (!IsGraphEditable())
-			{
-				return false;
-			}
-
 			if (MetasoundGraphEditor->GetSelectedNodes().IsEmpty())
 			{
 				return false;
@@ -2411,10 +2407,15 @@ namespace Metasound
 			const FGraphPanelSelectionSet& SelectedNodes = MetasoundGraphEditor->GetSelectedNodes();
 			for (FGraphPanelSelectionSet::TConstIterator SelectedIter(SelectedNodes); SelectedIter; ++SelectedIter)
 			{
+				// Allow deletion of comment nodes even on uneditable graphs 
+				// because they were unintentionally addable at one point
 				UEdGraphNode* Node = Cast<UEdGraphNode>(*SelectedIter);
 				if (Node && Node->CanUserDeleteNode())
 				{
-					return true;
+					if (Cast<UEdGraphNode_Comment>(Node) || IsGraphEditable())
+					{
+						return true;
+					}
 				}
 			}
 			return false;
