Allow MetaSound member default literal values and read only member fields (name, display name, description) to be copy/pasted in the details panel

diff --git a/Engine/Plugins/Runtime/Metasound/Source/MetasoundEditor/Private/MetasoundNodeDetailCustomization.cpp b/Engine/Plugins/Runtime/Metasound/Source/MetasoundEditor/Private/MetasoundNodeDetailCustomization.cpp
--- a/Engine/Plugins/Runtime/Metasound/Source/MetasoundEditor/Private/MetasoundNodeDetailCustomization.cpp
+++ b/Engine/Plugins/Runtime/Metasound/Source/MetasoundEditor/Private/MetasoundNodeDetailCustomization.cpp
@@ -6,6 +6,7 @@
 #include "Containers/Set.h"
 #include "Delegates/Delegate.h"
 #include "Framework/Notifications/NotificationManager.h"
+#include "HAL/PlatformApplicationMisc.h"
 #include "IDetailChildrenBuilder.h"
 #include "IDetailGroup.h"
 #include "Internationalization/Text.h"
@@ -934,22 +935,37 @@ PRAGMA_ENABLE_DEPRECATION_WARNINGS
 
 		void FMetasoundMemberDetailCustomization::CustomizeGeneralCategory(IDetailLayoutBuilder& InDetailLayout)
 		{
-			const bool bIsInterfaceMember = IsInterfaceMember();
-			const bool bIsGraphEditable = IsGraphEditable();
-			
 			IDetailCategoryBuilder& CategoryBuilder = GetGeneralCategoryBuilder(InDetailLayout);
+			const bool bIsReadOnly = IsInterfaceMember() || !IsGraphEditable(); 
+
+			// Override row copy action if it's disabled by the edit condition 
+			auto GenerateCopyPasteActions = [](FDetailWidgetRow& Row, const FString& Value)
+			{
+				FUIAction CopyAction(FExecuteAction::CreateLambda([Value]()
+				{
+					FPlatformApplicationMisc::ClipboardCopy(*Value);
+				}));
+				Row.CopyAction(CopyAction);
+
+				// Create a dummy paste action
+				// Needed because the custom copy action will only be set 
+				// if both the copy and paste actions are bound
+				// Pasting is still available directly via the text box if editable
+				const static FUIAction PasteAction(FExecuteAction::CreateLambda([]() {}), FCanExecuteAction::CreateLambda([]() { return false; }));
+				Row.PasteAction(PasteAction);
+			};
 
 			NameEditableTextBox = SNew(SEditableTextBox)
 				.Text(this, &FMetasoundMemberDetailCustomization::GetName)
 				.OnTextChanged(this, &FMetasoundMemberDetailCustomization::OnNameChanged)
 				.OnTextCommitted(this, &FMetasoundMemberDetailCustomization::OnNameCommitted)
-				.IsReadOnly(bIsInterfaceMember || !bIsGraphEditable)
+				.IsReadOnly(bIsReadOnly)
 				.SelectAllTextWhenFocused(true)
 				.Font(IDetailLayoutBuilder::GetDetailFont());
 
 			static const FText MemberNameToolTipFormat = LOCTEXT("GraphMember_NameDescriptionFormat", "Name used within the MetaSounds editor(s) and transacting systems (ex. Blueprints) if applicable to reference the given {0}.");
-			CategoryBuilder.AddCustomRow(LOCTEXT("GraphMember_NameProperty", "Name"))
-				.EditCondition(!bIsInterfaceMember && bIsGraphEditable, nullptr)
+			FDetailWidgetRow& NameRow = CategoryBuilder.AddCustomRow(LOCTEXT("GraphMember_NameProperty", "Name"))
+				.EditCondition(!bIsReadOnly, nullptr)
 				.NameContent()
 				[
 					SNew(STextBlock)
@@ -961,12 +977,15 @@ PRAGMA_ENABLE_DEPRECATION_WARNINGS
 				[
 					NameEditableTextBox.ToSharedRef()
 				];
+			GenerateCopyPasteActions(NameRow, GetName().ToString());
 
 			static const FText MemberDisplayNameText = LOCTEXT("GraphMember_DisplayNameProperty", "Display Name");
 			static const FText MemberDisplayNameToolTipFormat = LOCTEXT("GraphMember_DisplayNameDescriptionFormat", "Optional, localized name used within the MetaSounds editor(s) to describe the given {0}.");
 			const FText MemberDisplayNameTooltipText = FText::Format(MemberDisplayNameToolTipFormat, GraphMember->GetGraphMemberLabel());
-			CategoryBuilder.AddCustomRow(MemberDisplayNameText)
-				.EditCondition(!bIsInterfaceMember && bIsGraphEditable, nullptr)
+
+			TSharedRef<FGraphMemberEditableTextDisplayName> DisplayNameValueText = MakeShared<FGraphMemberEditableTextDisplayName>(GraphMember, MemberDisplayNameTooltipText);
+			FDetailWidgetRow& DisplayNameRow = CategoryBuilder.AddCustomRow(MemberDisplayNameText)
+				.EditCondition(!bIsReadOnly, nullptr)
 				.NameContent()
 				[
 					SNew(STextBlock)
@@ -976,17 +995,20 @@ PRAGMA_ENABLE_DEPRECATION_WARNINGS
 				]
 				.ValueContent()
 				[
-					SNew(STextPropertyEditableTextBox, MakeShared<FGraphMemberEditableTextDisplayName>(GraphMember, MemberDisplayNameTooltipText))
+					SNew(STextPropertyEditableTextBox, DisplayNameValueText)
 					.WrapTextAt(500)
 					.MinDesiredWidth(25.0f)
 					.MaxDesiredHeight(200)
 				];
+			GenerateCopyPasteActions(DisplayNameRow, DisplayNameValueText->GetText(0).ToString());
 
 			static const FText MemberDescriptionText = LOCTEXT("Member_DescriptionPropertyName", "Description");
 			static const FText MemberDescriptionToolTipFormat = LOCTEXT("Member_DescriptionToolTipFormat", "Description for {0}. For example, used as a tooltip when displayed on another graph's referencing node.");
 			const FText MemberDescriptionToolTipText = FText::Format(MemberDescriptionToolTipFormat, GraphMember->GetGraphMemberLabel());
-			CategoryBuilder.AddCustomRow(MemberDescriptionText)
-				.EditCondition(!bIsInterfaceMember && bIsGraphEditable, nullptr)
+			TSharedRef<FGraphMemberEditableTextDescription> DescriptionValueText = MakeShared<FGraphMemberEditableTextDescription>(GraphMember, MemberDescriptionToolTipText);
+			FDetailWidgetRow& DescriptionRow = CategoryBuilder.AddCustomRow(MemberDescriptionText)
+				.IsEnabled(true)
+				.EditCondition(!bIsReadOnly, nullptr)
 				.NameContent()
 				[
 					SNew(STextBlock)
@@ -996,13 +1018,14 @@ PRAGMA_ENABLE_DEPRECATION_WARNINGS
 				]
 				.ValueContent()
 				[
-					SNew(STextPropertyEditableTextBox, MakeShared<FGraphMemberEditableTextDescription>(GraphMember, MemberDescriptionToolTipText))
+					SNew(STextPropertyEditableTextBox, DescriptionValueText)
 					.WrapTextAt(500)
 					.MinDesiredWidth(25.0f)
 					.MaxDesiredHeight(200)
 				];
+			GenerateCopyPasteActions(DescriptionRow, DescriptionValueText->GetText(0).ToString());
 
-			DataTypeSelector.AddDataTypeSelector(InDetailLayout, MemberCustomizationStyle::DataTypeNameText, GraphMember, !bIsInterfaceMember && bIsGraphEditable);
+			DataTypeSelector.AddDataTypeSelector(InDetailLayout, MemberCustomizationStyle::DataTypeNameText, GraphMember, !bIsReadOnly);
 		}
 
 		void FMetasoundMemberDetailCustomization::CustomizeDetails(IDetailLayoutBuilder& InDetailLayout)
