Cleaning up metasound experimental cvars
#rnx

diff --git a/Engine/Plugins/Runtime/Metasound/Source/MetasoundEngine/Private/MetasoundOperatorCacheSubsystem.cpp b/Engine/Plugins/Runtime/Metasound/Source/MetasoundEngine/Private/MetasoundOperatorCacheSubsystem.cpp
--- a/Engine/Plugins/Runtime/Metasound/Source/MetasoundEngine/Private/MetasoundOperatorCacheSubsystem.cpp
+++ b/Engine/Plugins/Runtime/Metasound/Source/MetasoundEngine/Private/MetasoundOperatorCacheSubsystem.cpp
@@ -86,11 +86,6 @@ void UMetaSoundCacheSubsystem::PrecacheMetaSound(UMetaSoundSource* InMetaSound,
 	using namespace Audio;
 	using namespace Metasound;
 
-	if (!FMetasoundGenerator::GetManualPrecacheEnabled())
-	{
-		return;
-	}
-
 	TSharedPtr<FOperatorPool> OperatorPool;
 	IMetasoundGeneratorModule* Module = FModuleManager::GetModulePtr<IMetasoundGeneratorModule>("MetasoundGenerator");
 	if (!ensure(Module))
diff --git a/Engine/Plugins/Runtime/Metasound/Source/MetasoundEngine/Private/MetasoundSource.cpp b/Engine/Plugins/Runtime/Metasound/Source/MetasoundEngine/Private/MetasoundSource.cpp
--- a/Engine/Plugins/Runtime/Metasound/Source/MetasoundEngine/Private/MetasoundSource.cpp
+++ b/Engine/Plugins/Runtime/Metasound/Source/MetasoundEngine/Private/MetasoundSource.cpp
@@ -54,7 +54,7 @@ namespace Metasound
 {
 	namespace ConsoleVariables
 	{
-		bool bEnableExperimentalRuntimePresetGraphInflation = true;
+		bool bEnableExperimentalRuntimePresetGraphInflation = false;
 	}
 
 	namespace SourcePrivate
@@ -315,7 +315,7 @@ FAutoConsoleVariableRef CVarMetaSoundEnableExperimentalRUntimePresetGraphInflati
 	TEXT("au.MetaSound.Experimental.EnableRuntimePresetGraphInflation"),
 	Metasound::ConsoleVariables::bEnableExperimentalRuntimePresetGraphInflation,
 	TEXT("Enables experimental feature of MetaSounds which reduces overhead of preset graphs\n")
-	TEXT("Default: true"),
+	TEXT("Default: false"),
 	ECVF_Default);
 
 
diff --git a/Engine/Plugins/Runtime/Metasound/Source/MetasoundFrontend/Private/MetasoundAssetBase.cpp b/Engine/Plugins/Runtime/Metasound/Source/MetasoundFrontend/Private/MetasoundAssetBase.cpp
--- a/Engine/Plugins/Runtime/Metasound/Source/MetasoundFrontend/Private/MetasoundAssetBase.cpp
+++ b/Engine/Plugins/Runtime/Metasound/Source/MetasoundFrontend/Private/MetasoundAssetBase.cpp
@@ -133,7 +133,7 @@ namespace Metasound
 		}
 
 		FAutoConsoleVariableRef CVarMetaSoundDisableAsyncGraphRegistration(
-			TEXT("au.MetaSound.Experimental.DisableAsyncGraphRegistration"),
+			TEXT("au.MetaSound.DisableAsyncGraphRegistration"),
 			Metasound::Frontend::ConsoleVariables::bDisableAsyncGraphRegistration,
 			TEXT("Disables async registration of MetaSound graphs\n")
 			TEXT("Default: false"),
diff --git a/Engine/Plugins/Runtime/Metasound/Source/MetasoundGenerator/Private/MetasoundGenerator.cpp b/Engine/Plugins/Runtime/Metasound/Source/MetasoundGenerator/Private/MetasoundGenerator.cpp
--- a/Engine/Plugins/Runtime/Metasound/Source/MetasoundGenerator/Private/MetasoundGenerator.cpp
+++ b/Engine/Plugins/Runtime/Metasound/Source/MetasoundGenerator/Private/MetasoundGenerator.cpp
@@ -33,10 +33,8 @@ namespace Metasound
 	namespace ConsoleVariables
 	{
 		static bool bEnableAsyncMetaSoundGeneratorBuilder = true;
-		static bool bEnableExperimentalOneShotOperatorPool = false;
-		static bool bEnableExperimentalOperatorPool = false;
-		static bool bEnableExperimentalOperatorPoolManualPrecache = true;
-		static bool bEnableResetOnOperatorPoolInsertion = true;
+		static bool bEnableExperimentalAutoCachingForOneShotOperators = false;
+		static bool bEnableExperimentalAutoCachingForAllOperators = false;
 #if ENABLE_METASOUNDGENERATOR_INVALID_SAMPLE_VALUE_LOGGING
 		static bool bEnableMetaSoundGeneratorNonFiniteLogging = false;
 		static bool bEnableMetaSoundGeneratorInvalidSampleValueLogging = false;
@@ -190,7 +188,7 @@ namespace Metasound
 		{
 		}
 #endif
-	}
+	}	
 }
 
 FAutoConsoleVariableRef CVarMetaSoundEnableAsyncGeneratorBuilder(
@@ -200,35 +198,22 @@ FAutoConsoleVariableRef CVarMetaSoundEnableAsyncGeneratorBuilder(
 	TEXT("Default: true"),
 	ECVF_Default);
 
-FAutoConsoleVariableRef CVarMetaSoundEnableExperimentalOneShotOperatorPool(
-	TEXT("au.MetaSound.Experimental.EnableOneShotOperatorPool"),
-	Metasound::ConsoleVariables::bEnableExperimentalOneShotOperatorPool,
-	TEXT("Enables caching of MetaSound operators using the OneShot source interface.\n")
+FAutoConsoleVariableRef CVarMetaSoundEnableExperimentalAutoCachingForOneShotOperators(
+	TEXT("au.MetaSound.Experimental.EnableAutoCachingForOneShotOperators"),
+	Metasound::ConsoleVariables::bEnableExperimentalAutoCachingForOneShotOperators,
+	TEXT("Enables auto-caching of MetaSound operators using the OneShot source interface.\n")
+	TEXT("(see MetasoundOperatorCacheSubsystem.h for manual path).\n")
 	TEXT("Default: false"),
 	ECVF_Default);
 
-FAutoConsoleVariableRef CVarMetaSoundEnableExperimentalOperatorPool(
-	TEXT("au.MetaSound.Experimental.EnableOperatorPool"),
-	Metasound::ConsoleVariables::bEnableExperimentalOperatorPool,
-	TEXT("Enables caching of all MetaSound operators.\n")
+FAutoConsoleVariableRef CVarMetaSoundEnableExperimentalAutoCachingForAllOperators(
+	TEXT("au.MetaSound.Experimental.EnableAutoCachingForAllOperators"),
+	Metasound::ConsoleVariables::bEnableExperimentalAutoCachingForAllOperators,
+	TEXT("Enables auto-caching of all MetaSound operators.\n")
+	TEXT("(see MetasoundOperatorCacheSubsystem.h for manual path).\n")
 	TEXT("Default: false"),
 	ECVF_Default);
 
-FAutoConsoleVariableRef CVarMetaSoundEnableExperimentalOperatorPoolManualPrecache(
-	TEXT("au.MetaSound.Experimental.EnableOperatorPoolManualPrecache"),
-	Metasound::ConsoleVariables::bEnableExperimentalOperatorPoolManualPrecache,
-	TEXT("Enables manual pre-caching of explicit MetaSound asset operators.\n")
-	TEXT("Default: true"),
-	ECVF_Default);
-
-FAutoConsoleVariableRef CVarMetaSoundEnableResetOnOperatorPoolInsertion(
-	TEXT("au.MetaSound.Experimental.EnableResetOnOperatorPoolInsertion"),
-	Metasound::ConsoleVariables::bEnableResetOnOperatorPoolInsertion,
-	TEXT("Enables reset of MetaSound operators as they are inserted into the Operator pool.\n")
-	TEXT("This can save a significant amount memory at the cost of reseting pooled operators twice.\n")
-	TEXT("Default: true"),
-	ECVF_Default);
-
 #if ENABLE_METASOUNDGENERATOR_INVALID_SAMPLE_VALUE_LOGGING
 
 FAutoConsoleVariableRef CVarMetaSoundEnableGeneratorNonFiniteLogging(
@@ -703,11 +688,6 @@ PRAGMA_ENABLE_DEPRECATION_WARNINGS
 		return RenderTime;
 	}
 
-	bool FMetasoundGenerator::GetManualPrecacheEnabled()
-	{
-		return ConsoleVariables::bEnableExperimentalOperatorPoolManualPrecache;
-	}
-
 	int32 FMetasoundGenerator::FillWithBuffer(const Audio::FAlignedFloatBuffer& InBuffer, float* OutAudio, int32 MaxNumOutputSamples)
 	{
 		int32 InNum = InBuffer.Num();
@@ -849,12 +829,12 @@ PRAGMA_ENABLE_DEPRECATION_WARNINGS
 		InitBase(InParams);
 		// attempt to use operator cache instead of building a new operator.
 		bool bDidUseCachedOperator = false;
-		const bool bIsOperatorPoolEnabled = ConsoleVariables::bEnableExperimentalOneShotOperatorPool || ConsoleVariables::bEnableExperimentalOperatorPool;
+		const bool bIsOperatorPoolEnabled = ConsoleVariables::bEnableExperimentalAutoCachingForOneShotOperators || ConsoleVariables::bEnableExperimentalAutoCachingForAllOperators;
 		// Dynamic operators cannot use the operator cache because they can change their internal structure. 
 		// The operator cache assumes that the operator is unchanged from it's original structure. 
 		if (bIsOperatorPoolEnabled)
 		{
-			bUseOperatorPool = ConsoleVariables::bEnableExperimentalOperatorPool || MetasoundGeneratorPrivate::HasOneShotInterface(InParams.Graph->GetVertexInterface());
+			bUseOperatorPool = ConsoleVariables::bEnableExperimentalAutoCachingForAllOperators || MetasoundGeneratorPrivate::HasOneShotInterface(InParams.Graph->GetVertexInterface());
 		}
 
 		// check the cache for manually pre-cached operators
@@ -941,7 +921,8 @@ PRAGMA_ENABLE_DEPRECATION_WARNINGS
 				UE_LOG(LogMetasoundGenerator, VeryVerbose, TEXT("Caching operator %s"), *LexToString(OperatorID));
 
 				// give the operator a chance to reduce its memory footprint before being cached
-				if (ConsoleVariables::bEnableResetOnOperatorPoolInsertion && EnvironmentPtr.IsValid())
+				// in the future this should be a conanical phase of an operator's lifecycle (i.e. Reset, Execute, Hybernate)
+				if (EnvironmentPtr.IsValid())
 				{
 					if (IOperator::FResetFunction Reset = GraphOperator->GetResetFunction())
 					{
diff --git a/Engine/Plugins/Runtime/Metasound/Source/MetasoundGenerator/Public/MetasoundGenerator.h b/Engine/Plugins/Runtime/Metasound/Source/MetasoundGenerator/Public/MetasoundGenerator.h
--- a/Engine/Plugins/Runtime/Metasound/Source/MetasoundGenerator/Public/MetasoundGenerator.h
+++ b/Engine/Plugins/Runtime/Metasound/Source/MetasoundGenerator/Public/MetasoundGenerator.h
@@ -226,8 +226,6 @@ namespace Metasound
 		/** Fraction of a single CPU core used to render audio on a scale of 0.0 to 1.0 */
 		double GetCPUCoreUtilization() const;
 
-		/** Returns whether or not operator pool manual precache feature is enabled (true) or not (false). */
-		static bool GetManualPrecacheEnabled();
 
 		// Called when a new graph has been "compiled" and set up as this generator's graph.
 		// Note: We don't allow direct assignment to the FOnSetGraph delegate
