Fix for metering widgets freezing when changing Output Format of MetaSound Source - No longer destroying analyzers when changing Output Format.
The FAudioMeter etc were being deleted and recreated when Output Format changed, but their corresponding slate widget instances were kept in use. When the new FAudioMeter instances were created, their new slate widgets were not used.
Now also stopping active preview when changing Output Format, as when the analyzers are reinitialized with a different channel count they recreate their audio buses.

diff --git a/Engine/Plugins/Runtime/Metasound/Source/MetasoundEditor/Private/MetasoundDetailCustomization.cpp b/Engine/Plugins/Runtime/Metasound/Source/MetasoundEditor/Private/MetasoundDetailCustomization.cpp
--- a/Engine/Plugins/Runtime/Metasound/Source/MetasoundEditor/Private/MetasoundDetailCustomization.cpp
+++ b/Engine/Plugins/Runtime/Metasound/Source/MetasoundEditor/Private/MetasoundDetailCustomization.cpp
@@ -170,7 +170,7 @@ namespace Metasound
 									TSharedPtr<FEditor> ParentEditor = FGraphBuilder::GetEditorForMetasound(*Source.Get());
 									if (ParentEditor.IsValid())
 									{
-										ParentEditor->DestroyAnalyzers();
+										ParentEditor->Stop();
 									};
 								}
 							}));
diff --git a/Engine/Plugins/Runtime/Metasound/Source/MetasoundEditor/Private/MetasoundEditor.cpp b/Engine/Plugins/Runtime/Metasound/Source/MetasoundEditor/Private/MetasoundEditor.cpp
--- a/Engine/Plugins/Runtime/Metasound/Source/MetasoundEditor/Private/MetasoundEditor.cpp
+++ b/Engine/Plugins/Runtime/Metasound/Source/MetasoundEditor/Private/MetasoundEditor.cpp
@@ -1209,6 +1209,10 @@ namespace Metasound
 					{
 						OutputMeter = MakeShared<AudioWidgets::FAudioMeter>(MetaSoundSource->NumChannels, AudioDeviceId);
 					}
+					else if (OutputMeter->GetAudioBus()->GetNumChannels() != MetaSoundSource->NumChannels)
+					{
+						OutputMeter->Init(MetaSoundSource->NumChannels, AudioDeviceId, nullptr);
+					}
 
 					const uint32 MetaSoundNumChannels = static_cast<uint32>(MetaSoundSource->NumChannels);
 
@@ -1227,11 +1231,11 @@ namespace Metasound
 							OscilloscopeAnalysisPeriodMs,
 							OscilloscopePanelLayoutType);
 					}
-					else
+					else if (OutputOscilloscope->GetAudioBus()->GetNumChannels() != MetaSoundSource->NumChannels)
 					{
 						OutputOscilloscope->CreateAudioBus(MetaSoundNumChannels);
 						OutputOscilloscope->CreateDataProvider(AudioDeviceId, OscilloscopeTimeWindowMs, OscilloscopeMaxTimeWindowMs, OscilloscopeAnalysisPeriodMs, OscilloscopePanelLayoutType);
-						OutputOscilloscope->CreateOscilloscopeWidget(MetaSoundNumChannels, EAudioPanelLayoutType::Basic);
+						OutputOscilloscope->CreateOscilloscopeWidget(MetaSoundNumChannels, OscilloscopePanelLayoutType);
 					}
 
 					// Init Vectorscope
@@ -1249,7 +1253,7 @@ namespace Metasound
 							VectorscopeAnalysisPeriodMs,
 							VectorscopePanelLayoutType);
 					}
-					else
+					else if (OutputVectorscope->GetAudioBus()->GetNumChannels() != MetaSoundSource->NumChannels)
 					{
 						OutputVectorscope->CreateAudioBus(MetaSoundNumChannels);
 						OutputVectorscope->CreateDataProvider(AudioDeviceId, VectorscopeTimeWindowMs, VectorscopeMaxTimeWindowMs, VectorscopeAnalysisPeriodMs);
@@ -1260,6 +1264,10 @@ namespace Metasound
 					{
 						OutputSpectrumAnalyzer = MakeShared<AudioWidgets::FAudioSpectrumAnalyzer>(MetaSoundSource->NumChannels, AudioDeviceId);
 					}
+					else if (OutputSpectrumAnalyzer->GetAudioBus()->GetNumChannels() != MetaSoundSource->NumChannels)
+					{
+						OutputSpectrumAnalyzer->Init(MetaSoundSource->NumChannels, AudioDeviceId, nullptr);
+					}
 
 					return;
 				}
