MetaSound: Reduce number of time a dynamic operator has it's operators reordered.

#changelist validated
#virtualized

diff --git a/Engine/Plugins/Runtime/Metasound/Source/MetasoundGraphCore/Private/MetasoundDynamicOperator.cpp b/Engine/Plugins/Runtime/Metasound/Source/MetasoundGraphCore/Private/MetasoundDynamicOperator.cpp
--- a/Engine/Plugins/Runtime/Metasound/Source/MetasoundGraphCore/Private/MetasoundDynamicOperator.cpp
+++ b/Engine/Plugins/Runtime/Metasound/Source/MetasoundGraphCore/Private/MetasoundDynamicOperator.cpp
@@ -290,14 +290,6 @@ namespace Metasound
 
 			InOutGraphOperatorData.OperatorOrder = Order;
 
-			// Update the entry into the various runtime tables dependent upon when the operator exists
-			// in the operator order array.
-			for (const TPair<FOperatorID, FOperatorInfo>& OpIDAndInfo : InOutGraphOperatorData.OperatorMap)
-			{
-				IOperator* Operator = OpIDAndInfo.Value.Operator.Get();
-				UpdateGraphRuntimeTableEntries(OpIDAndInfo.Key, Operator, InOutGraphOperatorData);
-			}
-
 			// Sort operator tables to be in the correct order.
 			FTableSorter::SortTable(Order, InOutGraphOperatorData.ExecuteTable);
 			FTableSorter::SortTable(Order, InOutGraphOperatorData.PostExecuteTable);
@@ -307,8 +299,9 @@ namespace Metasound
 		}
 
 		// Add an operator to the grpah
-		FAddOperator::FAddOperator(FOperatorID InOperatorID, FOperatorInfo&& InInfo)
+		FAddOperator::FAddOperator(FOperatorID InOperatorID, EExecutionOrderInsertLocation InLocation, FOperatorInfo&& InInfo)
 		: OperatorID(InOperatorID)
+		, Location(InLocation)
 		, OperatorInfo(MoveTemp(InInfo))
 		{
 		}
@@ -332,23 +325,50 @@ namespace Metasound
 					FRemoveOperator(OperatorID).Transform(InGraphOperatorData);
 				}
 
-				InGraphOperatorData.OperatorOrder.Add(OperatorID);
 				InGraphOperatorData.OperatorMap.Add(OperatorID, MoveTemp(OperatorInfo));
-
-				// Update execution tables
-				if (IOperator::FExecuteFunction ExecuteFunc = Operator->GetExecuteFunction())
+				switch (Location)
 				{
-					InGraphOperatorData.ExecuteTable.Add(FExecuteEntry(OperatorID, *Operator, ExecuteFunc));
-				}
+					case EExecutionOrderInsertLocation::First:
+						{
+							InGraphOperatorData.OperatorOrder.Insert(OperatorID, 0);
+							// Update execution tables
+							if (IOperator::FExecuteFunction ExecuteFunc = Operator->GetExecuteFunction())
+							{
+								InGraphOperatorData.ExecuteTable.Insert(FExecuteEntry(OperatorID, *Operator, ExecuteFunc), 0);
+							}
 
-				if (IOperator::FPostExecuteFunction PostExecuteFunc = Operator->GetPostExecuteFunction())
-				{
-					InGraphOperatorData.PostExecuteTable.Add(FPostExecuteEntry(OperatorID, *Operator, PostExecuteFunc));
-				}
+							if (IOperator::FPostExecuteFunction PostExecuteFunc = Operator->GetPostExecuteFunction())
+							{
+								InGraphOperatorData.PostExecuteTable.Insert(FPostExecuteEntry(OperatorID, *Operator, PostExecuteFunc), 0);
+							}
 
-				if (IOperator::FResetFunction ResetFunc = Operator->GetResetFunction())
-				{
-					InGraphOperatorData.ResetTable.Add(FResetEntry(OperatorID, *Operator, ResetFunc));
+							if (IOperator::FResetFunction ResetFunc = Operator->GetResetFunction())
+							{
+								InGraphOperatorData.ResetTable.Insert(FResetEntry(OperatorID, *Operator, ResetFunc), 0);
+							}
+						}
+						break;
+
+					case EExecutionOrderInsertLocation::Last:
+						{
+							InGraphOperatorData.OperatorOrder.Add(OperatorID);
+							// Update execution tables
+							if (IOperator::FExecuteFunction ExecuteFunc = Operator->GetExecuteFunction())
+							{
+								InGraphOperatorData.ExecuteTable.Add(FExecuteEntry(OperatorID, *Operator, ExecuteFunc));
+							}
+
+							if (IOperator::FPostExecuteFunction PostExecuteFunc = Operator->GetPostExecuteFunction())
+							{
+								InGraphOperatorData.PostExecuteTable.Add(FPostExecuteEntry(OperatorID, *Operator, PostExecuteFunc));
+							}
+
+							if (IOperator::FResetFunction ResetFunc = Operator->GetResetFunction())
+							{
+								InGraphOperatorData.ResetTable.Add(FResetEntry(OperatorID, *Operator, ResetFunc));
+							}
+						}
+						break;
 				}
 			}
 
diff --git a/Engine/Plugins/Runtime/Metasound/Source/MetasoundGraphCore/Private/MetasoundDynamicOperator.h b/Engine/Plugins/Runtime/Metasound/Source/MetasoundGraphCore/Private/MetasoundDynamicOperator.h
--- a/Engine/Plugins/Runtime/Metasound/Source/MetasoundGraphCore/Private/MetasoundDynamicOperator.h
+++ b/Engine/Plugins/Runtime/Metasound/Source/MetasoundGraphCore/Private/MetasoundDynamicOperator.h
@@ -20,6 +20,13 @@ namespace Metasound
 		struct FDynamicGraphOperatorData;
 		using FLiteralAssignmentFunction = void(*)(const FOperatorSettings& InOperatorSettings, const FLiteral& InLiteral, const FAnyDataReference& OutDataRef);
 
+		/** Where to place a new operator in the execution order. */
+		enum class EExecutionOrderInsertLocation : uint8
+		{
+			First, //< New operator executes before all existing operators.
+			Last   //< New operator executes after all existing operators.
+		};
+
 		/** Action to perform after a single transformation. */
 		enum class EDynamicOperatorTransformQueueAction : uint8
 		{
@@ -96,10 +103,12 @@ namespace Metasound
 		class FAddOperator : public IDynamicOperatorTransform
 		{
 		public:
-			FAddOperator(FOperatorID InOperatorID, FOperatorInfo&& InInfo);
+
+			FAddOperator(FOperatorID InOperatorID, EExecutionOrderInsertLocation InLocation, FOperatorInfo&& InInfo);
 			virtual EDynamicOperatorTransformQueueAction Transform(FDynamicGraphOperatorData& InGraphOperatorData) override;
 		private:
 			FOperatorID OperatorID;
+			EExecutionOrderInsertLocation Location;
 			FOperatorInfo OperatorInfo;
 		};
 
diff --git a/Engine/Plugins/Runtime/Metasound/Source/MetasoundGraphCore/Private/MetasoundDynamicOperatorTransactor.cpp b/Engine/Plugins/Runtime/Metasound/Source/MetasoundGraphCore/Private/MetasoundDynamicOperatorTransactor.cpp
--- a/Engine/Plugins/Runtime/Metasound/Source/MetasoundGraphCore/Private/MetasoundDynamicOperatorTransactor.cpp
+++ b/Engine/Plugins/Runtime/Metasound/Source/MetasoundGraphCore/Private/MetasoundDynamicOperatorTransactor.cpp
@@ -90,6 +90,7 @@ namespace Metasound
 		: OperatorBuilder(DynamicOperatorTransactorPrivate::GetOperatorBuilderSettings())
 		, Graph(InGraph)
 		{
+			CurrentOperatorOrder = DynamicOperatorTransactorPrivate::DetermineOperatorOrder(Graph);
 		}
 
 		FDynamicOperatorTransactor::FDynamicOperatorTransactor()
@@ -115,10 +116,13 @@ namespace Metasound
 				{
 					// Only add to this queue because we do not know at what point
 					// the other queues and dynamic operators were created.
-					Queue->Enqueue(CreateAddOperatorTransform(*GuidAndNode.Value, InOperatorSettings, InEnvironment));
+					Queue->Enqueue(CreateAddOperatorTransform(*GuidAndNode.Value, EExecutionOrderInsertLocation::Last, InOperatorSettings, InEnvironment));
 				}
 			}
 
+			// Force order to be synchronized to internal execution order.
+			Queue->Enqueue(MakeUnique<FSetOperatorOrder>(CurrentOperatorOrder));
+
 			return Queue;
 		}
 
@@ -137,7 +141,7 @@ namespace Metasound
 
 			Graph.AddNode(InNodeID, NodePtr);
 
-			EnqueueAddOperatorTransform(*NodePtr);
+			EnqueueAddOperatorTransform(*NodePtr, EExecutionOrderInsertLocation::Last);
 		}
 
 		void FDynamicOperatorTransactor::RemoveNode(const FGuid& InNodeID)
@@ -241,21 +245,16 @@ namespace Metasound
 				return;
 			}
 
-			/* enqueue an update. */
-			TArray<FOperatorID> OperatorOrder = DynamicOperatorTransactorPrivate::DetermineOperatorOrder(Graph);
-
 			if (ToNode->GetVertexInterface().GetInputVertex(InToVertex).DataTypeName == GetMetasoundDataTypeName<FAudioBuffer>())
 			{
 				// Handle audio edge removal with a fade out.
-				EnqueueFadeAndRemoveEdgeOperatorTransform(*FromNode, InFromVertex, *ToNode, InToVertex, *ReplacementLiteralNode, OperatorOrder);
+				EnqueueFadeAndRemoveEdgeOperatorTransform(*FromNode, InFromVertex, *ToNode, InToVertex, *ReplacementLiteralNode);
 			}
 			else
 			{
 				// Immediately disconnect non-audio edges. 
-				EnqueueRemoveEdgeOperatorTransform(*FromNode, InFromVertex, *ToNode, InToVertex, *ReplacementLiteralNode, OperatorOrder);
+				EnqueueRemoveEdgeOperatorTransform(*FromNode, InFromVertex, *ToNode, InToVertex, *ReplacementLiteralNode);
 			}
-
-
 		}
 
 		void FDynamicOperatorTransactor::SetValue(const FGuid& InNodeID, const FVertexName& InVertex, TUniquePtr<INode> InLiteralNode)
@@ -277,9 +276,12 @@ namespace Metasound
 				return;
 			}
 
+			// Always insert new literal nodes first in execution order.
+			CurrentOperatorOrder.Insert(DirectedGraphAlgo::GetOperatorID(*LiteralNode), 0);
+
 			auto CreateAddNodeTransform = [&](const FOperatorSettings& InOperatorSettings, const FMetasoundEnvironment& InEnvironment) -> TUniquePtr<IDynamicOperatorTransform>
 			{
-				return CreateAddOperatorTransform(*LiteralNode, InOperatorSettings, InEnvironment);
+				return CreateAddOperatorTransform(*LiteralNode, EExecutionOrderInsertLocation::First, InOperatorSettings, InEnvironment);
 			};
 
 			EnqueueTransformOnOperatorQueues(CreateAddNodeTransform);
@@ -423,27 +425,38 @@ namespace Metasound
 
 			/* add edge to internal graph. */
 			Graph.AddDataEdge(InFromNode, InFromVertex, InToNode, InToVertex);
-
-			// Find order of operators after adding edge. 
-			TArray<FOperatorID> OperatorOrder = DynamicOperatorTransactorPrivate::DetermineOperatorOrder(Graph);
 			
 			if (InputVertex->DataTypeName == GetMetasoundDataTypeName<FAudioBuffer>())
 			{
 				// If edge is audio, then the connection needs to be faded
-				EnqueueFadeAndAddEdgeOperatorTransform(InFromNode, InFromVertex, InToNode, InToVertex, PriorLiteralNode.Get(), OperatorOrder);
+				EnqueueFadeAndAddEdgeOperatorTransform(InFromNode, InFromVertex, InToNode, InToVertex, PriorLiteralNode.Get());
 			}
 			else
 			{
 				// If the edge is not audio, then no fading is performed. 
-				EnqueueAddEdgeOperatorTransform(InFromNode, InFromVertex, InToNode, InToVertex, PriorLiteralNode.Get(), OperatorOrder);
+				EnqueueAddEdgeOperatorTransform(InFromNode, InFromVertex, InToNode, InToVertex, PriorLiteralNode.Get());
 			}
 		}
 		
-		void FDynamicOperatorTransactor::EnqueueAddOperatorTransform(const INode& InNode)
+		void FDynamicOperatorTransactor::EnqueueAddOperatorTransform(const INode& InNode, EExecutionOrderInsertLocation InLocation)
 		{
+			// Update the CurrentOperatorOrder based on the insert location of new node. 
+			// This avoids resorting of entire graph on rendering dynamic metasounds by
+			// simply appending to a specific location.
+			switch (InLocation)
+			{
+				case EExecutionOrderInsertLocation::First:
+					CurrentOperatorOrder.Insert(DirectedGraphAlgo::GetOperatorID(InNode), 0);
+					break;
+
+				case EExecutionOrderInsertLocation::Last:
+					CurrentOperatorOrder.Add(DirectedGraphAlgo::GetOperatorID(InNode));
+					break;
+			}
+
 			auto CreateAddNodeTransform = [&](const FOperatorSettings& InOperatorSettings, const FMetasoundEnvironment& InEnvironment) -> TUniquePtr<IDynamicOperatorTransform>
 			{
-				return CreateAddOperatorTransform(InNode, InOperatorSettings, InEnvironment);
+				return CreateAddOperatorTransform(InNode, InLocation, InOperatorSettings, InEnvironment);
 			};
 
 			EnqueueTransformOnOperatorQueues(CreateAddNodeTransform);
@@ -456,6 +469,7 @@ namespace Metasound
 			EnqueueBeginFadeOperatorTransform(InNode, EAudioFadeType::FadeOut, InputsToFade, InOutputsToFade);
 
 			const FOperatorID OperatorID = DirectedGraphAlgo::GetOperatorID(InNode);
+			CurrentOperatorOrder.RemoveSingle(OperatorID);
 
 			auto CreateEndFadeOutTransform = [&](const FOperatorSettings& InOperatorSettings, const FMetasoundEnvironment& InEnvironment) -> TUniquePtr<IDynamicOperatorTransform>
 			{
@@ -469,12 +483,27 @@ namespace Metasound
 			EnqueueTransformOnOperatorQueues(CreateEndFadeOutTransform);
 		}
 
-		void FDynamicOperatorTransactor::EnqueueAddEdgeOperatorTransform(const INode& InFromNode, const FVertexName& InFromVertex, const INode& InToNode, const FVertexName& InToVertex, const INode* InPriorLiteralNode, const TArray<FOperatorID>& InNewOperatorOrder)
+		void FDynamicOperatorTransactor::EnqueueAddEdgeOperatorTransform(const INode& InFromNode, const FVertexName& InFromVertex, const INode& InToNode, const FVertexName& InToVertex, const INode* InPriorLiteralNode)
 		{
 			/* enqueue an update. */
 			FOperatorID FromOperatorID = DirectedGraphAlgo::GetOperatorID(InFromNode);
 			FOperatorID ToOperatorID = DirectedGraphAlgo::GetOperatorID(InToNode);
 
+			if (InPriorLiteralNode)
+			{
+				CurrentOperatorOrder.RemoveSingle(DirectedGraphAlgo::GetOperatorID(InPriorLiteralNode));
+			}
+
+			// Find order of operators after adding edge. 
+			TArray<FOperatorID> NewOperatorOrder = DynamicOperatorTransactorPrivate::DetermineOperatorOrder(Graph);
+
+			// Only set new order if it's different than existing.
+			bool bSetNewOperatorOrder = NewOperatorOrder != CurrentOperatorOrder;
+			if (bSetNewOperatorOrder)
+			{
+				CurrentOperatorOrder = NewOperatorOrder;
+			}
+
 			auto CreateAddEdgeTransform = [&](const FOperatorSettings& InOperatorSettings, const FMetasoundEnvironment& InEnvironment) -> TUniquePtr<IDynamicOperatorTransform>
 			{
 				TArray<TUniquePtr<IDynamicOperatorTransform>> AtomicTransforms;
@@ -483,7 +512,10 @@ namespace Metasound
 				{
 					AtomicTransforms.Add(MakeUnique<FRemoveOperator>(DirectedGraphAlgo::GetOperatorID(InPriorLiteralNode)));
 				}
-				AtomicTransforms.Add(MakeUnique<FSetOperatorOrder>(InNewOperatorOrder));
+				if (bSetNewOperatorOrder)
+				{
+					AtomicTransforms.Add(MakeUnique<FSetOperatorOrder>(NewOperatorOrder));
+				}
 				AtomicTransforms.Add(MakeUnique<FConnectOperators>(FromOperatorID, InFromVertex, ToOperatorID, InToVertex));
 
 				return MakeUnique<FAtomicTransform>(MoveTemp(AtomicTransforms));
@@ -492,7 +524,7 @@ namespace Metasound
 			EnqueueTransformOnOperatorQueues(CreateAddEdgeTransform);
 		}
 
-		void FDynamicOperatorTransactor::EnqueueFadeAndAddEdgeOperatorTransform(const INode& InFromNode, const FVertexName& InFromVertex, const INode& InToNode, const FVertexName& InToVertex, const INode* InPriorLiteralNode, const TArray<FOperatorID>& InNewOperatorOrder)
+		void FDynamicOperatorTransactor::EnqueueFadeAndAddEdgeOperatorTransform(const INode& InFromNode, const FVertexName& InFromVertex, const INode& InToNode, const FVertexName& InToVertex, const INode* InPriorLiteralNode)
 		{
 			FOperatorID FromOperatorID = DirectedGraphAlgo::GetOperatorID(InFromNode);
 			FOperatorID ToOperatorID = DirectedGraphAlgo::GetOperatorID(InToNode);
@@ -502,6 +534,21 @@ namespace Metasound
 			TArrayView<const FVertexName> InputsToFade(&InToVertex, 1);
 			TArrayView<const FVertexName> OutputsToFade;
 
+			if (InPriorLiteralNode)
+			{
+				CurrentOperatorOrder.RemoveSingle(DirectedGraphAlgo::GetOperatorID(InPriorLiteralNode));
+			}
+
+			// Find order of operators after removing literal and adding edge. 
+			TArray<FOperatorID> NewOperatorOrder = DynamicOperatorTransactorPrivate::DetermineOperatorOrder(Graph);
+
+			// Only set new order if it's different than existing.
+			bool bSetNewOperatorOrder = NewOperatorOrder != CurrentOperatorOrder;
+			if (bSetNewOperatorOrder)
+			{
+				CurrentOperatorOrder = NewOperatorOrder;
+			}
+
 			auto CreateBeginFadeAndAddEdgeTransform = [&](const FOperatorSettings& InOperatorSettings, const FMetasoundEnvironment& InEnvironment) -> TUniquePtr<IDynamicOperatorTransform>
 			{
 				TArray<TUniquePtr<IDynamicOperatorTransform>> AtomicTransforms;
@@ -513,7 +560,10 @@ namespace Metasound
 					AtomicTransforms.Add(MakeUnique<FRemoveOperator>(DirectedGraphAlgo::GetOperatorID(InPriorLiteralNode)));
 				}
 				
-				AtomicTransforms.Add(MakeUnique<FSetOperatorOrder>(InNewOperatorOrder));
+				if (bSetNewOperatorOrder)
+				{
+					AtomicTransforms.Add(MakeUnique<FSetOperatorOrder>(NewOperatorOrder));
+				}
 				AtomicTransforms.Add(MakeUnique<FConnectOperators>(FromOperatorID, InFromVertex, ToOperatorID, InToVertex));
 				AtomicTransforms.Add(MakeUnique<FBeginAudioFadeTransform>(ToOperatorID, EAudioFadeType::FadeIn, InputsToFade, OutputsToFade));
 				
@@ -568,6 +618,10 @@ namespace Metasound
 		{
 			FOperatorID OperatorID = DirectedGraphAlgo::GetOperatorID(InNode);
 
+			// Removing operator does not require a re-sort because existing dependencies
+			// are still met due to DAG structure. 
+			CurrentOperatorOrder.RemoveSingle(OperatorID);
+
 			auto CreateRemoveNodeTransform = [&](const FOperatorSettings& InOperatorSettings, const FMetasoundEnvironment& InEnvironment) -> TUniquePtr<IDynamicOperatorTransform>
 			{
 				return MakeUnique<FRemoveOperator>(OperatorID);
@@ -575,7 +629,7 @@ namespace Metasound
 			EnqueueTransformOnOperatorQueues(CreateRemoveNodeTransform);
 		}
 
-		void FDynamicOperatorTransactor::EnqueueRemoveEdgeOperatorTransform(const INode& InFromNode, const FVertexName& InFromVertex, const INode& InToNode, const FVertexName& InToVertex, const INode& InReplacementLiteralNode, const TArray<FOperatorID>& InNewOperatorOrder)
+		void FDynamicOperatorTransactor::EnqueueRemoveEdgeOperatorTransform(const INode& InFromNode, const FVertexName& InFromVertex, const INode& InToNode, const FVertexName& InToVertex, const INode& InReplacementLiteralNode)
 		{
 			using namespace DynamicOperatorTransactorPrivate;
 
@@ -583,25 +637,26 @@ namespace Metasound
 			const FOperatorID ToOperatorID = DirectedGraphAlgo::GetOperatorID(InToNode);
 			const FOperatorID LiteralOperatorID = DirectedGraphAlgo::GetOperatorID(InReplacementLiteralNode);
 
+			// Put literals in the front of the execution stack to simplify updating 
+			// runtime instances. No need to sort the entire graph if we are just
+			// inserting something at the beginning of the execution stack. 
+			CurrentOperatorOrder.Insert(LiteralOperatorID, 0);
+
 			auto CreateRemoveEdgeTransform = [&](const FOperatorSettings& InOperatorSettings, const FMetasoundEnvironment& InEnvironment) -> TUniquePtr<IDynamicOperatorTransform>
 			{
 				// Add the literal node.
-				TUniquePtr<IDynamicOperatorTransform> AddNodeTransform = CreateAddOperatorTransform(InReplacementLiteralNode, InOperatorSettings, InEnvironment);
+				TUniquePtr<IDynamicOperatorTransform> AddNodeTransform = CreateAddOperatorTransform(InReplacementLiteralNode, EExecutionOrderInsertLocation::First, InOperatorSettings, InEnvironment);
 
 				// Swap prior connection with new connections.
 				TUniquePtr<IDynamicOperatorTransform> ConnectOperatorsTransform = MakeUnique<FSwapOperatorConnection>(FromOperatorID, InFromVertex, LiteralOperatorID, DynamicOperatorTransactorPrivate::LiteralNodeOutputVertexName, ToOperatorID, InToVertex);
 
-				// Reorder the node graph.
-				TUniquePtr<IDynamicOperatorTransform> SetOrderTransform = MakeUnique<FSetOperatorOrder>(InNewOperatorOrder);
-
-				if (AddNodeTransform.IsValid() && ConnectOperatorsTransform.IsValid() && SetOrderTransform.IsValid())
+				if (AddNodeTransform.IsValid() && ConnectOperatorsTransform.IsValid())
 				{
 					// Create an atomic transform so all sub-transforms happen before next execution.
 					TArray<TUniquePtr<IDynamicOperatorTransform>> AtomicTransforms;
 
 					AtomicTransforms.Emplace(MoveTemp(AddNodeTransform));
 					AtomicTransforms.Emplace(MoveTemp(ConnectOperatorsTransform));
-					AtomicTransforms.Emplace(MoveTemp(SetOrderTransform));
 
 					return MakeUnique<FAtomicTransform>(MoveTemp(AtomicTransforms));
 				}
@@ -616,7 +671,7 @@ namespace Metasound
 			EnqueueTransformOnOperatorQueues(CreateRemoveEdgeTransform);
 		}
 
-		void FDynamicOperatorTransactor::EnqueueFadeAndRemoveEdgeOperatorTransform(const INode& InFromNode, const FVertexName& InFromVertex, const INode& InToNode, const FVertexName& InToVertex, const INode& InReplacementLiteralNode, const TArray<FOperatorID>& InNewOperatorOrder)
+		void FDynamicOperatorTransactor::EnqueueFadeAndRemoveEdgeOperatorTransform(const INode& InFromNode, const FVertexName& InFromVertex, const INode& InToNode, const FVertexName& InToVertex, const INode& InReplacementLiteralNode)
 		{
 			// Fade the input to the node getting disconnected rather than the output of the source node. The source node
 			// may be connected to other nodes and fading it's output would fade all the other connected nodes' inputs. 
@@ -633,7 +688,7 @@ namespace Metasound
 			// If we ever find ourselves creating audio buffers with literals which 
 			// are anything other than silent buffers, we should rework this operation
 			// to do either a cross-fade, or an additional "fade in" to the new value. 
-			EnqueueRemoveEdgeOperatorTransform(InFromNode, InFromVertex, InToNode, InToVertex, InReplacementLiteralNode, InNewOperatorOrder);
+			EnqueueRemoveEdgeOperatorTransform(InFromNode, InFromVertex, InToNode, InToVertex, InReplacementLiteralNode);
 
 			// Remove fade operation. 
 			EnqueueEndFadeOperatorTransform(InToNode);
@@ -641,7 +696,7 @@ namespace Metasound
 
 
 
-		TUniquePtr<IDynamicOperatorTransform> FDynamicOperatorTransactor::CreateAddOperatorTransform(const INode& InNode, const FOperatorSettings& InOperatorSettings, const FMetasoundEnvironment& InEnvironment) const
+		TUniquePtr<IDynamicOperatorTransform> FDynamicOperatorTransactor::CreateAddOperatorTransform(const INode& InNode, EExecutionOrderInsertLocation InLocation, const FOperatorSettings& InOperatorSettings, const FMetasoundEnvironment& InEnvironment) const
 		{
 			using namespace DynamicOperatorTransactorPrivate;
 
@@ -678,7 +733,7 @@ namespace Metasound
 					MoveTemp(InterfaceData)
 				};
 
-				return MakeUnique<FAddOperator>(OperatorID, MoveTemp(OpInfo));
+				return MakeUnique<FAddOperator>(OperatorID, InLocation, MoveTemp(OpInfo));
 			}
 			else
 			{
diff --git a/Engine/Plugins/Runtime/Metasound/Source/MetasoundGraphCore/Public/MetasoundDynamicOperatorTransactor.h b/Engine/Plugins/Runtime/Metasound/Source/MetasoundGraphCore/Public/MetasoundDynamicOperatorTransactor.h
--- a/Engine/Plugins/Runtime/Metasound/Source/MetasoundGraphCore/Public/MetasoundDynamicOperatorTransactor.h
+++ b/Engine/Plugins/Runtime/Metasound/Source/MetasoundGraphCore/Public/MetasoundDynamicOperatorTransactor.h
@@ -25,6 +25,7 @@ namespace Metasound
 		class IDynamicOperatorTransform;
 		class FDynamicOperator;
 		enum class EAudioFadeType : uint8;
+		enum class EExecutionOrderInsertLocation : uint8;
 
 		using FLiteralAssignmentFunction = void(*)(const FOperatorSettings& InOperatorSettings, const FLiteral& InLiteral, const FAnyDataReference& OutDataRef);
 		using FReferenceCreationFunction = TOptional<FAnyDataReference>(*)(const FOperatorSettings& InSettings, FName DataType, const FLiteral& InLiteral, EDataReferenceAccessType InAccessType);
@@ -114,22 +115,22 @@ namespace Metasound
 		private:
 			using FOperatorID = uintptr_t;
 
-			void EnqueueAddOperatorTransform(const INode& InNode);
+			void EnqueueAddOperatorTransform(const INode& InNode, EExecutionOrderInsertLocation InLocation);
 			void EnqueueFadeAndRemoveOperatorTransform(const INode& InNode, TArrayView<const FVertexName> InOutputsToFade);
 			void EnqueueRemoveOperatorTransform(const INode& InNode);
 			void EnqueueBeginFadeOperatorTransform(const INode& InNode, EAudioFadeType InFadeType, TArrayView<const FVertexName> InInputsToFade, TArrayView<const FVertexName> InOutputsToFade);
 			void EnqueueEndFadeOperatorTransform(const INode& InNode);
 
-			void EnqueueRemoveEdgeOperatorTransform(const INode& InFromNode, const FVertexName& InFromVertex, const INode& InToNode, const FVertexName& InToVertex, const INode& InReplacementLiteralNode, const TArray<FOperatorID>& InNewOperatorOrder);
-			void EnqueueFadeAndRemoveEdgeOperatorTransform(const INode& InFromNode, const FVertexName& InFromVertex, const INode& InToNode, const FVertexName& InToVertex, const INode& InReplacementLiteralNode, const TArray<FOperatorID>& InNewOperatorOrder);
-			void EnqueueAddEdgeOperatorTransform(const INode& InFromNode, const FVertexName& InFromVertex, const INode& InToNode, const FVertexName& InToVertex, const INode* InPriorLiteralNode, const TArray<FOperatorID>& InNewOperatorOrder);
-			void EnqueueFadeAndAddEdgeOperatorTransform(const INode& InFromNode, const FVertexName& InFromVertex, const INode& InToNode, const FVertexName& InToVertex, const INode* InPriorLiteralNode, const TArray<FOperatorID>& InNewOperatorOrder);
+			void EnqueueRemoveEdgeOperatorTransform(const INode& InFromNode, const FVertexName& InFromVertex, const INode& InToNode, const FVertexName& InToVertex, const INode& InReplacementLiteralNode);
+			void EnqueueFadeAndRemoveEdgeOperatorTransform(const INode& InFromNode, const FVertexName& InFromVertex, const INode& InToNode, const FVertexName& InToVertex, const INode& InReplacementLiteralNode);
+			void EnqueueAddEdgeOperatorTransform(const INode& InFromNode, const FVertexName& InFromVertex, const INode& InToNode, const FVertexName& InToVertex, const INode* InPriorLiteralNode);
+			void EnqueueFadeAndAddEdgeOperatorTransform(const INode& InFromNode, const FVertexName& InFromVertex, const INode& InToNode, const FVertexName& InToVertex, const INode* InPriorLiteralNode);
 
 			void AddDataEdgeInternal(const INode& InFromNode, const FVertexName& InFromVertex, const FGuid& InToNodeID, const INode& InToNode, const FVertexName& InToVertex);
 
 			using FCreateTransformFunctionRef = TFunctionRef<TUniquePtr<IDynamicOperatorTransform>(const FOperatorSettings& InOperatorSettings, const FMetasoundEnvironment& InEnvironment)>;
 
-			TUniquePtr<IDynamicOperatorTransform> CreateAddOperatorTransform(const INode& InNode, const FOperatorSettings& InOperatorSettings, const FMetasoundEnvironment& InEnvironment) const;
+			TUniquePtr<IDynamicOperatorTransform> CreateAddOperatorTransform(const INode& InNode, EExecutionOrderInsertLocation InLocation, const FOperatorSettings& InOperatorSettings, const FMetasoundEnvironment& InEnvironment) const;
 
 			void EnqueueTransformOnOperatorQueues(FCreateTransformFunctionRef InFunc);
 
@@ -144,6 +145,7 @@ namespace Metasound
 			};
 
 			TArray<FDynamicOperatorInfo> OperatorInfos;
+			TArray<FOperatorID> CurrentOperatorOrder;
 
 			struct FLiteralNodeID
 			{
