Fix race condition in metasound editor and async graph registration

#changelist validated
#virtualized

diff --git a/Engine/Plugins/Runtime/Metasound/Source/MetasoundEditor/Private/MetasoundEditorGraphBuilder.cpp b/Engine/Plugins/Runtime/Metasound/Source/MetasoundEditor/Private/MetasoundEditorGraphBuilder.cpp
--- a/Engine/Plugins/Runtime/Metasound/Source/MetasoundEditor/Private/MetasoundEditorGraphBuilder.cpp
+++ b/Engine/Plugins/Runtime/Metasound/Source/MetasoundEditor/Private/MetasoundEditorGraphBuilder.cpp
@@ -1827,6 +1827,10 @@ namespace Metasound
 			FMetaSoundAssetRegistrationOptions RegOptions;
 			RegOptions.bForceReregister = true;
 			RegOptions.bForceViewSynchronization = bInForceViewSynchronization;
+			// Protect against race conditions by utilizing a copy of the graph. Race conditions can happen
+			// if registration is performed asynchronously while the editor is still modifying the graph.
+			RegOptions.bRegisterCopyIfAsync = true; 
+			
 			// if EditedReferencingMetaSounds is empty, then no MetaSounds are open
 			// that reference this MetaSound, so just register this asset. Otherwise,
 			// this graph will recursively get updated when the open referencing graphs
diff --git a/Engine/Plugins/Runtime/Metasound/Source/MetasoundFrontend/Private/MetasoundAssetBase.cpp b/Engine/Plugins/Runtime/Metasound/Source/MetasoundFrontend/Private/MetasoundAssetBase.cpp
--- a/Engine/Plugins/Runtime/Metasound/Source/MetasoundFrontend/Private/MetasoundAssetBase.cpp
+++ b/Engine/Plugins/Runtime/Metasound/Source/MetasoundFrontend/Private/MetasoundAssetBase.cpp
@@ -269,13 +269,16 @@ void FMetasoundAssetBase::RegisterGraphWithFrontend(Metasound::Frontend::FMetaSo
 	check(Owner);
 	const FString AssetName = Owner->GetName();
 
-	// Async registration is only available if:
-	// 1. The IMetaSoundDocumentInterface is not actively modified by a builder
-	//    (built graph must be released synchronously to avoid a race condition on
-	//    reading/writing the IMetaSoundDocumentInterface on the Game Thread)
-	// 2. Async registration is globally disabled via console variable.
-	const bool bAsync = !(IsBuilderActive() || ConsoleVariables::bDisableAsyncGraphRegistration);
-	GraphRegistryKey = FRegistryContainerImpl::Get().RegisterGraph(Owner, bAsync);
+	// Register graphs async by default;
+	const bool bAsync = !ConsoleVariables::bDisableAsyncGraphRegistration;
+	// Force a copy if async registration is enabled and we need to protect against
+	// race conditions from external modifications.
+	bool bForceCopy = (IsBuilderActive() && bAsync);
+#if WITH_EDITOR
+	bForceCopy |= InRegistrationOptions.bRegisterCopyIfAsync;
+#endif // WITH_EDITOR
+
+	GraphRegistryKey = FRegistryContainerImpl::Get().RegisterGraph(Owner, bAsync, bForceCopy);
 
 	if (GraphRegistryKey.IsValid())
 	{
diff --git a/Engine/Plugins/Runtime/Metasound/Source/MetasoundFrontend/Private/MetasoundFrontendRegistryContainerImpl.cpp b/Engine/Plugins/Runtime/Metasound/Source/MetasoundFrontend/Private/MetasoundFrontendRegistryContainerImpl.cpp
--- a/Engine/Plugins/Runtime/Metasound/Source/MetasoundFrontend/Private/MetasoundFrontendRegistryContainerImpl.cpp
+++ b/Engine/Plugins/Runtime/Metasound/Source/MetasoundFrontend/Private/MetasoundFrontendRegistryContainerImpl.cpp
@@ -42,7 +42,7 @@ namespace Metasound::Frontend
 {
 	namespace RegistryPrivate
 	{
-		TScriptInterface<IMetaSoundDocumentInterface> BuildRegistryDocument(TScriptInterface<IMetaSoundDocumentInterface> DocumentInterface)
+		TScriptInterface<IMetaSoundDocumentInterface> BuildRegistryDocument(TScriptInterface<IMetaSoundDocumentInterface> DocumentInterface, bool bForceCopy)
 		{
 			using namespace Metasound::Frontend;
 
@@ -87,7 +87,14 @@ namespace Metasound::Frontend
 			}
 #endif // !NO_LOGGING
 
-			return DocumentInterface;
+			if (bForceCopy)
+			{
+				return &UMetaSoundBuilderDocument::Create(*DocumentInterface.GetInterface());
+			}
+			else
+			{
+				return DocumentInterface;
+			}
 		}
 
 		// FGraphNode is used to create unique INodes based off of a IGraph.
@@ -430,7 +437,7 @@ namespace Metasound::Frontend
 		return MakeUnique<FNodeRegistryTransactionStream>(TransactionBuffer);
 	}
 
-	FGraphRegistryKey FRegistryContainerImpl::RegisterGraph(const TScriptInterface<IMetaSoundDocumentInterface>& InDocumentInterface, bool bAsync)
+	FGraphRegistryKey FRegistryContainerImpl::RegisterGraph(const TScriptInterface<IMetaSoundDocumentInterface>& InDocumentInterface, bool bAsync, bool bForceCopy)
 	{
 		using namespace UE;
 
@@ -442,7 +449,7 @@ namespace Metasound::Frontend
 		// Use the asset path of the provided document interface object for identification, *NOT* the
 		// built version as the build process may in fact create a new object with a transient path.
 		const FTopLevelAssetPath AssetPath = InDocumentInterface->GetAssetPathChecked();
-		const TScriptInterface<IMetaSoundDocumentInterface> RegistryDocInterface = RegistryPrivate::BuildRegistryDocument(InDocumentInterface);
+		const TScriptInterface<IMetaSoundDocumentInterface> RegistryDocInterface = RegistryPrivate::BuildRegistryDocument(InDocumentInterface, bForceCopy);
 
 		UObject* OwningObject = RegistryDocInterface.GetObject();
 		check(OwningObject);
diff --git a/Engine/Plugins/Runtime/Metasound/Source/MetasoundFrontend/Private/MetasoundFrontendRegistryContainerImpl.h b/Engine/Plugins/Runtime/Metasound/Source/MetasoundFrontend/Private/MetasoundFrontendRegistryContainerImpl.h
--- a/Engine/Plugins/Runtime/Metasound/Source/MetasoundFrontend/Private/MetasoundFrontendRegistryContainerImpl.h
+++ b/Engine/Plugins/Runtime/Metasound/Source/MetasoundFrontend/Private/MetasoundFrontendRegistryContainerImpl.h
@@ -128,7 +128,7 @@ namespace Metasound
 			virtual TArray<FConverterNodeInfo> GetPossibleConverterNodes(const FName& FromDataType, const FName& ToDataType) override;
 
 			// Register a graph from an IMetaSoundDocumentInterface
-			FGraphRegistryKey RegisterGraph(const TScriptInterface<IMetaSoundDocumentInterface>& InDocumentInterface, bool bAsync = true);
+			FGraphRegistryKey RegisterGraph(const TScriptInterface<IMetaSoundDocumentInterface>& InDocumentInterface, bool bAsync = true, bool bForceCopy = false);
 
 			// Unregister a graph
 			bool UnregisterGraph(const FGraphRegistryKey& InRegistryKey, const TScriptInterface<IMetaSoundDocumentInterface>& InDocumentInterface, bool bAsync = true);
diff --git a/Engine/Plugins/Runtime/Metasound/Source/MetasoundFrontend/Public/MetasoundAssetManager.h b/Engine/Plugins/Runtime/Metasound/Source/MetasoundFrontend/Public/MetasoundAssetManager.h
--- a/Engine/Plugins/Runtime/Metasound/Source/MetasoundFrontend/Public/MetasoundAssetManager.h
+++ b/Engine/Plugins/Runtime/Metasound/Source/MetasoundFrontend/Public/MetasoundAssetManager.h
@@ -62,6 +62,10 @@ namespace Metasound
 
 			// Attempt to rebuild referenced classes (only run if class not registered or set to force re-register)
 			bool bRebuildReferencedAssetClasses = true;
+
+			// If true and the registration will be perform asynchronously, a copy of the graph will be used for
+			// async registration. 
+			bool bRegisterCopyIfAsync = false;
 #endif // WITH_EDITOR
 		};
 
