Fix missing reset on post executable input node.

diff --git a/Engine/Plugins/Runtime/Metasound/Source/MetasoundEngineTest/Private/EngineTestMetaSoundAutomatedNodeTest_Reset.cpp b/Engine/Plugins/Runtime/Metasound/Source/MetasoundEngineTest/Private/EngineTestMetaSoundAutomatedNodeTest_Reset.cpp
--- a/Engine/Plugins/Runtime/Metasound/Source/MetasoundEngineTest/Private/EngineTestMetaSoundAutomatedNodeTest_Reset.cpp
+++ b/Engine/Plugins/Runtime/Metasound/Source/MetasoundEngineTest/Private/EngineTestMetaSoundAutomatedNodeTest_Reset.cpp
@@ -152,14 +152,14 @@ bool FMetasoundAutomatedNodeTest_Reset::RunTest(const FString& InRegistryKeyStri
 		}
 		else if (OpExecFunc || OpPostExecFunc)
 		{
-			AddError(FString::Printf(TEXT("Missing initialize function when execute function exists for node %s - %s"), *InRegistryKeyString, *GetPrettyName(RegistryKey)));
+			AddError(FString::Printf(TEXT("Missing Reset(...) function when execute function exists for node %s - %s"), *InRegistryKeyString, *GetPrettyName(RegistryKey)));
 		}
 
 		// Check that after returning all inputs to their original state and calling
 		// reset on the operator, that all output values have returned to their initail state. 
 		if (!OutputTester.AreAllOutputValuesEqualToCapturedValues())
 		{
-			AddError(FString::Printf(TEXT("Reset function resulted in different starting conditions for node %s - %s"), *InRegistryKeyString, *GetPrettyName(RegistryKey)));
+			AddError(FString::Printf(TEXT("Reset(...) function resulted in different starting conditions for node %s - %s"), *InRegistryKeyString, *GetPrettyName(RegistryKey)));
 		}
 	};
 
diff --git a/Engine/Plugins/Runtime/Metasound/Source/MetasoundFrontend/Private/MetasoundInputNode.cpp b/Engine/Plugins/Runtime/Metasound/Source/MetasoundFrontend/Private/MetasoundInputNode.cpp
--- a/Engine/Plugins/Runtime/Metasound/Source/MetasoundFrontend/Private/MetasoundInputNode.cpp
+++ b/Engine/Plugins/Runtime/Metasound/Source/MetasoundFrontend/Private/MetasoundInputNode.cpp
@@ -15,20 +15,6 @@ namespace Metasound
 {
 	namespace MetasoundInputNodePrivate
 	{
-		FDataReferenceCollection FInputOperatorBase::GetInputs() const
-		{
-			// This is slated to be deprecated and removed.
-			checkNoEntry();
-			return {};
-		}
-
-		FDataReferenceCollection FInputOperatorBase::GetOutputs() const
-		{
-			// This is slated to be deprecated and removed.
-			checkNoEntry();
-			return {};
-		}
-
 		FNonExecutableInputOperatorBase::FNonExecutableInputOperatorBase(const FVertexName& InVertexName, FAnyDataReference&& InDataRef)
 		: VertexName(InVertexName)
 		, DataRef(MoveTemp(InDataRef))
diff --git a/Engine/Plugins/Runtime/Metasound/Source/MetasoundFrontend/Public/MetasoundInputNode.h b/Engine/Plugins/Runtime/Metasound/Source/MetasoundFrontend/Public/MetasoundInputNode.h
--- a/Engine/Plugins/Runtime/Metasound/Source/MetasoundFrontend/Public/MetasoundInputNode.h
+++ b/Engine/Plugins/Runtime/Metasound/Source/MetasoundFrontend/Public/MetasoundInputNode.h
@@ -18,14 +18,7 @@ namespace Metasound
 {
 	namespace MetasoundInputNodePrivate
 	{
-		class METASOUNDFRONTEND_API FInputOperatorBase : public IOperator
-		{
-		public:
-			virtual FDataReferenceCollection GetInputs() const override;
-			virtual FDataReferenceCollection GetOutputs() const override;
-		};
-		
-		class METASOUNDFRONTEND_API FNonExecutableInputOperatorBase : public FInputOperatorBase
+		class METASOUNDFRONTEND_API FNonExecutableInputOperatorBase : public IOperator
 		{	
 		public:
 			virtual void BindInputs(FInputVertexInterfaceData& InOutVertexData) override;
@@ -80,17 +73,17 @@ namespace Metasound
 		};
 
 		template<typename DataType>
-		class TExecutableInputOperator : public FInputOperatorBase
+		class TExecutableInputOperator : public IOperator
 		{
 			static_assert(TExecutableDataType<DataType>::bIsExecutable, "TExecutableInputOperatorBase should only be used with executable data types");
 		public:
 			using FDataWriteReference = TDataWriteReference<DataType>;
-			using FDataWriteReferenceFactory = TDataWriteReferenceLiteralFactory<DataType>;
 
+			UE_DEPRECATED(5.4, "The Executable data types will no longer be supported. Please use PostExecutable data types.") 
 			TExecutableInputOperator(const FVertexName& InDataReferenceName, TDataWriteReference<DataType> InValue)
 				: DataReferenceName(InDataReferenceName)
 				, InputValue(InValue)
-				, OutputValue(FDataWriteReferenceFactory::CreateNew(*InValue))
+				, OutputValue(FDataWriteReference::CreateNew(*InValue))
 			{
 			}
 
@@ -109,6 +102,11 @@ namespace Metasound
 				return &Execute;
 			}
 
+			virtual FExecuteFunction GetPostExecuteFunction() override
+			{
+				return nullptr;
+			}
+
 			virtual FResetFunction GetResetFunction() override
 			{
 				return nullptr;
@@ -169,7 +167,7 @@ namespace Metasound
 		};
 
 		template<typename DataType>
-		class TPostExecutableInputOperator : public FInputOperatorBase
+		class TPostExecutableInputOperator : public IOperator
 		{
 			static_assert(TPostExecutableDataType<DataType>::bIsPostExecutable, "TPostExecutableInputOperator should only be used with post executable data types");
 			static_assert(!TExecutableDataType<DataType>::bIsExecutable, "A data type cannot be Executable and PostExecutable");
@@ -190,8 +188,7 @@ namespace Metasound
 
 			virtual void BindOutputs(FOutputVertexInterfaceData& InOutVertexData) override
 			{
-				TDataReadReference<DataType> DataReadRef = DataRef.GetDataReadReference<DataType>();
-				InOutVertexData.BindReadVertex(DataReferenceName, DataReadRef);
+				InOutVertexData.BindVertex(DataReferenceName, DataRef.GetDataReadReference<DataType>());
 			}
 
 			virtual FExecuteFunction GetExecuteFunction() override
@@ -201,15 +198,42 @@ namespace Metasound
 
 			virtual FPostExecuteFunction GetPostExecuteFunction() override
 			{
-				return &PostExecute;
+				// This condition is checked at runtime as its possible dynamic graphs may reassign ownership
+				// of underlying data to operate on in post execute. In this case, the expectation is that the
+				// data reference is now owned by another provider/operator.
+				if (DataRef.GetAccessType() == EDataReferenceAccessType::Write)
+				{
+					return &PostExecute;
+				}
+				else
+				{
+					return nullptr;
+				}
 			}
 
 			virtual FResetFunction GetResetFunction() override
 			{
-				return nullptr;
+				// This condition is checked at runtime as its possible dynamic graphs may reassign ownership
+				// of underlying data to operate on in post execute. In this case, the expectation is that the
+				// data reference is now owned by another provider/operator.
+				if (DataRef.GetAccessType() == EDataReferenceAccessType::Write)
+				{
+					return &NoOpReset;
+				}
+				else
+				{
+					return nullptr;
+				}
 			}
 
 		protected:
+			static void NoOpReset(IOperator* InOperator, const IOperator::FResetParams& InParams)
+			{
+				// All post executable nodes must have a reset.  This is a special
+				// case of a non-owning node performing post execute on a data type
+				// owned by an external system.
+			}
+
 			static void PostExecute(IOperator* InOperator)
 			{
 				using FPostExecutableInputOperator = TPostExecutableInputOperator<DataType>;
@@ -217,14 +241,10 @@ namespace Metasound
 				FPostExecutableInputOperator* DerivedOperator = static_cast<FPostExecutableInputOperator*>(InOperator);
 				check(nullptr != DerivedOperator);
 
-				// This condition is checked at runtime as its possible dynamic graphs may reassign ownership
-				// of underlying data to operate on in post execute. In this case, the expectation is that the
-				// data reference is now owned by another provider/operator.
-				FAnyDataReference& DerivedDataRef = DerivedOperator->DataRef;
-				if (DerivedDataRef.GetAccessType() == EDataReferenceAccessType::Write)
+				DataType* Value = DerivedOperator->DataRef.template GetWritableValue<DataType>();
+				if (ensure(Value != nullptr))
 				{
-					FDataWriteReference DataWriteRef = DerivedDataRef.GetDataWriteReference<DataType>();
-					TPostExecutableDataType<DataType>::PostExecute(*DataWriteRef);
+					TPostExecutableDataType<DataType>::PostExecute(*Value);
 				}
 			}
 
@@ -238,6 +258,7 @@ namespace Metasound
 		public:
 			using FDataWriteReference = TDataWriteReference<DataType>;
 			using FDataWriteReferenceFactory = TDataWriteReferenceLiteralFactory<DataType>;
+			using TPostExecutableInputOperator<DataType>::DataRef;
 
 			TResetablePostExecutableInputOperator(const FVertexName& InDataReferenceName, const FOperatorSettings& InSettings, const FLiteral& InLiteral)
 			: TPostExecutableInputOperator<DataType>(InDataReferenceName, FDataWriteReferenceFactory::CreateExplicitArgs(InSettings, InLiteral))
@@ -247,7 +268,15 @@ namespace Metasound
 
 			virtual IOperator::FResetFunction GetResetFunction() override
 			{
-				return &Reset;
+				if (DataRef.GetAccessType() == EDataReferenceAccessType::Write)
+				{
+					return &Reset;
+				}
+				else
+				{
+					// If DataRef is not writable, reference is assumed to be reset by another owning operator.
+					return nullptr;
+				}
 			}
 
 		private:
@@ -259,11 +288,10 @@ namespace Metasound
 				FResetablePostExecutableInputOperator* Operator = static_cast<FResetablePostExecutableInputOperator*>(InOperator);
 				check(nullptr != Operator);
 
-				// If DataRef is not writable, reference is assumed to be reset by another owning operator.
-				if (Operator->DataRef.GetAccessType() == EDataReferenceAccessType::Write)
+				DataType* Value = Operator->DataRef.template GetWritableValue<DataType>();
+				if (ensure(Value != nullptr))
 				{
-					FAnyDataReference& Ref = Operator->DataRef;
-					*Ref.GetDataWriteReference<DataType>() = TDataTypeLiteralFactory<DataType>::CreateExplicitArgs(InParams.OperatorSettings, Operator->Literal);
+					*Value = TDataTypeLiteralFactory<DataType>::CreateExplicitArgs(InParams.OperatorSettings, Operator->Literal);
 				}
 			}
 
