Ensures in Metasound Operator Cache Subsystem
#rnx

diff --git a/Engine/Plugins/Runtime/Metasound/Source/MetasoundEngine/Private/MetasoundOperatorCacheSubsystem.cpp b/Engine/Plugins/Runtime/Metasound/Source/MetasoundEngine/Private/MetasoundOperatorCacheSubsystem.cpp
--- a/Engine/Plugins/Runtime/Metasound/Source/MetasoundEngine/Private/MetasoundOperatorCacheSubsystem.cpp
+++ b/Engine/Plugins/Runtime/Metasound/Source/MetasoundEngine/Private/MetasoundOperatorCacheSubsystem.cpp
@@ -96,13 +96,20 @@ void UMetaSoundCacheSubsystem::PrecacheMetaSound(UMetaSoundSource* InMetaSound,
 	OperatorPool = Module->GetOperatorPool();
 	const FMixerDevice* MixerDevice = GetMixerDevice();
 
-	if (!ensure(MixerDevice && InMetaSound && OperatorPool))
+	if (!ensure(MixerDevice && OperatorPool))
 	{
 		return;
 	}
 
-	if (!ensure(InNumInstances > 0))
+	if (!InMetaSound)
 	{
+		UE_LOG(LogMetaSound, Error, TEXT("PrecacheMetaSound called without being provided a MetaSound, ignoring request"));
+		return;
+	}
+
+	if (InNumInstances < 1)
+	{
+		UE_LOG(LogMetaSound, Error, TEXT("PrecacheMetaSound called with invaled NumInstances %i, ignoring request"), InNumInstances);
 		return;
 	}
 
diff --git a/Engine/Plugins/Runtime/Metasound/Source/MetasoundEngine/Public/MetasoundOperatorCacheSubsystem.h b/Engine/Plugins/Runtime/Metasound/Source/MetasoundEngine/Public/MetasoundOperatorCacheSubsystem.h
--- a/Engine/Plugins/Runtime/Metasound/Source/MetasoundEngine/Public/MetasoundOperatorCacheSubsystem.h
+++ b/Engine/Plugins/Runtime/Metasound/Source/MetasoundEngine/Public/MetasoundOperatorCacheSubsystem.h
@@ -28,11 +28,11 @@ public:
 	//~ Begin USubsystem interface
 	virtual bool ShouldCreateSubsystem(UObject* Outer) const override;
 	virtual void Initialize(FSubsystemCollectionBase& Collection) override;
-	virtual void Update() override; // todo: compile out
+	virtual void Update() override;
 	//~ End USubsystem interface
 	 
 	/* Builds the requested number of MetaSound operators (asynchronously) and puts them in the pool for playback.
-	(If these operators are not yet available when the MetaSound attepmts to play, one will be created Independent of this request.) */
+	(If these operators are not yet available when the MetaSound attempts to play, one will be created Independent of this request.) */
 	UFUNCTION(BlueprintCallable, Category = "MetaSound")
 	void PrecacheMetaSound(UMetaSoundSource* InMetaSound, int32 InNumInstances = 1);
 
