Fix MetaSound Array Random node shared state not applied when No Repeats is connected to a dynamic value that changes before the Next trigger is called

diff --git a/Engine/Plugins/Runtime/Metasound/Source/MetasoundFrontend/Public/MetasoundArrayRandomNode.h b/Engine/Plugins/Runtime/Metasound/Source/MetasoundFrontend/Public/MetasoundArrayRandomNode.h
--- a/Engine/Plugins/Runtime/Metasound/Source/MetasoundFrontend/Public/MetasoundArrayRandomNode.h
+++ b/Engine/Plugins/Runtime/Metasound/Source/MetasoundFrontend/Public/MetasoundArrayRandomNode.h
@@ -248,8 +248,6 @@
 			PrevArraySize = InputArrayRef.Num();
 			PrevNoRepeatOrder = FMath::Clamp(*NoRepeatOrder, 0, PrevArraySize - 1);
 
-			InitializeState(PrevArraySize);
-
 			*OutValue = TDataTypeFactory<ElementType>::CreateAny(InParams.OperatorSettings);
 			TriggerOnNext->Reset();
 			TriggerOnReset->Reset();
@@ -266,24 +264,8 @@
 			TriggerOnReset->AdvanceBlock();
 
 			const ArrayType& InputArrayRef = *InputArray;
-
-			// Determine if the state of the random number generator needs to be
-			// reinitialized.
-			const bool bIsArrayNonEmpty = InputArrayRef.Num() != 0; // Skip reinit if the array is empty because it represents an invalid state for this node.
-			const bool bIsArraySizeChanged = PrevArraySize != InputArrayRef.Num(); // Need to reinit for array size changes. 
-			const bool bIsSharedStateEnablementInconsistent = (*bEnableSharedState != bSharedStateInitialized); // Need to reinit if the shared state enablement has been updated.
-
-			const bool bIsStateReinitializationNeeded = bIsArrayNonEmpty && (bIsArraySizeChanged || bIsSharedStateEnablementInconsistent);
-
-			PrevArraySize = InputArrayRef.Num();
-
-			if (bIsStateReinitializationNeeded)
+			if (InputArrayRef.Num() == 0)
 			{
-				InitializeState(PrevArraySize);
-			}
-
-			if (PrevArraySize == 0)
-			{
 #if WITH_METASOUND_DEBUG_ENVIRONMENT
 				if (!bHasLoggedEmptyArrayWarning)
 				{
@@ -293,12 +275,89 @@
 #endif // WITH_METASOUND_DEBUG_ENVIRONMENT
 				return;
 			}
+			
+ 			TriggerReset->ExecuteBlock(
+				[&](int32 StartFrame, int32 EndFrame)
+				{
+				},
+				[this](int32 StartFrame, int32 EndFrame)
+				{
+					if (!bSharedStateInitialized)
+					{
+						InitializeState(PrevArraySize);
+					}
 
-			// Check for a seed change
-			if (PrevSeedValue != *SeedValue)
+					if (UseSharedState())
+					{
+						FSharedStateRandomGetManager& RGM = FSharedStateRandomGetManager::Get();
+						RGM.ResetSeed(SharedStateUniqueId);
+					}
+					else
+					{
+						check(ArrayRandomGet.IsValid());
+						ArrayRandomGet->ResetSeed();
+					}
+					TriggerOnReset->TriggerFrame(StartFrame);
+				}
+			);
+ 
+			TriggerNext->ExecuteBlock(
+				[&](int32 StartFrame, int32 EndFrame)
+				{
+				},
+				[this](int32 StartFrame, int32 EndFrame)
+				{
+					ExecuteTriggerNext(StartFrame);
+				}
+			);
+		}
+
+	private:
+		void ExecuteTriggerNext(int32 StartFrame)
+		{
+			const ArrayType& InputArrayRef = *InputArray;
+			int32 OutRandomIndex = INDEX_NONE;
+
+			// Determine if the state of the random number generator needs to be
+			// reinitialized.
+			const bool bIsArrayNonEmpty = InputArrayRef.Num() != 0; // Skip reinit if the array is empty because it represents an invalid state for this node.
+			const bool bIsArraySizeChanged = PrevArraySize != InputArrayRef.Num(); // Need to reinit for array size changes. 
+			const bool bIsSharedStateEnablementInconsistent = (*bEnableSharedState != bSharedStateInitialized); // Need to reinit if the shared state enablement has been updated.
+			const bool bIsNonSharedStateInitializationNeeded = !ArrayRandomGet.IsValid(); // For the first initialization of the non shared state random get (bIsSharedStateEnablementInconsistent will take care of that for shared state)
+
+			const bool bIsStateReinitializationNeeded = (bIsArrayNonEmpty && (bIsArraySizeChanged || bIsSharedStateEnablementInconsistent || bIsNonSharedStateInitializationNeeded));
+
+			const bool bSeedValueChanged = PrevSeedValue != *SeedValue;
+			const bool bNoRepeatOrderChanged = PrevNoRepeatOrder != *NoRepeatOrder;
+			const bool bWeightsArrayChanged = WeightsArray != *InputWeightsArray;
+
+			// Update cached values if changed
+			if (bIsArraySizeChanged)
 			{
+				PrevArraySize = InputArrayRef.Num();
+			}
+			if (bSeedValueChanged)
+			{
 				PrevSeedValue = *SeedValue;
+			}
+			if (bNoRepeatOrderChanged)
+			{
+				PrevNoRepeatOrder = *NoRepeatOrder;
+			}
+			if (bWeightsArrayChanged)
+			{
+				WeightsArray = *InputWeightsArray;
+			}
 
+			// Reinitialize state (with new values) if needed 
+			if (bIsStateReinitializationNeeded)
+			{
+				InitializeState(PrevArraySize);
+			}
+
+			// Update other state (which was not necessarily changed by new state initialization)
+			if (bSeedValueChanged)
+			{
 				if (UseSharedState())
 				{
 					FSharedStateRandomGetManager& RGM = FSharedStateRandomGetManager::Get();
@@ -310,10 +369,8 @@
 					ArrayRandomGet->SetSeed(PrevSeedValue);
 				}
 			}
-
-			if (PrevNoRepeatOrder != *NoRepeatOrder)
+			if (bNoRepeatOrderChanged)
 			{
-				PrevNoRepeatOrder = *NoRepeatOrder;
 				if (UseSharedState())
 				{
 					FSharedStateRandomGetManager& RGM = FSharedStateRandomGetManager::Get();
@@ -325,76 +382,43 @@
 					ArrayRandomGet->SetNoRepeatOrder(PrevNoRepeatOrder);
 				}
 			}
+			if (bWeightsArrayChanged)
+			{
+				if (UseSharedState())
+				{
+					FSharedStateRandomGetManager& RGM = FSharedStateRandomGetManager::Get();
+					RGM.SetRandomWeights(SharedStateUniqueId, WeightsArray);
+				}
+				else
+				{
+					check(ArrayRandomGet.IsValid());
+					ArrayRandomGet->SetRandomWeights(WeightsArray);
+				}
+			}
 
-			WeightsArray = *InputWeightsArray;
+			// Get next value 
 			if (UseSharedState())
 			{
 				FSharedStateRandomGetManager& RGM = FSharedStateRandomGetManager::Get();
-				RGM.SetRandomWeights(SharedStateUniqueId, WeightsArray);
+				OutRandomIndex = RGM.NextValue(SharedStateUniqueId);
 			}
 			else
 			{
 				check(ArrayRandomGet.IsValid());
-				ArrayRandomGet->SetRandomWeights(WeightsArray);
+				OutRandomIndex = ArrayRandomGet->NextValue();
 			}
- 
-			// Don't do anything if our array is empty
-			if (InputArrayRef.Num() == 0)
-			{
-				return;
-			}
- 
- 			TriggerReset->ExecuteBlock(
-				[&](int32 StartFrame, int32 EndFrame)
-				{
-				},
-				[this](int32 StartFrame, int32 EndFrame)
-				{
-					if (UseSharedState())
-					{
-						FSharedStateRandomGetManager& RGM = FSharedStateRandomGetManager::Get();
-						RGM.ResetSeed(SharedStateUniqueId);
-					}
-					else
-					{
-						check(ArrayRandomGet.IsValid());
-						ArrayRandomGet->ResetSeed();
-					}
-					TriggerOnReset->TriggerFrame(StartFrame);
-				}
-			);
- 
-			TriggerNext->ExecuteBlock(
-				[&](int32 StartFrame, int32 EndFrame)
-				{
-				},
-				[this](int32 StartFrame, int32 EndFrame)
-				{
-					const ArrayType& InputArrayRef = *InputArray;
-					int32 OutRandomIndex = INDEX_NONE;
 
-					if (UseSharedState())
-					{
-						FSharedStateRandomGetManager& RGM = FSharedStateRandomGetManager::Get();
-						OutRandomIndex = RGM.NextValue(SharedStateUniqueId);
-					}
-					else
-					{
-						check(ArrayRandomGet.IsValid());
-						OutRandomIndex = ArrayRandomGet->NextValue();
-					}
+			check(OutRandomIndex != INDEX_NONE);
+#if WITH_METASOUND_DEBUG_ENVIRONMENT
+			UE_LOG(LogMetaSound, Verbose, TEXT("Array Random Get: Index chosen: '%u'"), OutRandomIndex);
+#endif // WITH_METASOUND_DEBUG_ENVIRONMENT
 
-					check(OutRandomIndex != INDEX_NONE);
+			// The input array size may have changed, so make sure it's wrapped into range of the input array
+			*OutValue = InputArrayRef[OutRandomIndex % InputArrayRef.Num()];
 
-					// The input array size may have changed, so make sure it's wrapped into range of the input array
-					*OutValue = InputArrayRef[OutRandomIndex % InputArrayRef.Num()];
-
-					TriggerOnNext->TriggerFrame(StartFrame);
-				}
-			);
+			TriggerOnNext->TriggerFrame(StartFrame);
 		}
 
-	private:
 		void InitializeState(int32 InArraySize)
 		{
 			bSharedStateInitialized = false;
