Update MetaSound FTime accessiblity from MetaSound outputs.
#10697
#10663

diff --git a/Engine/Plugins/Runtime/Metasound/Source/MetasoundEngine/Private/MetasoundEngineModule.cpp b/Engine/Plugins/Runtime/Metasound/Source/MetasoundEngine/Private/MetasoundEngineModule.cpp
--- a/Engine/Plugins/Runtime/Metasound/Source/MetasoundEngine/Private/MetasoundEngineModule.cpp
+++ b/Engine/Plugins/Runtime/Metasound/Source/MetasoundEngine/Private/MetasoundEngineModule.cpp
@@ -122,6 +122,7 @@ public:
 		METASOUND_REGISTER_VERTEX_ANALYZER_FACTORY(Frontend::FVertexAnalyzerForwardBool)
 		METASOUND_REGISTER_VERTEX_ANALYZER_FACTORY(Frontend::FVertexAnalyzerForwardFloat)
 		METASOUND_REGISTER_VERTEX_ANALYZER_FACTORY(Frontend::FVertexAnalyzerForwardInt)
+		METASOUND_REGISTER_VERTEX_ANALYZER_FACTORY(Frontend::FVertexAnalyzerForwardTime)
 		METASOUND_REGISTER_VERTEX_ANALYZER_FACTORY(Frontend::FVertexAnalyzerForwardString)
 		METASOUND_REGISTER_VERTEX_ANALYZER_FACTORY(Frontend::FVertexAnalyzerTriggerDensity)
 		METASOUND_REGISTER_VERTEX_ANALYZER_FACTORY(Frontend::FVertexAnalyzerTriggerToTime)
@@ -143,6 +144,10 @@ public:
 			GetMetasoundDataTypeName<FString>(),
 			Frontend::FVertexAnalyzerForwardString::GetAnalyzerName(),
 			Frontend::FVertexAnalyzerForwardString::FOutputs::GetValue().Name);
+		UMetasoundGeneratorHandle::RegisterPassthroughAnalyzerForType(
+			GetMetasoundDataTypeName<FTime>(),
+			Frontend::FVertexAnalyzerForwardTime::GetAnalyzerName(),
+			Frontend::FVertexAnalyzerForwardTime::FOutputs::GetValue().Name);
 		UMetasoundGeneratorHandle::RegisterPassthroughAnalyzerForType(
 			GetMetasoundDataTypeName<FTrigger>(),
 			Frontend::FVertexAnalyzerTriggerToTime::GetAnalyzerName(),
diff --git a/Engine/Plugins/Runtime/Metasound/Source/MetasoundEngine/Private/MetasoundOutput.cpp b/Engine/Plugins/Runtime/Metasound/Source/MetasoundEngine/Private/MetasoundOutput.cpp
--- a/Engine/Plugins/Runtime/Metasound/Source/MetasoundEngine/Private/MetasoundOutput.cpp
+++ b/Engine/Plugins/Runtime/Metasound/Source/MetasoundEngine/Private/MetasoundOutput.cpp
@@ -119,7 +119,7 @@ bool UMetasoundOutputBlueprintAccess::IsTime(const FMetaSoundOutput& Output)
 	return Output.IsType<Metasound::FTime>();
 }
 
-float UMetasoundOutputBlueprintAccess::GetTimeSeconds(const FMetaSoundOutput& Output, bool& Success)
+double UMetasoundOutputBlueprintAccess::GetTimeSeconds(const FMetaSoundOutput& Output, bool& Success)
 {
 	Metasound::FTime TimeValue;
 	Success = Output.Get(TimeValue);
diff --git a/Engine/Plugins/Runtime/Metasound/Source/MetasoundEngine/Public/MetasoundOutput.h b/Engine/Plugins/Runtime/Metasound/Source/MetasoundEngine/Public/MetasoundOutput.h
--- a/Engine/Plugins/Runtime/Metasound/Source/MetasoundEngine/Public/MetasoundOutput.h
+++ b/Engine/Plugins/Runtime/Metasound/Source/MetasoundEngine/Public/MetasoundOutput.h
@@ -164,5 +164,5 @@ class UMetasoundOutputBlueprintAccess final : public UBlueprintFunctionLibrary
 	UFUNCTION(BlueprintCallable, Category="MetaSoundOutput")
 	static bool IsTime(const FMetaSoundOutput& Output);
 	UFUNCTION(BlueprintCallable, Category="MetaSoundOutput")
-	static float GetTimeSeconds(const FMetaSoundOutput& Output, bool& Success);
+	static double GetTimeSeconds(const FMetaSoundOutput& Output, bool& Success);
 };
diff --git a/Engine/Plugins/Runtime/Metasound/Source/MetasoundFrontend/Public/Analysis/MetasoundFrontendVertexAnalyzerForwardValue.h b/Engine/Plugins/Runtime/Metasound/Source/MetasoundFrontend/Public/Analysis/MetasoundFrontendVertexAnalyzerForwardValue.h
--- a/Engine/Plugins/Runtime/Metasound/Source/MetasoundFrontend/Public/Analysis/MetasoundFrontendVertexAnalyzerForwardValue.h
+++ b/Engine/Plugins/Runtime/Metasound/Source/MetasoundFrontend/Public/Analysis/MetasoundFrontendVertexAnalyzerForwardValue.h
@@ -106,5 +106,6 @@ namespace Metasound
 		METASOUND_DECLARE_FORWARD_VALUE_ANALYZER(FVertexAnalyzerForwardFloat, "UE.Forward.Float", float)
 		METASOUND_DECLARE_FORWARD_VALUE_ANALYZER(FVertexAnalyzerForwardInt, "UE.Forward.Int32", int32)
 		METASOUND_DECLARE_FORWARD_VALUE_ANALYZER(FVertexAnalyzerForwardString, "UE.Forward.String", FString)
+		METASOUND_DECLARE_FORWARD_VALUE_ANALYZER(FVertexAnalyzerForwardTime, "UE.Forward.Time", FTime)
 	} // namespace Frontend
 } // namespace Metasound
