Fix for metasound ensure during ResavePackages commandlet.

diff --git a/Engine/Plugins/Runtime/Metasound/Source/MetasoundEngine/Private/MetasoundEngineAsset.h b/Engine/Plugins/Runtime/Metasound/Source/MetasoundEngine/Private/MetasoundEngineAsset.h
--- a/Engine/Plugins/Runtime/Metasound/Source/MetasoundEngine/Private/MetasoundEngineAsset.h
+++ b/Engine/Plugins/Runtime/Metasound/Source/MetasoundEngine/Private/MetasoundEngineAsset.h
@@ -8,6 +8,7 @@
 #include "MetasoundFrontendDocumentIdGenerator.h"
 #include "MetasoundFrontendRegistryKey.h"
 #include "MetasoundUObjectRegistry.h"
+#include "Misc/App.h"
 #include "Serialization/Archive.h"
 
 #if WITH_EDITORONLY_DATA
@@ -109,7 +110,7 @@ namespace Metasound
 
 			if (UMetasoundEditorGraphBase* MetaSoundGraph = Cast<UMetasoundEditorGraphBase>(InMetaSound.GetGraph()))
 			{
-				if (InSaveContext.IsCooking())
+				if (InSaveContext.IsCooking() || IsRunningCommandlet())
 				{
 					// Use deterministic ID generation so more can be done at cook rather than runtime
 					if (MetaSoundEnableCookDeterministicIDGeneration != 0)
@@ -121,11 +122,16 @@ namespace Metasound
 						}
 					}
 				}
-				else
+ 				else if (FApp::CanEverRenderAudio())
 				{
 					MetaSoundGraph->RegisterGraphWithFrontend();
 					MetaSoundGraph->GetModifyContext().SetForceRefreshViews();
 				}
+				else
+				{
+					UE_LOG(LogMetaSound, Warning, TEXT("PreSaveAsset for MetaSound: (%s) is doing nothing because InSaveContext.IsCooking, IsRunningCommandlet, and FApp::CanEverRenderAudio were all false")
+						, *InMetaSound.GetPathName());
+				}
 			}
 #endif // WITH_EDITORONLY_DATA
 		}
