- Fix for shutdown crash in FWaveAsset construction due to IStreamingManager becoming invalid while operator cache is building proxies async
- Fix for shutdown stall caused by operator cache not canceling async build gracefully
#rnx
#tests AudioUnitTests, ran bug repro

diff --git a/Engine/Plugins/Runtime/Metasound/Source/MetasoundEngine/Private/MetasoundWave.cpp b/Engine/Plugins/Runtime/Metasound/Source/MetasoundEngine/Private/MetasoundWave.cpp
--- a/Engine/Plugins/Runtime/Metasound/Source/MetasoundEngine/Private/MetasoundWave.cpp
+++ b/Engine/Plugins/Runtime/Metasound/Source/MetasoundEngine/Private/MetasoundWave.cpp
@@ -54,7 +54,11 @@ namespace Metasound
 						{
 							if (SoundWaveProxy->GetNumChunks() > 1)
 							{
-								IStreamingManager::Get().GetAudioStreamingManager().RequestChunk(SoundWaveProxy, 1, [](EAudioChunkLoadResult) {});
+								if (FStreamingManagerCollection* StreamingMgr = IStreamingManager::Get_Concurrent())
+								{
+									IAudioStreamingManager& AudioStreamingMgr = StreamingMgr->GetAudioStreamingManager();
+									AudioStreamingMgr.RequestChunk(SoundWaveProxy, 1, [](EAudioChunkLoadResult) {});
+								}
 							}
 						}
 					}
diff --git a/Engine/Plugins/Runtime/Metasound/Source/MetasoundEngine/Public/MetasoundOperatorCacheSubsystem.h b/Engine/Plugins/Runtime/Metasound/Source/MetasoundEngine/Public/MetasoundOperatorCacheSubsystem.h
--- a/Engine/Plugins/Runtime/Metasound/Source/MetasoundEngine/Public/MetasoundOperatorCacheSubsystem.h
+++ b/Engine/Plugins/Runtime/Metasound/Source/MetasoundEngine/Public/MetasoundOperatorCacheSubsystem.h
@@ -33,18 +33,18 @@ public:
 	 
 	/* Builds the requested number of MetaSound operators (asynchronously) and puts them in the pool for playback.
 	(If these operators are not yet available when the MetaSound attempts to play, one will be created Independent of this request.) */
-	UFUNCTION(BlueprintCallable, Category = "MetaSound")
-	void PrecacheMetaSound(UMetaSoundSource* InMetaSound, int32 InNumInstances = 1);
+	UFUNCTION(BlueprintCallable, Category = "MetaSound", meta = (DisplayName = "Precache MetaSound"))
+	void PrecacheMetaSound(UPARAM(DisplayName = "MetaSound Source") UMetaSoundSource* InMetaSound, UPARAM(DisplayName = "Num Instances") int32 InNumInstances = 1);
 
 	/* same as PrecacheMetaSound except cached operator that already exists in the cache will be moved to the top instead of building,
 	any operators that we couldn't move to the top, will be built.
 	(i.e. if 2 operators are already cached and Num Instances is 4, it will construct 2 and move the existing 2 to the top of the cache) */
-	UFUNCTION(BlueprintCallable, Category = "MetaSound")
-	void TouchOrPrecacheMetaSound(UMetaSoundSource* InMetaSound, int32 InNumInstances = 1);
+	UFUNCTION(BlueprintCallable, Category = "MetaSound", meta = (DisplayName = "Touch or Precache MetaSound"))
+	void TouchOrPrecacheMetaSound(UPARAM(DisplayName = "MetaSound Source") UMetaSoundSource* InMetaSound, UPARAM(DisplayName = "Num Instances") int32 InNumInstances = 1);
 
 	/* Clear the operator pool of any operators associated with the given MetaSound */
-	UFUNCTION(BlueprintCallable, Category = "MetaSound")
-	void RemoveCachedOperatorsForMetaSound(UMetaSoundSource* InMetaSound);
+	UFUNCTION(BlueprintCallable, Category = "MetaSound", meta = (DisplayName = "RemoveCached Operators for MetaSound"))
+	void RemoveCachedOperatorsForMetaSound(UPARAM(DisplayName = "MetaSound Source") UMetaSoundSource* InMetaSound);
 
 private:
 	FSoundGeneratorInitParams BuildParams;
diff --git a/Engine/Plugins/Runtime/Metasound/Source/MetasoundGenerator/Private/MetasoundGeneratorModule.cpp b/Engine/Plugins/Runtime/Metasound/Source/MetasoundGenerator/Private/MetasoundGeneratorModule.cpp
--- a/Engine/Plugins/Runtime/Metasound/Source/MetasoundGenerator/Private/MetasoundGeneratorModule.cpp
+++ b/Engine/Plugins/Runtime/Metasound/Source/MetasoundGenerator/Private/MetasoundGeneratorModule.cpp
@@ -57,6 +57,9 @@ namespace Metasound
 
 	void FMetasoundGeneratorModule::ShutdownModule() 
 	{
+		// Have to cancel independent of resetting the shared pointer as tasks
+		// may still be holding on to pool and therefore keeping it alive
+		OperatorPool->CancelAllBuildEvents();
 		OperatorPool.Reset();
 		OperatorInstanceCounterManager.Reset();
 	}
diff --git a/Engine/Plugins/Runtime/Metasound/Source/MetasoundGenerator/Private/MetasoundOperatorCache.cpp b/Engine/Plugins/Runtime/Metasound/Source/MetasoundGenerator/Private/MetasoundOperatorCache.cpp
--- a/Engine/Plugins/Runtime/Metasound/Source/MetasoundGenerator/Private/MetasoundOperatorCache.cpp
+++ b/Engine/Plugins/Runtime/Metasound/Source/MetasoundGenerator/Private/MetasoundOperatorCache.cpp
@@ -179,6 +179,11 @@ namespace Metasound
 	{
 	}
 
+	FOperatorPool::~FOperatorPool()
+	{
+		CancelAllBuildEvents();
+	}
+
 	FOperatorAndInputs FOperatorPool::ClaimOperator(const FGuid& InOperatorID)
 	{
 		FScopeLock Lock(&CriticalSection);
@@ -242,10 +247,82 @@ namespace Metasound
 		Trim();
 	}
 
+	void FOperatorPool::BuildAndAddAsync(TUniqueFunction<void()>&& InBuildFunc)
+	{
+		struct FBuildAndAddOpTask : public FAsyncGraphTaskBase
+		{
+			TUniqueFunction<void()> TaskFunc;
+
+			FBuildAndAddOpTask(TUniqueFunction<void()>&& InTaskFunc)
+				: TaskFunc(MoveTemp(InTaskFunc))
+			{
+			}
+
+			void DoTask(ENamedThreads::Type, const FGraphEventRef& EventRef)
+			{
+				TaskFunc();
+
+				const FName ModuleName = TEXT("MetasoundGenerator");
+				if (FModuleManager::Get().IsModuleLoaded(ModuleName))
+				{
+					FMetasoundGeneratorModule& Module = FModuleManager::GetModuleChecked<FMetasoundGeneratorModule>(ModuleName);
+					TSharedPtr<FOperatorPool> Pool = Module.GetOperatorPool();
+					if (Pool.IsValid())
+					{
+						Pool->RemoveBuildEvent(EventRef);
+					}
+				}
+			}
+
+			ENamedThreads::Type GetDesiredThread()
+			{
+				return ENamedThreads::AnyThread;
+			}
+		};
+
+		FGraphEventRef EventRef = TGraphTask<FBuildAndAddOpTask>::CreateTask()
+			.ConstructAndDispatchWhenReady(MoveTemp(InBuildFunc));
+
+		{
+			FScopeLock Lock(&CriticalSection);
+			ActiveBuildEvents.Add(MoveTemp(EventRef));
+		}
+	}
+
+	void FOperatorPool::CancelAllBuildEvents()
+	{
+		if (!ActiveBuildEvents.IsEmpty())
+		{
+			UE_LOG(LogMetasoundGenerator, Display, TEXT("Cancelling active MetaSound Cache Pool Operator build requests..."));
+
+			bStopping.store(true);
+			for (FGraphEventRef& EventRef : ActiveBuildEvents)
+			{
+				if (EventRef.IsValid())
+				{
+					EventRef->Wait();
+				}
+			}
+			ActiveBuildEvents.Reset();
+			bStopping.store(false);
+		}
+	}
+
+	void FOperatorPool::RemoveBuildEvent(const FGraphEventRef& InEventRef)
+	{
+		FScopeLock Lock(&CriticalSection);
+		ActiveBuildEvents.Remove(InEventRef);
+	}
+
 	void FOperatorPool::BuildAndAddOperator(TUniquePtr<FOperatorBuildData> InBuildData)
 	{
 		using namespace OperatorPoolPrivate;
 
+		if (bStopping.load())
+		{
+			return;
+		}
+
 		if (!ensure(InBuildData))
 		{
 			return;
@@ -263,15 +340,16 @@ namespace Metasound
 			}
 		}
 
-		FMetasoundGeneratorModule& Module = FModuleManager::GetModuleChecked<FMetasoundGeneratorModule>("MetasoundGenerator");
-		AsyncTask(ENamedThreads::AnyThread, [Graph = Graph, PreCacheData = MoveTemp(InBuildData), OperatorPool = Module.GetOperatorPool()] ()
+		// Build operations should never keep the operator pool alive as this can delay app shutdown arbitrarily.
+		TWeakPtr<FOperatorPool> WeakOpPool = AsShared();
+		BuildAndAddAsync([Graph, PreCacheData = MoveTemp(InBuildData), WeakOpPool]()
 		{
 			using namespace OperatorPoolPrivate;
 
 			METASOUND_LLM_SCOPE;
 			METASOUND_TRACE_CPUPROFILER_EVENT_SCOPE(Metasound::FOperatorPool::AsyncOperatorPrecache)
 
-			if (!ensure(PreCacheData && OperatorPool))
+			if (!ensure(PreCacheData))
 			{
 				return;
 			}
@@ -302,10 +380,16 @@ namespace Metasound
 			const int32 NumInstances = PreCacheData->NumInstances;
 			for (int32 i = 0; i < NumInstances; ++i)
 			{
+				TSharedPtr<FOperatorPool> OperatorPool = WeakOpPool.Pin();
+				if (!OperatorPool.IsValid() || OperatorPool->IsStopping())
+				{
+					break;
+				}
+
 				FBuildResults BuildResults;
 				FOperatorAndInputs OperatorAndInputs = GeneratorBuilder::BuildGraphOperator(PreCacheData->InitParams.OperatorSettings, PreCacheData->InitParams, BuildResults);
 				GeneratorBuilder::LogBuildErrors(PreCacheData->InitParams.MetaSoundName, BuildResults);
-	
+
 				const FGuid& GraphID = PreCacheData->InitParams.Graph->GetInstanceID();
 				OperatorPool->AddOperator(GraphID, MoveTemp(OperatorAndInputs));
 				OperatorPool->AddAssetIdToGraphIdLookUp(PreCacheData->AssetClassID, GraphID);
diff --git a/Engine/Plugins/Runtime/Metasound/Source/MetasoundGenerator/Public/MetasoundOperatorCache.h b/Engine/Plugins/Runtime/Metasound/Source/MetasoundGenerator/Public/MetasoundOperatorCache.h
--- a/Engine/Plugins/Runtime/Metasound/Source/MetasoundGenerator/Public/MetasoundOperatorCache.h
+++ b/Engine/Plugins/Runtime/Metasound/Source/MetasoundGenerator/Public/MetasoundOperatorCache.h
@@ -82,10 +82,11 @@ namespace Metasound
 
 	// Pool of re-useable metasound operators to be used / put back by the metasound generator
 	// operators can also be pre-constructed via the UMetasoundCacheSubsystem BP api.
-	class METASOUNDGENERATOR_API FOperatorPool
+	class METASOUNDGENERATOR_API FOperatorPool : public TSharedFromThis<FOperatorPool>
 	{
 	public:
 		FOperatorPool(const FOperatorPoolSettings& InSettings);
+		~FOperatorPool();
 
 		FOperatorAndInputs ClaimOperator(const FGuid& InOperatorID);
 
@@ -96,6 +97,8 @@ namespace Metasound
 		void TouchOperators(const FGuid& InOpeoratorID, const int32& NumToTouch = 1);
 		void TouchOperatorsViaAssetClassID(const FGuid& InAssetClassID, const int32& NumToTouch = 1);
 
+		bool IsStopping() const { return bStopping.load(); }
+
 		void RemoveOperatorsWithID(const FGuid& InOperatorID);
 		void RemoveOperatorsWithAssetClassID(const FGuid& InAssetClassID);
 
@@ -109,7 +112,12 @@ namespace Metasound
 		void UpdateHitRateTracker();
 #endif // #if METASOUND_OPERATORCACHEPROFILER_ENABLED
 
+		void CancelAllBuildEvents();
+
 	private:
+		void RemoveBuildEvent(const FGraphEventRef& InEventRef);
+
+		void BuildAndAddAsync(TUniqueFunction<void()>&& InBuildFunc);
 		void Trim();
 
 		FOperatorPoolSettings Settings;
@@ -119,6 +127,9 @@ namespace Metasound
 		OperatorPoolPrivate::FWindowedHitRate HitRateTracker;
 #endif // #if METASOUND_OPERATORCACHEPROFILER_ENABLED
 
+		std::atomic<bool> bStopping;
+		TSet<FGraphEventRef> ActiveBuildEvents;
+
 		TMap<FGuid, TArray<FOperatorAndInputs>> Operators;
 		TMap<FGuid, FGuid> AssetIdToGraphIdLookUp;
 		TArray<FGuid> Stack;
