Fix losing namespace when select click renaming MetaSound graph members

diff --git a/Engine/Plugins/Runtime/Metasound/Source/MetasoundEditor/Private/MetasoundEditor.cpp b/Engine/Plugins/Runtime/Metasound/Source/MetasoundEditor/Private/MetasoundEditor.cpp
--- a/Engine/Plugins/Runtime/Metasound/Source/MetasoundEditor/Private/MetasoundEditor.cpp
+++ b/Engine/Plugins/Runtime/Metasound/Source/MetasoundEditor/Private/MetasoundEditor.cpp
@@ -510,14 +510,36 @@ namespace Metasound
 				{
 					if (UMetasoundEditorGraphMember* GraphMember = MetasoundAction->GetGraphMember())
 					{
+						// Check if new name has changed
+						// Check against the non namespaced member name because
+						// this text box is only for the non namespaced part of the name
+						// (namespace is in parent menu items)
+						FName Namespace;
+						FName Name;
+						Audio::FParameterPath::SplitName(GraphMember->GetMemberName(), Namespace, Name);
+
+						if (Name == InNewText.ToString())
+						{
+							return;
+						}
+
 						const FText TransactionLabel = FText::Format(LOCTEXT("Rename Graph Member", "Set MetaSound {0}'s Name"), GraphMember->GetGraphMemberLabel());
 						const FScopedTransaction Transaction(TransactionLabel);
 
 						constexpr bool bPostTransaction = false;
 						GraphMember->SetDisplayName(FText::GetEmpty(), bPostTransaction);
-						GraphMember->SetMemberName(FName(*InNewText.ToString()), bPostTransaction);
+
+						// Add back namespace if needed
+						FString NewName = InNewText.ToString();
+						if (!Namespace.IsNone())
+						{
+							NewName = Namespace.ToString() + Audio::FParameterPath::NamespaceDelimiter + NewName;
+						}
+						GraphMember->SetMemberName(FName(NewName), bPostTransaction);
 					}
 				}
+
+
 			}
 
 			virtual TSharedRef<SWidget> CreateTextSlotWidget(FCreateWidgetForActionData* const InCreateData, TAttribute<bool> bIsReadOnly) override
