Fix MetaSound Array Random Get Node not initializing properly

diff --git a/Engine/Plugins/Runtime/Metasound/Source/MetasoundFrontend/Public/MetasoundArrayRandomNode.h b/Engine/Plugins/Runtime/Metasound/Source/MetasoundFrontend/Public/MetasoundArrayRandomNode.h
--- a/Engine/Plugins/Runtime/Metasound/Source/MetasoundFrontend/Public/MetasoundArrayRandomNode.h
+++ b/Engine/Plugins/Runtime/Metasound/Source/MetasoundFrontend/Public/MetasoundArrayRandomNode.h
@@ -282,7 +282,7 @@ namespace Metasound
 				},
 				[this](int32 StartFrame, int32 EndFrame)
 				{
-					if (!bSharedStateInitialized)
+					if (IsStateInitializationNeeded())
 					{
 						InitializeState(PrevArraySize);
 					}
@@ -318,15 +318,8 @@ namespace Metasound
 			const ArrayType& InputArrayRef = *InputArray;
 			int32 OutRandomIndex = INDEX_NONE;
 
-			// Determine if the state of the random number generator needs to be
-			// reinitialized.
-			const bool bIsArrayNonEmpty = InputArrayRef.Num() != 0; // Skip reinit if the array is empty because it represents an invalid state for this node.
-			const bool bIsArraySizeChanged = PrevArraySize != InputArrayRef.Num(); // Need to reinit for array size changes. 
-			const bool bIsSharedStateEnablementInconsistent = (*bEnableSharedState != bSharedStateInitialized); // Need to reinit if the shared state enablement has been updated.
-			const bool bIsNonSharedStateInitializationNeeded = !ArrayRandomGet.IsValid(); // For the first initialization of the non shared state random get (bIsSharedStateEnablementInconsistent will take care of that for shared state)
-
-			const bool bIsStateReinitializationNeeded = (bIsArrayNonEmpty && (bIsArraySizeChanged || bIsSharedStateEnablementInconsistent || bIsNonSharedStateInitializationNeeded));
-
+			const bool bIsStateReinitializationNeeded = IsStateInitializationNeeded();
+			const bool bIsArraySizeChanged = PrevArraySize != InputArrayRef.Num(); 
 			const bool bSeedValueChanged = PrevSeedValue != *SeedValue;
 			const bool bNoRepeatOrderChanged = PrevNoRepeatOrder != *NoRepeatOrder;
 			const bool bWeightsArrayChanged = WeightsArray != *InputWeightsArray;
@@ -419,6 +412,17 @@ namespace Metasound
 			TriggerOnNext->TriggerFrame(StartFrame);
 		}
 
+		bool IsStateInitializationNeeded()
+		{
+			const ArrayType& InputArrayRef = *InputArray;
+			const bool bIsArrayNonEmpty = InputArrayRef.Num() != 0; // Skip reinit if the array is empty because it represents an invalid state for this node.
+			const bool bIsArraySizeChanged = PrevArraySize != InputArrayRef.Num(); // Need to reinit for array size changes. 
+			const bool bIsSharedStateEnablementInconsistent = (*bEnableSharedState != bSharedStateInitialized); // Need to reinit if the shared state enablement has been updated.
+			const bool bIsNonSharedStateInitializationNeeded = !*bEnableSharedState && !ArrayRandomGet.IsValid(); // For the first initialization of the non shared state random get (bIsSharedStateEnablementInconsistent will take care of that for shared state)
+
+			return (bIsArrayNonEmpty && (bIsArraySizeChanged || bIsSharedStateEnablementInconsistent || bIsNonSharedStateInitializationNeeded));
+		}
+
 		void InitializeState(int32 InArraySize)
 		{
 			bSharedStateInitialized = false;
