Minor pool clean-up & additional protections in light of review of last submission
#rnx
#tests AudioUnitTests, ran bug repro

diff --git a/Engine/Plugins/Runtime/Metasound/Source/MetasoundGenerator/Private/MetasoundGeneratorModule.cpp b/Engine/Plugins/Runtime/Metasound/Source/MetasoundGenerator/Private/MetasoundGeneratorModule.cpp
--- a/Engine/Plugins/Runtime/Metasound/Source/MetasoundGenerator/Private/MetasoundGeneratorModule.cpp
+++ b/Engine/Plugins/Runtime/Metasound/Source/MetasoundGenerator/Private/MetasoundGeneratorModule.cpp
@@ -55,12 +55,19 @@ namespace Metasound
 		OperatorInstanceCounterManager = MakeShared<FConcurrentInstanceCounterManager>(InstanceCounterCategory);
 	}
 
-	void FMetasoundGeneratorModule::ShutdownModule() 
+	void FMetasoundGeneratorModule::ShutdownModule()
 	{
-		// Have to cancel independent of resetting the shared pointer as tasks
-		// may still be holding on to pool and therefore keeping it alive
-		OperatorPool->CancelAllBuildEvents();
-		OperatorPool.Reset();
+		if (OperatorPool.IsValid())
+		{
+			TSharedPtr<FOperatorPool> PoolShuttingDown = OperatorPool;
+			OperatorPool.Reset();
+
+			// Clear the pool reference and cancel independent of resetting
+			// the shared pointer to ensure if any references are held elsewhere,
+			// they are properly invalidate.
+			PoolShuttingDown->CancelAllBuildEvents();
+		}
+
 		OperatorInstanceCounterManager.Reset();
 	}
 
diff --git a/Engine/Plugins/Runtime/Metasound/Source/MetasoundGenerator/Private/MetasoundOperatorCache.cpp b/Engine/Plugins/Runtime/Metasound/Source/MetasoundGenerator/Private/MetasoundOperatorCache.cpp
--- a/Engine/Plugins/Runtime/Metasound/Source/MetasoundGenerator/Private/MetasoundOperatorCache.cpp
+++ b/Engine/Plugins/Runtime/Metasound/Source/MetasoundGenerator/Private/MetasoundOperatorCache.cpp
@@ -267,7 +267,8 @@ namespace Metasound
 				{
 					FMetasoundGeneratorModule& Module = FModuleManager::GetModuleChecked<FMetasoundGeneratorModule>(ModuleName);
 					TSharedPtr<FOperatorPool> Pool = Module.GetOperatorPool();
-					if (Pool.IsValid())
+
+					if (Pool.IsValid() && !Pool->IsStopping())
 					{
 						Pool->RemoveBuildEvent(EventRef);
 					}
diff --git a/Engine/Plugins/Runtime/Metasound/Source/MetasoundGenerator/Public/MetasoundOperatorCache.h b/Engine/Plugins/Runtime/Metasound/Source/MetasoundGenerator/Public/MetasoundOperatorCache.h
--- a/Engine/Plugins/Runtime/Metasound/Source/MetasoundGenerator/Public/MetasoundOperatorCache.h
+++ b/Engine/Plugins/Runtime/Metasound/Source/MetasoundGenerator/Public/MetasoundOperatorCache.h
@@ -49,7 +49,6 @@ namespace Metasound
 	
 			void FirstUpdate();
 			void SetWindowLength(const float InNewLengthSeconds);
-			void AddResult(IntermediateResult&& InResultToAdd);
 			void ExpireResult(const IntermediateResult& InResultToExpire);
 			void TickResults(const float DeltaTimeSeconds);
 	
