Set automated metasound node tests to be picked up by Enginetest nightly
#rnx

diff --git a/Engine/Plugins/Runtime/Metasound/Source/MetasoundEngineTest/Private/EngineTestMetaSoundAutomatedNodeTest_Bind.cpp b/Engine/Plugins/Runtime/Metasound/Source/MetasoundEngineTest/Private/EngineTestMetaSoundAutomatedNodeTest_Bind.cpp
--- a/Engine/Plugins/Runtime/Metasound/Source/MetasoundEngineTest/Private/EngineTestMetaSoundAutomatedNodeTest_Bind.cpp
+++ b/Engine/Plugins/Runtime/Metasound/Source/MetasoundEngineTest/Private/EngineTestMetaSoundAutomatedNodeTest_Bind.cpp
@@ -19,7 +19,7 @@
 
 #if WITH_EDITORONLY_DATA
 
-IMPLEMENT_COMPLEX_AUTOMATION_TEST(FMetasoundAutomatedNodeTest_Bind, "Audio.Metasound.AutomatedNodeTest.Bind", EAutomationTestFlags::EditorContext | EAutomationTestFlags::StressFilter)
+IMPLEMENT_COMPLEX_AUTOMATION_TEST(FMetasoundAutomatedNodeTest_Bind, "Audio.Metasound.AutomatedNodeTest.Bind", EAutomationTestFlags::EditorContext | EAutomationTestFlags::EngineFilter)
 void FMetasoundAutomatedNodeTest_Bind::GetTests(TArray<FString>& OutBeautifiedNames, TArray<FString>& OutTestCommands) const
 {
 	using namespace Metasound::EngineTest;
diff --git a/Engine/Plugins/Runtime/Metasound/Source/MetasoundEngineTest/Private/EngineTestMetaSoundAutomatedNodeTest_Reset.cpp b/Engine/Plugins/Runtime/Metasound/Source/MetasoundEngineTest/Private/EngineTestMetaSoundAutomatedNodeTest_Reset.cpp
--- a/Engine/Plugins/Runtime/Metasound/Source/MetasoundEngineTest/Private/EngineTestMetaSoundAutomatedNodeTest_Reset.cpp
+++ b/Engine/Plugins/Runtime/Metasound/Source/MetasoundEngineTest/Private/EngineTestMetaSoundAutomatedNodeTest_Reset.cpp
@@ -18,7 +18,7 @@
 #if WITH_EDITORONLY_DATA
 
 
-IMPLEMENT_COMPLEX_AUTOMATION_TEST(FMetasoundAutomatedNodeTest_Reset, "Audio.Metasound.AutomatedNodeTest.Reset", EAutomationTestFlags::EditorContext | EAutomationTestFlags::StressFilter)
+IMPLEMENT_COMPLEX_AUTOMATION_TEST(FMetasoundAutomatedNodeTest_Reset, "Audio.Metasound.AutomatedNodeTest.Reset", EAutomationTestFlags::EditorContext | EAutomationTestFlags::EngineFilter)
 void FMetasoundAutomatedNodeTest_Reset::GetTests(TArray<FString>& OutBeautifiedNames, TArray<FString>& OutTestCommands) const
 {
 	using namespace Metasound::EngineTest;
diff --git a/Engine/Plugins/Runtime/Metasound/Source/MetasoundEngineTest/Private/EngineTestMetaSoundAutomatedNodeTest_ValidOutputs.cpp b/Engine/Plugins/Runtime/Metasound/Source/MetasoundEngineTest/Private/EngineTestMetaSoundAutomatedNodeTest_ValidOutputs.cpp
--- a/Engine/Plugins/Runtime/Metasound/Source/MetasoundEngineTest/Private/EngineTestMetaSoundAutomatedNodeTest_ValidOutputs.cpp
+++ b/Engine/Plugins/Runtime/Metasound/Source/MetasoundEngineTest/Private/EngineTestMetaSoundAutomatedNodeTest_ValidOutputs.cpp
@@ -18,7 +18,7 @@
 #if WITH_EDITORONLY_DATA
 
 
-IMPLEMENT_COMPLEX_AUTOMATION_TEST(FMetasoundAutomatedNodeTest_ValidOutputs, "Audio.Metasound.AutomatedNodeTest.ValidOutputs", EAutomationTestFlags::EditorContext | EAutomationTestFlags::StressFilter)
+IMPLEMENT_COMPLEX_AUTOMATION_TEST(FMetasoundAutomatedNodeTest_ValidOutputs, "Audio.Metasound.AutomatedNodeTest.ValidOutputs", EAutomationTestFlags::EditorContext | EAutomationTestFlags::EngineFilter)
 void FMetasoundAutomatedNodeTest_ValidOutputs::GetTests(TArray<FString>& OutBeautifiedNames, TArray<FString>& OutTestCommands) const
 {
 	using namespace Metasound::EngineTest;
