Adding constructor pin to grain delay node for max delay

diff --git a/Engine/Plugins/Runtime/Metasound/Source/MetasoundStandardNodes/Private/MetasoundGrainDelayNode.cpp b/Engine/Plugins/Runtime/Metasound/Source/MetasoundStandardNodes/Private/MetasoundGrainDelayNode.cpp
--- a/Engine/Plugins/Runtime/Metasound/Source/MetasoundStandardNodes/Private/MetasoundGrainDelayNode.cpp
+++ b/Engine/Plugins/Runtime/Metasound/Source/MetasoundStandardNodes/Private/MetasoundGrainDelayNode.cpp
@@ -32,6 +32,7 @@ namespace Metasound
 		METASOUND_PARAM(InGrainDelayEnvelope, "Grain Envelope", "The type of envelope to use for the grains.")
 		METASOUND_PARAM(InGrainMaxCount, "Max Grain Count", "The maximum number of grains to render at a time (between 1 and 100). More grains will cost more CPU and potentially clip.")
 		METASOUND_PARAM(InGrainDelayFeedbackAmount, "Feedback Amount", "The amount of feedback of each grain. The grain delay will feed its audio output back into itself.")
+		METASOUND_PARAM(InGrainMaxDelayTimeSeconds, "Max Delay Time", "The maximum amount of time to delay the audio.")
 
 		METASOUND_PARAM(OutAudio, "Out Audio", "Output audio buffer which has been grain-delayed.")
 	}
@@ -62,7 +63,9 @@ namespace Metasound
 			const FFloatReadRef& InGrainPitchShiftDelta,
 			const FGrainDelayEnvelopeReadRef& InGrainDelayEnvelope,
 			const FInt32ReadRef& InGrainMaxCount,	
-			const FFloatReadRef& InGrainDelayFeedbackAmount);
+			const FFloatReadRef& InGrainDelayFeedbackAmount,
+			float InGrainMaxDelaySeconds
+		);
 
 		static const FNodeClassMetadata& GetNodeInfo();
 		static const FVertexInterface& GetVertexInterface();
@@ -89,6 +92,7 @@ namespace Metasound
 		FGrainDelayEnvelopeReadRef GrainDelayEnvelope;
 		FInt32ReadRef GrainMaxCount;
 		FFloatReadRef GrainDelayFeedbackAmount;
+		float GrainMaxDelayTimeSeconds = 2.0f;
 
 		// Output references
 		FAudioBufferWriteRef AudioOutput;
@@ -113,7 +117,8 @@ namespace Metasound
 											  const FFloatReadRef& InGrainPitchShiftDelta,
 											  const FGrainDelayEnvelopeReadRef& InGrainDelayEnvelope,
 											  const FInt32ReadRef& InGrainMaxCount,
-											  const FFloatReadRef& InGrainDelayFeedbackAmount)
+											  const FFloatReadRef& InGrainDelayFeedbackAmount,
+											  float InGrainMaxDelayTimeSeconds)
 		: OperatorSettings(InSettings),
 	      AudioInput(InAudioInput),
 		  GrainSpawnTrigger(InGrainSpawnTrigger),
@@ -126,8 +131,9 @@ namespace Metasound
 		  GrainDelayEnvelope(InGrainDelayEnvelope),
 		  GrainMaxCount(InGrainMaxCount),
 		  GrainDelayFeedbackAmount(InGrainDelayFeedbackAmount),
+		  GrainMaxDelayTimeSeconds(InGrainMaxDelayTimeSeconds),
 	      AudioOutput(FAudioBufferWriteRef::CreateNew(OperatorSettings)),
-		  GrainDelayProcessor(OperatorSettings.GetSampleRate())
+		  GrainDelayProcessor(OperatorSettings.GetSampleRate(), GrainMaxDelayTimeSeconds)
 	{
 		PreviousGrainDelayMsec = GrainDelayProcessor.GetGrainDelayClamped(*GrainDelay);
 		PreviousGrainEnvelopeType = *InGrainDelayEnvelope;
@@ -165,6 +171,9 @@ namespace Metasound
 
 		auto CreateDefaultInterface = []()-> FVertexInterface
 		{
+			FDataVertexMetadata MaxDelayTimeMetadata = METASOUND_GET_PARAM_METADATA(InGrainMaxDelayTimeSeconds);
+			MaxDelayTimeMetadata.bIsAdvancedDisplay = true;
+
 			FInputVertexInterface InputInterface;
 			InputInterface.Add(TInputDataVertex<FAudioBuffer>(METASOUND_GET_PARAM_NAME_AND_METADATA(InAudio)));
 			InputInterface.Add(TInputDataVertex<FTrigger>(METASOUND_GET_PARAM_NAME_AND_METADATA(InGrainSpawnTrigger)));
@@ -177,8 +186,9 @@ namespace Metasound
 			InputInterface.Add(TInputDataVertex<FEnumGrainDelayEnvelope>(METASOUND_GET_PARAM_NAME_AND_METADATA(InGrainDelayEnvelope), (int32)Audio::Grain::EEnvelope::Gaussian));
 			InputInterface.Add(TInputDataVertex<int32>(METASOUND_GET_PARAM_NAME_AND_METADATA(InGrainMaxCount), 16));
 			InputInterface.Add(TInputDataVertex<float>(METASOUND_GET_PARAM_NAME_AND_METADATA(InGrainDelayFeedbackAmount), 0.0f));
+			InputInterface.Add(TInputConstructorVertex<FTime>(METASOUND_GET_PARAM_NAME(InGrainMaxDelayTimeSeconds), MaxDelayTimeMetadata, 2.0f));
 
-			FOutputVertexInterface OutputInterface;
+			FOutputVertexInterface OutputInterface; 
 			OutputInterface.Add(TOutputDataVertex<FAudioBuffer>(METASOUND_GET_PARAM_NAME_AND_METADATA(OutAudio)));
 
 			return FVertexInterface(InputInterface, OutputInterface);
@@ -205,6 +215,7 @@ namespace Metasound
 		FGrainDelayEnvelopeReadRef GrainDelayEnvelope = InputData.GetOrCreateDefaultDataReadReference<FEnumGrainDelayEnvelope>(METASOUND_GET_PARAM_NAME(InGrainDelayEnvelope), InParams.OperatorSettings);
 		FInt32ReadRef GrainMaxCount = InputData.GetOrCreateDefaultDataReadReference<int32>(METASOUND_GET_PARAM_NAME(InGrainMaxCount), InParams.OperatorSettings);
 		FFloatReadRef GrainDelayFeedbackAmount = InputData.GetOrCreateDefaultDataReadReference<float>(METASOUND_GET_PARAM_NAME(InGrainDelayFeedbackAmount), InParams.OperatorSettings);
+		FTime GrainMaxDelayTimeSeconds = InputData.GetOrCreateDefaultValue<FTime>(METASOUND_GET_PARAM_NAME(InGrainMaxDelayTimeSeconds), InParams.OperatorSettings);
 
 		return MakeUnique<FGrainDelayOperator>(InParams.OperatorSettings, 
 			AudioInput,
@@ -217,7 +228,8 @@ namespace Metasound
 			GrainPitchShiftDelta,
 			GrainDelayEnvelope,
 			GrainMaxCount,
-			GrainDelayFeedbackAmount);
+			GrainDelayFeedbackAmount,
+			GrainMaxDelayTimeSeconds.GetSeconds());
 	}
 
 	void FGrainDelayOperator::BindInputs(FInputVertexInterfaceData& InOutVertexData)
@@ -235,6 +247,7 @@ namespace Metasound
 		InOutVertexData.BindReadVertex(METASOUND_GET_PARAM_NAME(InGrainDelayEnvelope), GrainDelayEnvelope);
 		InOutVertexData.BindReadVertex(METASOUND_GET_PARAM_NAME(InGrainMaxCount), GrainMaxCount);
 		InOutVertexData.BindReadVertex(METASOUND_GET_PARAM_NAME(InGrainDelayFeedbackAmount), GrainDelayFeedbackAmount);
+		InOutVertexData.SetValue(METASOUND_GET_PARAM_NAME(InGrainMaxDelayTimeSeconds), FTime::FromSeconds(GrainMaxDelayTimeSeconds));
 	}
 
 	void FGrainDelayOperator::BindOutputs(FOutputVertexInterfaceData& InOutVertexData)
diff --git a/Engine/Source/Runtime/SignalProcessing/Private/GrainDelay.cpp b/Engine/Source/Runtime/SignalProcessing/Private/GrainDelay.cpp
--- a/Engine/Source/Runtime/SignalProcessing/Private/GrainDelay.cpp
+++ b/Engine/Source/Runtime/SignalProcessing/Private/GrainDelay.cpp
@@ -7,22 +7,23 @@ namespace Audio
 	namespace GrainDelay
 	{
 		// TODO: consider making this a construction pin or cvar
-		static constexpr float MaxDelaySeconds = 2.0f;
 		static constexpr int32 GrainEnvelopeSampleCount = 1024;
 		static constexpr float MaxAbsPitchShiftInOctaves = 6.0f;
 
-		FGrainDelay::FGrainDelay(const float InSampleRate)
+		FGrainDelay::FGrainDelay(const float InSampleRate, const float InMaxDelaySeconds)
 			: SampleRate(InSampleRate)
+			, MaxDelaySeconds(FMath::Max(InMaxDelaySeconds, 0.0f))
 		{
 			// DelayLine must accommodate the delay of the grain and the pitch shifter.
 			// The pitch shifter has a maximum delay and requires an extra millisecond 
 			// for phase offsets. An additional sample is added to account for floating 
 			// point rounding.
+
 			const float RoundingErrorBufferInSeconds = 1.f / InSampleRate;
 			const float MaxTapDelayInSeconds = (FTapDelayPitchShifter::MaxDelayLength + 1.f) / 1000.f;
 			const float DelayLineSizeInSeconds = MaxDelaySeconds + MaxTapDelayInSeconds + RoundingErrorBufferInSeconds;
 			GrainDelayLine.Init(SampleRate, DelayLineSizeInSeconds);
-			
+
 			// Generate the grain data
 			GenerateEnvelopeData(GrainEnvelope, GrainEnvelopeSampleCount, GrainEnvelopeType);
 			
diff --git a/Engine/Source/Runtime/SignalProcessing/Public/DSP/GrainDelay.h b/Engine/Source/Runtime/SignalProcessing/Public/DSP/GrainDelay.h
--- a/Engine/Source/Runtime/SignalProcessing/Public/DSP/GrainDelay.h
+++ b/Engine/Source/Runtime/SignalProcessing/Public/DSP/GrainDelay.h
@@ -13,7 +13,7 @@ namespace Audio
 		class FGrainDelay
 		{
 		public:
-			SIGNALPROCESSING_API FGrainDelay(const float InSampleRate);
+			SIGNALPROCESSING_API FGrainDelay(const float InSampleRate, const float InMaxDelaySeconds = 2.0f);
 			SIGNALPROCESSING_API ~FGrainDelay();
 
 			SIGNALPROCESSING_API void Reset();
@@ -112,6 +112,8 @@ namespace Audio
 			
 			// The base pitch shift ratio (pitch scale) of all grains in the grain delay
 			float PitchShiftRatioBase = 1.0f;
+
+			float MaxDelaySeconds = 2.0f;
 		};
 	}
 }
