Metasound Custom Pin Types and colors

#tests PIE

diff --git a/Engine/Plugins/Runtime/Metasound/Source/MetasoundEditor/Private/MetasoundEditorGraphBuilder.cpp b/Engine/Plugins/Runtime/Metasound/Source/MetasoundEditor/Private/MetasoundEditorGraphBuilder.cpp
--- a/Engine/Plugins/Runtime/Metasound/Source/MetasoundEditor/Private/MetasoundEditorGraphBuilder.cpp
+++ b/Engine/Plugins/Runtime/Metasound/Source/MetasoundEditor/Private/MetasoundEditorGraphBuilder.cpp
@@ -718,6 +718,12 @@ namespace Metasound
 				return Settings->WaveTablePinTypeColor;
 			}
 
+			// custom colors
+			if (const FLinearColor* Color = Settings->CustomPinTypeColors.Find(PinType.PinCategory))
+			{
+				return *Color;
+			}
+
 			return Settings->DefaultPinTypeColor;
 		}
 
diff --git a/Engine/Plugins/Runtime/Metasound/Source/MetasoundEditor/Private/MetasoundEditorModule.cpp b/Engine/Plugins/Runtime/Metasound/Source/MetasoundEditor/Private/MetasoundEditorModule.cpp
--- a/Engine/Plugins/Runtime/Metasound/Source/MetasoundEditor/Private/MetasoundEditorModule.cpp
+++ b/Engine/Plugins/Runtime/Metasound/Source/MetasoundEditor/Private/MetasoundEditorModule.cpp
@@ -529,7 +529,7 @@ namespace Metasound
 
 						// Types like triggers & AudioBuffer are specialized, so ignore their preferred
 						// literal types to classify the category.
-						if (!FGraphBuilder::IsPinCategoryMetaSoundCustomDataType(PinCategory))
+						if (!FGraphBuilder::IsPinCategoryMetaSoundCustomDataType(PinCategory) && !CustomPinCategories.Contains(PinCategory))
 						{
 							// Primitives
 							switch (RegistryInfo.PreferredLiteralType)
@@ -631,6 +631,26 @@ namespace Metasound
 				PinTypes.Emplace(InDataTypeName, MoveTemp(PinConfiguration));
 			}
 
+			virtual void RegisterCustomPinType(FName InDataTypeName, const FGraphPinParams& Params) override
+			{
+				RegisterPinType(InDataTypeName, Params.PinCategory, Params.PinSubcategory, Params.PinConnectedIcon, Params.PinDisconnectedIcon);
+				if (Params.PinCategory.IsNone())
+				{
+					return;
+				}
+				
+				if (FGraphBuilder::IsPinCategoryMetaSoundCustomDataType(InDataTypeName))
+				{
+					UE_LOG(LogMetasoundEditor, Warning, TEXT("Attempted to register a \"Custom Pin Type\": \"%s\", but this is already a Metasound Custom Data Type"), *InDataTypeName.ToString());
+					return;
+				}
+				
+				CustomPinCategories.Add(Params.PinCategory);
+				UMetasoundEditorSettings* Settings = GetMutableDefault<UMetasoundEditorSettings>();
+				Settings->CustomPinTypeColors.Add(Params.PinCategory, Params.PinColor ? *Params.PinColor : Settings->DefaultPinTypeColor);
+
+			}
+
 			void ShutdownAssetClassRegistry()
 			{
 				if (FAssetRegistryModule* AssetRegistryModule = static_cast<FAssetRegistryModule*>(FModuleManager::Get().GetModule("AssetRegistry")))
@@ -968,6 +988,7 @@ namespace Metasound
 			TArray<TSharedPtr<FAssetTypeActions_Base>> AssetActions;
 			TMap<EMetasoundFrontendLiteralType, const TSubclassOf<UMetasoundEditorGraphMemberDefaultLiteral>> InputDefaultLiteralClassRegistry;
 			TMap<FName, FGraphPinConfiguration> PinTypes;
+			TSet<FName> CustomPinCategories;
 
 			TMap<UClass*, TUniquePtr<IMemberDefaultLiteralCustomizationFactory>> LiteralCustomizationFactories;
 
diff --git a/Engine/Plugins/Runtime/Metasound/Source/MetasoundEditor/Private/MetasoundEditorSettings.h b/Engine/Plugins/Runtime/Metasound/Source/MetasoundEditor/Private/MetasoundEditorSettings.h
--- a/Engine/Plugins/Runtime/Metasound/Source/MetasoundEditor/Private/MetasoundEditorSettings.h
+++ b/Engine/Plugins/Runtime/Metasound/Source/MetasoundEditor/Private/MetasoundEditorSettings.h
@@ -86,6 +86,9 @@ public:
 	UPROPERTY(EditAnywhere, config, Category=General)
 	FString DefaultAuthor;
 
+	/** Maps Pin Category To Pin Color */
+	TMap<FName, FLinearColor> CustomPinTypeColors;
+	
 	/** Default pin type color */
 	UPROPERTY(EditAnywhere, config, Category=PinColors)
 	FLinearColor DefaultPinTypeColor;
diff --git a/Engine/Plugins/Runtime/Metasound/Source/MetasoundEditor/Public/MetasoundEditorModule.h b/Engine/Plugins/Runtime/Metasound/Source/MetasoundEditor/Public/MetasoundEditorModule.h
--- a/Engine/Plugins/Runtime/Metasound/Source/MetasoundEditor/Public/MetasoundEditorModule.h
+++ b/Engine/Plugins/Runtime/Metasound/Source/MetasoundEditor/Public/MetasoundEditorModule.h
@@ -51,6 +51,15 @@ namespace Metasound
 			Complete = 3
 		};
 
+		struct FGraphPinParams
+		{
+			FName PinCategory = {};
+			FName PinSubcategory = {};
+			const FLinearColor* PinColor = nullptr;
+			const FSlateBrush* PinConnectedIcon = nullptr;
+			const FSlateBrush* PinDisconnectedIcon = nullptr;
+		};
+
 		class METASOUNDEDITOR_API FMetasoundDefaultLiteralCustomizationBase
 		{
 		protected:
@@ -107,6 +116,8 @@ namespace Metasound
 			
 			virtual void RegisterPinType(FName InDataTypeName, FName InPinCategory = { }, FName InPinSubCategory = { },
 				const FSlateBrush* InPinConnectedIcon = nullptr, const FSlateBrush* InPinDisconnectedIcon = nullptr) = 0;
+			
+			virtual void RegisterCustomPinType(FName InDataTypeName, const FGraphPinParams& Params) = 0;
 		};
 	} // namespace Editor
 } // namespace Metasound
