Enable pasting objects and object arrays from the content browser and BP into MetaSound default values

diff --git a/Engine/Plugins/Runtime/Metasound/Source/MetasoundEditor/Private/MetasoundEditor.cpp b/Engine/Plugins/Runtime/Metasound/Source/MetasoundEditor/Private/MetasoundEditor.cpp
--- a/Engine/Plugins/Runtime/Metasound/Source/MetasoundEditor/Private/MetasoundEditor.cpp
+++ b/Engine/Plugins/Runtime/Metasound/Source/MetasoundEditor/Private/MetasoundEditor.cpp
@@ -54,6 +54,7 @@
 #include "MetasoundFrontendTransform.h"
 #include "MetasoundGenerator.h"
 #include "MetasoundLog.h"
+#include "MetasoundNodeDetailCustomization.h"
 #include "MetasoundSource.h"
 #include "MetasoundUObjectRegistry.h"
 #include "Misc/Attribute.h"
@@ -1195,6 +1196,8 @@ namespace Metasound
 
 			FPropertyEditorModule& PropertyModule = FModuleManager::LoadModuleChecked<FPropertyEditorModule>("PropertyEditor");
 			MetasoundDetails = PropertyModule.CreateDetailView(Args);
+			MetasoundDetails->SetExtensionHandler(MakeShared<FMetaSoundNodeExtensionHandler>());
+
 			// Set details selection to the MetaSound's source settings 
 			SetSelection({ Metasound });
 			InterfacesDetails = PropertyModule.CreateDetailView(Args);
diff --git a/Engine/Plugins/Runtime/Metasound/Source/MetasoundEditor/Private/MetasoundNodeDetailCustomization.cpp b/Engine/Plugins/Runtime/Metasound/Source/MetasoundEditor/Private/MetasoundNodeDetailCustomization.cpp
--- a/Engine/Plugins/Runtime/Metasound/Source/MetasoundEditor/Private/MetasoundNodeDetailCustomization.cpp
+++ b/Engine/Plugins/Runtime/Metasound/Source/MetasoundEditor/Private/MetasoundNodeDetailCustomization.cpp
@@ -104,6 +104,152 @@ namespace Metasound
 					? CreateElementTypeNameFromArrayTypeName(InDataTypeInfo.DataTypeName)
 					: InDataTypeInfo.DataTypeName;
 			}
+
+			// Paste execute action for object member default values
+			FExecuteAction CreateDefaultValueObjectPasteExecuteAction(TSharedPtr<IPropertyHandle> PropertyHandle, bool bIsArray)
+			{
+				return FExecuteAction::CreateLambda([PropertyHandle = PropertyHandle, bIsArray = bIsArray]()
+				{
+					const FScopedTransaction Transaction(LOCTEXT("PasteObjectArrayProperty", "Paste Property"));
+
+					FString ClipboardValue;
+					FPlatformApplicationMisc::ClipboardPaste(ClipboardValue);
+					if (ClipboardValue.IsEmpty())
+					{
+						return;
+					}
+					
+					Frontend::FDataTypeRegistryInfo DataTypeInfo;
+					MemberCustomizationPrivate::GetDataTypeFromElementPropertyHandle(PropertyHandle, DataTypeInfo);
+					UClass* ProxyGenClass = DataTypeInfo.ProxyGeneratorClass;
+					if (!ProxyGenClass)
+					{
+						return;
+					}
+					FTopLevelAssetPath ClassPath = ProxyGenClass->GetClassPathName();
+
+					// Try to reformat string 
+					// Split into array of objects
+					TArray<FString> Values;
+					// Copying from other MS, still parse to verify object type 
+					// or copying from BP
+					if ((ClipboardValue.StartsWith("((") && ClipboardValue.EndsWith("))")) || 
+						(ClipboardValue.StartsWith("(\"") && ClipboardValue.EndsWith("\")")))
+					{
+						// Remove first and last parentheses
+						ClipboardValue.LeftChop(1).RightChop(1).ParseIntoArrayWS(Values, TEXT(","), true);
+					}
+					// Copying from content browser 
+					else
+					{
+						ClipboardValue.ParseIntoArrayWS(Values, TEXT(","), true);
+					}
+
+					if (Values.Num() > 0)
+					{
+						TStringBuilder<512> Builder;
+						if (bIsArray)
+						{
+							Builder << TEXT("(");
+						}
+
+						for (FString& Value : Values)
+						{
+							// Remove (Object= ) wrapper (other MetaSound case)
+							if (Value.Contains(TEXT("Object=")))
+							{
+								Value = Value.LeftChop(2).RightChop(9);
+							}
+							// Validate the class path (before the first ')
+							FString ValueClassPath = Value.Left(Value.Find("'"));
+							// Remove beginning quote (BP case)
+							if (ValueClassPath.StartsWith(TEXT("\"")))
+							{
+								ValueClassPath.RightChopInline(1);
+							}
+							
+							// Wrap objects in (Object=*)
+							if (ValueClassPath == ClassPath.ToString())
+							{
+								Builder << TEXT("(Object=");
+								Builder << Value;
+								Builder << TEXT("),");
+							}
+							else
+							{
+								UE_LOG(LogMetaSound, Warning, TEXT("Failed to paste object of type %s which does not match default value type %s"), *ValueClassPath, *ClassPath.ToString());
+								return;
+							}
+						}
+
+						// Remove last comma
+						if (Builder.Len() > 0)
+						{
+							Builder.RemoveSuffix(1);
+						}
+
+						if (bIsArray)
+						{
+							Builder << TEXT(")");
+						}
+
+						FString FormattedString = Builder.ToString();
+						PropertyHandle->SetValueFromFormattedString(FormattedString, EPropertyValueSetFlags::InstanceObjects);
+					}
+				});
+			}
+
+			// Create copy/paste actions for member default value for object and object array types
+			void CreateDefaultValueObjectCopyPasteActions(FDetailWidgetRow& InWidgetRow, TSharedPtr<IPropertyHandle> PropertyHandle, bool bIsArray)
+			{
+				// Copy action
+				FUIAction CopyAction;
+				CopyAction.ExecuteAction = FExecuteAction::CreateLambda([PropertyHandle = PropertyHandle]()
+				{
+					FString Value;
+					if (PropertyHandle->GetValueAsFormattedString(Value, PPF_Copy) == FPropertyAccess::Success)
+					{
+						FPlatformApplicationMisc::ClipboardCopy(*Value);
+					}
+				});
+
+				// Paste action
+				TArray<UObject*> OuterObjects;
+				UMetasoundEditorGraphMember* GraphMember = nullptr;
+				PropertyHandle->GetOuterObjects(OuterObjects);
+				if (!OuterObjects.IsEmpty())
+				{
+					UObject* Outer = OuterObjects[0]->GetOuter();
+					GraphMember = Cast<UMetasoundEditorGraphMember>(Outer);
+				}
+
+				FUIAction PasteAction;
+				// Paste only enabled if graph is editable (for variables/outputs)
+				// or if graph is editable and input is not an interface member and is overridden (for inputs)
+				PasteAction.CanExecuteAction = FCanExecuteAction::CreateLambda([GraphMember = GraphMember]()
+				{
+					if (!GraphMember) { return false; }
+					const bool bIsGraphEditable = GraphMember->GetOwningGraph()->IsEditable();
+
+					if (UMetasoundEditorGraphInput* Input = Cast<UMetasoundEditorGraphInput>(GraphMember))
+					{
+						Frontend::FConstNodeHandle InputNodeHandle = Input->GetConstNodeHandle();
+						const TSet<FName>& InputsInheritingDefault = InputNodeHandle->GetOwningGraph()->GetInputsInheritingDefault();
+						FName NodeName = InputNodeHandle->GetNodeName();
+						return !Input->IsInterfaceMember() && (bIsGraphEditable || !InputsInheritingDefault.Contains(NodeName));
+					}
+					else
+					{
+						return bIsGraphEditable;
+					}
+					return false;
+				});
+
+				PasteAction.ExecuteAction = CreateDefaultValueObjectPasteExecuteAction(PropertyHandle, bIsArray);
+
+				InWidgetRow.CopyAction(CopyAction);
+				InWidgetRow.PasteAction(PasteAction);
+			}
 		} // namespace MemberCustomizationPrivate
 
 		FMetasoundFloatLiteralCustomization::~FMetasoundFloatLiteralCustomization()
@@ -631,6 +777,8 @@ PRAGMA_ENABLE_DEPRECATION_WARNINGS
 			[
 				ValueWidget
 			];
+
+			MemberCustomizationPrivate::CreateDefaultValueObjectCopyPasteActions(ValueRow, StructPropertyHandle, /*bIsArray=*/false);
 		}
 
 		void FMetasoundDefaultMemberElementDetailCustomizationBase::CustomizeHeader(TSharedRef<IPropertyHandle> StructPropertyHandle, FDetailWidgetRow& HeaderRow, IPropertyTypeCustomizationUtils& StructCustomizationUtils)
@@ -1506,6 +1654,15 @@ PRAGMA_ENABLE_DEPRECATION_WARNINGS
 			return EVisibility::Collapsed;
 		}
 
+		bool FMetaSoundNodeExtensionHandler::IsPropertyExtendable(const UClass* InObjectClass, const IPropertyHandle& PropertyHandle) const
+		{
+			return InObjectClass == UMetasoundEditorGraphMemberDefaultObjectArray::StaticClass();
+		}
+
+		void FMetaSoundNodeExtensionHandler::ExtendWidgetRow(FDetailWidgetRow& InWidgetRow, const IDetailLayoutBuilder& InDetailBuilder, const UClass* InObjectClass, TSharedPtr<IPropertyHandle> PropertyHandle)
+		{
+			MemberCustomizationPrivate::CreateDefaultValueObjectCopyPasteActions(InWidgetRow, PropertyHandle, /*bIsArray=*/true);
+		}
 	} // namespace Editor
 } // namespace Metasound
 #undef LOCTEXT_NAMESPACE
diff --git a/Engine/Plugins/Runtime/Metasound/Source/MetasoundEditor/Private/MetasoundNodeDetailCustomization.h b/Engine/Plugins/Runtime/Metasound/Source/MetasoundEditor/Private/MetasoundNodeDetailCustomization.h
--- a/Engine/Plugins/Runtime/Metasound/Source/MetasoundEditor/Private/MetasoundNodeDetailCustomization.h
+++ b/Engine/Plugins/Runtime/Metasound/Source/MetasoundEditor/Private/MetasoundNodeDetailCustomization.h
@@ -9,6 +9,7 @@
 #include "EdGraph/EdGraphSchema.h"
 #include "Framework/Application/SlateApplication.h"
 #include "IDetailCustomization.h"
+#include "IDetailPropertyExtensionHandler.h"
 #include "IDetailPropertyRow.h"
 #include "IPropertyTypeCustomization.h"
 #include "Layout/Visibility.h"
@@ -380,6 +381,14 @@ namespace Metasound
 		};
 
 		using FMetasoundOutputDetailCustomization = FMetasoundVertexDetailCustomization;
+
+		class FMetaSoundNodeExtensionHandler : public IDetailPropertyExtensionHandler
+		{
+		public:
+			// IDetailPropertyExtensionHandler interface
+			virtual bool IsPropertyExtendable(const UClass* InObjectClass, const IPropertyHandle& PropertyHandle) const override;
+			virtual void ExtendWidgetRow(FDetailWidgetRow& InWidgetRow, const IDetailLayoutBuilder& InDetailBuilder, const UClass* InObjectClass, TSharedPtr<IPropertyHandle> PropertyHandle) override;
+		};
 	} // namespace Editor
 } // namespace Metasound
 #undef LOCTEXT_NAMESPACE
