[AudioBus] Fix for incorrect clamp at higher channel counts.

diff --git a/Engine/Plugins/Runtime/Metasound/Source/MetasoundEngine/Private/MetasoundAudioBusReaderNode.cpp b/Engine/Plugins/Runtime/Metasound/Source/MetasoundEngine/Private/MetasoundAudioBusReaderNode.cpp
--- a/Engine/Plugins/Runtime/Metasound/Source/MetasoundEngine/Private/MetasoundAudioBusReaderNode.cpp
+++ b/Engine/Plugins/Runtime/Metasound/Source/MetasoundEngine/Private/MetasoundAudioBusReaderNode.cpp
@@ -235,7 +235,7 @@ namespace Metasound
 					if (FAudioDevice* AudioDevice = ADM->GetAudioDeviceRaw(AudioDeviceId))
 					{
 						// Start the audio bus in case it's not already started
-						AudioBusChannels = static_cast<uint32>(FMath::Min(AudioBusProxy->NumChannels, static_cast<int32>(EAudioBusChannels::SevenPointOne)));
+						AudioBusChannels = static_cast<uint32>(FMath::Min(AudioBusProxy->NumChannels, static_cast<int32>(EAudioBusChannels::MaxChannelCount)));
 						AudioBusId = AudioBusProxy->AudioBusId;
 						
 						const Audio::FAudioBusKey AudioBusKey = Audio::FAudioBusKey(AudioBusId);
diff --git a/Engine/Source/Runtime/Engine/Classes/Sound/AudioBus.h b/Engine/Source/Runtime/Engine/Classes/Sound/AudioBus.h
--- a/Engine/Source/Runtime/Engine/Classes/Sound/AudioBus.h
+++ b/Engine/Source/Runtime/Engine/Classes/Sound/AudioBus.h
@@ -16,7 +16,8 @@ enum class EAudioBusChannels : uint8
 	Stereo = 1,
 	Quad = 3,
 	FivePointOne = 5,
-	SevenPointOne = 7
+	SevenPointOne = 7,
+	MaxChannelCount = 8
 };
 
 namespace AudioBusUtils
