[MetasoundEngine] Provide UnwatchOutput API as inverse of WatchOutput operation.

diff --git a/Engine/Plugins/Runtime/Metasound/Source/MetasoundEngine/Private/MetasoundGeneratorHandle.cpp b/Engine/Plugins/Runtime/Metasound/Source/MetasoundEngine/Private/MetasoundGeneratorHandle.cpp
--- a/Engine/Plugins/Runtime/Metasound/Source/MetasoundEngine/Private/MetasoundGeneratorHandle.cpp
+++ b/Engine/Plugins/Runtime/Metasound/Source/MetasoundEngine/Private/MetasoundGeneratorHandle.cpp
@@ -146,6 +146,33 @@ namespace Metasound
 		return WatchOutputInternal(OutputName, FWatchOutputUnifiedDelegate(OnOutputValueChanged), AnalyzerName, AnalyzerOutputName);
 	}
 
+	bool FMetasoundGeneratorHandle::UnwatchOutput(
+		FName OutputName,
+		const FOnMetasoundOutputValueChanged& OnOutputValueChanged,
+		FName AnalyzerName,
+		FName AnalyzerOutputName)
+	{
+		return UnwatchOutputInternal(OutputName, FWatchOutputUnifiedDelegate(OnOutputValueChanged), AnalyzerName, AnalyzerOutputName);
+	}
+
+	bool FMetasoundGeneratorHandle::UnwatchOutput(
+		FName OutputName,
+		const FOnMetasoundOutputValueChangedNative& OnOutputValueChanged,
+		FName AnalyzerName,
+		FName AnalyzerOutputName)
+	{
+		return UnwatchOutputInternal(OutputName, FWatchOutputUnifiedDelegate(OnOutputValueChanged), AnalyzerName, AnalyzerOutputName);
+	}
+
+	bool FMetasoundGeneratorHandle::UnwatchOutput(
+		FName OutputName,
+		const FDelegateHandle& OnOutputValueChanged,
+		FName AnalyzerName,
+		FName AnalyzerOutputName)
+	{
+		return UnwatchOutputInternal(OutputName, FWatchOutputUnifiedDelegate(OnOutputValueChanged), AnalyzerName, AnalyzerOutputName);
+	}
+
 	void FMetasoundGeneratorHandle::UpdateOutputWatchers()
 	{
 		METASOUND_LLM_SCOPE;
@@ -346,6 +373,58 @@ namespace Metasound
 			return false;
 		}
 
+		Frontend::FAnalyzerAddress AnalyzerAddress;
+		if (!TryCreateAnalyzerAddress(OutputName, AnalyzerName, AnalyzerOutputName, AnalyzerAddress))
+		{
+			return false;
+		}
+
+		// Create the watcher
+		CreateOutputWatcher(AnalyzerAddress, OnOutputValueChanged);
+
+		// Update the generator's analyzers if necessary
+		FixUpOutputWatchers();
+
+		return true;
+	}
+
+	bool FMetasoundGeneratorHandle::UnwatchOutputInternal(
+		const FName OutputName,
+		const FWatchOutputUnifiedDelegate& OnOutputValueChanged,
+		const FName AnalyzerName,
+		const FName AnalyzerOutputName)
+	{
+		METASOUND_LLM_SCOPE;
+		METASOUND_TRACE_CPUPROFILER_EVENT_SCOPE(FMetasoundGeneratorHandle::UnwatchOutputInternal);
+
+		check(IsInGameThread());
+
+		if (!IsValid())
+		{
+			return false;
+		}
+
+		Frontend::FAnalyzerAddress AnalyzerAddress;
+		if (!TryCreateAnalyzerAddress(OutputName, AnalyzerName, AnalyzerOutputName, AnalyzerAddress))
+		{
+			return false;
+		}
+
+		// Remove the watcher
+		RemoveOutputWatcher(AnalyzerAddress, OnOutputValueChanged);
+
+		// Update the generator's analyzers if necessary
+		FixUpOutputWatchers();
+
+		return true;
+	}
+
+	bool FMetasoundGeneratorHandle::TryCreateAnalyzerAddress(
+		const FName OutputName,
+		const FName AnalyzerName,
+		const FName AnalyzerOutputName,
+		Frontend::FAnalyzerAddress& OutAnalyzerAddress)
+	{
 		// Make the analyzer address.
 		Frontend::FAnalyzerAddress AnalyzerAddress;
 		AnalyzerAddress.InstanceID = GetAudioComponentId();
@@ -405,12 +484,7 @@ namespace Metasound
 			}
 		}
 
-		// Create the watcher
-		CreateOutputWatcher(AnalyzerAddress, OnOutputValueChanged);
-
-		// Update the generator's analyzers if necessary
-		FixUpOutputWatchers();
-
+		OutAnalyzerAddress = AnalyzerAddress;
 		return true;
 	}
 
@@ -430,13 +504,21 @@ namespace Metasound
 		{
 			// For each watcher, make sure the generator has a corresponding analyzer
 			// (will fail gracefully on duplicates or non-existent outputs)
+			// we can also remove any analyzer that has no further bindings
 			for (const auto& Watcher : OutputWatchers)
 			{
-				PinnedGenerator->AddOutputVertexAnalyzer(Watcher.Value.AnalyzerAddress);
+				if (Watcher.Value.OnOutputValueChanged.IsBound())
+				{
+					PinnedGenerator->AddOutputVertexAnalyzer(Watcher.Value.AnalyzerAddress);
+				}
+				else
+				{
+					PinnedGenerator->RemoveOutputVertexAnalyzer(Watcher.Value.AnalyzerAddress);
+				}
 			}
 		}
 	}
-	
+
 	void FMetasoundGeneratorHandle::CreateOutputWatcher(
 		const Frontend::FAnalyzerAddress& AnalyzerAddress,
 		const FWatchOutputUnifiedDelegate& OnOutputValueChanged)
@@ -465,6 +547,29 @@ namespace Metasound
 		}
 	}
 
+	void FMetasoundGeneratorHandle::RemoveOutputWatcher(
+		const Frontend::FAnalyzerAddress& AnalyzerAddress,
+		const FWatchOutputUnifiedDelegate& OnOutputValueChanged)
+	{
+		METASOUND_LLM_SCOPE;
+		METASOUND_TRACE_CPUPROFILER_EVENT_SCOPE(FMetasoundGeneratorHandle::RemoveOutputWatcher);
+
+		check(IsInGameThread()); // modifying watchers isn't thread-safe
+
+		// Watcher must exist in order to be removed
+		const FOutputWatcherKey WatcherKey
+		{
+			AnalyzerAddress.OutputName,
+			AnalyzerAddress.AnalyzerName,
+			AnalyzerAddress.AnalyzerMemberName
+		};
+
+		if (FOutputWatcher* Watcher = OutputWatchers.Find(WatcherKey))
+		{
+			Watcher->OnOutputValueChanged.Remove(OnOutputValueChanged);
+		}		
+	}
+
 	void FMetasoundGeneratorHandle::HandleGeneratorCreated(
 		const uint64 InAudioComponentId,
 		TSharedPtr<FMetasoundGenerator> InGenerator)
diff --git a/Engine/Plugins/Runtime/Metasound/Source/MetasoundEngine/Public/MetasoundGeneratorHandle.h b/Engine/Plugins/Runtime/Metasound/Source/MetasoundEngine/Public/MetasoundGeneratorHandle.h
--- a/Engine/Plugins/Runtime/Metasound/Source/MetasoundEngine/Public/MetasoundGeneratorHandle.h
+++ b/Engine/Plugins/Runtime/Metasound/Source/MetasoundEngine/Public/MetasoundGeneratorHandle.h
@@ -122,6 +122,33 @@ namespace Metasound
 			FName AnalyzerName = NAME_None,
 			FName AnalyzerOutputName = NAME_None);
 
+		/**
+		 * Unwatch/Stop watching an output value
+		 *
+		 * @param OutputName - The user-specified name of the output in the Metasound
+		 * @param OnOutputValueChanged - The event or handle previously watched
+		 * @param AnalyzerName - (optional) The name of the analyzer to use on the output, defaults to a passthrough
+		 * @param AnalyzerOutputName - (optional) The name of the output on the analyzer to watch, defaults to the passthrough output
+		 * @returns true if the unwatch setup succeeded, false otherwise
+		 */
+		bool UnwatchOutput(
+			FName OutputName,
+			const FOnMetasoundOutputValueChanged& OnOutputValueChanged,
+			FName AnalyzerName = NAME_None,
+			FName AnalyzerOutputName = NAME_None);
+
+		bool UnwatchOutput(
+			FName OutputName,
+			const FOnMetasoundOutputValueChangedNative& OnOutputValueChanged,
+			FName AnalyzerName = NAME_None,
+			FName AnalyzerOutputName = NAME_None);
+
+		bool UnwatchOutput(
+			FName OutputName,
+			const FDelegateHandle& OnOutputValueChanged,
+			FName AnalyzerName = NAME_None,
+			FName AnalyzerOutputName = NAME_None);
+
 		/**
 		 * Update any watched outputs
 		 */
@@ -165,10 +192,12 @@ namespace Metasound
 		{
 			FOnMetasoundOutputValueChanged WatchDelegate;
 			FOnMetasoundOutputValueChangedNative NativeWatchDelegate;
+			FDelegateHandle NativeWatchDelegateHandle;
 
 			FWatchOutputUnifiedDelegate() {}
 			FWatchOutputUnifiedDelegate(const FOnMetasoundOutputValueChanged& Delegate) : WatchDelegate(Delegate) {}
-			FWatchOutputUnifiedDelegate(const FOnMetasoundOutputValueChangedNative& Delegate) : NativeWatchDelegate(Delegate) {}
+			FWatchOutputUnifiedDelegate(const FOnMetasoundOutputValueChangedNative& Delegate) : NativeWatchDelegate(Delegate), NativeWatchDelegateHandle(Delegate.GetHandle()){}
+			FWatchOutputUnifiedDelegate(const FDelegateHandle& DelegateHandle) : NativeWatchDelegateHandle(DelegateHandle) {}
 		};
 
 		bool WatchOutputInternal(
@@ -177,12 +206,28 @@ namespace Metasound
 			FName AnalyzerName = NAME_None,
 			FName AnalyzerOutputName = NAME_None);
 
+		bool UnwatchOutputInternal(
+			FName OutputName,
+			const FWatchOutputUnifiedDelegate& OnOutputValueChanged,
+			FName AnalyzerName = NAME_None,
+			FName AnalyzerOutputName = NAME_None);
+
+		bool TryCreateAnalyzerAddress(
+			const FName OutputName,
+			const FName AnalyzerName,
+			const FName AnalyzerOutputName,
+			Frontend::FAnalyzerAddress& OutAnalyzerAddress);
+
 		void FixUpOutputWatchers();
 
 		void CreateOutputWatcher(
 			const Frontend::FAnalyzerAddress& AnalyzerAddress,
 			const FWatchOutputUnifiedDelegate& OnOutputValueChanged);
-		
+
+		void RemoveOutputWatcher(
+			const Frontend::FAnalyzerAddress& AnalyzerAddress,
+			const FWatchOutputUnifiedDelegate& OnOutputValueChanged);
+
 		TWeakObjectPtr<UAudioComponent> AudioComponent;
 		const uint64 AudioComponentId;
 
@@ -263,6 +308,34 @@ namespace Metasound
 				}
 			}
 
+			void Remove(const FWatchOutputUnifiedDelegate& Delegate)
+			{
+				if (Delegate.WatchDelegate.IsBound())
+				{
+					WatchDelegates.Remove(Delegate.WatchDelegate);
+				}
+
+				if (Delegate.NativeWatchDelegateHandle.IsValid())
+				{
+					NativeWatchDelegates.Remove(Delegate.NativeWatchDelegateHandle);
+				}
+			}
+
+			bool IsBound() const
+			{
+				if (WatchDelegates.IsBound())
+				{
+					return true;
+				}
+
+				if (NativeWatchDelegates.IsBound())
+				{
+					return true;
+				}
+
+				return false;
+			}
+
 			void Broadcast(FName OutputName, const FMetaSoundOutput& Output) const
 			{
 				WatchDelegates.Broadcast(OutputName, Output);
