Move operator cache to use async build pipe instead of task graph (follow-up to listed bug fix per initial review feedback)
#rnx

diff --git a/Engine/Plugins/Runtime/Metasound/Source/MetasoundEngine/Private/MetasoundOperatorCacheSubsystem.cpp b/Engine/Plugins/Runtime/Metasound/Source/MetasoundEngine/Private/MetasoundOperatorCacheSubsystem.cpp
--- a/Engine/Plugins/Runtime/Metasound/Source/MetasoundEngine/Private/MetasoundOperatorCacheSubsystem.cpp
+++ b/Engine/Plugins/Runtime/Metasound/Source/MetasoundEngine/Private/MetasoundOperatorCacheSubsystem.cpp
@@ -8,44 +8,41 @@
 #include "Modules/ModuleManager.h"
 #include "Misc/Optional.h"
 
-static TOptional<Metasound::FMetasoundGeneratorInitParams> CreateInitParams(UMetaSoundSource* InMetaSound, const FSoundGeneratorInitParams& InParams)
+namespace Metasound::OperatorCachePrivate
 {
-	using namespace Metasound;
-	using namespace Metasound::Frontend;
-	using namespace Metasound::Engine;
-	using namespace Metasound::SourcePrivate;
-
-	// InMetaSound was null
-	if (!ensure(InMetaSound))
-	{
-		return {};
+	TOptional<FMetasoundGeneratorInitParams> CreateInitParams(UMetaSoundSource& InMetaSound, const FSoundGeneratorInitParams& InParams)
+	{
+		using namespace Metasound::Frontend;
+		using namespace Metasound::Engine;
+		using namespace Metasound::SourcePrivate;
+
+		// Dynamic MetaSounds cannot be precached
+		if (InMetaSound.IsDynamic())
+		{
+			return { };
+		}
+
+		FOperatorSettings InSettings = InMetaSound.GetOperatorSettings(static_cast<FSampleRate>(InParams.SampleRate));
+		FMetasoundEnvironment Environment = InMetaSound.CreateEnvironment(InParams);
+
+		FOperatorBuilderSettings BuilderSettings = FOperatorBuilderSettings::GetDefaultSettings();
+		// Graph analyzer currently only enabled for preview sounds (but can theoretically be supported for all sounds)
+		BuilderSettings.bPopulateInternalDataReferences = InParams.bIsPreviewSound;
+
+		return FMetasoundGeneratorInitParams
+		{
+			InSettings,
+			MoveTemp(BuilderSettings),
+			{}, // Graph, retrieved from the FrontEnd Registry in FOperatorPool::BuildAndAddOperator()
+			Environment,
+			InMetaSound.GetName(),
+			InMetaSound.GetOutputAudioChannelOrder(),
+			{}, // DefaultParameters
+			true, // bBuildSynchronous
+			{} // DataChannel
+		};
 	}
-	// we cannot precache dynamic metasounds
-	if (!ensure(!InMetaSound->IsDynamic()))
-	{
-		return {};
-	}
-
-	FOperatorSettings InSettings = InMetaSound->GetOperatorSettings(static_cast<FSampleRate>(InParams.SampleRate));
-	FMetasoundEnvironment Environment = InMetaSound->CreateEnvironment(InParams);
-
-	FOperatorBuilderSettings BuilderSettings = FOperatorBuilderSettings::GetDefaultSettings();
-	// Graph analyzer currently only enabled for preview sounds (but can theoretically be supported for all sounds)
-	BuilderSettings.bPopulateInternalDataReferences = InParams.bIsPreviewSound;
-
-	return FMetasoundGeneratorInitParams
-	{
-		InSettings,
-		MoveTemp(BuilderSettings),
-		{}, // Graph, retrieved from the FrontEnd Registry in FOperatorPool::BuildAndAddOperator()
-		Environment,
-		InMetaSound->GetName(),
-		InMetaSound->GetOutputAudioChannelOrder(),
-		{}, // DefaultParameters
-		true, // bBuildSynchronous
-		{} // DataChannel
-	};
-}
+} // namespace Metasound::OperatorCachePrivate
 
 bool UMetaSoundCacheSubsystem::ShouldCreateSubsystem(UObject* Outer) const
 {
@@ -81,22 +78,19 @@ void UMetaSoundCacheSubsystem::Update()
 #endif // #if METASOUND_OPERATORCACHEPROFILER_ENABLED
 }
 
-void UMetaSoundCacheSubsystem::PrecacheMetaSound(UMetaSoundSource* InMetaSound, int32 InNumInstances)
+void UMetaSoundCacheSubsystem::PrecacheMetaSoundInternal(UMetaSoundSource* InMetaSound, int32 InNumInstances, bool bTouchExisting)
 {
 	using namespace Audio;
 	using namespace Metasound;
 
-	TSharedPtr<FOperatorPool> OperatorPool;
 	IMetasoundGeneratorModule* Module = FModuleManager::GetModulePtr<IMetasoundGeneratorModule>("MetasoundGenerator");
 	if (!ensure(Module))
 	{
 		return;
 	}
 
-	OperatorPool = Module->GetOperatorPool();
-	const FMixerDevice* MixerDevice = GetMixerDevice();
-
-	if (!ensure(MixerDevice && OperatorPool))
+	TSharedPtr<FOperatorPool> OperatorPool = Module->GetOperatorPool();
+	if (!ensure(OperatorPool))
 	{
 		return;
 	}
@@ -109,64 +103,40 @@ void UMetaSoundCacheSubsystem::PrecacheMetaSound(UMetaSoundSource* InMetaSound,
 
 	if (InNumInstances < 1)
 	{
-		UE_LOG(LogMetaSound, Error, TEXT("PrecacheMetaSound called with invaled NumInstances %i, ignoring request"), InNumInstances);
+		UE_LOG(LogMetaSound, Error, TEXT("PrecacheMetaSound called with invalid NumInstances %i, ignoring request"), InNumInstances);
 		return;
 	}
 
 	InMetaSound->InitResources();
-	BuildParams.GraphName = InMetaSound->GetOwningAssetName();
-
-	if (InMetaSound->IsDynamic())
-	{
-		return;
-	}
 
-	TOptional<FMetasoundGeneratorInitParams> InitParams = CreateInitParams(InMetaSound, BuildParams);
+	BuildParams.GraphName = InMetaSound->GetOwningAssetName();
+	TOptional<FMetasoundGeneratorInitParams> InitParams = OperatorCachePrivate::CreateInitParams(*InMetaSound, BuildParams);
 	if (!InitParams.IsSet())
 	{
 		return;
 	}
-	
 
 	TUniquePtr<FOperatorBuildData> Data = MakeUnique<FOperatorBuildData>(
 		  MoveTemp(InitParams.GetValue())
 		, InMetaSound->GetGraphRegistryKey()
 		, InMetaSound->AssetClassID
 		, InNumInstances
+		, bTouchExisting
 	);
 
 	OperatorPool->BuildAndAddOperator(MoveTemp(Data));
 }
 
-void UMetaSoundCacheSubsystem::TouchOrPrecacheMetaSound(UMetaSoundSource* InMetaSound, int32 InNumInstances)
+void UMetaSoundCacheSubsystem::PrecacheMetaSound(UMetaSoundSource* InMetaSound, int32 InNumInstances)
 {
-	using namespace Metasound;
-
-	TSharedPtr<FOperatorPool> OperatorPool;
-	IMetasoundGeneratorModule* Module = FModuleManager::GetModulePtr<IMetasoundGeneratorModule>("MetasoundGenerator");
-	if (!ensure(Module))
-	{
-		return;
-	}
-
-	OperatorPool = Module->GetOperatorPool();
-	if (!ensure(OperatorPool))
-	{
-		return;
-	}
-
-	// get the number of instances already in the cache
-	const int32 NumInCache = OperatorPool->GetNumCachedOperatorsWithAssetClassID(InMetaSound->AssetClassID);
-
-	// move pre-existing to the top of the cache
-	OperatorPool->TouchOperatorsViaAssetClassID(InMetaSound->AssetClassID, FMath::Min(NumInCache, InNumInstances));
+	constexpr bool bTouchExisting = false;
+	PrecacheMetaSoundInternal(InMetaSound, InNumInstances, bTouchExisting);
+}
 
-	// build the difference (InNumInstances - existing)
-	const int32 NumToBuild = InNumInstances - NumInCache;
-	if (NumToBuild > 0)
-	{
-		PrecacheMetaSound(InMetaSound, NumToBuild);
-	}
+void UMetaSoundCacheSubsystem::TouchOrPrecacheMetaSound(UMetaSoundSource* InMetaSound, int32 InNumInstances)
+{
+	constexpr bool bTouchExisting = true;
+	PrecacheMetaSoundInternal(InMetaSound, InNumInstances, bTouchExisting);
 }
 
 void UMetaSoundCacheSubsystem::RemoveCachedOperatorsForMetaSound(UMetaSoundSource* InMetaSound)
diff --git a/Engine/Plugins/Runtime/Metasound/Source/MetasoundEngine/Public/MetasoundOperatorCacheSubsystem.h b/Engine/Plugins/Runtime/Metasound/Source/MetasoundEngine/Public/MetasoundOperatorCacheSubsystem.h
--- a/Engine/Plugins/Runtime/Metasound/Source/MetasoundEngine/Public/MetasoundOperatorCacheSubsystem.h
+++ b/Engine/Plugins/Runtime/Metasound/Source/MetasoundEngine/Public/MetasoundOperatorCacheSubsystem.h
@@ -47,6 +47,8 @@ public:
 	void RemoveCachedOperatorsForMetaSound(UPARAM(DisplayName = "MetaSound Source") UMetaSoundSource* InMetaSound);
 
 private:
+	void PrecacheMetaSoundInternal(UMetaSoundSource* InMetaSound, int32 InNumInstances, bool bTouchExisting);
+
 	FSoundGeneratorInitParams BuildParams;
 }; //UMetaSoundCacheSubsystem
 
diff --git a/Engine/Plugins/Runtime/Metasound/Source/MetasoundGenerator/Private/MetasoundGeneratorModule.cpp b/Engine/Plugins/Runtime/Metasound/Source/MetasoundGenerator/Private/MetasoundGeneratorModule.cpp
--- a/Engine/Plugins/Runtime/Metasound/Source/MetasoundGenerator/Private/MetasoundGeneratorModule.cpp
+++ b/Engine/Plugins/Runtime/Metasound/Source/MetasoundGenerator/Private/MetasoundGeneratorModule.cpp
@@ -64,8 +64,8 @@ namespace Metasound
 
 			// Clear the pool reference and cancel independent of resetting
 			// the shared pointer to ensure if any references are held elsewhere,
-			// they are properly invalidate.
-			PoolShuttingDown->CancelAllBuildEvents();
+			// they are properly invalidated.
+			PoolShuttingDown->StopAsyncTasks();
 		}
 
 		OperatorInstanceCounterManager.Reset();
diff --git a/Engine/Plugins/Runtime/Metasound/Source/MetasoundGenerator/Private/MetasoundOperatorCache.cpp b/Engine/Plugins/Runtime/Metasound/Source/MetasoundGenerator/Private/MetasoundOperatorCache.cpp
--- a/Engine/Plugins/Runtime/Metasound/Source/MetasoundGenerator/Private/MetasoundOperatorCache.cpp
+++ b/Engine/Plugins/Runtime/Metasound/Source/MetasoundGenerator/Private/MetasoundOperatorCache.cpp
@@ -165,51 +165,58 @@ namespace Metasound
 		, Frontend::FGraphRegistryKey InRegistryKey
 		, FGuid InAssetClassID
 		, int32 InNumInstances
+		, bool bInTouchExisting
 	)
 	: InitParams(InInitParams)
 	, RegistryKey(InRegistryKey)
 	, AssetClassID(InAssetClassID)
 	, NumInstances(InNumInstances)
+	, bTouchExisting(bInTouchExisting)
 	{
 	}
 
 
 	FOperatorPool::FOperatorPool(const FOperatorPoolSettings& InSettings)
 	: Settings(InSettings)
+	, AsyncBuildPipe(UE_SOURCE_LOCATION)
 	{
 	}
 
 	FOperatorPool::~FOperatorPool()
 	{
-		CancelAllBuildEvents();
+		StopAsyncTasks();
 	}
 
 	FOperatorAndInputs FOperatorPool::ClaimOperator(const FGuid& InOperatorID)
 	{
-		FScopeLock Lock(&CriticalSection);
-
 		FOperatorAndInputs OpAndInputs;
 
-		bool bCacheHit = false;
-		if (TArray<FOperatorAndInputs>* OperatorsWithID = Operators.Find(InOperatorID))
+		if (!IsStopping())
 		{
-			if (OperatorsWithID->Num() > 0)
+			FScopeLock Lock(&CriticalSection);
+
+			bool bCacheHit = false;
+			if (TArray<FOperatorAndInputs>* OperatorsWithID = Operators.Find(InOperatorID))
 			{
-				OpAndInputs = OperatorsWithID->Pop();
-				Stack.RemoveAt(Stack.FindLast(InOperatorID));
+				if (OperatorsWithID->Num() > 0)
+				{
+					OpAndInputs = OperatorsWithID->Pop();
+					Stack.RemoveAt(Stack.FindLast(InOperatorID));
 #if METASOUND_OPERATORCACHEPROFILER_ENABLED
-				OperatorPoolPrivate::CacheHitCount++;
-				bCacheHit = true;
-				TRACE_COUNTER_DECREMENT(MetaSound_OperatorPool_NumOperators);
+					OperatorPoolPrivate::CacheHitCount++;
+					bCacheHit = true;
+					TRACE_COUNTER_DECREMENT(MetaSound_OperatorPool_NumOperators);
 #endif
+				}
 			}
-		}
 
 #if METASOUND_OPERATORCACHEPROFILER_ENABLED
-		bCacheHit? HitRateTracker.AddHit() : HitRateTracker.AddMiss();
-		OperatorPoolPrivate::CacheAttemptCount++;
-		TRACE_COUNTER_SET(MetaSound_OperatorPool_HitRatio, OperatorPoolPrivate::GetHitRatio());
+			bCacheHit ? HitRateTracker.AddHit() : HitRateTracker.AddMiss();
+			OperatorPoolPrivate::CacheAttemptCount++;
+			TRACE_COUNTER_SET(MetaSound_OperatorPool_HitRatio, OperatorPoolPrivate::GetHitRatio());
 #endif
+		}
+
 		return OpAndInputs;
 	}
 
@@ -218,7 +225,52 @@ namespace Metasound
 		AddOperator(InOperatorID, { MoveTemp(InOperator), MoveTemp(InInputData) });
 	}
 
-	void FOperatorPool::AddOperator(const FGuid& InOperatorID, FOperatorAndInputs && OperatorAndInputs)
+	bool FOperatorPool::ExecuteTaskAsync(FOperatorPool::FTaskFunction&& InFunction)
+	{
+		using namespace UE::Tasks;
+
+		if (IsStopping())
+		{
+			return false;
+		}
+
+		TWeakPtr<FOperatorPool> WeakOpPool = AsShared();
+		const int32 TaskId = ++LastTaskId;
+		FTask NewTask = AsyncBuildPipe.Launch(UE_SOURCE_LOCATION, [WeakOpPool, TaskId, PoolFunction = MoveTemp(InFunction)]() mutable
+		{
+			PoolFunction(TaskId, WeakOpPool);
+
+			if (TSharedPtr<FOperatorPool> ThisPool = WeakOpPool.Pin())
+			{
+				FScopeLock Lock(&ThisPool->CriticalSection);
+				ThisPool->ActiveBuildTasks.Remove(TaskId);
+			}
+		});
+
+		{
+			FScopeLock Lock(&CriticalSection);
+			ActiveBuildTasks.Add(TaskId, MoveTemp(NewTask));
+		}
+
+		return true;
+	}
+
+	void FOperatorPool::AddOperator(const FGuid& InOperatorID, FOperatorAndInputs&& OperatorAndInputs)
+	{
+		using namespace UE::Tasks;
+
+		check(OperatorAndInputs.Operator.IsValid());
+
+		ExecuteTaskAsync([OperatorID = InOperatorID, OpAndInputs = MoveTemp(OperatorAndInputs)](FTaskId, TWeakPtr<FOperatorPool> WeakPoolPtr) mutable
+		{
+			if (TSharedPtr<FOperatorPool> Pool = WeakPoolPtr.Pin())
+			{
+				Pool->AddOperatorInternal(OperatorID, MoveTemp(OpAndInputs));
+			}
+		});
+	}	
+
+	void FOperatorPool::AddOperatorInternal(const FGuid& InOperatorID, FOperatorAndInputs && OperatorAndInputs)
 	{
 		if (!OperatorAndInputs.Operator.IsValid())
 		{
@@ -247,83 +299,54 @@ namespace Metasound
 		Trim();
 	}
 
-	void FOperatorPool::BuildAndAddAsync(TUniqueFunction<void()>&& InBuildFunc)
+	bool FOperatorPool::IsStopping() const
 	{
-		struct FBuildAndAddOpTask : public FAsyncGraphTaskBase
-		{
-			TUniqueFunction<void()> TaskFunc;
-
-			FBuildAndAddOpTask(TUniqueFunction<void()>&& InTaskFunc)
-				: TaskFunc(MoveTemp(InTaskFunc))
-			{
-			}
-
-			void DoTask(ENamedThreads::Type, const FGraphEventRef& EventRef)
-			{
-				TaskFunc();
-
-				const FName ModuleName = TEXT("MetasoundGenerator");
-				if (FModuleManager::Get().IsModuleLoaded(ModuleName))
-				{
-					FMetasoundGeneratorModule& Module = FModuleManager::GetModuleChecked<FMetasoundGeneratorModule>(ModuleName);
-					TSharedPtr<FOperatorPool> Pool = Module.GetOperatorPool();
+		return bStopping.load();
+	}
 
-					if (Pool.IsValid() && !Pool->IsStopping())
-					{
-						Pool->RemoveBuildEvent(EventRef);
-					}
-				}
-			}
+	void FOperatorPool::CancelAllBuildEvents()
+	{
+		StopAsyncTasks();
+	}
 
-			ENamedThreads::Type GetDesiredThread()
-			{
-				return ENamedThreads::AnyThread;
-			}
-		};
+	void FOperatorPool::StopAsyncTasks()
+	{
+		using namespace UE::Tasks;
 
-		FGraphEventRef EventRef = TGraphTask<FBuildAndAddOpTask>::CreateTask()
-			.ConstructAndDispatchWhenReady(MoveTemp(InBuildFunc));
+		bStopping.store(true);
 
+		// Move tasks to local copy in crit section to allow for safe mutation
+		// of ActiveBuildTasks from within tasks and avoid deadlocks with
+		// mutation of other pool resources while canceling remaining tasks.
+		TMap<FTaskId, FTask> TasksToCancel;
 		{
 			FScopeLock Lock(&CriticalSection);
-			ActiveBuildEvents.Add(MoveTemp(EventRef));
+			TasksToCancel = MoveTemp(ActiveBuildTasks);
+			ActiveBuildTasks.Reset();
 		}
-	}
 
-	void FOperatorPool::CancelAllBuildEvents()
-	{
-		if (!ActiveBuildEvents.IsEmpty())
+		if (!TasksToCancel.IsEmpty())
 		{
 			UE_LOG(LogMetasoundGenerator, Display, TEXT("Cancelling active MetaSound Cache Pool Operator build requests..."));
 
-			bStopping.store(true);
-			for (FGraphEventRef& EventRef : ActiveBuildEvents)
+			for (TPair<FTaskId, FTask>& Pair : TasksToCancel)
 			{
-				if (EventRef.IsValid())
+				FTask& TaskToCancel = Pair.Value;
+				if (!TaskToCancel.IsCompleted())
 				{
-					EventRef->Wait();
+					TaskToCancel.BusyWait();
 				}
 			}
-			ActiveBuildEvents.Reset();
-			bStopping.store(false);
 		}
-	}
 
-	void FOperatorPool::RemoveBuildEvent(const FGraphEventRef& InEventRef)
-	{
-		FScopeLock Lock(&CriticalSection);
-		ActiveBuildEvents.Remove(InEventRef);
+		bStopping.store(false);
 	}
 
 	void FOperatorPool::BuildAndAddOperator(TUniquePtr<FOperatorBuildData> InBuildData)
 	{
+		using namespace UE::Tasks;
 		using namespace OperatorPoolPrivate;
 
-		if (bStopping.load())
-		{
-			return;
-		}
-
 		if (!ensure(InBuildData))
 		{
 			return;
@@ -342,8 +365,7 @@ namespace Metasound
 		}
 
 		// Build operations should never keep the operator pool alive as this can delay app shutdown arbitrarily.
-		TWeakPtr<FOperatorPool> WeakOpPool = AsShared();
-		BuildAndAddAsync([Graph, PreCacheData = MoveTemp(InBuildData), WeakOpPool]()
+		ExecuteTaskAsync([Graph, PreCacheData = MoveTemp(InBuildData)](FTaskId, TWeakPtr<FOperatorPool> WeakPoolPtr)
 		{
 			using namespace OperatorPoolPrivate;
 
@@ -378,13 +400,30 @@ namespace Metasound
 				return;
 			}
 
-			const int32 NumInstances = PreCacheData->NumInstances;
-			for (int32 i = 0; i < NumInstances; ++i)
+			int32 NumToBuild = PreCacheData->NumInstances;
+
+			if (PreCacheData->bTouchExisting)
 			{
-				TSharedPtr<FOperatorPool> OperatorPool = WeakOpPool.Pin();
+				TSharedPtr<FOperatorPool> OperatorPool = WeakPoolPtr.Pin();
+				if (OperatorPool.IsValid())
+				{
+					// Get the number of instances already in the cache & move pre-existing to the top of the cache
+					const int32 NumInCache = OperatorPool->GetNumCachedOperatorsWithAssetClassID(PreCacheData->AssetClassID);
+					OperatorPool->TouchOperatorsViaAssetClassID(PreCacheData->AssetClassID, FMath::Min(NumInCache, NumToBuild));
+					NumToBuild -= NumInCache;
+				}
+			}
+
+			for (int32 i = 0; i < NumToBuild; ++i)
+			{
+				// These build operations can take a fair bit of time, so
+				// check continually for the validity of the operator pool
+				// on each build request to abort if necessary if cancellation
+				// is requested.
+				TSharedPtr<FOperatorPool> OperatorPool = WeakPoolPtr.Pin();
 				if (!OperatorPool.IsValid() || OperatorPool->IsStopping())
 				{
-					break;
+					return;
 				}
 
 				FBuildResults BuildResults;
@@ -392,55 +431,73 @@ namespace Metasound
 				GeneratorBuilder::LogBuildErrors(PreCacheData->InitParams.MetaSoundName, BuildResults);
 
 				const FGuid& GraphID = PreCacheData->InitParams.Graph->GetInstanceID();
-				OperatorPool->AddOperator(GraphID, MoveTemp(OperatorAndInputs));
-				OperatorPool->AddAssetIdToGraphIdLookUp(PreCacheData->AssetClassID, GraphID);
+				OperatorPool->AddOperatorInternal(GraphID, MoveTemp(OperatorAndInputs));
+				OperatorPool->AddAssetIdToGraphIdLookUpInternal(PreCacheData->AssetClassID, GraphID);
 			}
 		});
 	}
 
-	void FOperatorPool::TouchOperators(const FGuid& InOpeoratorID, const int32& NumToTouch)
+	void FOperatorPool::TouchOperators(const FGuid& InOperatorID, int32 NumToTouch)
 	{
-		const int32 NumToMove = FMath::Min(NumToTouch, GetNumCachedOperatorsWithID(InOpeoratorID));
-		if (!NumToMove)
+		if (!IsStopping())
 		{
-			return;
-		}
+			FScopeLock Lock(&CriticalSection);
 
-		FScopeLock Lock(&CriticalSection);
+			const int32 NumToMove = FMath::Min(NumToTouch, GetNumCachedOperatorsWithID(InOperatorID));
+			if (!NumToMove)
+			{
+				return;
+			}
 
-		// add to the "top" (end)
-		for (int32 i = 0; i < NumToMove; ++i)
-		{
-			Stack.Add(InOpeoratorID);
-		}
+			// add to the "top" (end)
+			for (int32 i = 0; i < NumToMove; ++i)
+			{
+				Stack.Add(InOperatorID);
+			}
 
-		// remove from the "bottom" (begining)
-		for (int32 i = 0; i < NumToMove; ++i)
-		{
-			Stack.RemoveSingle(InOpeoratorID);
+			// remove from the "bottom" (beginning)
+			for (int32 i = 0; i < NumToMove; ++i)
+			{
+				Stack.RemoveSingle(InOperatorID);
+			}
 		}
-
 	}
 
-	void FOperatorPool::TouchOperatorsViaAssetClassID(const FGuid& InAssetClassID, const int32& NumToTouch)
+	void FOperatorPool::TouchOperatorsViaAssetClassID(const FGuid& InAssetClassID, int32 NumToTouch)
 	{
-		FScopeLock Lock(&CriticalSection);
-		FGuid* GraphIdPtr = AssetIdToGraphIdLookUp.Find(InAssetClassID);
-		if (GraphIdPtr)
+		if (!IsStopping())
 		{
-			TouchOperators(*GraphIdPtr, NumToTouch);
+			FScopeLock Lock(&CriticalSection);
+			FGuid* GraphIdPtr = AssetIdToGraphIdLookUp.Find(InAssetClassID);
+			if (GraphIdPtr)
+			{
+				TouchOperators(*GraphIdPtr, NumToTouch);
+			}
 		}
 	}
 
 	void FOperatorPool::RemoveOperatorsWithID(const FGuid& InOperatorID)
 	{
-		FScopeLock Lock(&CriticalSection);
-		Operators.Remove(InOperatorID);
-		const int32 NumRemoved = Stack.Remove(InOperatorID);
+		using namespace UE::Tasks;
+
+		if (IsStopping())
+		{
+			return;
+		}
+
+		ExecuteTaskAsync([OperatorID = InOperatorID](FTaskId, TWeakPtr<FOperatorPool> WeakPoolPtr)
+		{
+			if (TSharedPtr<FOperatorPool> Pool = WeakPoolPtr.Pin())
+			{
+				FScopeLock Lock(&Pool->CriticalSection);
+				Pool->Operators.Remove(OperatorID);
+				const int32 NumRemoved = Pool->Stack.Remove(OperatorID);
 
 #if METASOUND_OPERATORCACHEPROFILER_ENABLED
-		TRACE_COUNTER_SUBTRACT(MetaSound_OperatorPool_NumOperators, int64(NumRemoved));
+				TRACE_COUNTER_SUBTRACT(MetaSound_OperatorPool_NumOperators, int64(NumRemoved));
 #endif // #if METASOUND_OPERATORCACHEPROFILER_ENABLED
+			}
+		});
 	}
 
 	void FOperatorPool::RemoveOperatorsWithAssetClassID(const FGuid& InAssetClassID)
@@ -456,9 +513,12 @@ namespace Metasound
 		
 	void FOperatorPool::SetMaxNumOperators(uint32 InMaxNumOperators)
 	{
-		FScopeLock Lock(&CriticalSection);
-		Settings.MaxNumOperators = InMaxNumOperators;
-		Trim();
+		if (!IsStopping())
+		{
+			FScopeLock Lock(&CriticalSection);
+			Settings.MaxNumOperators = InMaxNumOperators;
+			Trim();
+		}
 	}
 
 	int32 FOperatorPool::GetNumCachedOperatorsWithID(const FGuid& InOperatorID) const
@@ -483,7 +543,7 @@ namespace Metasound
 		return 0;
 	}
 
-	void FOperatorPool::AddAssetIdToGraphIdLookUp(const FGuid& InAssetClassID, const FGuid& InOperatorID)
+	void FOperatorPool::AddAssetIdToGraphIdLookUpInternal(const FGuid& InAssetClassID, const FGuid& InOperatorID)
 	{
 		FScopeLock Lock(&CriticalSection);
 		AssetIdToGraphIdLookUp.Add(InAssetClassID, InOperatorID);
@@ -506,7 +566,6 @@ namespace Metasound
 			for (int32 i = 0; i < NumToTrim; i++)
 			{
 				UE_LOG(LogMetasoundGenerator, VeryVerbose, TEXT("Trimming operator with ID %s"), *LexToString(Stack[i]));
-				// Destructor called by TUniquePtr going out of scope.
 				TArray<FOperatorAndInputs>* OperatorArray = Operators.Find(Stack[i]);
 				if (ensure(OperatorArray))
 				{
@@ -528,4 +587,4 @@ namespace Metasound
 
 		// todo: prune AssetIdToGraphIdLookUp?
 	}
-} // namespace Metasound
\ No newline at end of file
+} // namespace Metasound
diff --git a/Engine/Plugins/Runtime/Metasound/Source/MetasoundGenerator/Public/MetasoundOperatorCache.h b/Engine/Plugins/Runtime/Metasound/Source/MetasoundGenerator/Public/MetasoundOperatorCache.h
--- a/Engine/Plugins/Runtime/Metasound/Source/MetasoundGenerator/Public/MetasoundOperatorCache.h
+++ b/Engine/Plugins/Runtime/Metasound/Source/MetasoundGenerator/Public/MetasoundOperatorCache.h
@@ -8,6 +8,8 @@
 #include "MetasoundGenerator.h"
 #include "Misc/Guid.h"
 #include "HAL/CriticalSection.h"
+#include "Tasks/Pipe.h"
+#include "Tasks/Task.h"
 #include "Templates/UniquePtr.h"
 
 #ifndef METASOUND_OPERATORCACHEPROFILER_ENABLED
@@ -69,12 +71,16 @@ namespace Metasound
 		FGuid AssetClassID;
 		int32 NumInstances;
 
+		// If true, touches existing assets and only builds remaining number if required
+		bool bTouchExisting = false; 
+
 		FOperatorBuildData() = delete;
 		FOperatorBuildData(
 			  FMetasoundGeneratorInitParams&& InInitParams
 			, Frontend::FGraphRegistryKey InRegistryKey
 			, FGuid InAssetID
 			, int32 InNumInstances = 1
+			, bool bInTouchExisting = false
 		);
 
 	}; // struct FOperatorPrecacheData
@@ -90,13 +96,14 @@ namespace Metasound
 		FOperatorAndInputs ClaimOperator(const FGuid& InOperatorID);
 
 		void AddOperator(const FGuid& InOperatorID, TUniquePtr<IOperator>&& InOperator, FInputVertexInterfaceData&& InputData);
-		void AddOperator(const FGuid& InOperatorID, FOperatorAndInputs && OperatorAndInputs);
+		void AddOperator(const FGuid& InOperatorID, FOperatorAndInputs&& OperatorAndInputs);
+
 		void BuildAndAddOperator(TUniquePtr<FOperatorBuildData> InBuildData);
 
-		void TouchOperators(const FGuid& InOpeoratorID, const int32& NumToTouch = 1);
-		void TouchOperatorsViaAssetClassID(const FGuid& InAssetClassID, const int32& NumToTouch = 1);
+		void TouchOperators(const FGuid& InOperatorID, int32 NumToTouch = 1);
+		void TouchOperatorsViaAssetClassID(const FGuid& InAssetClassID, int32 NumToTouch = 1);
 
-		bool IsStopping() const { return bStopping.load(); }
+		bool IsStopping() const;
 
 		void RemoveOperatorsWithID(const FGuid& InOperatorID);
 		void RemoveOperatorsWithAssetClassID(const FGuid& InAssetClassID);
@@ -104,19 +111,28 @@ namespace Metasound
 		int32 GetNumCachedOperatorsWithID(const FGuid& InOperatorID) const;
 		int32 GetNumCachedOperatorsWithAssetClassID(const FGuid& InAssetClassID) const;
 
-		void AddAssetIdToGraphIdLookUp(const FGuid& InAssetClassID, const FGuid& InOperatorID);
+		UE_DEPRECATED(5.5, "Adding id to look-up is now private implementation")
+		void AddAssetIdToGraphIdLookUp(const FGuid& InAssetClassID, const FGuid& InOperatorID) { }
 
 		void SetMaxNumOperators(uint32 InMaxNumOperators);
 #if METASOUND_OPERATORCACHEPROFILER_ENABLED
 		void UpdateHitRateTracker();
 #endif // #if METASOUND_OPERATORCACHEPROFILER_ENABLED
 
+		UE_DEPRECATED(5.5, "Use StopAsyncTasks")
 		void CancelAllBuildEvents();
 
+		void StopAsyncTasks();
+
+		using FTaskId = int32;
+		using FTaskFunction = TUniqueFunction<void(FOperatorPool::FTaskId, TWeakPtr<FOperatorPool>)>;
+
 	private:
-		void RemoveBuildEvent(const FGraphEventRef& InEventRef);
+		FTaskId LastTaskId = 0;
 
-		void BuildAndAddAsync(TUniqueFunction<void()>&& InBuildFunc);
+		void AddAssetIdToGraphIdLookUpInternal(const FGuid& InAssetClassID, const FGuid& InOperatorID);
+		void AddOperatorInternal(const FGuid& InOperatorID, FOperatorAndInputs&& OperatorAndInputs);
+		bool ExecuteTaskAsync(FTaskFunction&& InFunction);
 		void Trim();
 
 		FOperatorPoolSettings Settings;
@@ -126,14 +142,17 @@ namespace Metasound
 		OperatorPoolPrivate::FWindowedHitRate HitRateTracker;
 #endif // #if METASOUND_OPERATORCACHEPROFILER_ENABLED
 
+		// Notifies active build tasks to abort as soon as possible
+		// and gates additional build tasks from being added.
 		std::atomic<bool> bStopping;
-		TSet<FGraphEventRef> ActiveBuildEvents;
+
+		TMap<FTaskId, UE::Tasks::FTask> ActiveBuildTasks;
+		UE::Tasks::FPipe AsyncBuildPipe;
 
 		TMap<FGuid, TArray<FOperatorAndInputs>> Operators;
 		TMap<FGuid, FGuid> AssetIdToGraphIdLookUp;
 		TArray<FGuid> Stack;
 	};
-
 } // namespace Metasound
 
 
