Fix clamping on metasound block rate override
#rnx

diff --git a/Engine/Plugins/Runtime/Metasound/Source/MetasoundFrontend/Private/MetasoundAssetBase.cpp b/Engine/Plugins/Runtime/Metasound/Source/MetasoundFrontend/Private/MetasoundAssetBase.cpp
--- a/Engine/Plugins/Runtime/Metasound/Source/MetasoundFrontend/Private/MetasoundAssetBase.cpp
+++ b/Engine/Plugins/Runtime/Metasound/Source/MetasoundFrontend/Private/MetasoundAssetBase.cpp
@@ -193,7 +193,7 @@ namespace Metasound
 		
 		TRange<float> GetBlockRateClampRange()
 		{
-			return TRange<float>(100.f,1000.f);
+			return TRange<float>(1.f,1000.f);
 		}
 
 		TRange<int32> GetSampleRateClampRange()
