[MetaSound Operator Cache]  Added defense and warning against calling Remove without a MetaSound source
#swarm 32060562

diff --git a/Engine/Plugins/Runtime/Metasound/Source/MetasoundEngine/Private/MetasoundOperatorCacheSubsystem.cpp b/Engine/Plugins/Runtime/Metasound/Source/MetasoundEngine/Private/MetasoundOperatorCacheSubsystem.cpp
--- a/Engine/Plugins/Runtime/Metasound/Source/MetasoundEngine/Private/MetasoundOperatorCacheSubsystem.cpp
+++ b/Engine/Plugins/Runtime/Metasound/Source/MetasoundEngine/Private/MetasoundOperatorCacheSubsystem.cpp
@@ -143,6 +143,12 @@ void UMetaSoundCacheSubsystem::RemoveCachedOperatorsForMetaSound(UMetaSoundSourc
 {
 	using namespace Metasound;
 
+	if (!InMetaSound)
+	{
+		UE_LOG(LogMetaSound, Warning, TEXT("Remove Cached Operators called without being provided a MetaSound, ignoring request"));
+		return;
+	}
+
 	IMetasoundGeneratorModule* Module = FModuleManager::GetModulePtr<IMetasoundGeneratorModule>("MetasoundGenerator");
 	if (!ensure(Module))
 	{
