Fix frequency interpolation in MetaSound low pass filter node

diff --git a/Engine/Plugins/Runtime/Metasound/Source/MetasoundStandardNodes/Private/MetasoundBasicFilters.cpp b/Engine/Plugins/Runtime/Metasound/Source/MetasoundStandardNodes/Private/MetasoundBasicFilters.cpp
--- a/Engine/Plugins/Runtime/Metasound/Source/MetasoundStandardNodes/Private/MetasoundBasicFilters.cpp
+++ b/Engine/Plugins/Runtime/Metasound/Source/MetasoundStandardNodes/Private/MetasoundBasicFilters.cpp
@@ -633,7 +633,7 @@ namespace Metasound
 	void FOnePoleLowPassFilterOperator::Execute()
 	{
 		float ClampedFreq = FMath::Clamp(0.0f, *Frequency, SampleRate);
-		OnePoleLowPassFilter.StartFrequencyInterpolation(*Frequency, 1);
+		OnePoleLowPassFilter.StartFrequencyInterpolation(*Frequency, AudioInput->Num());
 		OnePoleLowPassFilter.ProcessAudioBuffer(AudioInput->GetData(), AudioOutput->GetData(), AudioInput->Num());
 		OnePoleLowPassFilter.StopFrequencyInterpolation();
 	}
