Refactor reroute templating in preparation of new templates coming in subsequent changes (i.e. inputs, maybe more node config templates, etc.)
#tests AudioUnitTests, engine tests, tested adding, removing, and auditioning reroutes

diff --git a/Engine/Plugins/Runtime/Metasound/Source/MetasoundEditor/Private/MetasoundEditor.cpp b/Engine/Plugins/Runtime/Metasound/Source/MetasoundEditor/Private/MetasoundEditor.cpp
--- a/Engine/Plugins/Runtime/Metasound/Source/MetasoundEditor/Private/MetasoundEditor.cpp
+++ b/Engine/Plugins/Runtime/Metasound/Source/MetasoundEditor/Private/MetasoundEditor.cpp
@@ -49,6 +49,7 @@
 #include "MetasoundFrontend.h"
 #include "MetasoundFrontendDocument.h"
 #include "MetasoundFrontendDocumentBuilder.h"
+#include "MetasoundFrontendNodeTemplateRegistry.h"
 #include "MetasoundFrontendRegistries.h"
 #include "MetasoundFrontendSearchEngine.h"
 #include "MetasoundFrontendTransform.h"
@@ -2586,22 +2587,31 @@ namespace Metasound
 					{
 						FMetasoundFrontendClass FrontendClass;
 						const FMetasoundFrontendClassName ClassName = ExternalNode->GetClassName();
-						if (ClassName == FRerouteNodeTemplate::ClassName)
+						if (const INodeTemplate* NodeTemplate = INodeTemplateRegistry::Get().FindTemplate(ClassName))
 						{
 							bool bIsValid = false;
 							if (!ExternalNode->Pins.IsEmpty())
 							{
-								const FName DataType = FGraphBuilder::GetPinDataType(ExternalNode->Pins.Last());
-								if (!DataType.IsNone())
+								FNodeTemplateGenerateInterfaceParams Params;
+								for (const UEdGraphPin* Pin : ExternalNode->Pins)
 								{
-									const FNodeRegistryKey& RerouteTemplateKey = FRerouteNodeTemplate::GetRegistryKey();
-									FMetasoundFrontendNodeInterface NodeInterface = FRerouteNodeTemplate::CreateNodeInterfaceFromDataType(DataType);
+									check(Pin);
 
-									FNodeHandle NodeHandle = MetasoundAsset->GetRootGraphHandle()->AddTemplateNode(RerouteTemplateKey, MoveTemp(NodeInterface));
-									FGraphBuilder::InitGraphNodeIDFromNodeHandle(NodeHandle, ExternalNode);
-
-									bIsValid = NodeHandle->IsValid();
+									FName PinDataType = FGraphBuilder::GetPinDataType(Pin);
+									if (Pin->Direction == EGPD_Input)
+									{
+										Params.InputsToConnect.Add(PinDataType);
+									}
+									else // EGPD_Output
+									{
+										Params.OutputsToConnect.Add(PinDataType);
+									}
 								}
+
+								FNodeHandle NodeHandle = MetasoundAsset->GetRootGraphHandle()->AddTemplateNode(*NodeTemplate, MoveTemp(Params));
+								FGraphBuilder::InitGraphNodeIDFromNodeHandle(NodeHandle, ExternalNode);
+
+								bIsValid = NodeHandle->IsValid();
 							}
 
 							if (!bIsValid)
diff --git a/Engine/Plugins/Runtime/Metasound/Source/MetasoundEditor/Private/MetasoundEditorGraphBuilder.cpp b/Engine/Plugins/Runtime/Metasound/Source/MetasoundEditor/Private/MetasoundEditorGraphBuilder.cpp
--- a/Engine/Plugins/Runtime/Metasound/Source/MetasoundEditor/Private/MetasoundEditorGraphBuilder.cpp
+++ b/Engine/Plugins/Runtime/Metasound/Source/MetasoundEditor/Private/MetasoundEditorGraphBuilder.cpp
@@ -753,7 +753,26 @@ namespace Metasound
 
 		Frontend::FConstInputHandle FGraphBuilder::GetConstInputHandleFromPin(const UEdGraphPin* InPin)
 		{
-			return GetInputHandleFromPin(InPin);
+			using namespace Frontend;
+			using namespace VariableNames;
+
+			if (InPin && ensure(InPin->Direction == EGPD_Input))
+			{
+				if (const UMetasoundEditorGraphVariableNode* EdVariableNode = Cast<UMetasoundEditorGraphVariableNode>(InPin->GetOwningNode()))
+				{
+					// UEdGraphPins on variable nodes use the variable's name for display
+					// purposes instead of the underlying vertex's name. The frontend vertices
+					// of a variable node have consistent names no matter what the 
+					// variable is named.
+					return EdVariableNode->GetConstNodeHandle()->GetConstInputWithVertexName(METASOUND_GET_PARAM_NAME(InputData));
+				}
+				else if (const UMetasoundEditorGraphNode* EdNode = CastChecked<UMetasoundEditorGraphNode>(InPin->GetOwningNode()))
+				{
+					return EdNode->GetConstNodeHandle()->GetConstInputWithVertexName(InPin->GetFName());
+				}
+			}
+
+			return IInputController::GetInvalidHandle();
 		}
 
 		FName FGraphBuilder::GetPinDataType(const UEdGraphPin* InPin)
@@ -777,6 +796,133 @@ namespace Metasound
 			return { };
 		}
 
+		FMetasoundFrontendVertexHandle FGraphBuilder::GetPinVertexHandle(const FMetaSoundFrontendDocumentBuilder& InBuilder, const UEdGraphPin* InPin)
+		{
+			using namespace VariableNames;
+
+			if (!InPin)
+			{
+				return { };
+			}
+
+			const UMetasoundEditorGraphNode* OwningNode = CastChecked<UMetasoundEditorGraphNode>(InPin->GetOwningNode());
+
+			const FGuid NodeID = OwningNode->GetNodeID();
+			const FMetasoundFrontendNode* Node = InBuilder.FindNode(NodeID);
+			if (!Node)
+			{
+				return { };
+			}
+
+			const FMetasoundFrontendClass* Class = InBuilder.FindDependency(Node->ClassID);
+			if (!Class)
+			{
+				return { };
+			}
+
+			const FMetasoundFrontendVertex* Vertex = nullptr;
+			switch (Class->Metadata.GetType())
+			{
+				case EMetasoundFrontendClassType::Variable:
+				case EMetasoundFrontendClassType::VariableAccessor:
+				case EMetasoundFrontendClassType::VariableDeferredAccessor:
+				case EMetasoundFrontendClassType::VariableMutator:
+				{
+					// All variables nodes use the same pin name for user-modifiable node
+					// inputs and outputs and the editor does not display the pin's name. The
+					// editor instead displays the variable's name in place of the pin name to
+					// maintain a consistent look and behavior to input and output nodes.
+					Vertex = InPin->Direction == EGPD_Input
+						? InBuilder.FindNodeInput(NodeID, METASOUND_GET_PARAM_NAME(InputData))
+						: InBuilder.FindNodeOutput(NodeID, METASOUND_GET_PARAM_NAME(OutputData));
+					break;
+				}
+
+				case EMetasoundFrontendClassType::Input:
+				{
+					Vertex = &Node->Interface.Inputs.Last();
+					break;
+				}
+
+				case EMetasoundFrontendClassType::Output:
+				{
+					Vertex = &Node->Interface.Outputs.Last();
+					break;
+				}
+
+				default:
+				{
+					Vertex = InPin->Direction == EGPD_Input
+						? InBuilder.FindNodeInput(NodeID, InPin->GetFName())
+						: InBuilder.FindNodeOutput(NodeID, InPin->GetFName());
+				}
+			}
+
+			FMetasoundFrontendVertexHandle VertexHandle { NodeID };
+			if (Vertex)
+			{
+				VertexHandle.VertexID = Vertex->VertexID;
+			}
+			return VertexHandle;
+		}
+
+		const FMetasoundFrontendVertex* FGraphBuilder::GetPinVertex(const FMetaSoundFrontendDocumentBuilder& InBuilder, const UEdGraphPin* InPin)
+		{
+			using namespace VariableNames;
+
+			if (!InPin)
+			{
+				return nullptr;
+			}
+
+			const UMetasoundEditorGraphNode* OwningNode = CastChecked<UMetasoundEditorGraphNode>(InPin->GetOwningNode());
+
+			const FGuid NodeID = OwningNode->GetNodeID();
+			const FMetasoundFrontendNode* Node = InBuilder.FindNode(NodeID);
+			if (!Node)
+			{
+				return nullptr;
+			}
+
+			const FMetasoundFrontendClass* Class = InBuilder.FindDependency(Node->ClassID);
+			if (!Class)
+			{
+				return nullptr;
+			}
+
+			switch (Class->Metadata.GetType())
+			{
+				case EMetasoundFrontendClassType::Variable:
+				case EMetasoundFrontendClassType::VariableAccessor:
+				case EMetasoundFrontendClassType::VariableDeferredAccessor:
+				case EMetasoundFrontendClassType::VariableMutator:
+				{
+					// All variables nodes use the same pin name for user-modifiable node
+					// inputs and outputs and the editor does not display the pin's name. The
+					// editor instead displays the variable's name in place of the pin name to
+					// maintain a consistent look and behavior to input and output nodes.
+					return InPin->Direction == EGPD_Input
+						? InBuilder.FindNodeInput(NodeID, METASOUND_GET_PARAM_NAME(InputData))
+						: InBuilder.FindNodeOutput(NodeID, METASOUND_GET_PARAM_NAME(OutputData));
+				}
+				case EMetasoundFrontendClassType::Input:
+				{
+					return &Node->Interface.Inputs.Last();
+				}
+				case EMetasoundFrontendClassType::Output:
+				{
+					return &Node->Interface.Outputs.Last();
+				}
+
+				default:
+				{
+					return InPin->Direction == EGPD_Input
+						? InBuilder.FindNodeInput(NodeID, InPin->GetFName())
+						: InBuilder.FindNodeOutput(NodeID, InPin->GetFName());
+				}
+			}
+		}
+
 		Frontend::FOutputHandle FGraphBuilder::GetOutputHandleFromPin(const UEdGraphPin* InPin)
 		{
 			using namespace Frontend;
@@ -1252,12 +1398,6 @@ namespace Metasound
 			return NodeHandle;
 		}
 
-		Frontend::FNodeHandle FGraphBuilder::AddInputNodeHandle(UObject& InMetaSound, const FName InTypeName, const FMetasoundFrontendLiteral* InDefaultValue, const FName* InNameBase)
-		{
-			UE_LOG(LogMetaSound, Error, TEXT("FGraphBuilder::AddInputNodeHandle with these parameters is no longer supported and should not be called. Use the one with FCreateNodeVertexParams instead."));
-			return Frontend::INodeController::GetInvalidHandle();
-		}
-
 		Frontend::FNodeHandle FGraphBuilder::AddInputNodeHandle(UObject& InMetaSound, const FCreateNodeVertexParams& InParams, const FMetasoundFrontendLiteral* InDefaultValue, const FName* InNameBase)
 		{
 			FMetasoundAssetBase* MetaSoundAsset = IMetasoundUObjectRegistry::Get().GetObjectAsAssetBase(&InMetaSound);
@@ -1284,12 +1424,6 @@ namespace Metasound
 
 			return MetaSoundAsset->GetRootGraphHandle()->AddInputVertex(ClassInput);
 		}
-		
-		Frontend::FNodeHandle FGraphBuilder::AddOutputNodeHandle(UObject& InMetaSound, const FName InTypeName, const FName* InNameBase)
-		{
-			UE_LOG(LogMetaSound, Error, TEXT("FGraphBuilder::AddOutputNodeHandle with these parameters is no longer supported and should not be called. Use the one with FCreateNodeVertexParams instead."));
-			return Frontend::INodeController::GetInvalidHandle();
-		}
 
 		Frontend::FNodeHandle FGraphBuilder::AddOutputNodeHandle(UObject& InMetaSound, const FCreateNodeVertexParams& InParams, const FName* InNameBase)
 		{
diff --git a/Engine/Plugins/Runtime/Metasound/Source/MetasoundEditor/Private/MetasoundEditorGraphBuilder.h b/Engine/Plugins/Runtime/Metasound/Source/MetasoundEditor/Private/MetasoundEditorGraphBuilder.h
--- a/Engine/Plugins/Runtime/Metasound/Source/MetasoundEditor/Private/MetasoundEditorGraphBuilder.h
+++ b/Engine/Plugins/Runtime/Metasound/Source/MetasoundEditor/Private/MetasoundEditorGraphBuilder.h
@@ -1,12 +1,14 @@
 // Copyright Epic Games, Inc. All Rights Reserved.
 #pragma once
 
+#include "Internationalization/Text.h"
+#include "Logging/TokenizedMessage.h"
 #include "MetasoundEditorGraphValidation.h"
 #include "MetasoundFrontend.h"
 #include "MetasoundFrontendController.h"
+#include "MetasoundFrontendDocumentBuilder.h"
 #include "MetasoundFrontendDocument.h"
-#include "Internationalization/Text.h"
-#include "Logging/TokenizedMessage.h"
+#include "MetasoundFrontendNodeTemplateRegistry.h"
 #include "Templates/SharedPointer.h"
 #include "UObject/NameTypes.h"
 #include "UObject/NoExportTypes.h"
@@ -128,13 +130,6 @@ namespace Metasound
 			// Adds a corresponding UMetasoundEditorGraphInputNode for the provided node handle.
 			static UMetasoundEditorGraphInputNode* AddInputNode(UObject& InMetaSound, Frontend::FNodeHandle InNodeHandle, FVector2D InLocation, bool bInSelectNewNode = true);
 
-			UE_DEPRECATED(5.1, "Use AddInputNodeHandle with FCreateNodeVertexParams instead.")
-			static Frontend::FNodeHandle AddInputNodeHandle(
-				UObject& InMetaSound,
-				const FName InTypeName,
-				const FMetasoundFrontendLiteral* InDefaultValue = nullptr,
-				const FName* InNameBase = nullptr);
-
 			// Generates FNodeHandle for the given external node data. Does not bind or create EdGraph representation of given node.
 			static Frontend::FNodeHandle AddInputNodeHandle(
 				UObject& InMetaSound,
@@ -155,9 +150,6 @@ namespace Metasound
 			// Adds an output node to the editor graph that corresponds to the provided node handle.
 			static UMetasoundEditorGraphOutputNode* AddOutputNode(UObject& InMetaSound, Frontend::FNodeHandle& InNodeHandle, FVector2D InLocation, bool bInSelectNewNode = true);
 
-			UE_DEPRECATED(5.1, "Use AddOutputNodeHandle with FCreateNodeVertexParams instead.")
-			static Frontend::FNodeHandle AddOutputNodeHandle(UObject& InMetaSound, const FName InTypeName, const FName* InNameBase = nullptr);
-
 			// Generates analogous FNodeHandle for the given internal node data. Does not bind nor create EdGraph representation of given node.
 			static Frontend::FNodeHandle AddOutputNodeHandle(UObject& InMetaSound, const FCreateNodeVertexParams& InParams, const FName* InNameBase = nullptr);
 
@@ -206,6 +198,8 @@ namespace Metasound
 			static Frontend::FConstInputHandle GetConstInputHandleFromPin(const UEdGraphPin* InPin);
 
 			static FName GetPinDataType(const UEdGraphPin* InPin);
+			static FMetasoundFrontendVertexHandle GetPinVertexHandle(const FMetaSoundFrontendDocumentBuilder& InBuilder, const UEdGraphPin* InPin);
+			static const FMetasoundFrontendVertex* GetPinVertex(const FMetaSoundFrontendDocumentBuilder& InBuilder, const UEdGraphPin* InPin);
 
 			static const FMetasoundFrontendEdgeStyle* GetOutputEdgeStyle(Frontend::FConstOutputHandle InOutputHandle);
 			static const FMetasoundFrontendEdgeStyle* GetOutputEdgeStyle(const UEdGraphPin* InPin);
diff --git a/Engine/Plugins/Runtime/Metasound/Source/MetasoundEditor/Private/MetasoundEditorGraphNode.cpp b/Engine/Plugins/Runtime/Metasound/Source/MetasoundEditor/Private/MetasoundEditorGraphNode.cpp
--- a/Engine/Plugins/Runtime/Metasound/Source/MetasoundEditor/Private/MetasoundEditorGraphNode.cpp
+++ b/Engine/Plugins/Runtime/Metasound/Source/MetasoundEditor/Private/MetasoundEditorGraphNode.cpp
@@ -11,6 +11,7 @@
 #include "Logging/TokenizedMessage.h"
 #include "Metasound.h"
 #include "MetasoundAssetManager.h"
+#include "MetasoundBuilderSubsystem.h"
 #include "MetasoundEditorCommands.h"
 #include "MetasoundEditorGraph.h"
 #include "MetasoundEditorGraphBuilder.h"
@@ -133,6 +134,40 @@ int32 UMetasoundEditorGraphNode::EstimateNodeWidth() const
 	}
 }
 
+const FMetaSoundFrontendDocumentBuilder& UMetasoundEditorGraphNode::GetFrontendBuilderChecked() const
+{
+	UObject* Outermost = GetOutermostObject();
+	check(Outermost);
+
+	const UMetaSoundBuilderBase& BuilderBase = UMetaSoundBuilderSubsystem::GetChecked().AttachBuilderToAssetChecked(*Outermost);
+	return BuilderBase.GetConstBuilder();
+}
+
+const FMetasoundFrontendNode* UMetasoundEditorGraphNode::GetFrontendNode() const
+{
+	if (UObject* Outermost = GetOutermostObject())
+	{
+		const FGuid NodeID = GetNodeID();
+		const UMetaSoundBuilderBase& Builder = UMetaSoundBuilderSubsystem::GetChecked().AttachBuilderToAssetChecked(*Outermost);
+		return Builder.GetConstBuilder().FindNode(NodeID);
+	}
+
+	return nullptr;
+}
+
+const FMetasoundFrontendNode& UMetasoundEditorGraphNode::GetFrontendNodeChecked() const
+{
+	UObject* Outermost = GetOutermostObject();
+	check(Outermost);
+
+	const FGuid NodeID = GetNodeID();
+	const UMetaSoundBuilderBase& Builder = UMetaSoundBuilderSubsystem::GetChecked().AttachBuilderToAssetChecked(*Outermost);
+
+	const FMetasoundFrontendNode* FrontendNode = Builder.GetConstBuilder().FindNode(NodeID);
+	check(FrontendNode);
+	return *FrontendNode;
+}
+
 UObject& UMetasoundEditorGraphNode::GetMetasoundChecked()
 {
 	UMetasoundEditorGraph* EdGraph = CastChecked<UMetasoundEditorGraph>(GetGraph());
@@ -904,7 +939,7 @@ void UMetasoundEditorGraphExternalNode::Validate(Metasound::Editor::FGraphNodeVa
 			{
 #if WITH_EDITOR
 				FString Message;
-				if (!Template->HasRequiredConnections(NodeHandle, &Message))
+				if (!Template->HasRequiredConnections(GetFrontendBuilderChecked(), GetNodeID(), &Message))
 				{
 					OutResult.SetMessage(EMessageSeverity::Warning, Message);
 				}
diff --git a/Engine/Plugins/Runtime/Metasound/Source/MetasoundEditor/Private/MetasoundEditorGraphSchema.cpp b/Engine/Plugins/Runtime/Metasound/Source/MetasoundEditor/Private/MetasoundEditorGraphSchema.cpp
--- a/Engine/Plugins/Runtime/Metasound/Source/MetasoundEditor/Private/MetasoundEditorGraphSchema.cpp
+++ b/Engine/Plugins/Runtime/Metasound/Source/MetasoundEditor/Private/MetasoundEditorGraphSchema.cpp
@@ -17,6 +17,7 @@
 #include "Metasound.h"
 #include "MetasoundAssetBase.h"
 #include "MetasoundAssetManager.h"
+#include "MetasoundBuilderSubsystem.h"
 #include "MetasoundDataReference.h"
 #include "MetasoundEditor.h"
 #include "MetasoundEditorCommands.h"
@@ -30,6 +31,7 @@
 #include "MetasoundFrontend.h"
 #include "MetasoundFrontendGraphLinter.h"
 #include "MetasoundFrontendNodesCategories.h"
+#include "MetasoundFrontendNodeTemplateRegistry.h"
 #include "MetasoundFrontendRegistries.h"
 #include "MetasoundFrontendSearchEngine.h"
 #include "MetasoundLiteral.h"
@@ -37,6 +39,7 @@
 #include "MetasoundStandardNodesCategories.h"
 #include "MetasoundUObjectRegistry.h"
 #include "MetasoundVariableNodes.h"
+#include "MetasoundVertex.h"
 #include "NodeTemplates/MetasoundFrontendNodeTemplateReroute.h"
 #include "ScopedTransaction.h"
 #include "Styling/SlateStyleRegistry.h"
@@ -891,27 +894,41 @@ FMetasoundGraphSchemaAction_NewReroute::FMetasoundGraphSchemaAction_NewReroute(c
 
 UEdGraphNode* FMetasoundGraphSchemaAction_NewReroute::PerformAction(UEdGraph* ParentGraph, UEdGraphPin* FromPin, const FVector2D Location, bool bSelectNewNode /* = true*/)
 {
+	using namespace Metasound;
 	using namespace Metasound::Editor;
 	using namespace Metasound::Frontend;
 
-	const FName DataType = FGraphBuilder::GetPinDataType(FromPin);
-	if (DataType.IsNone())
-	{
-		return nullptr;
-	}
-
-	const FScopedTransaction Transaction(FText::Format(LOCTEXT("AddNewRerouteNode", "Add {0} Reroute Node"), FText::FromName(DataType)));
+	check(ParentGraph);
 	UMetasoundEditorGraph* MetaSoundGraph = CastChecked<UMetasoundEditorGraph>(ParentGraph);
 	UObject& ParentMetasound = MetaSoundGraph->GetMetasoundChecked();
+	UMetaSoundBuilderBase& Builder = UMetaSoundBuilderSubsystem::GetChecked().AttachBuilderToAssetChecked(ParentMetasound);
+
+	const FScopedTransaction Transaction(LOCTEXT("AddNewRerouteNode", "Add Reroute Node"));
 	ParentMetasound.Modify();
 	ParentGraph->Modify();
 
 	FMetasoundAssetBase* MetaSoundAsset = Metasound::IMetasoundUObjectRegistry::Get().GetObjectAsAssetBase(&ParentMetasound);
 	check(MetaSoundAsset);
 
-	const FNodeRegistryKey& RerouteTemplateKey = FRerouteNodeTemplate::GetRegistryKey();
-	FMetasoundFrontendNodeInterface NodeInterface = FRerouteNodeTemplate::CreateNodeInterfaceFromDataType(DataType);
-	FNodeHandle NodeHandle = MetaSoundAsset->GetRootGraphHandle()->AddTemplateNode(RerouteTemplateKey, MoveTemp(NodeInterface));
+	const INodeTemplate* RerouteTemplate = INodeTemplateRegistry::Get().FindTemplate(FRerouteNodeTemplate::ClassName);
+	check(RerouteTemplate);
+
+	const FMetaSoundFrontendDocumentBuilder& DocBuilder = Builder.GetConstBuilder();
+
+	// Provided 'FromPin' is what to connect to, so if its an input, its the output of the generated node needs to
+	// match the from pin and vice versa.
+	const FName FromVertexDataType = FGraphBuilder::GetPinDataType(FromPin);
+	FNodeTemplateGenerateInterfaceParams Params;
+	if (FromPin->Direction == EGPD_Input)
+	{
+		Params.OutputsToConnect.Add(FromVertexDataType);
+	}
+	else
+	{
+		Params.InputsToConnect.Add(FromVertexDataType);
+	}
+
+	FNodeHandle NodeHandle = MetaSoundAsset->GetRootGraphHandle()->AddTemplateNode(*RerouteTemplate, MoveTemp(Params));
 
 	if (UMetasoundEditorGraphExternalNode* NewGraphNode = FGraphBuilder::AddExternalNode(ParentMetasound, NodeHandle, Location, bSelectNewNode))
 	{
@@ -1313,6 +1330,7 @@ const FPinConnectionResponse UMetasoundEditorGraphSchema::CanCreateConnection(co
 
 void UMetasoundEditorGraphSchema::OnPinConnectionDoubleCicked(UEdGraphPin* PinA, UEdGraphPin* PinB, const FVector2D& GraphPosition) const
 {
+	using namespace Metasound;
 	using namespace Metasound::Editor;
 
 	if (!PinA || !PinB)
@@ -1327,11 +1345,19 @@ void UMetasoundEditorGraphSchema::OnPinConnectionDoubleCicked(UEdGraphPin* PinA,
 	UMetasoundEditorGraph* ParentGraph = Cast<UMetasoundEditorGraph>(PinA->GetOwningNode()->GetGraph());
 	if (ParentGraph->IsEditable())
 	{
-		const FName DataType = FGraphBuilder::GetPinDataType(PinA);
-		const FScopedTransaction Transaction(FText::Format(LOCTEXT("AddConnectNewRerouteNode", "Add & Connect {0} Reroute Node"), FText::FromName(DataType)));
-
 		UMetasoundEditorGraph* MetaSoundGraph = CastChecked<UMetasoundEditorGraph>(ParentGraph);
 		UObject& ParentMetasound = MetaSoundGraph->GetMetasoundChecked();
+		UMetaSoundBuilderBase& Builder = UMetaSoundBuilderSubsystem::GetChecked().AttachBuilderToAssetChecked(ParentMetasound);
+		const FMetaSoundFrontendDocumentBuilder& DocBuilder = Builder.GetConstBuilder();
+		
+		FName VertexDataType;
+		const FMetasoundFrontendVertex* Vertex = FGraphBuilder::GetPinVertex(DocBuilder, PinA);
+		if (ensure(Vertex))
+		{
+			VertexDataType = Vertex->TypeName;
+		}
+		const FScopedTransaction Transaction(FText::Format(LOCTEXT("AddConnectNewRerouteNode", "Add & Connect {0} Reroute Node"), FText::FromName(VertexDataType)));
+
 		ParentMetasound.Modify();
 		ParentGraph->Modify();
 
@@ -1341,8 +1367,7 @@ void UMetasoundEditorGraphSchema::OnPinConnectionDoubleCicked(UEdGraphPin* PinA,
 		constexpr bool bShouldTransact = false;
 		TSharedPtr<FMetasoundGraphSchemaAction_NewReroute> RerouteAction = MakeShared<FMetasoundGraphSchemaAction_NewReroute>(IconColor, bShouldTransact);
 
-		UEdGraphNode& Node = *PinA->GetOwningNode();
-		UEdGraphNode* NewNode = RerouteAction->PerformAction(Node.GetGraph(), OutputPin, GraphPosition, true);
+		UEdGraphNode* NewNode = RerouteAction->PerformAction(ParentGraph, OutputPin, GraphPosition, true);
 
 		if (ensure(NewNode))
 		{
diff --git a/Engine/Plugins/Runtime/Metasound/Source/MetasoundEditor/Private/SMetasoundGraphPin.cpp b/Engine/Plugins/Runtime/Metasound/Source/MetasoundEditor/Private/SMetasoundGraphPin.cpp
--- a/Engine/Plugins/Runtime/Metasound/Source/MetasoundEditor/Private/SMetasoundGraphPin.cpp
+++ b/Engine/Plugins/Runtime/Metasound/Source/MetasoundEditor/Private/SMetasoundGraphPin.cpp
@@ -38,8 +38,8 @@ namespace Metasound
 				{
 					if (const INodeTemplate* Template = INodeTemplateRegistry::Get().FindTemplate(FRerouteNodeTemplate::GetRegistryKey()))
 					{
-						FConstNodeHandle NodeHandle = OwningNode->GetConstNodeHandle();
-						bHasRequiredConnections = Template->HasRequiredConnections(NodeHandle);
+						const FMetaSoundFrontendDocumentBuilder& Builder = OwningNode->GetFrontendBuilderChecked();
+						bHasRequiredConnections = Template->HasRequiredConnections(Builder, OwningNode->GetNodeID());
 					}
 				}
 			}
diff --git a/Engine/Plugins/Runtime/Metasound/Source/MetasoundEditor/Public/MetasoundEditorGraphNode.h b/Engine/Plugins/Runtime/Metasound/Source/MetasoundEditor/Public/MetasoundEditorGraphNode.h
--- a/Engine/Plugins/Runtime/Metasound/Source/MetasoundEditor/Public/MetasoundEditorGraphNode.h
+++ b/Engine/Plugins/Runtime/Metasound/Source/MetasoundEditor/Public/MetasoundEditorGraphNode.h
@@ -83,6 +83,10 @@ public:
 		return false;
 	}
 
+	const FMetaSoundFrontendDocumentBuilder& GetFrontendBuilderChecked() const;
+	const FMetasoundFrontendNode* GetFrontendNode() const;
+	const FMetasoundFrontendNode& GetFrontendNodeChecked() const;
+
 	UObject& GetMetasoundChecked();
 	const UObject& GetMetasoundChecked() const;
 
diff --git a/Engine/Plugins/Runtime/Metasound/Source/MetasoundFrontend/Private/MetasoundFrontendGraphController.cpp b/Engine/Plugins/Runtime/Metasound/Source/MetasoundFrontend/Private/MetasoundFrontendGraphController.cpp
--- a/Engine/Plugins/Runtime/Metasound/Source/MetasoundFrontend/Private/MetasoundFrontendGraphController.cpp
+++ b/Engine/Plugins/Runtime/Metasound/Source/MetasoundFrontend/Private/MetasoundFrontendGraphController.cpp
@@ -1045,29 +1045,6 @@ namespace Metasound
 			return INodeController::GetInvalidHandle();
 		}
 
-		FNodeHandle FGraphController::AddInputVertex(const FVertexName& InName, const FName InTypeName, const FMetasoundFrontendLiteral* InDefaultValue)
-		{
-			const FGuid VertexID = FGuid::NewGuid();
-
-			FMetasoundFrontendClassInput Description;
-
-			Description.Name = InName;
-			Description.TypeName = InTypeName;
-			Description.VertexID = VertexID;
-
-			if (InDefaultValue)
-			{
-				Description.DefaultLiteral = *InDefaultValue;
-			}
-			else
-			{
-				Metasound::FLiteral Literal = IDataTypeRegistry::Get().CreateDefaultLiteral(InTypeName);
-				Description.DefaultLiteral.SetFromLiteral(Literal);
-			}
-
-			return AddInputVertex(Description);
-		}
-
 		bool FGraphController::RemoveInputVertex(const FVertexName& InName)
 		{
 			auto IsInputNodeWithSameName = [&](const FMetasoundFrontendClass& InClass, const FMetasoundFrontendNode& InNode)
@@ -1424,35 +1401,29 @@ namespace Metasound
 			return AddNode(FNodeRegistryKey(InClassMetadata), InNodeGuid);
 		}
 
-		FNodeHandle FGraphController::AddTemplateNode(const FNodeRegistryKey& InKey, FMetasoundFrontendNodeInterface&& InNodeInterface, FGuid InNodeGuid)
+		FNodeHandle FGraphController::AddTemplateNode(const INodeTemplate& InTemplate, FNodeTemplateGenerateInterfaceParams Params, FGuid InNodeGuid)
 		{
-			if (const INodeTemplate* Template = INodeTemplateRegistry::Get().FindTemplate(InKey))
+			// Construct a FNodeClassInfo from this lookup key.
+			const FNodeRegistryKey Key(InTemplate.GetFrontendClass().Metadata);
+			FConstClassAccessPtr Class = OwningDocument->FindOrAddClass(Key);
+			const bool bIsValidClass = (nullptr != Class.Get());
+
+			if (bIsValidClass)
 			{
-				const bool bIsValidInterface = Template->IsValidNodeInterface(InNodeInterface);
-				if (ensureAlwaysMsgf(bIsValidInterface, TEXT("Cannot implement interface when attempting to add node using template with key '%s'"), *InKey.ToString()))
+				if (FMetasoundFrontendGraphClass* GraphClass = GraphClassPtr.Get())
 				{
-					// Construct a FNodeClassInfo from this lookup key.
-					FConstClassAccessPtr Class = OwningDocument->FindOrAddClass(InKey);
-					const bool bIsValidClass = (nullptr != Class.Get());
-
-					if (bIsValidClass)
+					if (const FMetasoundFrontendClass* NodeClass = Class.Get())
 					{
-						if (FMetasoundFrontendGraphClass* GraphClass = GraphClassPtr.Get())
-						{
-							if (const FMetasoundFrontendClass* NodeClass = Class.Get())
-							{
-								FMetasoundFrontendNode& Node = GraphClass->Graph.Nodes.Emplace_GetRef(*NodeClass);
-								Node.UpdateID(InNodeGuid);
-								Node.Interface = InNodeInterface;
-								FNodeAccessPtr NodePtr = GraphClassPtr.GetNodeWithNodeID(Node.GetID());
-								return GetNodeHandle(FGraphController::FNodeAndClass{ NodePtr, Class });
-							}
-						}
+						FMetasoundFrontendNode& Node = GraphClass->Graph.Nodes.Emplace_GetRef(*NodeClass);
+						Node.UpdateID(InNodeGuid);
+						Node.Interface = InTemplate.GenerateNodeInterface(MoveTemp(Params));
+						FNodeAccessPtr NodePtr = GraphClassPtr.GetNodeWithNodeID(Node.GetID());
+						return GetNodeHandle(FGraphController::FNodeAndClass{ NodePtr, Class });
 					}
 				}
 			}
 
-			UE_LOG(LogMetaSound, Warning, TEXT("Failed to find or add node template class info with registry key [Key:%s]"), *InKey.ToString());
+			UE_LOG(LogMetaSound, Warning, TEXT("Failed to find or add node template class info with registry key [Key:%s]"), *Key.ToString());
 			return INodeController::GetInvalidHandle();
 		}
 
diff --git a/Engine/Plugins/Runtime/Metasound/Source/MetasoundFrontend/Private/MetasoundFrontendGraphController.h b/Engine/Plugins/Runtime/Metasound/Source/MetasoundFrontend/Private/MetasoundFrontendGraphController.h
--- a/Engine/Plugins/Runtime/Metasound/Source/MetasoundFrontend/Private/MetasoundFrontendGraphController.h
+++ b/Engine/Plugins/Runtime/Metasound/Source/MetasoundFrontend/Private/MetasoundFrontendGraphController.h
@@ -7,10 +7,15 @@
 #include "MetasoundFrontendDocument.h"
 #include "MetasoundFrontendDocumentAccessPtr.h"
 
+
+// Forward Declarations
 class UClass;
 
 namespace Metasound
 {
+	// Forward Declarations
+	class INodeTemplate;
+
 	namespace Frontend
 	{
 		/** FGraphController represents a Metasound graph class. */
@@ -117,7 +122,6 @@ namespace Metasound
 			FNodeHandle GetOutputNodeWithName(const FVertexName& InName) override;
 
 			FNodeHandle AddInputVertex(const FMetasoundFrontendClassInput& InDescription) override;
-			FNodeHandle AddInputVertex(const FVertexName& InName, const FName InTypeName, const FMetasoundFrontendLiteral* InDefaultValue) override;
 			bool RemoveInputVertex(const FVertexName& InName) override;
 
 			FNodeHandle AddOutputVertex(const FMetasoundFrontendClassOutput& InDescription) override;
@@ -177,7 +181,7 @@ namespace Metasound
 			FNodeHandle AddNode(const FNodeRegistryKey& InNodeClass, FGuid InNodeGuid) override;
 			FNodeHandle AddNode(const FMetasoundFrontendClassMetadata& InClassMetadata, FGuid InNodeGuid) override;
 			FNodeHandle AddDuplicateNode(const INodeController& InNode) override;
-			FNodeHandle AddTemplateNode(const FNodeRegistryKey& InKey, FMetasoundFrontendNodeInterface&& InNodeInterface, FGuid InNodeGuid) override;
+			FNodeHandle AddTemplateNode(const INodeTemplate& InNodeTemplate, FNodeTemplateGenerateInterfaceParams Params, FGuid InNodeGuid = FGuid::NewGuid()) override;
 
 			// Remove the node corresponding to this node handle.
 			// On success, invalidates the received node handle.
diff --git a/Engine/Plugins/Runtime/Metasound/Source/MetasoundFrontend/Private/MetasoundFrontendInputController.cpp b/Engine/Plugins/Runtime/Metasound/Source/MetasoundFrontend/Private/MetasoundFrontendInputController.cpp
--- a/Engine/Plugins/Runtime/Metasound/Source/MetasoundFrontend/Private/MetasoundFrontendInputController.cpp
+++ b/Engine/Plugins/Runtime/Metasound/Source/MetasoundFrontend/Private/MetasoundFrontendInputController.cpp
@@ -7,6 +7,7 @@
 #include "MetasoundFrontendDocumentAccessPtr.h"
 #include "MetasoundFrontendGraphLinter.h"
 #include "MetasoundFrontendInvalidController.h"
+#include "MetasoundFrontendNodeTemplateRegistry.h"
 #include "Misc/Guid.h"
 
 #define LOCTEXT_NAMESPACE "MetasoundFrontendInputController"
@@ -78,6 +79,10 @@ namespace Metasound
 								AccessType = ClassInput->AccessType;
 								return;
 							}
+
+							// Likely template node with no set class input interface, so valid to return unset
+							AccessType = EMetasoundFrontendVertexAccessType::Unset;
+							return;
 						}
 
 						const EMetasoundFrontendVertexAccessType RerouteAccessType = ReroutedInput->GetVertexAccessType();
@@ -211,6 +216,17 @@ namespace Metasound
 
 		bool FBaseInputController::IsConnectionUserModifiable() const
 		{
+			FConstNodeHandle Owner = GetOwningNode();
+			if (Owner->GetClassMetadata().GetType() == EMetasoundFrontendClassType::Template)
+			{
+				const FNodeRegistryKey Key(GetOwningNode()->GetClassMetadata());
+				const INodeTemplate* Template = INodeTemplateRegistry::Get().FindTemplate(Key);
+				if (ensure(Template))
+				{
+					return Template->IsInputConnectionUserModifiable();
+				}
+			}
+
 			return true;
 		}
 
@@ -466,7 +482,6 @@ namespace Metasound
 			return Access;
 		}
 
-
 		//
 		// FOutputNodeInputController
 		//
diff --git a/Engine/Plugins/Runtime/Metasound/Source/MetasoundFrontend/Private/MetasoundFrontendInvalidController.h b/Engine/Plugins/Runtime/Metasound/Source/MetasoundFrontend/Private/MetasoundFrontendInvalidController.h
--- a/Engine/Plugins/Runtime/Metasound/Source/MetasoundFrontend/Private/MetasoundFrontendInvalidController.h
+++ b/Engine/Plugins/Runtime/Metasound/Source/MetasoundFrontend/Private/MetasoundFrontendInvalidController.h
@@ -364,7 +364,6 @@ namespace Metasound
 			virtual FConstClassOutputAccessPtr FindClassOutputWithName(const FVertexName& InName) const override { return FConstClassOutputAccessPtr(); }
 
 			virtual FNodeHandle AddInputVertex(const FMetasoundFrontendClassInput& InDescription) override { return INodeController::GetInvalidHandle(); }
-			virtual TSharedRef<INodeController> AddInputVertex(const FVertexName& InName, const FName InTypeName, const FMetasoundFrontendLiteral* InDefaultValue) override { return INodeController::GetInvalidHandle(); }
 			virtual bool RemoveInputVertex(const FVertexName& InputName) override { return false; }
 
 			virtual FNodeHandle AddOutputVertex(const FMetasoundFrontendClassOutput& InDescription) override { return INodeController::GetInvalidHandle(); }
@@ -412,7 +411,7 @@ namespace Metasound
 			virtual FNodeHandle AddNode(const FNodeRegistryKey& InNodeClass, FGuid InNodeGuid) override { return INodeController::GetInvalidHandle(); }
 			virtual FNodeHandle AddNode(const FMetasoundFrontendClassMetadata& InNodeClass, FGuid InNodeGuid) override { return INodeController::GetInvalidHandle(); }
 			virtual FNodeHandle AddDuplicateNode(const INodeController& InNode) override { return INodeController::GetInvalidHandle(); }
-			virtual FNodeHandle AddTemplateNode(const FNodeRegistryKey& InNodeClass, FMetasoundFrontendNodeInterface&& InNodeInterface, FGuid InNodeGuid) override { return INodeController::GetInvalidHandle(); }
+			virtual FNodeHandle AddTemplateNode(const INodeTemplate& InNodeTemplate, FNodeTemplateGenerateInterfaceParams Params, FGuid InNodeGuid = FGuid::NewGuid()) override { return INodeController::GetInvalidHandle(); }
 
 			// Remove the node corresponding to this node handle.
 			// On success, invalidates the received node handle.
diff --git a/Engine/Plugins/Runtime/Metasound/Source/MetasoundFrontend/Private/MetasoundFrontendNodeTemplateRegistry.cpp b/Engine/Plugins/Runtime/Metasound/Source/MetasoundFrontend/Private/MetasoundFrontendNodeTemplateRegistry.cpp
--- a/Engine/Plugins/Runtime/Metasound/Source/MetasoundFrontend/Private/MetasoundFrontendNodeTemplateRegistry.cpp
+++ b/Engine/Plugins/Runtime/Metasound/Source/MetasoundFrontend/Private/MetasoundFrontendNodeTemplateRegistry.cpp
@@ -5,35 +5,16 @@
 #include "MetasoundFrontendDocument.h"
 #include "MetasoundFrontendRegistryContainerImpl.h"
 #include "MetasoundFrontendRegistryTransaction.h"
+#include "MetasoundFrontendTransform.h"
 #include "MetasoundLog.h"
 #include "MetasoundTrace.h"
 
 
 namespace Metasound::Frontend
 {
-	bool FNodeTemplateBase::IsValidNodeInterface(const FMetasoundFrontendNodeInterface& InNodeInterface) const
+	TUniquePtr<INodeTransform> INodeTemplate::GenerateNodeTransform(FMetasoundFrontendDocument& InDocument) const
 	{
-		return true;
-	}
-
-	bool FNodeTemplateBase::IsInputAccessTypeDynamic() const
-	{
-		return false;
-	}
-
-	bool FNodeTemplateBase::IsOutputAccessTypeDynamic() const
-	{
-		return false;
-	}
-
-	EMetasoundFrontendVertexAccessType FNodeTemplateBase::GetNodeInputAccessType(const FMetaSoundFrontendDocumentBuilder& InBuilder, const FGuid& InNodeID, const FGuid& InVertexID) const
-	{
-		return EMetasoundFrontendVertexAccessType::Unset;
-	}
-
-	EMetasoundFrontendVertexAccessType FNodeTemplateBase::GetNodeOutputAccessType(const FMetaSoundFrontendDocumentBuilder& InBuilder, const FGuid& InNodeID, const FGuid& InVertexID) const
-	{
-		return EMetasoundFrontendVertexAccessType::Unset;
+		return nullptr;
 	}
 
 	class FNodeTemplateRegistry : public INodeTemplateRegistry
@@ -43,12 +24,14 @@ namespace Metasound::Frontend
 		virtual ~FNodeTemplateRegistry() = default;
 
 		virtual const INodeTemplate* FindTemplate(const FNodeRegistryKey& InKey) const override;
+		virtual const INodeTemplate* FindTemplate(const FMetasoundFrontendClassName& InClassName) const override;
 
 		void Register(TUniquePtr<INodeTemplate>&& InEntry);
 		void Unregister(const FNodeRegistryKey& InKey);
 
 	private:
-		TMap<FNodeRegistryKey, TUniquePtr<INodeTemplate>> Templates;
+		TMap<FNodeRegistryKey, TUniquePtr<const INodeTemplate>> Templates;
+		TMultiMap<FMetasoundFrontendClassName, const INodeTemplate*> TemplateByClassName;
 	};
 
 	void FNodeTemplateRegistry::Register(TUniquePtr<INodeTemplate>&& InTemplate)
@@ -58,19 +41,24 @@ namespace Metasound::Frontend
 			const FNodeRegistryKey Key = FNodeRegistryKey(InTemplate->GetFrontendClass().Metadata);
 			if (ensure(Key.IsValid()))
 			{
-				Templates.Add(Key, MoveTemp(InTemplate));
+				TUniquePtr<const INodeTemplate>& Entry = Templates.Add(Key, MoveTemp(InTemplate));
+				TemplateByClassName.Add(Entry->GetFrontendClass().Metadata.GetClassName(), Entry.Get());
 			}
 		}
 	}
 
 	void FNodeTemplateRegistry::Unregister(const FNodeRegistryKey& InKey)
 	{
-		ensure(Templates.Remove(InKey) > 0);
+		TUniquePtr<const INodeTemplate> Removed;
+		if (ensure(Templates.RemoveAndCopyValue(InKey, Removed)))
+		{
+			ensure(TemplateByClassName.Remove(Removed->GetFrontendClass().Metadata.GetClassName()));
+		}
 	}
 
 	const INodeTemplate* FNodeTemplateRegistry::FindTemplate(const FNodeRegistryKey& InKey) const
 	{
-		if (const TUniquePtr<INodeTemplate>* TemplatePtr = Templates.Find(InKey))
+		if (const TUniquePtr<const INodeTemplate>* TemplatePtr = Templates.Find(InKey))
 		{
 			return TemplatePtr->Get();
 		}
@@ -78,12 +66,49 @@ namespace Metasound::Frontend
 		return nullptr;
 	}
 
+	const INodeTemplate* FNodeTemplateRegistry::FindTemplate(const FMetasoundFrontendClassName& InClassName) const
+	{
+		if (TemplateByClassName.Contains(InClassName))
+		{
+			TArray<const INodeTemplate*> FoundTemplates;
+			constexpr bool bMaintainOrder = false;
+			TemplateByClassName.MultiFind(InClassName, FoundTemplates, bMaintainOrder);
+			Algo::Sort(FoundTemplates, [](const INodeTemplate* A, const INodeTemplate* B) { return A->GetVersionNumber() < B->GetVersionNumber(); });
+
+			return FoundTemplates.Last();
+		}
+
+		return nullptr;
+	}
+
 	INodeTemplateRegistry& INodeTemplateRegistry::Get()
 	{
 		static FNodeTemplateRegistry Registry;
 		return Registry;
 	}
 
+#if WITH_EDITOR
+	FText FNodeTemplateBase::GetNodeDisplayName(const IMetaSoundDocumentInterface& Interface, const FGuid& InNodeID) const
+	{
+		return { };
+	}
+
+	FText FNodeTemplateBase::GetInputPinDisplayName(const Frontend::IInputController& InInput) const
+	{
+		return InInput.GetDisplayName();
+	}
+
+	FText FNodeTemplateBase::GetOutputPinDisplayName(const Frontend::IOutputController& InOutput) const
+	{
+		return InOutput.GetDisplayName();
+	}
+
+	bool FNodeTemplateBase::HasRequiredConnections(const FMetaSoundFrontendDocumentBuilder& InBuilder, const FGuid& InNodeID, FString* OutMessage) const
+	{
+		return true;
+	}
+#endif // WITH_EDITOR
+
 	void RegisterNodeTemplate(TUniquePtr<INodeTemplate>&& InTemplate)
 	{
 		class FTemplateRegistryEntry : public INodeRegistryTemplateEntry
diff --git a/Engine/Plugins/Runtime/Metasound/Source/MetasoundFrontend/Private/NodeTemplates/MetasoundFrontendNodeTemplateReroute.cpp b/Engine/Plugins/Runtime/Metasound/Source/MetasoundFrontend/Private/NodeTemplates/MetasoundFrontendNodeTemplateReroute.cpp
--- a/Engine/Plugins/Runtime/Metasound/Source/MetasoundFrontend/Private/NodeTemplates/MetasoundFrontendNodeTemplateReroute.cpp
+++ b/Engine/Plugins/Runtime/Metasound/Source/MetasoundFrontend/Private/NodeTemplates/MetasoundFrontendNodeTemplateReroute.cpp
@@ -5,9 +5,11 @@
 #include "Algo/AnyOf.h"
 #include "MetasoundAssetManager.h"
 #include "MetasoundFrontendDataTypeRegistry.h"
+#include "MetasoundFrontendDocument.h"
 #include "MetasoundFrontendDocumentBuilder.h"
 #include "MetasoundFrontendNodeTemplateRegistry.h"
 #include "MetasoundFrontendRegistries.h"
+#include "MetasoundFrontendTransform.h"
 
 
 namespace Metasound::Frontend
@@ -85,6 +87,75 @@ namespace Metasound::Frontend
 
 			return false;
 		}
+
+		const FMetasoundFrontendVertex* FindReroutedOutputVertex(const FMetaSoundFrontendDocumentBuilder& InBuilder, const FMetasoundFrontendNode& InOutputOwningNode, const FMetasoundFrontendVertex& InOutputVertex)
+		{
+			const FMetasoundFrontendVertex* ReroutedVertex = &InOutputVertex;
+			if (const FMetasoundFrontendClass* Class = InBuilder.FindDependency(InOutputOwningNode.ClassID))
+			{
+				if (Class->Metadata.GetClassName() == FRerouteNodeTemplate::ClassName)
+				{
+					const FGuid& OutputOwningNodeID = InOutputOwningNode.GetID();
+					TArray<const FMetasoundFrontendVertex*> Inputs = InBuilder.FindNodeInputs(OutputOwningNodeID);
+					if (!Inputs.IsEmpty())
+					{
+						if (const FMetasoundFrontendVertex* RerouteInput = Inputs.Last())
+						{
+							const FMetasoundFrontendNode* ConnectedNode = nullptr;
+							if (const FMetasoundFrontendVertex* OutputVertex = InBuilder.FindNodeOutputConnectedToNodeInput(OutputOwningNodeID, RerouteInput->VertexID, &ConnectedNode))
+							{
+								check(ConnectedNode);
+								return FindReroutedOutputVertex(InBuilder, *ConnectedNode, *OutputVertex);
+							}
+						}
+					}
+				}
+			}
+
+			return ReroutedVertex;
+		}
+
+		void FindReroutedInputVertices(
+			const FMetaSoundFrontendDocumentBuilder& InBuilder,
+			const FMetasoundFrontendNode& InInputOwningNode,
+			const FMetasoundFrontendVertex& InInputVertex,
+			TArray<const FMetasoundFrontendNode*>& InOutReroutedInputOwningNodes,
+			TArray<const FMetasoundFrontendVertex*>& InOutReroutedInputVertices)
+		{
+			using namespace Frontend;
+
+			if (const FMetasoundFrontendClass* Class = InBuilder.FindDependency(InInputOwningNode.ClassID))
+			{
+				if (Class->Metadata.GetClassName() == FRerouteNodeTemplate::ClassName)
+				{
+					const FGuid& InputOwningNodeID = InInputOwningNode.GetID();
+					TArray<const FMetasoundFrontendVertex*> Outputs = InBuilder.FindNodeOutputs(InputOwningNodeID);
+					for (const FMetasoundFrontendVertex* Output : Outputs)
+					{
+						check(Output);
+
+						TArray<const FMetasoundFrontendNode*> ConnectedInputNodes;
+						TArray<const FMetasoundFrontendVertex*> ConnectedInputVertices = InBuilder.FindNodeInputsConnectedToNodeOutput(InputOwningNodeID, Output->VertexID, &ConnectedInputNodes);
+						check(ConnectedInputNodes.Num() == ConnectedInputVertices.Num());
+
+						for (int32 Index = 0; Index < ConnectedInputNodes.Num(); ++Index)
+						{
+							const FMetasoundFrontendNode* ConnectedInputOwningNode = ConnectedInputNodes[Index];
+							check(ConnectedInputOwningNode);
+							const FMetasoundFrontendVertex* ConnectedInputVertex = ConnectedInputVertices[Index];
+							check(ConnectedInputVertex);
+
+							FindReroutedInputVertices(InBuilder, *ConnectedInputOwningNode, *ConnectedInputVertex, InOutReroutedInputOwningNodes, InOutReroutedInputVertices);
+						}
+					}
+
+					return;
+				}
+			}
+
+			InOutReroutedInputOwningNodes.Add(&InInputOwningNode);
+			InOutReroutedInputVertices.Add(&InInputVertex);
+		}
 	} // namespace ReroutePrivate
 
 	const FMetasoundFrontendClassName FRerouteNodeTemplate::ClassName { "UE", "Reroute", "" };
@@ -93,12 +164,44 @@ namespace Metasound::Frontend
 
 	const FMetasoundFrontendClassName& FRerouteNodeTemplate::GetClassName() const
 	{
-		return ClassName;
+		return FRerouteNodeTemplate::ClassName;
 	}
 
-	TUniquePtr<INodeTransform> FRerouteNodeTemplate::GenerateNodeTransform(FMetasoundFrontendDocument& InPreprocessedDocument) const
+#if WITH_EDITOR
+	FText FRerouteNodeTemplate::GetNodeDisplayName(const IMetaSoundDocumentInterface& DocumentInterface, const FGuid& InNodeID) const
 	{
-		return GenerateNodeTransform();
+		return { };
+	}
+#endif // WITH_EDITOR
+
+	FMetasoundFrontendNodeInterface FRerouteNodeTemplate::GenerateNodeInterface(FNodeTemplateGenerateInterfaceParams InParams) const
+	{
+		FName DataType;
+
+		if (!InParams.InputsToConnect.IsEmpty())
+		{
+			DataType = InParams.InputsToConnect.Last();
+		}
+
+		if (!InParams.OutputsToConnect.IsEmpty())
+		{
+			const FName OutputDataType = InParams.OutputsToConnect.Last();
+			if (DataType.IsNone())
+			{
+				DataType = OutputDataType;
+			}
+			else
+			{
+				checkf(DataType == OutputDataType, TEXT("Cannot generate MetasoundFrontendNodeInterface via reroute template with params of unmatched or unset input/output DataType"));
+			}
+		}
+
+		static const FName VertexName = "Value";
+		FMetasoundFrontendNodeInterface NewInterface;
+		NewInterface.Inputs.Add(FMetasoundFrontendVertex { VertexName, DataType, FGuid::NewGuid() });
+		NewInterface.Outputs.Add(FMetasoundFrontendVertex { VertexName, DataType, FGuid::NewGuid() });
+
+		return NewInterface;
 	}
 
 	TUniquePtr<INodeTransform> FRerouteNodeTemplate::GenerateNodeTransform() const
@@ -138,17 +241,6 @@ namespace Metasound::Frontend
 		return FrontendClass;
 	}
 
-	FMetasoundFrontendNodeInterface FRerouteNodeTemplate::CreateNodeInterfaceFromDataType(FName InDataType)
-	{
-		auto CreateNewVertex = [&] { return FMetasoundFrontendVertex { "Value", InDataType, FGuid::NewGuid() }; };
-
-		FMetasoundFrontendNodeInterface NewInterface;
-		NewInterface.Inputs.Add(CreateNewVertex());
-		NewInterface.Outputs.Add(CreateNewVertex());
-
-		return NewInterface;
-	}
-
 	EMetasoundFrontendVertexAccessType FRerouteNodeTemplate::GetNodeInputAccessType(const FMetaSoundFrontendDocumentBuilder& InBuilder, const FGuid& InNodeID, const FGuid& InVertexID) const
 	{
 		// Recursive search up DAG for first connected non-reroute node's input access type
@@ -216,17 +308,29 @@ namespace Metasound::Frontend
 	}
 
 #if WITH_EDITOR
-	bool FRerouteNodeTemplate::HasRequiredConnections(FConstNodeHandle InNodeHandle, FString* OutMessage) const
+	bool FRerouteNodeTemplate::HasRequiredConnections(const FMetaSoundFrontendDocumentBuilder& InBuilder, const FGuid& InNodeID, FString* OutMessage) const
 	{
-		TArray<FConstOutputHandle> Outputs = InNodeHandle->GetConstOutputs();
-		TArray<FConstInputHandle> Inputs = InNodeHandle->GetConstInputs();
+		const FMetasoundFrontendNode* Node = InBuilder.FindNode(InNodeID);
+		if (!Node)
+		{
+			return false;
+		}
+
+		TArray<const FMetasoundFrontendVertex*> Outputs = InBuilder.FindNodeOutputs(InNodeID);
+		const bool bConnectedToNonRerouteOutputs = Algo::AnyOf(Outputs, [&InBuilder, &Node](const FMetasoundFrontendVertex* OutputVertex)
+		{
+			using namespace ReroutePrivate;
+			return FindReroutedOutputVertex(InBuilder, *Node, *OutputVertex) != nullptr;
+		});
 
-		const bool bConnectedToNonRerouteOutputs = Algo::AnyOf(Outputs, [](const FConstOutputHandle& OutputHandle) { return Frontend::FindReroutedOutput(OutputHandle)->IsValid(); });
-		const bool bConnectedToNonRerouteInputs = Algo::AnyOf(Inputs, [](const FConstInputHandle& InputHandle)
+		TArray<const FMetasoundFrontendVertex*> Inputs = InBuilder.FindNodeInputs(InNodeID);
+		const bool bConnectedToNonRerouteInputs = Algo::AnyOf(Inputs, [&InBuilder, &Node](const FMetasoundFrontendVertex* InputVertex)
 		{
-			TArray<FConstInputHandle> Inputs;
-			Frontend::FindReroutedInputs(InputHandle, Inputs);
-			return !Inputs.IsEmpty();
+			using namespace ReroutePrivate;
+			TArray<const FMetasoundFrontendVertex*> InputVertices;
+			TArray<const FMetasoundFrontendNode*> InputVerticesOwningNodes;
+			FindReroutedInputVertices(InBuilder, *Node, *InputVertex, InputVerticesOwningNodes, InputVertices);
+			return !InputVertices.IsEmpty();
 		});
 
 		const bool bHasRequiredConnections = bConnectedToNonRerouteOutputs || bConnectedToNonRerouteOutputs == bConnectedToNonRerouteInputs;
@@ -244,6 +348,16 @@ namespace Metasound::Frontend
 		return true;
 	}
 
+	bool FRerouteNodeTemplate::IsInputConnectionUserModifiable() const
+	{
+		return true;
+	}
+
+	bool FRerouteNodeTemplate::IsOutputConnectionUserModifiable() const
+	{
+		return true;
+	}
+
 	bool FRerouteNodeTemplate::IsOutputAccessTypeDynamic() const
 	{
 		return true;
diff --git a/Engine/Plugins/Runtime/Metasound/Source/MetasoundFrontend/Public/MetasoundFrontendController.h b/Engine/Plugins/Runtime/Metasound/Source/MetasoundFrontend/Public/MetasoundFrontendController.h
--- a/Engine/Plugins/Runtime/Metasound/Source/MetasoundFrontend/Public/MetasoundFrontendController.h
+++ b/Engine/Plugins/Runtime/Metasound/Source/MetasoundFrontend/Public/MetasoundFrontendController.h
@@ -7,6 +7,7 @@
 #include "MetasoundFrontend.h"
 #include "MetasoundFrontendDocument.h"
 #include "MetasoundFrontendDocumentAccessPtr.h"
+#include "MetasoundFrontendNodeTemplateRegistry.h"
 #include "MetasoundFrontendRegistries.h"
 #include "MetasoundGraph.h"
 #include "MetasoundVertex.h"
@@ -50,6 +51,8 @@ namespace Metasound
 		class IOutputController;
 		class IVariableController;
 
+		struct FNodeTemplateGenerateInterfaceParams;
+
 		// Metasound Frontend Handles are all TSharedRefs of various Metasound Frontend Controllers.
 		using FInputHandle = TSharedRef<IInputController>;
 		using FConstInputHandle = TSharedRef<const IInputController>;
@@ -828,9 +831,6 @@ namespace Metasound
 			 */
 			virtual FNodeHandle AddInputVertex(const FMetasoundFrontendClassInput& InDescription) = 0;
 
-			UE_DEPRECATED(5.1, "Use AddInputVertex method which specifies EMetasoundFrontendVertexAccessType")
-			virtual FNodeHandle AddInputVertex(const FVertexName& InName, const FName InTypeName, const FMetasoundFrontendLiteral* InDefaultValue) = 0;
-
 			/** Remove the input with the given name. Returns true if successfully removed, false otherwise. */
 			virtual bool RemoveInputVertex(const FVertexName& InputName) = 0;
 
@@ -960,14 +960,16 @@ namespace Metasound
 
 			/** Add a new template node to this graph, providing the defined interface as expected by the caller.
 			 *
-			 * @param InKey - Class key (must correspond with a class in the registry that was registered as a template).
-			 * @param InNodeInterface - Interface for node class.  Validated by template class registered in the node class registry. If invalid, node is not created/added.
+			 * @param InNodeTemplate - Node template to use to generate node
 			 * @param InNodeGuid - (Optional) Explicit guid for the new node. Must be unique within the graph.
 			 * Only useful to specify explicitly when caller is managing or tracking the graph's guids (ex. replacing removed node).
 			 *
 			 * @return Node handle for class. On error, an invalid handle is returned.
 			 */
-			virtual FNodeHandle AddTemplateNode(const FNodeRegistryKey& InKey, FMetasoundFrontendNodeInterface&& InNodeInterface, FGuid InNodeGuid = FGuid::NewGuid()) = 0;
+			virtual FNodeHandle AddTemplateNode(const INodeTemplate& InNodeTemplate, FNodeTemplateGenerateInterfaceParams Params, FGuid InNodeGuid = FGuid::NewGuid()) = 0;
+
+			UE_DEPRECATED(5.5, "Use overload that provides INodeTemplate instead")
+			virtual FNodeHandle AddTemplateNode(const FNodeRegistryKey& InKey, FMetasoundFrontendNodeInterface&& InNodeInterface, FGuid InNodeGuid = FGuid::NewGuid()) { return INodeController::GetInvalidHandle(); }
 
 			/** Remove the node corresponding to this node handle.
 			 *
@@ -1133,7 +1135,10 @@ namespace Metasound
 		};
 
 		METASOUNDFRONTEND_API FConstOutputHandle FindReroutedOutput(FConstOutputHandle InOutputHandle);
+
+		UE_DEPRECATED(5.5, "Reroute input recursion is now privately implemented in the reroute template node class.")
 		METASOUNDFRONTEND_API void FindReroutedInputs(FConstInputHandle InHandleToCheck, TArray<FConstInputHandle>& InOutInputHandles);
+
 		METASOUNDFRONTEND_API void IterateReroutedInputs(FConstInputHandle InHandleToCheck, TFunctionRef<void(FConstInputHandle)> Func);
 	} // namespace Frontend
 } // namespace Metasound
diff --git a/Engine/Plugins/Runtime/Metasound/Source/MetasoundFrontend/Public/MetasoundFrontendNodeTemplateRegistry.h b/Engine/Plugins/Runtime/Metasound/Source/MetasoundFrontend/Public/MetasoundFrontendNodeTemplateRegistry.h
--- a/Engine/Plugins/Runtime/Metasound/Source/MetasoundFrontend/Public/MetasoundFrontendNodeTemplateRegistry.h
+++ b/Engine/Plugins/Runtime/Metasound/Source/MetasoundFrontend/Public/MetasoundFrontendNodeTemplateRegistry.h
@@ -2,20 +2,39 @@
 
 #pragma once
 
-#include "MetasoundFrontendController.h"
 #include "MetasoundFrontendDocument.h"
-#include "MetasoundFrontendTransform.h"
+#include "MetasoundFrontendRegistryKey.h"
+#include "MetasoundVertex.h"
 #include "Templates/SharedPointer.h"
 #include "Templates/UniquePtr.h"
 #include "UObject/NoExportTypes.h"
-
+#include "UObject/ScriptInterface.h"
 
 // Forward Declarations
+class IMetaSoundDocumentInterface;
 struct FMetaSoundFrontendDocumentBuilder;
 
 
 namespace Metasound::Frontend
 {
+	class INodeTransform;
+
+#if WITH_EDITOR
+	class INodeController;
+	class IInputController;
+	class IOutputController;
+
+	using FConstNodeHandle = TSharedRef<const INodeController>;
+#endif // WITH_EDITOR
+
+	struct METASOUNDFRONTEND_API FNodeTemplateGenerateInterfaceParams
+	{
+		// Currently just DataType FName. Subsequent change will replace this with a VertexHandle
+		// and will add builder reference to once builder supports template nodes.
+		TArray<FName> InputsToConnect;
+		TArray<FName> OutputsToConnect;
+	};
+
 	/**
 		* Base interface for a node template, which acts in place of frontend node class and respective instance(s).
 		* Instances are preprocessed, allowing for custom graph manipulation prior to generating a respective runtime
@@ -24,16 +43,24 @@ namespace Metasound::Frontend
 	class METASOUNDFRONTEND_API INodeTemplate
 	{
 	public:
+		virtual FMetasoundFrontendNodeInterface GenerateNodeInterface(FNodeTemplateGenerateInterfaceParams InParams) const = 0;
+
 		virtual ~INodeTemplate() = default;
 
 		// Returns note template class name.
 		virtual const FMetasoundFrontendClassName& GetClassName() const = 0;
 
 		UE_DEPRECATED(5.4, "Use version that does not provide a preprocessed document")
-		virtual TUniquePtr<INodeTransform> GenerateNodeTransform(FMetasoundFrontendDocument& InDocument) const { return nullptr; }
+		virtual TUniquePtr<INodeTransform> GenerateNodeTransform(FMetasoundFrontendDocument& InDocument) const;
+
+#if WITH_EDITOR
+		virtual FText GetNodeDisplayName(const IMetaSoundDocumentInterface& Interface, const FGuid& InNodeID) const = 0;
+		virtual FText GetInputPinDisplayName(const Frontend::IInputController& InInput) const = 0;
+		virtual FText GetOutputPinDisplayName(const Frontend::IOutputController& InOutput) const = 0;
+#endif // WITH_EDITOR
 
 		// Generates node transform that is used to preprocess nodes.
-		virtual TUniquePtr<INodeTransform> GenerateNodeTransform() const { return nullptr; }
+		virtual TUniquePtr<INodeTransform> GenerateNodeTransform() const = 0;
 
 		// Returns the class definition for the given node class template.
 		virtual const FMetasoundFrontendClass& GetFrontendClass() const = 0;
@@ -51,31 +78,31 @@ namespace Metasound::Frontend
 		virtual const FMetasoundFrontendVersionNumber& GetVersionNumber() const = 0;
 
 #if WITH_EDITOR
+		UE_DEPRECATED(5.5, "Use version that takes a Frontend NodeID and builder instead")
+		virtual bool HasRequiredConnections(FConstNodeHandle InNodeHandle, FString* OutMessage = nullptr) const { return false; }
+
 		// Returns whether or not the given node template has the necessary
 		// required connections to be preprocessed (editor only).
-		virtual bool HasRequiredConnections(FConstNodeHandle InNodeHandle, FString* OutMessage = nullptr) const = 0;
+		virtual bool HasRequiredConnections(const FMetaSoundFrontendDocumentBuilder& InBuilder, const FGuid& InNodeID, FString* OutMessage = nullptr) const = 0;
 #endif // WITH_EDITOR
 
 		// Returns whether template can dynamically assign a node's input access type (as opposed to it being assigned on the class input definition)
 		virtual bool IsInputAccessTypeDynamic() const = 0;
 
+		// Whether or not input connections are user modifiable
+		virtual bool IsInputConnectionUserModifiable() const = 0;
+
 		// Returns whether template can dynamically assign a node's output's access type (as opposed to it being assigned on the class output definition)
 		virtual bool IsOutputAccessTypeDynamic() const = 0;
 
+		// Whether or not output connections are user modifiable
+		virtual bool IsOutputConnectionUserModifiable() const = 0;
+
 		// Given the provided node interface, returns whether or not it conforms to an expected format
 		// that can be successfully manipulated by a generated node template transform.
 		virtual bool IsValidNodeInterface(const FMetasoundFrontendNodeInterface& InNodeInterface) const = 0;
 	};
 
-	class METASOUNDFRONTEND_API FNodeTemplateBase : public INodeTemplate
-	{
-		virtual bool IsValidNodeInterface(const FMetasoundFrontendNodeInterface& InNodeInterface) const override;
-		virtual bool IsInputAccessTypeDynamic() const override;
-		virtual bool IsOutputAccessTypeDynamic() const override;
-		virtual EMetasoundFrontendVertexAccessType GetNodeInputAccessType(const FMetaSoundFrontendDocumentBuilder& InBuilder, const FGuid& InNodeID, const FGuid& InVertexID) const override;
-		virtual EMetasoundFrontendVertexAccessType GetNodeOutputAccessType(const FMetaSoundFrontendDocumentBuilder& InBuilder, const FGuid& InNodeID, const FGuid& InVertexID) const override;
-	};
-
 	class METASOUNDFRONTEND_API INodeTemplateRegistry
 	{
 	public:
@@ -86,6 +113,20 @@ namespace Metasound::Frontend
 
 		// Find a template with the given key. Returns null if entry not found with given key.
 		virtual const INodeTemplate* FindTemplate(const FNodeRegistryKey& InKey) const = 0;
+
+		// Find a template with the given class name with the highest version. Returns null if entry not found with given name.
+		virtual const INodeTemplate* FindTemplate(const FMetasoundFrontendClassName& InClassName) const = 0;
+	};
+
+	class METASOUNDFRONTEND_API FNodeTemplateBase : public INodeTemplate
+	{
+	public:
+#if WITH_EDITOR
+		virtual FText GetNodeDisplayName(const IMetaSoundDocumentInterface& Interface, const FGuid& InNodeID) const override;
+		virtual FText GetInputPinDisplayName(const IInputController& InInput) const override;
+		virtual FText GetOutputPinDisplayName(const IOutputController& InOutput) const override;
+		virtual bool HasRequiredConnections(const FMetaSoundFrontendDocumentBuilder& InBuilder, const FGuid& InNodeID, FString* OutMessage = nullptr) const override;
+#endif // WITH_EDITOR
 	};
 
 	// Register & Unregister are not publicly accessible implementation as the API
diff --git a/Engine/Plugins/Runtime/Metasound/Source/MetasoundFrontend/Public/MetasoundFrontendTransform.h b/Engine/Plugins/Runtime/Metasound/Source/MetasoundFrontend/Public/MetasoundFrontendTransform.h
--- a/Engine/Plugins/Runtime/Metasound/Source/MetasoundFrontend/Public/MetasoundFrontendTransform.h
+++ b/Engine/Plugins/Runtime/Metasound/Source/MetasoundFrontend/Public/MetasoundFrontendTransform.h
@@ -64,16 +64,10 @@ namespace Metasound::Frontend
 	public:
 		virtual ~INodeTransform() = default;
 
-		UE_DEPRECATED(5.3, "Deprecated: unused function which encouraged breaking const behavior of transform state.")
-		virtual FMetasoundFrontendDocument& GetOwningDocument() const { static FMetasoundFrontendDocument BaseDoc; return BaseDoc; };
-
-		UE_DEPRECATED(5.3, "Deprecated: unused function which encouraged breaking const behavior of transform state.")
-		virtual FMetasoundFrontendGraph& GetOwningGraph() const { static FMetasoundFrontendGraph BaseGraph; return BaseGraph; }
-
 		/** Return true if the node was modified, false otherwise. */
 		virtual bool Transform(const FGuid& InNodeID, FMetaSoundFrontendDocumentBuilder& OutBuilder) const;
 
-		UE_DEPRECATED(5.4, "Use transform which provides builder and const node representation to apply node transform to mutate underlying builder document")
+		UE_DEPRECATED(5.4, "Use transform overload with node ID and builder parameters to apply node transform to builder's underlying document")
 		virtual bool Transform(FMetasoundFrontendNode& InOutNode) const;
 	};
 
diff --git a/Engine/Plugins/Runtime/Metasound/Source/MetasoundFrontend/Public/NodeTemplates/MetasoundFrontendNodeTemplateReroute.h b/Engine/Plugins/Runtime/Metasound/Source/MetasoundFrontend/Public/NodeTemplates/MetasoundFrontendNodeTemplateReroute.h
--- a/Engine/Plugins/Runtime/Metasound/Source/MetasoundFrontend/Public/NodeTemplates/MetasoundFrontendNodeTemplateReroute.h
+++ b/Engine/Plugins/Runtime/Metasound/Source/MetasoundFrontend/Public/NodeTemplates/MetasoundFrontendNodeTemplateReroute.h
@@ -12,7 +12,7 @@
 
 namespace Metasound::Frontend
 {
-	class METASOUNDFRONTEND_API FRerouteNodeTemplate : public INodeTemplate
+	class METASOUNDFRONTEND_API FRerouteNodeTemplate : public FNodeTemplateBase
 	{
 	public:
 		static const FMetasoundFrontendClassName ClassName;
@@ -21,26 +21,34 @@ namespace Metasound::Frontend
 
 		static const FNodeRegistryKey& GetRegistryKey();
 
-		static FMetasoundFrontendNodeInterface CreateNodeInterfaceFromDataType(FName InDataType);
+		UE_DEPRECATED(5.5, "Look-up template via registry and use non-static GenerateNodeInterface instead with provided params")
+		static FMetasoundFrontendNodeInterface CreateNodeInterfaceFromDataType(FName InDataType) { return { }; }
 
 		virtual ~FRerouteNodeTemplate() = default;
 
 		virtual const FMetasoundFrontendClassName& GetClassName() const override;
 
+#if WITH_EDITOR
+		virtual FText GetNodeDisplayName(const IMetaSoundDocumentInterface& DocumentInterface, const FGuid& InNodeID) const override;
+#endif // WITH_EDITOR
+
 		UE_DEPRECATED(5.4, "Use version that does not require mutating a provided PreprocessedDocument")
-		virtual TUniquePtr<INodeTransform> GenerateNodeTransform(FMetasoundFrontendDocument& InPreprocessedDocument) const override;
+		virtual TUniquePtr<INodeTransform> GenerateNodeTransform(FMetasoundFrontendDocument& InPreprocessedDocument) const override { return { }; }
 
+		virtual FMetasoundFrontendNodeInterface GenerateNodeInterface(FNodeTemplateGenerateInterfaceParams InParams) const override;
 		virtual TUniquePtr<INodeTransform> GenerateNodeTransform() const override;
 		virtual const FMetasoundFrontendClass& GetFrontendClass() const override;
 		virtual EMetasoundFrontendVertexAccessType GetNodeInputAccessType(const FMetaSoundFrontendDocumentBuilder& InBuilder, const FGuid& InNodeID, const FGuid& InVertexID) const override;
 		virtual EMetasoundFrontendVertexAccessType GetNodeOutputAccessType(const FMetaSoundFrontendDocumentBuilder& InBuilder, const FGuid& InNodeID, const FGuid& InVertexID) const override;
 		virtual const FMetasoundFrontendVersionNumber& GetVersionNumber() const override;
 		virtual bool IsInputAccessTypeDynamic() const override;
+		virtual bool IsInputConnectionUserModifiable() const override;
+		virtual bool IsOutputConnectionUserModifiable() const override;
 		virtual bool IsOutputAccessTypeDynamic() const override;
 		virtual bool IsValidNodeInterface(const FMetasoundFrontendNodeInterface& InNodeInterface) const override;
 
 #if WITH_EDITOR
-		virtual bool HasRequiredConnections(FConstNodeHandle InNodeHandle, FString* OutMessage = nullptr) const override;
+		virtual bool HasRequiredConnections(const FMetaSoundFrontendDocumentBuilder& InBuilder, const FGuid& InNodeID, FString* OutMessage = nullptr) const override;
 #endif // WITH_EDITOR
 	};
 } // namespace Metasound::Frontend
