Removed unreachable code

#rnx

diff --git a/Engine/Plugins/Runtime/Metasound/Source/MetasoundEditor/Private/FindInMetasound.cpp b/Engine/Plugins/Runtime/Metasound/Source/MetasoundEditor/Private/FindInMetasound.cpp
--- a/Engine/Plugins/Runtime/Metasound/Source/MetasoundEditor/Private/FindInMetasound.cpp
+++ b/Engine/Plugins/Runtime/Metasound/Source/MetasoundEditor/Private/FindInMetasound.cpp
@@ -61,11 +61,7 @@ namespace Metasound::Editor
 		{
 			return LOCTEXT("FindMetasound_PinCategory", "Pin");
 		}
-		else
-		{
-			return LOCTEXT("FindMetasound_NodeCategory", "Node");
-		}
-		return FText::GetEmpty();
+		return LOCTEXT("FindMetasound_NodeCategory", "Node");
 	}
 
 	TSharedRef<SWidget> FFindInMetasoundResult::CreateIcon() const
