Fix non-cooked quality settings by adding a cook stage the MetasoundSource to correctly save them out per platform.
Previously the PerPlatformSettings was added to the MetasoundSettings, however as this is not a true asset they weren't being correctly cooked.
The code previously pulled the global settings, then allowed overrides upon this. This decision is now made during cook.
#tests Shelve tested back and forth with Sondra Moyles

diff --git a/Engine/Plugins/Runtime/Harmonix/Source/HarmonixMetasound/Private/HarmonixMetasound/Components/MetasoundOfflinePlayerComponent.cpp b/Engine/Plugins/Runtime/Harmonix/Source/HarmonixMetasound/Private/HarmonixMetasound/Components/MetasoundOfflinePlayerComponent.cpp
--- a/Engine/Plugins/Runtime/Harmonix/Source/HarmonixMetasound/Private/HarmonixMetasound/Components/MetasoundOfflinePlayerComponent.cpp
+++ b/Engine/Plugins/Runtime/Harmonix/Source/HarmonixMetasound/Private/HarmonixMetasound/Components/MetasoundOfflinePlayerComponent.cpp
@@ -40,10 +40,6 @@ UMetasoundGeneratorHandle* UMetasoundOfflinePlayerComponent::CreateGeneratorBase
 
 	TArray<FAudioParameter> DefaultParams;
 
-	// HACK: Override the MetaSound Source's sample rate to work around the bug in UMetaSoundSource::GetOperatorSettings
-	// where it will use the hard-coded default sample rate of 48000 instead of whatever you pass in.
-	MetasoundSource->SampleRateOverride = SampleRate;
-
 	Generator = MetasoundSource->CreateSoundGenerator(InitParams, TArray<FAudioParameter>());
 	if (!Generator)
 	{
@@ -93,10 +89,7 @@ TSharedPtr<Metasound::FMetasoundGeneratorHandle> UMetasoundOfflinePlayerComponen
 
 	TArray<FAudioParameter> DefaultParams;
 
-	// HACK: Override the MetaSound Source's sample rate to work around the bug in UMetaSoundSource::GetOperatorSettings
-	// where it will use the hard-coded default sample rate of 48000 instead of whatever you pass in.
-	MetasoundSource->SampleRateOverride = SampleRate;
-
+	
 	Generator = MetasoundSource->CreateSoundGenerator(InitParams, TArray<FAudioParameter>());
 	if (!Generator)
 	{
diff --git a/Engine/Plugins/Runtime/Metasound/Source/MetasoundEngine/Private/MetasoundBuilderSubsystem.cpp b/Engine/Plugins/Runtime/Metasound/Source/MetasoundEngine/Private/MetasoundBuilderSubsystem.cpp
--- a/Engine/Plugins/Runtime/Metasound/Source/MetasoundEngine/Private/MetasoundBuilderSubsystem.cpp
+++ b/Engine/Plugins/Runtime/Metasound/Source/MetasoundEngine/Private/MetasoundBuilderSubsystem.cpp
@@ -1593,7 +1593,9 @@ void UMetaSoundSourceBuilder::OnRemovingOutput(int32 OutputIndex) const
 
 void UMetaSoundSourceBuilder::SetBlockRateOverride(float BlockRate)
 {
+#if WITH_EDITORONLY_DATA
 	GetMetaSoundSource().BlockRateOverride.Default = BlockRate;
+#endif //WITH_EDITORONLY_DATA
 }
 
 void UMetaSoundSourceBuilder::SetFormat(EMetaSoundOutputAudioFormat OutputFormat, EMetaSoundBuilderResult& OutResult)
@@ -1653,12 +1655,16 @@ void UMetaSoundSourceBuilder::SetPlatformSampleRateOverride(const FPerPlatformIn
 
 void UMetaSoundSourceBuilder::SetQuality(FName Quality)
 {
+#if WITH_EDITORONLY_DATA
 	GetMetaSoundSource().QualitySetting = Quality;
+#endif //WITH_EDITORONLY_DATA	
 }
 
 void UMetaSoundSourceBuilder::SetSampleRateOverride(int32 SampleRate)
 {
+#if WITH_EDITORONLY_DATA
 	GetMetaSoundSource().SampleRateOverride.Default = SampleRate;
+#endif //WITH_EDITORONLY_DATA	
 }
 
 UMetaSoundPatchBuilder* UMetaSoundBuilderSubsystem::CreatePatchBuilder(FName BuilderName, EMetaSoundBuilderResult& OutResult)
diff --git a/Engine/Plugins/Runtime/Metasound/Source/MetasoundEngine/Private/MetasoundSettings.cpp b/Engine/Plugins/Runtime/Metasound/Source/MetasoundEngine/Private/MetasoundSettings.cpp
--- a/Engine/Plugins/Runtime/Metasound/Source/MetasoundEngine/Private/MetasoundSettings.cpp
+++ b/Engine/Plugins/Runtime/Metasound/Source/MetasoundEngine/Private/MetasoundSettings.cpp
@@ -113,6 +113,11 @@ void UMetaSoundSettings::PostEditChangeProperty(FPropertyChangedEvent& PropertyC
 	DenyListCacheChangeID++;
 }
 
+FName UMetaSoundSettings::GetQualitySettingPropertyName()
+{
+	return GET_MEMBER_NAME_CHECKED(UMetaSoundSettings, QualitySettings);
+}
+
 #endif //WITH_EDITORONLY_DATA
 
 TArray<FName> UMetaSoundQualityHelper::GetQualityList()
@@ -122,7 +127,7 @@ TArray<FName> UMetaSoundQualityHelper::GetQualityList()
 #if WITH_EDITORONLY_DATA
 	if (const UMetaSoundSettings* Settings = GetDefault<UMetaSoundSettings>())
 	{
-		Algo::Transform(Settings->QualitySettings, Names, [](const FMetaSoundQualitySettings& Quality) -> FName
+		Algo::Transform(Settings->GetQualitySettings(), Names, [](const FMetaSoundQualitySettings& Quality) -> FName
 		{
 			return Quality.Name;
 		});
diff --git a/Engine/Plugins/Runtime/Metasound/Source/MetasoundEngine/Private/MetasoundSource.cpp b/Engine/Plugins/Runtime/Metasound/Source/MetasoundEngine/Private/MetasoundSource.cpp
--- a/Engine/Plugins/Runtime/Metasound/Source/MetasoundEngine/Private/MetasoundSource.cpp
+++ b/Engine/Plugins/Runtime/Metasound/Source/MetasoundEngine/Private/MetasoundSource.cpp
@@ -8,6 +8,7 @@
 #include "AudioDeviceManager.h"
 #include "Containers/Ticker.h"
 #include "IAudioParameterInterfaceRegistry.h"
+#include "Interfaces/ITargetPlatform.h"
 #include "Interfaces/MetasoundOutputFormatInterfaces.h"
 #include "Interfaces/MetasoundFrontendSourceInterface.h"
 #include "Internationalization/Text.h"
@@ -59,6 +60,18 @@ namespace Metasound
 
 	namespace SourcePrivate
 	{
+		static constexpr float DefaultBlockRateConstant = 100.f;
+		static constexpr float DefaultSampleRateConstant = 48000.f;
+
+		static bool IsCookedForEditor(const FArchive& InArchive, const UObject* InObj)
+		{
+#if WITH_EDITORONLY_DATA
+			return ((InArchive.GetPortFlags() & PPF_Duplicate) == 0) && InObj->GetPackage()->HasAnyPackageFlags(PKG_Cooked);
+#else //WITH_EDITORONLY_DATA
+			return false;
+#endif //WITH_EDITORONLY_DATA
+		}
+		
 		static const FLazyName TriggerName = "Trigger";
 
 		// Holds onto a global static TSet for tracking which error/warning logs have been
@@ -366,8 +379,8 @@ void UMetaSoundSource::PostEditChangeProperty(FPropertyChangedEvent& InEvent)
 	{
 		PostEditChangeOutputFormat();
 	}
-	if (InEvent.GetPropertyName() == GET_MEMBER_NAME_CHECKED(UMetaSoundSource, SampleRateOverride) ||
-		InEvent.GetPropertyName() == GET_MEMBER_NAME_CHECKED(UMetaSoundSource, BlockRateOverride) ||
+	if (InEvent.GetMemberPropertyName() == GET_MEMBER_NAME_CHECKED(UMetaSoundSource, SampleRateOverride) ||
+		InEvent.GetMemberPropertyName() == GET_MEMBER_NAME_CHECKED(UMetaSoundSource, BlockRateOverride) ||
 		InEvent.GetPropertyName() == GET_MEMBER_NAME_CHECKED(UMetaSoundSource, QualitySetting) )
 	{
 		PostEditChangeQualitySettings();
@@ -444,9 +457,6 @@ void UMetaSoundSource::PostEditChangeOutputFormat()
 
 void UMetaSoundSource::PostEditChangeQualitySettings()
 {
-	// Re-cache Operator settings by clearing the Optional.
-	OperatorSettings.Reset();
-
 	// Refresh the SampleRate (which is what the engine sees from the operator settings).
 	SampleRate = GetOperatorSettings(CachedAudioDeviceSampleRate).GetSampleRate();
 
@@ -454,7 +464,7 @@ void UMetaSoundSource::PostEditChangeQualitySettings()
 	if (const UMetaSoundSettings* Settings = GetDefault<UMetaSoundSettings>())	
 	{
 		auto FindByName = [&Name = QualitySetting](const FMetaSoundQualitySettings& Q) -> bool { return Q.Name == Name; };
-		if (const FMetaSoundQualitySettings* Found = Settings->QualitySettings.FindByPredicate(FindByName))
+		if (const FMetaSoundQualitySettings* Found = Settings->GetQualitySettings().FindByPredicate(FindByName))
 		{
 			QualitySettingGuid = Found->UniqueId;
 		}
@@ -515,6 +525,89 @@ void UMetaSoundSource::Serialize(FArchive& InArchive)
 {
 	Super::Serialize(InArchive);
 	Metasound::FMetaSoundEngineAssetHelper::SerializeToArchive(*this, InArchive);
+	
+	using namespace Metasound::SourcePrivate;	
+	
+	// Load/Save cooked data.
+	if (InArchive.IsCooking() || (FPlatformProperties::RequiresCookedData() && InArchive.IsLoading()) || IsCookedForEditor(InArchive, this))
+	{
+		const FName PlatformName = InArchive.CookingTarget() ? *InArchive.CookingTarget()->IniPlatformName() : FName(FPlatformProperties::IniPlatformName());
+		SerializeCookedQualitySettings(PlatformName, InArchive);
+	}
+}
+
+bool UMetaSoundSource::GetQualitySettings(
+	const FName InPlatformName, Metasound::SourcePrivate::FCookedQualitySettings& OutQualitySettings) const
+{
+#if WITH_EDITORONLY_DATA
+	
+	// Query Project settings. 
+	if (const UMetaSoundSettings* Settings = GetDefault<UMetaSoundSettings>())
+	{
+		if (const FMetaSoundQualitySettings* Found = Settings->GetQualitySettings().FindByPredicate([&QT = QualitySetting](const FMetaSoundQualitySettings& Q) -> bool { return Q.Name == QT; }))
+		{
+			// Allow partial applications of settings, if some are non-zero.
+			if (const float Value = Found->BlockRate.GetValueForPlatform(InPlatformName); Value > 0.f)
+			{
+				UE_LOG(LogMetaSound, VeryVerbose, TEXT("Metasound [%s] using Quality '%s', BlockRate=%3.3f" ), *GetName(), *QualitySetting.ToString(), Value);
+				OutQualitySettings.BlockRate = Value;
+			}
+			if (const int32 Value = Found->SampleRate.GetValueForPlatform(InPlatformName); Value > 0)
+			{
+				UE_LOG(LogMetaSound, VeryVerbose, TEXT("Metasound [%s] using Quality '%s', SampleRate=%d" ), *GetName(), *QualitySetting.ToString(), Value);
+				OutQualitySettings.SampleRate = Value;
+			}
+		}
+	}
+
+	// Query overrides defined on this asset.
+	if (const float SerializedBlockRate = BlockRateOverride.GetValueForPlatform(InPlatformName); SerializedBlockRate > 0.0f)
+	{
+		UE_LOG(LogMetaSound, VeryVerbose, TEXT("Metasound [%s] BlockRate Override: %3.3f"), *GetName(), SerializedBlockRate);
+		OutQualitySettings.BlockRate = SerializedBlockRate;
+	}
+	if (const int32 SerializedSampleRate = SampleRateOverride.GetValueForPlatform(InPlatformName); SerializedSampleRate > 0)
+	{
+		UE_LOG(LogMetaSound, VeryVerbose, TEXT("Metasound [%s] SampleRate Override: %d"), *GetName(), SerializedSampleRate);
+		OutQualitySettings.SampleRate = SerializedSampleRate;
+	}
+	
+	// Success.
+	UE_LOG(LogMetaSound, Verbose, TEXT("Metasound [%s] using SampleRate=%d, BlockRate=%2.3f (not-cooked)"),
+	 	*GetName(), OutQualitySettings.SampleRate.GetValue(), OutQualitySettings.BlockRate.GetValue()) ;
+	
+	return true;
+
+#else //WITH_EDITORONLY_DATA	
+
+	// If we've been cooked, this should contain the quality settings.
+	if (CookedQualitySettings)
+	{
+		OutQualitySettings = *CookedQualitySettings;
+		UE_LOG(LogMetaSound, Verbose, TEXT("Metasound [%s] using SampleRate=%d, BlockRate=%2.3f (cooked)"),
+			*GetName(), OutQualitySettings.SampleRate.GetValue(), OutQualitySettings.BlockRate.GetValue()) ;
+		return true;
+	}
+	
+	// Fail.
+	return false;
+	
+#endif //WITH_EDITORONLY_DATA
+}
+void UMetaSoundSource::SerializeCookedQualitySettings(const FName PlatformName, FArchive& Ar)
+{
+	Metasound::SourcePrivate::FCookedQualitySettings Settings;
+	if (Ar.IsSaving())
+	{
+		GetQualitySettings(PlatformName, Settings);
+	}
+	// Use Struct Serializer.
+	FMetaSoundQualitySettings::StaticStruct()->SerializeItem(Ar,&Settings,nullptr);
+		
+	if (Ar.IsLoading())
+	{
+		CookedQualitySettings = MakePimpl<Metasound::SourcePrivate::FCookedQualitySettings>(Settings);
+	}
 }
 
 #if WITH_EDITOR
@@ -601,7 +694,7 @@ void UMetaSoundSource::PostLoadQualitySettings()
 		{
 			if (
 				WeakSource.IsValid() &&
-				InEvent.GetMemberPropertyName() == GET_MEMBER_NAME_CHECKED(UMetaSoundSettings, QualitySettings)
+				InEvent.GetMemberPropertyName() == UMetaSoundSettings::GetQualitySettingPropertyName()
 			)
 			{
 				WeakSource->ResolveQualitySettings(CastChecked<UMetaSoundSettings>(InObj));
@@ -616,8 +709,6 @@ void UMetaSoundSource::PostLoadQualitySettings()
 			{
 				WeakSource->ResolveQualitySettings(GetMutableDefault<UMetaSoundSettings>());
 			
-				WeakSource->OperatorSettings.Reset();
-						
 				// Override SampleRate with the Operator settings version which uses our Quality settings.
 				WeakSource->SampleRate = WeakSource->GetOperatorSettings(WeakSource->CachedAudioDeviceSampleRate).GetSampleRate();
 			}
@@ -630,19 +721,20 @@ void UMetaSoundSource::PostLoadQualitySettings()
 
 void UMetaSoundSource::ResolveQualitySettings(const UMetaSoundSettings* Settings)
 {
+#if WITH_EDITORONLY_DATA
+
 	const FMetaSoundQualitySettings* Resolved = nullptr;
 
 	// 1. Try and resolve by name. (most should resolve unless its been renamed, deleted).
 	auto FindByName = [&Name = QualitySetting](const FMetaSoundQualitySettings& Q) -> bool { return Q.Name == Name; };
-	Resolved = Settings->QualitySettings.FindByPredicate(FindByName);
+	Resolved = Settings->GetQualitySettings().FindByPredicate(FindByName);
 
-#if WITH_EDITORONLY_DATA
 
 	// 2. If that failed, try by guid (if its been renamed in the settings, we can still find it).
 	if (!Resolved && QualitySettingGuid.IsValid())
 	{
 		auto FindByGuid = [&Guid = QualitySettingGuid](const FMetaSoundQualitySettings& Q) -> bool { return Q.UniqueId == Guid; };
-		Resolved = Settings->QualitySettings.FindByPredicate(FindByName);
+		Resolved = Settings->GetQualitySettings().FindByPredicate(FindByName);
 	}
 
 	// 3. If still failed to resolve, use defaults and warn.
@@ -654,11 +746,11 @@ void UMetaSoundSource::ResolveQualitySettings(const UMetaSoundSettings* Settings
 		// Reset to defaults. (and make sure they are sane)
 		QualitySetting = GetDefault<UMetaSoundSource>()->QualitySetting;
 		QualitySettingGuid = GetDefault<UMetaSoundSource>()->QualitySettingGuid;
-		if (!Settings->QualitySettings.FindByPredicate(FindByName) && !Settings->QualitySettings.IsEmpty())
+		if (!Settings->GetQualitySettings().FindByPredicate(FindByName) && !Settings->GetQualitySettings().IsEmpty())
 		{
 			// Default doesn't point to anything, use first one in the list.
-			QualitySetting = Settings->QualitySettings[0].Name;
-			QualitySettingGuid = Settings->QualitySettings[0].UniqueId;
+			QualitySetting = Settings->GetQualitySettings()[0].Name;
+			QualitySettingGuid = Settings->GetQualitySettings()[0].UniqueId;
 		}				
 	}
 
@@ -1324,79 +1416,41 @@ TSharedPtr<Audio::IParameterTransmitter> UMetaSoundSource::CreateParameterTransm
 
 Metasound::FOperatorSettings UMetaSoundSource::GetOperatorSettings(Metasound::FSampleRate InDeviceSampleRate) const
 {	
-	// We should recache the operator settings if the device rate has changed.
-	if (InDeviceSampleRate != CachedAudioDeviceSampleRate)
+	using namespace Metasound;
+	using namespace Metasound::SourcePrivate;
+	
+	// Default sensibly.
+	FCookedQualitySettings Settings;
+	Settings.BlockRate = DefaultBlockRateConstant;
+	Settings.SampleRate = InDeviceSampleRate > 0 ? InDeviceSampleRate : DefaultSampleRateConstant;
+
+	// Fetch our quality settings.
+	// If we are cooked these are baked, if we are editor these are queried from the project settings and this assets overrides.
+	ensure(GetQualitySettings(FPlatformProperties::IniPlatformName(), Settings));
+	
+	// Query CVars. (Override with CVars if they are > 0)
+	using namespace Metasound::Frontend;
+	const float BlockRateCVar = GetBlockRateOverride();
+	const int32 SampleRateCvar = GetSampleRateOverride();
+
+	if (SampleRateCvar > 0)
 	{
-		OperatorSettings.Reset();
+		Settings.SampleRate = SampleRateCvar;
 	}
-	
-	if (!OperatorSettings)
+	if (BlockRateCVar > 0)
 	{
-		using namespace Metasound;
-		using namespace Metasound::SourcePrivate;
-		
-		// Lazy Query and cache on the optional.
-		auto QueryQualitySettings = [&](Metasound::FSampleRate InSampleRate) -> Metasound::FOperatorSettings
-		{
-			static const float DefaultBlockRateConstant = 100.f;
-			static const float DefaultSampleRateConstant = 48000.f;
-			
-			// 1. Sensible defaults. (If Device SampleRate is sensible use that as default).
-			FSampleRate MetasoundSampleRate = InSampleRate > 0 ? InDeviceSampleRate : DefaultSampleRateConstant;
-			float MetasoundBlockRate = DefaultBlockRateConstant;
-
-			// 2. Query our quality settings.
-			if (const UMetaSoundSettings* Settings = GetDefault<UMetaSoundSettings>())
-			{
-				if (const FMetaSoundQualitySettings* Found = Settings->QualitySettings.FindByPredicate([&QT = QualitySetting](const FMetaSoundQualitySettings& Q) -> bool { return Q.Name == QT; }))
-				{
-					// Allow partial applications of settings, if some are non-zero.
-					if (const float Value = Found->BlockRate.GetValue(); Value > 0.f)
-					{
-						MetasoundBlockRate = Value;
-					}
-					if (const float Value = Found->SampleRate.GetValue(); Value > 0.f)
-					{
-						MetasoundSampleRate = Value;
-					}
-				}
-			}
-
-			// 3. Do per asset overrides.
-			if (const float SerializedBlockRate = BlockRateOverride.GetValue(); SerializedBlockRate > 0.0f)
-			{
-				MetasoundBlockRate = SerializedBlockRate;
-			}
-			if (const int32 SerializedSampleRate = SampleRateOverride.GetValue(); SerializedSampleRate > 0)
-			{
-				MetasoundSampleRate = SerializedSampleRate;
-			}
-
-			// 4. Query CVars. (Override with CVars if they are > 0)
-			using namespace Metasound::Frontend;
-			const float BlockRateCVar = GetBlockRateOverride();
-			const int32 SampleRateCvar = GetSampleRateOverride();
-
-			if (SampleRateCvar > 0)
-			{
-				MetasoundSampleRate = SampleRateCvar;
-			}
-			if (BlockRateCVar > 0)
-			{
-				MetasoundBlockRate = BlockRateCVar;
-			}
+		Settings.BlockRate = BlockRateCVar;
+	}
 
-			// 5. Sanity clamps.
-			const TRange<float> BlockRange = GetBlockRateClampRange();
-			const TRange<int32> RateRange = GetSampleRateClampRange();
-			MetasoundBlockRate = FMath::Clamp(MetasoundBlockRate, BlockRange.GetLowerBoundValue(), BlockRange.GetUpperBoundValue());
-			MetasoundSampleRate = FMath::Clamp(MetasoundSampleRate, RateRange.GetLowerBoundValue(), RateRange.GetUpperBoundValue());
+	// Sanity clamps.
+	const TRange<float> BlockRange = GetBlockRateClampRange();
+	const TRange<int32> RateRange = GetSampleRateClampRange();
+	Settings.BlockRate = FMath::Clamp(Settings.BlockRate.GetValue(), BlockRange.GetLowerBoundValue(), BlockRange.GetUpperBoundValue());
+	Settings.SampleRate = FMath::Clamp(Settings.SampleRate.GetValue(), RateRange.GetLowerBoundValue(), RateRange.GetUpperBoundValue());
 
-			return Metasound::FOperatorSettings(MetasoundSampleRate, MetasoundBlockRate);
-		};
-		OperatorSettings = QueryQualitySettings(InDeviceSampleRate);
-	}
-	return *OperatorSettings;
+	return Metasound::FOperatorSettings(
+		/* SampleRate */ Settings.SampleRate.GetValue(),
+		/* BlockRate */ Settings.BlockRate.GetValue());
 }
 
 Metasound::FMetasoundEnvironment UMetaSoundSource::CreateEnvironment() const
diff --git a/Engine/Plugins/Runtime/Metasound/Source/MetasoundEngine/Public/MetasoundSettings.h b/Engine/Plugins/Runtime/Metasound/Source/MetasoundEngine/Public/MetasoundSettings.h
--- a/Engine/Plugins/Runtime/Metasound/Source/MetasoundEngine/Public/MetasoundSettings.h
+++ b/Engine/Plugins/Runtime/Metasound/Source/MetasoundEngine/Public/MetasoundSettings.h
@@ -58,19 +58,19 @@ struct METASOUNDENGINE_API FMetaSoundQualitySettings
 	UPROPERTY()
 	FGuid UniqueId = {};
 
-#endif //WITH_EDITORONLY_DATA
-
 	/** Name of this quality setting. This will appear in the quality dropdown list.
 		The names should be unique and adequately describe the Entry. "High", "Low" etc. **/
 	UPROPERTY(EditAnywhere, Category = "Quality")
 	FName Name = {};
+	
+#endif //WITH_EDITORONLY_DATA	
 
 	/** Sample Rate (in Hz). NOTE: A Zero value will have no effect and use the Device Rate. **/
-	UPROPERTY(EditAnywhere, Category = "Quality", meta = (ClampMin = "0", ClampMax="96000"))
+	UPROPERTY(EditAnywhere, config, Category = "Quality", meta = (ClampMin = "0", ClampMax="96000"))
 	FPerPlatformInt SampleRate = 0;
 
 	/** Block Rate (in Hz). NOTE: A Zero value will have no effect and use the Default (100)  **/
-	UPROPERTY(EditAnywhere, Category = "Quality", meta = (ClampMin = "0", ClampMax="1000"))
+	UPROPERTY(EditAnywhere, config, Category = "Quality", meta = (ClampMin = "0", ClampMax="1000"))
 	FPerPlatformFloat BlockRate = 0.f;
 };
 
@@ -105,17 +105,26 @@ public:
 	  */
 	UPROPERTY(EditAnywhere, config, Category = Registration, meta = (RelativePath, LongPackageName))
 	TArray<FDirectoryPath> DirectoriesToRegister;
-	
-	/** Array of possible quality settings for Metasounds to chose from */
-	UPROPERTY(EditAnywhere, config, Category = Quality)
-	TArray<FMetaSoundQualitySettings> QualitySettings;
 		
 	UPROPERTY(Transient)
 	int32 DenyListCacheChangeID = 0;	
 
+#if WITH_EDITORONLY_DATA
+	const TArray<FMetaSoundQualitySettings>& GetQualitySettings() const { return QualitySettings; }
+	static FName GetQualitySettingPropertyName(); 
+#endif //WITH_EDITORONLY_DATA
+
+private:
+
+	/** Array of possible quality settings for Metasounds to chose from */
+	// NOTE: Ideally this would be wrapped with WITH_EDITORONLY_DATA, but standalone "-game" requires
+	// it to exist. Access is limited to the accessor above, which enforces it correctly.
+	UPROPERTY(EditAnywhere, config, Category = Quality)
+	TArray<FMetaSoundQualitySettings> QualitySettings;
+
 #if WITH_EDITOR
 private:
-	
+
 	virtual void PostEditChangeChainProperty(FPropertyChangedChainEvent& PropertyChangedEvent) override;
 
 	virtual void PostEditChangeProperty(FPropertyChangedEvent& PropertyChangedEvent) override;
diff --git a/Engine/Plugins/Runtime/Metasound/Source/MetasoundEngine/Public/MetasoundSource.h b/Engine/Plugins/Runtime/Metasound/Source/MetasoundEngine/Public/MetasoundSource.h
--- a/Engine/Plugins/Runtime/Metasound/Source/MetasoundEngine/Public/MetasoundSource.h
+++ b/Engine/Plugins/Runtime/Metasound/Source/MetasoundEngine/Public/MetasoundSource.h
@@ -24,6 +24,7 @@
 
 // Forward Declarations
 class UMetaSoundSettings;
+struct FMetaSoundQualitySettings;
 
 namespace Audio
 {
@@ -45,6 +46,7 @@ namespace Metasound
 	namespace SourcePrivate
 	{
 		class FParameterRouter;
+		using FCookedQualitySettings = FMetaSoundQualitySettings;
 	} // namespace SourcePrivate
 
 	namespace DynamicGraph
@@ -114,15 +116,15 @@ public:
 	UPROPERTY(EditAnywhere, BlueprintReadOnly, Category = Metasound)
 	EMetaSoundOutputAudioFormat OutputFormat;
 
+#if WITH_EDITORONLY_DATA
+
 	// The Quality this Metasound will use. These are defined in the MetaSounds project settings.
 	UPROPERTY(config, EditAnywhere, BlueprintReadWrite, meta = (GetOptions="MetasoundEngine.MetaSoundQualityHelper.GetQualityList"), Category = "Metasound")
 	FName QualitySetting;
 
-#if WITH_EDITORONLY_DATA
 	// This a editor only look up for the Quality Setting above. Preventing orphaning of the original name.
 	UPROPERTY()
 	FGuid QualitySettingGuid;
-#endif //WITH_EDITOR_DATA
 
 	// Override the BlockRate for this Sound (overrides Quality). NOTE: A Zero value will have no effect and use either the Quality setting (if set), or the defaults.
 	UPROPERTY(EditAnywhere, BlueprintReadOnly, AdvancedDisplay, Category = Metasound, meta = (UIMin = 0, UIMax = 1000.0, DisplayAfter="OutputFormat", DisplayName = "Override Block Rate (in Hz)"))
@@ -131,7 +133,9 @@ public:
 	// Override the SampleRate for this Sound (overrides Quality). NOTE: A Zero value will have no effect and use either the Quality setting (if set), or the Device Rate
 	UPROPERTY(EditAnywhere, BlueprintReadOnly, AdvancedDisplay, Category = Metasound, meta = (UIMin = 0, UIMax = 96000, DisplayName = "Override Sample Rate (in Hz)"))
 	FPerPlatformInt SampleRateOverride = 0;
-
+	
+#endif //WITH_EDITOR_DATA
+	
 	UPROPERTY(AssetRegistrySearchable)
 	FGuid AssetClassID;
 
@@ -226,7 +230,6 @@ public:
 	virtual void PostLoad() override;
 
 	void PostLoadQualitySettings();
-	void ResolveQualitySettings(const UMetaSoundSettings* Settings);
 
 	virtual bool ConformObjectDataToInterfaces() override;
 
@@ -335,10 +338,6 @@ private:
 	TSharedPtr<Metasound::DynamicGraph::FDynamicOperatorTransactor> GetDynamicGeneratorTransactor() const;
 
 	TSharedPtr<Metasound::DynamicGraph::FDynamicOperatorTransactor> DynamicTransactor;
-	/*
-	 * Lazy (Cached) Operator Settings. Built in GetOperatorSettings
-	 */
-	mutable TOptional<Metasound::FOperatorSettings> OperatorSettings;
 
 	// Cache the AudioDevice Samplerate. (so that if we have to regenerate operator settings without the device rate we can use this).
 	mutable Metasound::FSampleRate CachedAudioDeviceSampleRate = 0;
@@ -349,4 +348,10 @@ private:
 	// utilize the same base MetaSound, they may be able to share their operators in the operator pool. This makes for a more
 	// efficient use of the operator pool.
 	bool bIsPresetGraphInflationSupported = false;
+
+	// Quality settings. 
+	bool GetQualitySettings(const FName InPlatformName, Metasound::SourcePrivate::FCookedQualitySettings& OutQualitySettings) const;
+	void ResolveQualitySettings(const UMetaSoundSettings* Settings);	
+	void SerializeCookedQualitySettings(const FName PlatformName, FArchive& Ar);
+	TPimplPtr<Metasound::SourcePrivate::FCookedQualitySettings> CookedQualitySettings;
 };
