Fix for perf regression in Metasounds
#tests Local verification

diff --git a/Engine/Plugins/Runtime/Metasound/Source/MetasoundEngine/Private/MetasoundSource.cpp b/Engine/Plugins/Runtime/Metasound/Source/MetasoundEngine/Private/MetasoundSource.cpp
--- a/Engine/Plugins/Runtime/Metasound/Source/MetasoundEngine/Private/MetasoundSource.cpp
+++ b/Engine/Plugins/Runtime/Metasound/Source/MetasoundEngine/Private/MetasoundSource.cpp
@@ -583,7 +583,17 @@ bool UMetaSoundSource::GetQualitySettings(
 	// If we've been cooked, this should contain the quality settings.
 	if (CookedQualitySettings)
 	{
-		OutQualitySettings = *CookedQualitySettings;
+		if (const float SerializedBlockRate = CookedQualitySettings->BlockRate.GetValue(); SerializedBlockRate > 0.0f)
+		{
+			UE_LOG(LogMetaSound, VeryVerbose, TEXT("Metasound [%s] BlockRate: %3.3f"), *GetName(), SerializedBlockRate);
+			OutQualitySettings.BlockRate = SerializedBlockRate;
+		}
+		if (const int32 SerializedSampleRate = CookedQualitySettings->SampleRate.GetValue(); SerializedSampleRate > 0)
+		{
+			UE_LOG(LogMetaSound, VeryVerbose, TEXT("Metasound [%s] SampleRate: %d"), *GetName(), SerializedSampleRate);
+			OutQualitySettings.SampleRate = SerializedSampleRate;
+		}
+	
 		UE_LOG(LogMetaSound, Verbose, TEXT("Metasound [%s] using SampleRate=%d, BlockRate=%2.3f (cooked)"),
 			*GetName(), OutQualitySettings.SampleRate.GetValue(), OutQualitySettings.BlockRate.GetValue()) ;
 		return true;
@@ -1448,6 +1458,9 @@ Metasound::FOperatorSettings UMetaSoundSource::GetOperatorSettings(Metasound::FS
 	Settings.BlockRate = FMath::Clamp(Settings.BlockRate.GetValue(), BlockRange.GetLowerBoundValue(), BlockRange.GetUpperBoundValue());
 	Settings.SampleRate = FMath::Clamp(Settings.SampleRate.GetValue(), RateRange.GetLowerBoundValue(), RateRange.GetUpperBoundValue());
 
+	UE_LOG(LogMetaSound, Verbose, TEXT("Metasound [%s] GetOperatorSettings: SampleRate: %d, BlockRate: %3.3f"),
+		*GetName(), Settings.SampleRate.GetValue(), Settings.BlockRate.GetValue());
+	
 	return Metasound::FOperatorSettings(
 		/* SampleRate */ Settings.SampleRate.GetValue(),
 		/* BlockRate */ Settings.BlockRate.GetValue());
