Fix ensure on retrieving quality settings from none-asset metasound.

diff --git a/Engine/Plugins/Runtime/Metasound/Source/MetasoundEngine/Private/MetasoundSource.cpp b/Engine/Plugins/Runtime/Metasound/Source/MetasoundEngine/Private/MetasoundSource.cpp
--- a/Engine/Plugins/Runtime/Metasound/Source/MetasoundEngine/Private/MetasoundSource.cpp
+++ b/Engine/Plugins/Runtime/Metasound/Source/MetasoundEngine/Private/MetasoundSource.cpp
@@ -1436,7 +1436,8 @@ Metasound::FOperatorSettings UMetaSoundSource::GetOperatorSettings(Metasound::FS
 
 	// Fetch our quality settings.
 	// If we are cooked these are baked, if we are editor these are queried from the project settings and this assets overrides.
-	ensure(GetQualitySettings(FPlatformProperties::IniPlatformName(), Settings));
+	const bool bFoundQualitySettings = GetQualitySettings(FPlatformProperties::IniPlatformName(), Settings);
+	UE_CLOG(!bFoundQualitySettings && IsAsset(), LogMetaSound, Error, TEXT("Could not retrieve quality settings for asset %s"), *GetOwningAssetName());
 	
 	// Query CVars. (Override with CVars if they are > 0)
 	using namespace Metasound::Frontend;
