Log error when failing to find interface for metasound

diff --git a/Engine/Plugins/Runtime/Metasound/Source/MetasoundFrontend/Private/MetasoundAssetBase.cpp b/Engine/Plugins/Runtime/Metasound/Source/MetasoundFrontend/Private/MetasoundAssetBase.cpp
--- a/Engine/Plugins/Runtime/Metasound/Source/MetasoundFrontend/Private/MetasoundAssetBase.cpp
+++ b/Engine/Plugins/Runtime/Metasound/Source/MetasoundFrontend/Private/MetasoundAssetBase.cpp
@@ -587,7 +587,10 @@ void FMetasoundAssetBase::CacheRegistryMetadata()
 	{
 		const FInterfaceRegistryKey InterfaceKey = GetInterfaceRegistryKey(Version);
 		const IInterfaceRegistryEntry* Entry = IInterfaceRegistry::Get().FindInterfaceRegistryEntry(InterfaceKey);
-		if (ensure(Entry))
+
+		UE_CLOG(nullptr == Entry, LogMetaSound, Error, TEXT("Failed to find interface (%s) when caching registry data for %s. MetaSound inputs and outputs for asset may not function correctly."), *Version.ToString(), *GetOwningAssetName());
+
+		if (Entry)
 		{
 			for (const FMetasoundFrontendClassInput& InterfaceInput : Entry->GetInterface().Inputs)
 			{
