First pass on AudioMaterialStyle overriding
#rnx

diff --git a/Engine/Plugins/Runtime/AudioWidgets/Source/AudioWidgets/Private/AudioMaterialSlate/SAudioMaterialKnob.cpp b/Engine/Plugins/Runtime/AudioWidgets/Source/AudioWidgets/Private/AudioMaterialSlate/SAudioMaterialKnob.cpp
--- a/Engine/Plugins/Runtime/AudioWidgets/Source/AudioWidgets/Private/AudioMaterialSlate/SAudioMaterialKnob.cpp
+++ b/Engine/Plugins/Runtime/AudioWidgets/Source/AudioWidgets/Private/AudioMaterialSlate/SAudioMaterialKnob.cpp
@@ -26,12 +26,6 @@ void SAudioMaterialKnob::Construct(const FArguments& InArgs)
 	OnMouseCaptureBegin = InArgs._OnMouseCaptureBegin;
 	OnMouseCaptureEnd = InArgs._OnMouseCaptureEnd;
 
-	//For now check if owner is not a widget -> use the default style.
-	if (!Cast<UWidget>(Owner))
-	{
-		AudioMaterialKnobStyle = &FAudioWidgetsStyle::Get().GetWidgetStyle<FAudioMaterialKnobStyle>("AudioMaterialKnob.Style");
-	}
-
 	ApplyNewMaterial();
 
 	if (InArgs._Value.IsSet())
diff --git a/Engine/Plugins/Runtime/AudioWidgets/Source/AudioWidgets/Private/AudioMaterialSlate/SAudioMaterialLabeledKnob.cpp b/Engine/Plugins/Runtime/AudioWidgets/Source/AudioWidgets/Private/AudioMaterialSlate/SAudioMaterialLabeledKnob.cpp
--- a/Engine/Plugins/Runtime/AudioWidgets/Source/AudioWidgets/Private/AudioMaterialSlate/SAudioMaterialLabeledKnob.cpp
+++ b/Engine/Plugins/Runtime/AudioWidgets/Source/AudioWidgets/Private/AudioMaterialSlate/SAudioMaterialLabeledKnob.cpp
@@ -41,6 +41,7 @@ void SAudioMaterialLabeledKnob::Construct(const FArguments& InArgs)
 	// Underlying Knob widget
 	SAssignNew(Knob, SAudioMaterialKnob)
 		.Owner(InArgs._Owner)
+		.AudioMaterialKnobStyle(Style)
 		.OnFloatValueChanged_Lambda([this](float Value)
 		{
 			ValueAttribute.Set(Value);
diff --git a/Engine/Plugins/Runtime/Metasound/Source/MetasoundEditor/Private/MetasoundEditorSettings.cpp b/Engine/Plugins/Runtime/Metasound/Source/MetasoundEditor/Private/MetasoundEditorSettings.cpp
--- a/Engine/Plugins/Runtime/Metasound/Source/MetasoundEditor/Private/MetasoundEditorSettings.cpp
+++ b/Engine/Plugins/Runtime/Metasound/Source/MetasoundEditor/Private/MetasoundEditorSettings.cpp
@@ -1,7 +1,10 @@
 // Copyright Epic Games, Inc. All Rights Reserved.
 #include "MetasoundEditorSettings.h"
 
+#include "AudioMaterialSlate/AudioMaterialSlateTypes.h"
+#include "AudioWidgetsStyle.h"
 #include "Styling/AppStyle.h"
+#include "Styling/SlateWidgetStyleAsset.h"
 #include "UObject/UnrealType.h"
 
 #include UE_INLINE_GENERATED_CPP_BY_NAME(MetasoundEditorSettings)
@@ -30,3 +33,15 @@ UMetasoundEditorSettings::UMetasoundEditorSettings(const FObjectInitializer& Obj
 	VariableNodeTitleColor = FLinearColor(0.211f, 0.513f, 0.035f);				// copper
 }
 
+const FAudioMaterialKnobStyle* UMetasoundEditorSettings::GetKnobStyle() const
+{
+	if (const UObject* Style = KnobStyleOverride.TryLoad())
+	{
+		if (const USlateWidgetStyleAsset* SlateWidgetStyleAsset = CastChecked<USlateWidgetStyleAsset>(Style))
+		{
+			return SlateWidgetStyleAsset->GetStyle<FAudioMaterialKnobStyle>();
+		}
+	}
+
+	return &FAudioWidgetsStyle::Get().GetWidgetStyle<FAudioMaterialKnobStyle>("AudioMaterialKnob.Style");
+}
diff --git a/Engine/Plugins/Runtime/Metasound/Source/MetasoundEditor/Private/MetasoundEditorSettings.h b/Engine/Plugins/Runtime/Metasound/Source/MetasoundEditor/Private/MetasoundEditorSettings.h
--- a/Engine/Plugins/Runtime/Metasound/Source/MetasoundEditor/Private/MetasoundEditorSettings.h
+++ b/Engine/Plugins/Runtime/Metasound/Source/MetasoundEditor/Private/MetasoundEditorSettings.h
@@ -5,9 +5,13 @@
 #include "UObject/NoExportTypes.h"
 #include "UObject/Object.h"
 #include "UObject/ObjectMacros.h"
+#include "UObject/SoftObjectPath.h"
 
 #include "MetasoundEditorSettings.generated.h"
 
+class USlateWidgetStyleAsset;
+struct FAudioMaterialKnobStyle;
+
 UENUM()
 enum class EMetasoundActiveAnalyzerEnvelopeDirection : uint8
 {
@@ -164,4 +168,11 @@ public:
 	/** Determines which details view to show in Metasounds Editor */
 	UPROPERTY(Transient)
 	EMetasoundActiveDetailView DetailView = EMetasoundActiveDetailView::General;
+
+	/**  Metasound Style override, If set will use the set style to represent Knobs in the Metasound Editor */
+	UPROPERTY(EditAnywhere, config, Category = MetasoundStyling, meta = (AllowedClasses = "SlateWidgetStyleAsset", DisplayName = "Knob Style"))
+	FSoftObjectPath KnobStyleOverride;
+	
+	/** Get the AudioMaterialKnob Style. if KnobStyleOverride is not set, returns default style.*/
+	const FAudioMaterialKnobStyle* GetKnobStyle() const;
 };
diff --git a/Engine/Plugins/Runtime/Metasound/Source/MetasoundEditor/Private/SMetasoundGraphNode.cpp b/Engine/Plugins/Runtime/Metasound/Source/MetasoundEditor/Private/SMetasoundGraphNode.cpp
--- a/Engine/Plugins/Runtime/Metasound/Source/MetasoundEditor/Private/SMetasoundGraphNode.cpp
+++ b/Engine/Plugins/Runtime/Metasound/Source/MetasoundEditor/Private/SMetasoundGraphNode.cpp
@@ -23,6 +23,7 @@
 #include "MetasoundEditorGraphNode.h"
 #include "MetasoundEditorGraphSchema.h"
 #include "MetasoundEditorModule.h"
+#include "MetasoundEditorSettings.h"
 #include "MetasoundFrontendNodeTemplateRegistry.h"
 #include "MetasoundFrontendRegistries.h"
 #include "MetasoundTrace.h"
@@ -858,8 +859,12 @@ namespace Metasound
 
 							if (Metasound::Editor::GraphNodePrivate::UseAudioMaterialWidgets)
 							{
+								const UMetasoundEditorSettings* MetasoundSettings = GetDefault<UMetasoundEditorSettings>();
+								check(MetasoundSettings)
+
 								SAssignNew(InputWidget, SAudioMaterialLabeledKnob)
 									.Owner(GraphMember->GetOwningGraph())
+									.Style(MetasoundSettings->GetKnobStyle())
 									.OnValueChanged_Lambda(OnValueChangedLambda)
 									.OnMouseCaptureBegin_Lambda(OnRadialSliderMouseCaptureBeginLambda)
 									.OnMouseCaptureEnd_Lambda(OnRadialSliderMouseCaptureEndLambda);
