Optimize incremental topological sort for metasound dynamic operators.

diff --git a/Engine/Plugins/Runtime/Metasound/Source/MetasoundGraphCore/Private/MetasoundDynamicGraphAlgo.cpp b/Engine/Plugins/Runtime/Metasound/Source/MetasoundGraphCore/Private/MetasoundDynamicGraphAlgo.cpp
--- a/Engine/Plugins/Runtime/Metasound/Source/MetasoundGraphCore/Private/MetasoundDynamicGraphAlgo.cpp
+++ b/Engine/Plugins/Runtime/Metasound/Source/MetasoundGraphCore/Private/MetasoundDynamicGraphAlgo.cpp
@@ -2,9 +2,11 @@
 
 #include "MetasoundDynamicGraphAlgo.h"
 
+#include "Algo/IsSorted.h"
 #include "Containers/Array.h"
 #include "Containers/Map.h"
 #include "MetasoundDataReference.h"
+#include "MetasoundGraphAlgo.h"
 #include "MetasoundGraphAlgoPrivate.h"
 #include "MetasoundLog.h"
 #include "MetasoundTrace.h"
@@ -17,22 +19,103 @@ namespace Metasound
 	{
 		namespace DynamicGraphAlgoPrivate
 		{
+			template<typename EntryType>
+			struct TGetOrdinal
+			{
+				int32 operator()(const EntryType& InEntry) const
+				{
+					return InEntry.Ordinal;
+				}
+			};
+
+			template<typename EntryType>
+			void SortExecutionTable(TArrayView<EntryType> InTable)
+			{
+				Algo::SortBy(InTable, TGetOrdinal<EntryType>());
+			}
+
+			template<typename EntryType>
+			void SetOrdinalsAndSortTable(const TMap<FOperatorID, int32>& InOrdinals, TArray<EntryType>& InOutTable)
+			{
+				// Assign new ordinals to all entries
+				for (EntryType& Entry : InOutTable)
+				{
+					if (const int32* Ordinal = InOrdinals.Find(Entry.OperatorID))
+					{
+						Entry.Ordinal = *Ordinal;
+					}
+					else
+					{
+						Entry.Ordinal = ORDINAL_NONE;
+					}
+				}
+
+				// Sort table by ordinal
+				SortExecutionTable<EntryType>(InOutTable);
+
+				// Remove entries without an ordinal
+				InOutTable.RemoveAll([](const EntryType& InEntry) { return InEntry.Ordinal == ORDINAL_NONE; });
+			}
+
+			template<typename EntryType>
+			void SwapOrdinalsAndSortTable(int32 InMinOrdinal, int32 InMaxOrdinal, const TArray<FOrdinalSwap>& InSwaps, TArray<EntryType>& InOutTable)
+			{
+				int32 StartTableIndex = Algo::LowerBoundBy(InOutTable, InMinOrdinal, TGetOrdinal<EntryType>());
+				int32 EndTableIndex = Algo::UpperBoundBy(InOutTable, InMaxOrdinal, TGetOrdinal<EntryType>());
+
+				if ((StartTableIndex >= InOutTable.Num()) || (EndTableIndex <= StartTableIndex))
+				{
+					// No entries exist in the table which match the swaps. There
+					// is nothing to update.
+					return;
+				}
+
+				const FOrdinalSwap* SwapPtr = InSwaps.GetData();
+				const FOrdinalSwap* const EndSwapPtr = SwapPtr + InSwaps.Num();
+				EntryType* EntryPtr = InOutTable.GetData() + StartTableIndex;
+				EntryType* const EndEntryPtr = InOutTable.GetData() + EndTableIndex;
+
+				// Iterate through swaps and entries until we've worked through all the
+				// swaps or all the table entries in the range.
+				while ((SwapPtr != EndSwapPtr) && (EntryPtr != EndEntryPtr))
+				{
+					if (SwapPtr->OriginalOrdinal == EntryPtr->Ordinal)
+					{
+						// Found a match. Update ordinal and increment both pointers. 
+						EntryPtr->Ordinal = SwapPtr->NewOrdinal;
+						EntryPtr++;
+						SwapPtr++;
+					}
+					else if (SwapPtr->OriginalOrdinal < EntryPtr->Ordinal)
+					{
+						SwapPtr++;
+					}
+					else // if (EntryPtr->Ordinal < SwapPtr->OriginalOrdinal) <-- assumed because of earlier logic checks.
+					{
+						EntryPtr++;
+					}
+				}
+
+				// Sort the entries by ordinal. Only update table entries in the given range. 
+				SortExecutionTable(TArrayView<const EntryType>(InOutTable.GetData() + StartTableIndex, EndTableIndex - StartTableIndex));
+			}
+
 			template<typename EntryType, typename FunctionType>
-			void UpdateTableEntry(const TArray<FOperatorID>& InOperatorOrder, const FOperatorID& InOperatorID, IOperator* InOperator, TArray<EntryType>& InOutTable, FunctionType InFunction)
+			void UpdateTableEntry(const FOperatorID& InOperatorID, FOperatorInfo& InOperatorInfo, TArray<EntryType>& InOutTable, FunctionType InFunction)
 			{
 				using namespace DirectedGraphAlgo;
 
-				int32 EntryIndex = InOutTable.IndexOfByPredicate([&](const EntryType& InEntry) { return InEntry.OperatorID == InOperatorID; });
+				int32 OperatorOrdinal = InOperatorInfo.Ordinal;
+				int32 EntryIndex = Algo::BinarySearchBy(InOutTable, OperatorOrdinal, TGetOrdinal<EntryType>());
 
-				int32 OperatorIndex = InOperatorOrder.Find(InOperatorID);
-				bool bEntryShouldExist = (nullptr != InOperator) && (nullptr != InFunction) && (INDEX_NONE != OperatorIndex); // Remove operators if they have no function for this table, or if they or not in the operator execution order. 
+				bool bEntryShouldExist = (nullptr != InFunction); // Remove operators if they have no function for this table
 				bool bEntryCurrentlyExists = (EntryIndex != INDEX_NONE);
 
 				if (bEntryCurrentlyExists && bEntryShouldExist)
 				{
 					// Update the existing entry
 					InOutTable[EntryIndex].Function = InFunction;
-					InOutTable[EntryIndex].Operator = InOperator;
+					InOutTable[EntryIndex].Operator = InOperatorInfo.Operator.Get();
 				}
 				else if (bEntryCurrentlyExists)
 				{
@@ -41,37 +124,105 @@ namespace Metasound
 				}
 				else if (bEntryShouldExist)
 				{
-					check(InOperator);
-					// Add new entry
-					int32 PreceedingOperatorIndex = OperatorIndex - 1;
-					int32 InsertLocation = 0; // Default to inserting at the beginning
-					while (PreceedingOperatorIndex >= 0)
-					{
-						FOperatorID PreceedingOperatorID = InOperatorOrder[PreceedingOperatorIndex];
-						PreceedingOperatorIndex--;
+					check(InOperatorInfo.Operator.IsValid());
+
+					// Add missing entry
+					int32 InsertLocation = Algo::UpperBoundBy(InOutTable, OperatorOrdinal, TGetOrdinal<EntryType>());
+					InOutTable.Insert(EntryType{OperatorOrdinal, InOperatorID, *InOperatorInfo.Operator, InFunction}, InsertLocation);
+				}
+			}
+
+			void UpdateGraphRuntimeTableEntries(const FOperatorID& InOperatorID, FOperatorInfo& InOperatorInfo, FDynamicGraphOperatorData& InOutGraphOperatorData)
+			{
+				using namespace DynamicGraphAlgoPrivate;
+				METASOUND_TRACE_CPUPROFILER_EVENT_SCOPE(Metasound::DynamicGraphAlgo::UpdateOperatorRuntimeTableEntries)
+				IOperator* Operator = InOperatorInfo.Operator.Get();
+
+				UpdateTableEntry(InOperatorID, InOperatorInfo, InOutGraphOperatorData.ExecuteTable, Operator ? Operator->GetExecuteFunction() : nullptr);
+				UpdateTableEntry(InOperatorID, InOperatorInfo, InOutGraphOperatorData.PostExecuteTable, Operator ? Operator->GetPostExecuteFunction() : nullptr);
+				UpdateTableEntry(InOperatorID, InOperatorInfo, InOutGraphOperatorData.ResetTable, Operator ? Operator->GetResetFunction() : nullptr);
+			}
 
-						int32 PreceedingEntryIndex = InOutTable.IndexOfByPredicate([&](const EntryType& InEntry) { return InEntry.OperatorID == PreceedingOperatorID; });
-						if (PreceedingEntryIndex != INDEX_NONE)
+			void SetOutputVertexData(FDynamicGraphOperatorData& InOutGraphOperatorData)
+			{
+				using namespace DirectedGraphAlgo;
+
+				// Iterate through the output operators and force their output data references
+				// to be reflected in the graph's FOutputVertexInterfaceData
+				for (const TPair<FVertexName, FOperatorID>& OutputVertexInfo : InOutGraphOperatorData.OutputVertexMap)
+				{
+					const FVertexName& VertexName = OutputVertexInfo.Get<0>();
+					const FOperatorID& OperatorID = OutputVertexInfo.Get<1>();
+
+					if (const FGraphOperatorData::FOperatorInfo* OperatorInfo = InOutGraphOperatorData.OperatorMap.Find(OperatorID))
+					{
+						if (const FAnyDataReference* Ref = OperatorInfo->VertexData.GetOutputs().FindDataReference(OutputVertexInfo.Get<0>()))
+						{
+							InOutGraphOperatorData.VertexData.GetOutputs().SetVertex(VertexName, *Ref);
+						}
+						else if (InOutGraphOperatorData.VertexData.GetOutputs().IsVertexBound(VertexName))
 						{
-							InsertLocation = PreceedingEntryIndex + 1;
-							break;
+							UE_LOG(LogMetaSound, Error, TEXT("Output vertex (%s) lost data reference after rebinding graph"), *VertexName.ToString());
 						}
 					}
+					else
+					{
+						UE_LOG(LogMetaSound, Error, TEXT("Failed to update graph operator outputs. Could not find output operator info with ID %s for vertex %s"), *LexToString(OperatorID), *VertexName.ToString());
+					}
+				}
+			}
+		} // namespace DynamicGraphAlgoPrivate
 
-					InOutTable.Insert(EntryType{InOperatorID, *InOperator, InFunction}, InsertLocation);
+		void SetOrdinalsAndSort(const TMap<FOperatorID, int32>& InOrdinals, FDynamicGraphOperatorData& InOutGraphOperatorData)
+		{
+			using namespace DynamicGraphAlgoPrivate;
+
+			// Assign ordinals to operator map
+			for (TPair<FOperatorID, FOperatorInfo>& Pair : InOutGraphOperatorData.OperatorMap)
+			{
+				if (const int32* Ordinal = InOrdinals.Find(Pair.Key))
+				{
+					Pair.Value.Ordinal = *Ordinal;
+				}
+				else
+				{
+					Pair.Value.Ordinal = ORDINAL_NONE;
 				}
 			}
 
-		} // namespace DynamicGraphAlgoPrivate
+			// set ordinals to on execution tables. 
+			SetOrdinalsAndSortTable(InOrdinals, InOutGraphOperatorData.ExecuteTable);
+			SetOrdinalsAndSortTable(InOrdinals, InOutGraphOperatorData.PostExecuteTable);
+			SetOrdinalsAndSortTable(InOrdinals, InOutGraphOperatorData.ResetTable);
+		}
 
-		void UpdateGraphRuntimeTableEntries(const FOperatorID& InOperatorID, IOperator* InOperator, FDynamicGraphOperatorData& InOutGraphOperatorData)
+		void SwapOrdinalsAndSort(const TArray<FOrdinalSwap>& InSwaps, FDynamicGraphOperatorData& InOutGraphOperatorData)
 		{
 			using namespace DynamicGraphAlgoPrivate;
-			METASOUND_TRACE_CPUPROFILER_EVENT_SCOPE(Metasound::DynamicGraphAlgo::UpdateOperatorRuntimeTableEntries)
 
-			UpdateTableEntry(InOutGraphOperatorData.OperatorOrder, InOperatorID, InOperator, InOutGraphOperatorData.ExecuteTable, InOperator ? InOperator->GetExecuteFunction() : nullptr);
-			UpdateTableEntry(InOutGraphOperatorData.OperatorOrder, InOperatorID, InOperator, InOutGraphOperatorData.PostExecuteTable, InOperator ? InOperator->GetPostExecuteFunction() : nullptr);
-			UpdateTableEntry(InOutGraphOperatorData.OperatorOrder, InOperatorID, InOperator, InOutGraphOperatorData.ResetTable, InOperator ? InOperator->GetResetFunction() : nullptr);
+			checkfSlow(Algo::IsSorted(InSwaps, FOrdinalSwap::OriginalOrdinalLessThan), TEXT("Dynamic MetaSound ordinal swaps must be presorted by the original ordinal."));
+
+			if (InSwaps.Num() < 1)
+			{
+				return;
+			}
+
+			const int32 MinOrdinal = InSwaps[0].OriginalOrdinal;
+			const int32 MaxOrdinal = InSwaps.Last().OriginalOrdinal;
+			
+			// Update operator map
+			for (const FOrdinalSwap& Swap : InSwaps)
+			{
+				if (FOperatorInfo* OpInfo = InOutGraphOperatorData.OperatorMap.Find(Swap.OperatorID))
+				{
+					OpInfo->Ordinal = Swap.NewOrdinal;
+				}
+			}
+
+			// Update execution tables
+			SwapOrdinalsAndSortTable(MinOrdinal, MaxOrdinal, InSwaps, InOutGraphOperatorData.ExecuteTable);
+			SwapOrdinalsAndSortTable(MinOrdinal, MaxOrdinal, InSwaps, InOutGraphOperatorData.PostExecuteTable);
+			SwapOrdinalsAndSortTable(MinOrdinal, MaxOrdinal, InSwaps, InOutGraphOperatorData.ResetTable);
 		}
 
 		// Apply updates to data references through all the operators by following connections described in the FOperatorInfo map.
@@ -113,7 +264,7 @@ namespace Metasound
 					Operator.BindOutputs(OpInfo->VertexData.GetOutputs());
 
 					// Update execute/postexecute/reset tables in case those have changed after rebinding.
-					UpdateGraphRuntimeTableEntries(Current.OperatorID, &Operator, InOutGraphOperatorData);
+					UpdateGraphRuntimeTableEntries(Current.OperatorID, *OpInfo, InOutGraphOperatorData);
 
 					// See if binding altered the outputs. 
 					OutputUpdates.Reset();
@@ -141,22 +292,25 @@ namespace Metasound
 			}
 		}
 
-		FExecuteEntry::FExecuteEntry(DirectedGraphAlgo::FOperatorID InOperatorID, IOperator& InOperator, IOperator::FExecuteFunction InFunc)
-		: OperatorID(InOperatorID)
+		FExecuteEntry::FExecuteEntry(int32 InOrdinal, DirectedGraphAlgo::FOperatorID InOperatorID, IOperator& InOperator, IOperator::FExecuteFunction InFunc)
+		: Ordinal(InOrdinal)
+		, OperatorID(InOperatorID)
 		, Operator(&InOperator)
 		, Function(InFunc)
 		{
 		}
 
-		FPostExecuteEntry::FPostExecuteEntry(DirectedGraphAlgo::FOperatorID InOperatorID, IOperator& InOperator, IOperator::FPostExecuteFunction InFunc)
-		: OperatorID(InOperatorID)
+		FPostExecuteEntry::FPostExecuteEntry(int32 InOrdinal, DirectedGraphAlgo::FOperatorID InOperatorID, IOperator& InOperator, IOperator::FPostExecuteFunction InFunc)
+		: Ordinal(InOrdinal)
+		, OperatorID(InOperatorID)
 		, Operator(&InOperator)
 		, Function(InFunc)
 		{
 		}
 
-		FResetEntry::FResetEntry(DirectedGraphAlgo::FOperatorID InOperatorID, IOperator& InOperator, IOperator::FResetFunction InFunc)
-		: OperatorID(InOperatorID)
+		FResetEntry::FResetEntry(int32 InOrdinal, DirectedGraphAlgo::FOperatorID InOperatorID, IOperator& InOperator, IOperator::FResetFunction InFunc)
+		: Ordinal(InOrdinal)
+		, OperatorID(InOperatorID)
 		, Operator(&InOperator)
 		, Function(InFunc)
 		{
@@ -167,80 +321,80 @@ namespace Metasound
 		{
 		}
 
-		FDynamicGraphOperatorData::FDynamicGraphOperatorData(DirectedGraphAlgo::FGraphOperatorData&& InGraphOperatorData)
-		: DirectedGraphAlgo::FGraphOperatorData(MoveTemp(InGraphOperatorData))
-		{
-			InitTables();
-		}
 
-		FDynamicGraphOperatorData::FDynamicGraphOperatorData(DirectedGraphAlgo::FGraphOperatorData&& InGraphOperatorData, const FDynamicOperatorUpdateCallbacks& InCallbacks)
-		: DirectedGraphAlgo::FGraphOperatorData(MoveTemp(InGraphOperatorData))
+		FDynamicGraphOperatorData::FDynamicGraphOperatorData(const FOperatorSettings& InSettings, const FDynamicOperatorUpdateCallbacks& InCallbacks)
+		: DirectedGraphAlgo::FGraphOperatorData(InSettings)
 		, OperatorUpdateCallbacks(InCallbacks)
 		{
-			InitTables();
 		}
 
 		void FDynamicGraphOperatorData::InitTables()
 		{
 			// Populate execute/postexecute/reset stacks
-			for (const FOperatorID& OperatorID : OperatorOrder)
+			for (TPair<FOperatorID, FOperatorInfo>& Pair : OperatorMap)
 			{
-				if (FOperatorInfo* OperatorInfo = OperatorMap.Find(OperatorID))
-				{
-					check(OperatorInfo->Operator.IsValid());
-					IOperator& Operator = *(OperatorInfo->Operator);
+				const FOperatorID OperatorID = Pair.Key;
+				FOperatorInfo& OperatorInfo = Pair.Value;
 
-					if (IOperator::FExecuteFunction ExecuteFunc = Operator.GetExecuteFunction())
-					{
-						ExecuteTable.Emplace(FExecuteEntry{OperatorID, Operator, ExecuteFunc});
-					}
+				check(OperatorInfo.Operator.IsValid());
+				IOperator& Operator = *OperatorInfo.Operator;
 
-					if (IOperator::FPostExecuteFunction PostExecuteFunc = Operator.GetPostExecuteFunction())
-					{
-						PostExecuteTable.Emplace(FPostExecuteEntry{OperatorID, Operator, PostExecuteFunc});
-					}
+				if (IOperator::FExecuteFunction ExecuteFunc = Operator.GetExecuteFunction())
+				{
+					ExecuteTable.Emplace(FExecuteEntry{OperatorInfo.Ordinal, OperatorID, Operator, ExecuteFunc});
+				}
 
-					if (IOperator::FResetFunction ResetFunc = Operator.GetResetFunction())
-					{
-						ResetTable.Emplace(FResetEntry{OperatorID, Operator, ResetFunc});
-					}
+				if (IOperator::FPostExecuteFunction PostExecuteFunc = Operator.GetPostExecuteFunction())
+				{
+					PostExecuteTable.Emplace(FPostExecuteEntry{OperatorInfo.Ordinal, OperatorID, Operator, PostExecuteFunc});
+				}
+
+				if (IOperator::FResetFunction ResetFunc = Operator.GetResetFunction())
+				{
+					ResetTable.Emplace(FResetEntry{OperatorInfo.Ordinal, OperatorID, Operator, ResetFunc});
 				}
 			}
+
+			// Sort execution stacks
+			DynamicGraphAlgoPrivate::SortExecutionTable<FExecuteEntry>(ExecuteTable);
+			DynamicGraphAlgoPrivate::SortExecutionTable<FPostExecuteEntry>(PostExecuteTable);
+			DynamicGraphAlgoPrivate::SortExecutionTable<FResetEntry>(ResetTable);
 		}
 
 		void UpdateOutputVertexData(FDynamicGraphOperatorData& InOutGraphOperatorData)
 		{
 			using namespace DirectedGraphAlgo;
+			using namespace DynamicGraphAlgoPrivate;
+
 			METASOUND_TRACE_CPUPROFILER_EVENT_SCOPE(Metasound::DynamicGraphAlgo::UpdateOutputVertexData)
 
-			// Iterate through the output operators and force their output data references
-			// to be reflected in the graph's FOutputVertexInterfaceData
-			for (const TPair<FVertexName, FOperatorID>& OutputVertexInfo : InOutGraphOperatorData.OutputVertexMap)
+			// If we have operator updates to call out to, we need to check for
+			// any changes to the output vertex data.
+			if (InOutGraphOperatorData.OperatorUpdateCallbacks.OnOutputUpdated)
 			{
-				const FVertexName& VertexName = OutputVertexInfo.Get<0>();
-				const FOperatorID& OperatorID = OutputVertexInfo.Get<1>();
+				// Cache current graph output vertex data state. 	
+				TArray<FVertexDataState> OriginalOutputVertexState;
+				GetVertexInterfaceDataState(InOutGraphOperatorData.VertexData.GetOutputs(), OriginalOutputVertexState);
 
-				if (const FGraphOperatorData::FOperatorInfo* OperatorInfo = InOutGraphOperatorData.OperatorMap.Find(OperatorID))
-				{
-					if (const FAnyDataReference* Ref = OperatorInfo->VertexData.GetOutputs().FindDataReference(OutputVertexInfo.Get<0>()))
-					{
-						InOutGraphOperatorData.VertexData.GetOutputs().SetVertex(VertexName, *Ref);
+				// Force updates
+				SetOutputVertexData(InOutGraphOperatorData);
 
-						if (InOutGraphOperatorData.OperatorUpdateCallbacks.OnOutputUpdated)
-						{
-							InOutGraphOperatorData.OperatorUpdateCallbacks.OnOutputUpdated(VertexName, InOutGraphOperatorData.VertexData.GetOutputs());
-						}
-					}
-					else if (InOutGraphOperatorData.VertexData.GetOutputs().IsVertexBound(VertexName))
-					{
-						UE_LOG(LogMetaSound, Error, TEXT("Output vertex (%s) lost data reference after rebinding graph"), *VertexName.ToString());
-					}
-				}
-				else
+				// Check for any changes to the output vertex data state. 
+				TSortedVertexNameMap<FAnyDataReference> OutputVertexUpdates;
+				CompareVertexInterfaceDataToPriorState(InOutGraphOperatorData.VertexData.GetOutputs(), OriginalOutputVertexState, OutputVertexUpdates);
+
+				// Report any updates. 
+				for (const TPair<FVertexName, FAnyDataReference>& Updates : OutputVertexUpdates)
 				{
-					UE_LOG(LogMetaSound, Error, TEXT("Failed to update graph operator outputs. Could not find output operator info with ID %s for vertex %s"), *LexToString(OperatorID), *VertexName.ToString());
+					InOutGraphOperatorData.OperatorUpdateCallbacks.OnOutputUpdated(Updates.Key, InOutGraphOperatorData.VertexData.GetOutputs());
 				}
 			}
+			else
+			{
+				// Force updates on output. No need to update outside callers if
+				// there is no callback set. 
+				SetOutputVertexData(InOutGraphOperatorData);
+			}
 		}
 
 		void RebindWrappedOperator(const FOperatorID& InOperatorID, FOperatorInfo& InOperatorInfo, FDynamicGraphOperatorData& InOutGraphOperatorData)
@@ -263,7 +417,7 @@ namespace Metasound
 			InOperatorInfo.Operator->BindOutputs(InOperatorInfo.VertexData.GetOutputs());
 
 			// Update any execution tables that need updating after wrapping
-			UpdateGraphRuntimeTableEntries(InOperatorID, InOperatorInfo.Operator.Get(), InOutGraphOperatorData);
+			UpdateGraphRuntimeTableEntries(InOperatorID, InOperatorInfo, InOutGraphOperatorData);
 
 			// Determine if there have been changes to `OutputVertexData`. 
 			TSortedVertexNameMap<FAnyDataReference> OutputsToUpdate;
@@ -282,7 +436,6 @@ namespace Metasound
 					for (const FGraphOperatorData::FVertexDestination& Destination : *Destinations)
 					{
 						PropagateBindUpdate(Destination.OperatorID, Destination.VertexName, OutputToUpdate.Get<1>(), InOutGraphOperatorData);
-						
 					}
 				}
 			}
@@ -342,5 +495,78 @@ namespace Metasound
 			// simply reading the latest values.
 			InOutVertexData.Bind(InOutGraphOperatorData.VertexData.GetOutputs());
 		}
+
+		void InsertOperator(FOperatorID InOperatorID, FOperatorInfo InOperatorInfo, FDynamicGraphOperatorData& InOutGraphOperatorData)
+		{
+			using namespace DynamicGraphAlgoPrivate;
+
+			IOperator* Operator = InOperatorInfo.Operator.Get();
+			if (nullptr == Operator)
+			{
+				return;
+			}
+
+			if (FOperatorInfo* ExistingInfo = InOutGraphOperatorData.OperatorMap.Find(InOperatorID))
+			{
+				// The options here are not good. The prior operator will be 
+				// removed and replaced with this new operator.
+				// Another option would be to leave the existing operator unchanged. 
+				// Neither option is satisfactory.
+				UE_LOG(LogMetaSound, Warning, TEXT("Overriding existing operator with the same operator ID %d. Duplicate operator IDs will lead to undefined behavior. Remove existing operators before adding a new one with the same ID"), InOperatorID);
+
+				TArray<FOperatorID> ConnectionsToRemove;
+				RemoveOperator(InOperatorID, ConnectionsToRemove, InOutGraphOperatorData);
+			}
+
+			// insert operator to execution tables
+			if (InOperatorInfo.Ordinal != ORDINAL_NONE)
+			{
+				if (IOperator::FExecuteFunction ExecuteFunc = Operator->GetExecuteFunction())
+				{
+					int32 InsertLocation = Algo::UpperBoundBy(InOutGraphOperatorData.ExecuteTable, InOperatorInfo.Ordinal, TGetOrdinal<FExecuteEntry>());
+					InOutGraphOperatorData.ExecuteTable.Insert(FExecuteEntry(InOperatorInfo.Ordinal, InOperatorID, *Operator, ExecuteFunc), InsertLocation);
+				}
+
+				if (IOperator::FPostExecuteFunction PostExecuteFunc = Operator->GetPostExecuteFunction())
+				{
+					int32 InsertLocation = Algo::UpperBoundBy(InOutGraphOperatorData.PostExecuteTable, InOperatorInfo.Ordinal, TGetOrdinal<FPostExecuteEntry>());
+					InOutGraphOperatorData.PostExecuteTable.Insert(FPostExecuteEntry(InOperatorInfo.Ordinal, InOperatorID, *Operator, PostExecuteFunc), InsertLocation);
+				}
+
+				if (IOperator::FResetFunction ResetFunc = Operator->GetResetFunction())
+				{
+					int32 InsertLocation = Algo::UpperBoundBy(InOutGraphOperatorData.ResetTable, InOperatorInfo.Ordinal, TGetOrdinal<FResetEntry>());
+					InOutGraphOperatorData.ResetTable.Insert(FResetEntry(InOperatorInfo.Ordinal, InOperatorID, *Operator, ResetFunc), InsertLocation);
+				}
+			}
+
+			// Move operator info to local map.
+			InOutGraphOperatorData.OperatorMap.Add(InOperatorID, MoveTemp(InOperatorInfo));
+		}
+
+		void RemoveOperator(FOperatorID InOperatorID, const TArray<FOperatorID>& InOperatorsConnectedToInput, FDynamicGraphOperatorData& InOutGraphOperatorData)
+		{
+			using namespace DirectedGraphAlgo;
+
+			// Remove any other nodes connected to this node. 
+			for (const FOperatorID& ConnectedOperatorID : InOperatorsConnectedToInput)
+			{
+				if (FOperatorInfo* ConnectedOperatorInfo = InOutGraphOperatorData.OperatorMap.Find(ConnectedOperatorID))
+				{
+					for (TPair<FVertexName, TArray<FGraphOperatorData::FVertexDestination>>& VertexDestinations : ConnectedOperatorInfo->OutputConnections)
+					{
+						VertexDestinations.Value.RemoveAllSwap([&InOperatorID](const FGraphOperatorData::FVertexDestination& Destination) { return Destination.OperatorID == InOperatorID; });
+					}
+				}
+			}
+
+			// remove from map of operators
+			InOutGraphOperatorData.OperatorMap.Remove(InOperatorID);
+
+			// Update execution tables
+			InOutGraphOperatorData.ExecuteTable.RemoveAll([&InOperatorID](const FExecuteEntry& InEntry) { return InEntry.OperatorID == InOperatorID; });
+			InOutGraphOperatorData.PostExecuteTable.RemoveAll([&InOperatorID](const FPostExecuteEntry& InEntry) { return InEntry.OperatorID == InOperatorID; });
+			InOutGraphOperatorData.ResetTable.RemoveAll([&InOperatorID](const FResetEntry& InEntry) { return InEntry.OperatorID == InOperatorID; });
+		}
 	}
 }
diff --git a/Engine/Plugins/Runtime/Metasound/Source/MetasoundGraphCore/Private/MetasoundDynamicGraphAlgo.h b/Engine/Plugins/Runtime/Metasound/Source/MetasoundGraphCore/Private/MetasoundDynamicGraphAlgo.h
--- a/Engine/Plugins/Runtime/Metasound/Source/MetasoundGraphCore/Private/MetasoundDynamicGraphAlgo.h
+++ b/Engine/Plugins/Runtime/Metasound/Source/MetasoundGraphCore/Private/MetasoundDynamicGraphAlgo.h
@@ -3,6 +3,7 @@
 #pragma once
 
 #include "Containers/Array.h"
+#include "Math/NumericLimits.h"
 #include "MetasoundDynamicOperatorTransactor.h"
 #include "MetasoundGraphAlgoPrivate.h"
 #include "MetasoundOperatorInterface.h"
@@ -20,7 +21,7 @@ namespace Metasound
 		/* Convenience wrapper for execute function of an IOperator. */
 		struct FExecuteEntry final
 		{
-			FExecuteEntry(FOperatorID InOperatorID, IOperator& InOperator, IOperator::FExecuteFunction InFunc);
+			FExecuteEntry(int32 InOrdinal, FOperatorID InOperatorID, IOperator& InOperator, IOperator::FExecuteFunction InFunc);
 
 			void Execute()
 			{
@@ -29,6 +30,7 @@ namespace Metasound
 				Function(Operator);
 			}
 
+			int32 Ordinal;
 			FOperatorID OperatorID;
 			IOperator* Operator;
 			IOperator::FExecuteFunction Function;	
@@ -37,7 +39,7 @@ namespace Metasound
 		/* Convenience wrapper for post execute function of an IOperator. */
 		struct FPostExecuteEntry final
 		{
-			FPostExecuteEntry(FOperatorID InOperatorID, IOperator& InOperator, IOperator::FPostExecuteFunction InFunc);
+			FPostExecuteEntry(int32 InOrdinal, FOperatorID InOperatorID, IOperator& InOperator, IOperator::FPostExecuteFunction InFunc);
 
 			void PostExecute()
 			{
@@ -46,6 +48,7 @@ namespace Metasound
 				Function(Operator);
 			}
 
+			int32 Ordinal;
 			FOperatorID OperatorID;
 			IOperator* Operator;
 			IOperator::FPostExecuteFunction Function;	
@@ -54,7 +57,7 @@ namespace Metasound
 		/* Convenience wrapper for reset function of an IOperator. */
 		struct FResetEntry final
 		{
-			FResetEntry(FOperatorID InOperatorID, IOperator& InOperator, IOperator::FResetFunction InFunc);
+			FResetEntry(int32 InOrdinal, FOperatorID InOperatorID, IOperator& InOperator, IOperator::FResetFunction InFunc);
 
 			void Reset(const IOperator::FResetParams& InParams)
 			{
@@ -63,17 +66,18 @@ namespace Metasound
 				Function(Operator, InParams);
 			}
 
+			int32 Ordinal;
 			FOperatorID OperatorID;
 			IOperator* Operator;
 			IOperator::FResetFunction Function;	
 		};
 
+
 		/** Collection of data needed to support a dynamic operator*/
 		struct FDynamicGraphOperatorData : DirectedGraphAlgo::FGraphOperatorData
 		{
 			FDynamicGraphOperatorData(const FOperatorSettings& InSettings);
-			FDynamicGraphOperatorData(DirectedGraphAlgo::FGraphOperatorData&& InGraphOperatorData);
-			FDynamicGraphOperatorData(DirectedGraphAlgo::FGraphOperatorData&& InGraphOperatorData, const FDynamicOperatorUpdateCallbacks& InCallbacks);
+			FDynamicGraphOperatorData(const FOperatorSettings& InSettings, const FDynamicOperatorUpdateCallbacks& InCallbacks);
 
 			// Initialize the Execute/PostExecute/Reset tables.
 			void InitTables();
@@ -87,20 +91,37 @@ namespace Metasound
 			TArray<FResetEntry> ResetTable;
 		};
 
-		/** Update the runtime table entries.
-		 * 
-		 * This updates the function tables in the FDynamicGraphOperatorData for a given operator. For an operator to exist in the function tables
-		 * it must provide a function to be called from GetExecutionFunction(), GetPostExecutionFunction() and/or GetResetFunction() and the operator's ID
-		 * must also be present in the FDynamicGraphOperatorData's operator order array.  If either of these are not true, then the operator will 
-		 * not have an entry in a specific execution table. 
-		 * 
-		 * Note: This function does not update the order of the entry in the execution table. 
-		 * 
-		 * @param InOperatorID - ID of operator to update.
-		 * @param InOperator - Pointer to operator.
-		 * @param InOutGraphOperatorData - Structure containing the operator order and runtime tables.
+		/** Interface to allow FOperatorBuilder special access to internal FDynamicGraphOperatorData
+		 * structures when the operator is being built.
+		 */
+		class IDynamicGraphInPlaceBuildable 
+		{
+		public:
+			virtual ~IDynamicGraphInPlaceBuildable() = default;
+		private:
+			friend class ::Metasound::FOperatorBuilder;
+
+			/** This function gives FOperatorBuilder access to the graph's internal data. 
+			 * This allows the FOperatorBuilder to build the operator in place by modifying
+			 * the internal data structure of the operator. An alternative approach 
+			 * would be to pass the data to the operator as a constructor argument. 
+			 * "In Place" building was chosen because it simplifies and streamlines the
+			 * build process within the FOperatorBuilder.
+			 */
+			virtual FDynamicGraphOperatorData& GetDynamicGraphOperatorData() = 0;
+		};
+
+		/** Sets the ordinals of operators and sorts executions tables. 
+		 * @param InOrdinals - A map of FOperatorID to ordinal.
+		 * @param InOutGraphOperatorData - The graph data where the updated ordinals should be applied. 
+		 */
+		void SetOrdinalsAndSort(const TMap<FOperatorID, int32>& InOrdinals, FDynamicGraphOperatorData& InOutGraphOperatorData);
+
+		/** Applies the ordinal swaps to operators and sorts execution tables.
+		 * @param InSwaps - An array of ordinal swaps.
+		 * @param InOutGraphOperatorData - The graph data where the updated ordinals should be applied. 
 		 */
-		void UpdateGraphRuntimeTableEntries(const FOperatorID& InOperatorID, IOperator* InOperator, FDynamicGraphOperatorData& InOutGraphOperatorData);
+		void SwapOrdinalsAndSort(const TArray<FOrdinalSwap>& InSwaps, FDynamicGraphOperatorData& InOutGraphOperatorData);
 
 		/** Propagate FVertexInterfaceData updates through the operators in the Dynamic Graph Operator Data. 
 		 *
@@ -128,6 +149,22 @@ namespace Metasound
 
 		/** Rebind the graph inputs, updating internal operator bindings as needed. */
 		void RebindGraphOutputs(FOutputVertexInterfaceData& InOutVertexData, FDynamicGraphOperatorData& InOutGraphOperatorData);
+
+		/** Insert an operator into the graph data and add to execution tables.
+		 *
+		 * @param InOperatorID - ID of new operator.
+		 * @param InOperatorInfo - Info containing describing operator.
+		 * @param InOutGraphOperatorData - Graph data where the operator will be inserted. 
+		 */
+		void InsertOperator(FOperatorID InOperatorID, FOperatorInfo InOperatorInfo, FDynamicGraphOperatorData& InOutGraphOperatorData);
+
+		/** Remove operator and related connections from the graph data.
+		 *
+		 * @param InOperatorID - ID of operator to remove.
+		 * @param InOperatorsConnectedToInput - An array of operator IDs denoting which operators connect to the input of operator being removed. 
+		 * @param InOutGraphOperatorData - Graph data where the operator will be removed from.
+		 */
+		void RemoveOperator(FOperatorID InOperatorID, const TArray<FOperatorID>& InOperatorsConnectedToInput, FDynamicGraphOperatorData& InOutGraphOperatorData);
 	}
 }
 
diff --git a/Engine/Plugins/Runtime/Metasound/Source/MetasoundGraphCore/Private/MetasoundDynamicOperator.cpp b/Engine/Plugins/Runtime/Metasound/Source/MetasoundGraphCore/Private/MetasoundDynamicOperator.cpp
--- a/Engine/Plugins/Runtime/Metasound/Source/MetasoundGraphCore/Private/MetasoundDynamicOperator.cpp
+++ b/Engine/Plugins/Runtime/Metasound/Source/MetasoundGraphCore/Private/MetasoundDynamicOperator.cpp
@@ -4,6 +4,7 @@
 
 #include "Containers/Array.h"
 #include "HAL/IConsoleManager.h"
+#include "MetasoundDynamicGraphAlgo.h"
 #include "MetasoundDynamicOperatorAudioFade.h"
 #include "MetasoundNodeInterface.h"
 #include "MetasoundOperatorInterface.h"
@@ -36,40 +37,14 @@ namespace Metasound
 
 		namespace DynamicOperatorPrivate
 		{
-			float MetaSoundExperimentalTransformTimeoutInSeconds = 0.010f;
+			float MetaSoundExperimentalTransformTimeoutInSeconds = -1.0f;
 			FAutoConsoleVariableRef CVarMetaSoundExperimentalTransformTimeoutInSeconds(
 				TEXT("au.MetaSound.Experimental.DynamicOperatorTransformTimeoutInSeconds"),
 				MetaSoundExperimentalTransformTimeoutInSeconds,
 				TEXT("Sets the number of seconds allowed to process pending dynamic graph transformations for a single MetaSound render cycle .\n")
-				TEXT("[Less than zero]: Disabled, [Greater than zero]: Enabled, 0.010s (default)"),
+				TEXT("[Less than zero]: Disabled, [Greater than zero]: Enabled, (disabled by default)"),
 				ECVF_Default);
 			
-			// Table sorter helper so we don't rewrite this algorithm for each differe
-			// stack type (Execute/PostExecute/Reset)
-			struct FTableSorter
-			{
-				template<typename TableEntryType>
-				static void SortTable(const TArray<FOperatorID>& OperatorOrder, TArray<TableEntryType>& InOutTable)
-				{
-					uint32 TargetIndex = 0;
-					for (const FOperatorID& OperatorID : OperatorOrder)
-					{
-						uint32 CurrentIndex = InOutTable.IndexOfByPredicate([&](const TableEntryType& Entry) 
-							{ 
-								return Entry.OperatorID == OperatorID; 
-							}
-						);
-
-						// The operator may not exist in the stack if it doesn't
-						// have an applicable function.
-						if (INDEX_NONE != CurrentIndex)
-						{
-							InOutTable.Swap(TargetIndex, CurrentIndex);
-							TargetIndex++;
-						}
-					}
-				}
-			};
 
 			template<typename VertexInterfaceDataType>
 			class TScopeUnfreeze final
@@ -92,6 +67,7 @@ namespace Metasound
 				VertexInterfaceDataType& VertexData;
 			};
 
+
 		} // namespace DynamicOperatorPrivate
 
 		FDynamicOperator::FDynamicOperator(const FOperatorSettings& InSettings)
@@ -104,8 +80,8 @@ namespace Metasound
 			}
 		}
 
-		FDynamicOperator::FDynamicOperator(DirectedGraphAlgo::FGraphOperatorData&& InGraphOperatorData, TSharedPtr<TSpscQueue<TUniquePtr<IDynamicOperatorTransform>>> InTransformQueue, const FDynamicOperatorUpdateCallbacks& InOperatorUpdateCallbacks)
-		: DynamicOperatorData(MoveTemp(InGraphOperatorData), InOperatorUpdateCallbacks)
+		FDynamicOperator::FDynamicOperator(const FOperatorSettings& InSettings, TSharedPtr<TSpscQueue<TUniquePtr<IDynamicOperatorTransform>>> InTransformQueue, const FDynamicOperatorUpdateCallbacks& InOperatorUpdateCallbacks)
+		: DynamicOperatorData(InSettings, InOperatorUpdateCallbacks)
 		, TransformQueue(InTransformQueue)
 		{
 			// Ensure that a transform queue exists. 
@@ -156,6 +132,11 @@ namespace Metasound
 			}
 		}
 
+		FDynamicGraphOperatorData& FDynamicOperator::GetDynamicGraphOperatorData()
+		{
+			return DynamicOperatorData;
+		}
+
 		void FDynamicOperator::ApplyTransformsUntilFence()
 		{
 			METASOUND_TRACE_CPUPROFILER_EVENT_SCOPE(Metasound::FDynamicOperator::ApplyTransformsUntilFence)
@@ -277,121 +258,63 @@ namespace Metasound
 		}
 
 		// Set the order of operators in the graph
-		FSetOperatorOrder::FSetOperatorOrder(TArray<FOperatorID> InOrder)
-		: Order(MoveTemp(InOrder))
+		FSetOperatorOrdinalsAndSort::FSetOperatorOrdinalsAndSort(TMap<FOperatorID, int32> InOrdinals)
+		: Ordinals(MoveTemp(InOrdinals))
 		{
 		}
 
-		EDynamicOperatorTransformQueueAction FSetOperatorOrder::Transform(FDynamicGraphOperatorData& InOutGraphOperatorData)
+		EDynamicOperatorTransformQueueAction FSetOperatorOrdinalsAndSort::Transform(FDynamicGraphOperatorData& InOutGraphOperatorData)
 		{
-			METASOUND_TRACE_CPUPROFILER_EVENT_SCOPE(Metasound::DynamicOperator::SetOperatorOrder)
+			METASOUND_TRACE_CPUPROFILER_EVENT_SCOPE(Metasound::DynamicOperator::SetOperatorOrdinalsAndSort)
 
-			using namespace DynamicOperatorPrivate;
+			SetOrdinalsAndSort(Ordinals, InOutGraphOperatorData);
+
+			return EDynamicOperatorTransformQueueAction::Continue;
+		}
+
+		// Swaps order of operators. 
+		FSwapOperatorOrdinalsAndSort::FSwapOperatorOrdinalsAndSort(TArray<FOrdinalSwap> InSwaps)
+		: Swaps(MoveTemp(InSwaps))
+		{
+		}
 
-			InOutGraphOperatorData.OperatorOrder = Order;
+		EDynamicOperatorTransformQueueAction FSwapOperatorOrdinalsAndSort::Transform(FDynamicGraphOperatorData& InOutGraphOperatorData)
+		{
+			METASOUND_TRACE_CPUPROFILER_EVENT_SCOPE(Metasound::DynamicOperator::SwapOperatorOrdinalsAndSort)
 
-			// Sort operator tables to be in the correct order.
-			FTableSorter::SortTable(Order, InOutGraphOperatorData.ExecuteTable);
-			FTableSorter::SortTable(Order, InOutGraphOperatorData.PostExecuteTable);
-			FTableSorter::SortTable(Order, InOutGraphOperatorData.ResetTable);
+			SwapOrdinalsAndSort(Swaps, InOutGraphOperatorData);
 
 			return EDynamicOperatorTransformQueueAction::Continue;
 		}
 
 		// Add an operator to the grpah
-		FAddOperator::FAddOperator(FOperatorID InOperatorID, EExecutionOrderInsertLocation InLocation, FOperatorInfo&& InInfo)
+		FInsertOperator::FInsertOperator(FOperatorID InOperatorID, FOperatorInfo InInfo)
 		: OperatorID(InOperatorID)
-		, Location(InLocation)
 		, OperatorInfo(MoveTemp(InInfo))
 		{
 		}
 
-		EDynamicOperatorTransformQueueAction FAddOperator::Transform(FDynamicGraphOperatorData& InGraphOperatorData)
+		EDynamicOperatorTransformQueueAction FInsertOperator::Transform(FDynamicGraphOperatorData& InOutGraphOperatorData)
 		{
-			METASOUND_TRACE_CPUPROFILER_EVENT_SCOPE(Metasound::DynamicOperator::AddOperator)
-
-			if (OperatorInfo.Operator.IsValid())
-			{
-				IOperator* Operator = OperatorInfo.Operator.Get();
+			METASOUND_TRACE_CPUPROFILER_EVENT_SCOPE(Metasound::DynamicOperator::InsertOperator)
 
-				if (FOperatorInfo* ExistingInfo = InGraphOperatorData.OperatorMap.Find(OperatorID))
-				{
-					// The options here are not good. The prior operator will be 
-					// removed and replaced with this new operator.
-					// Another option would be to leave the existing operator unchanged. 
-					// Neither option is satisfactory.
-					UE_LOG(LogMetaSound, Warning, TEXT("Overriding existing operator with the same operator ID %d. Duplicate operator IDs will lead to undefined behavior. Remove existing operators before adding a new one with the same ID"), OperatorID);
-
-					FRemoveOperator(OperatorID).Transform(InGraphOperatorData);
-				}
-
-				InGraphOperatorData.OperatorMap.Add(OperatorID, MoveTemp(OperatorInfo));
-				switch (Location)
-				{
-					case EExecutionOrderInsertLocation::First:
-						{
-							InGraphOperatorData.OperatorOrder.Insert(OperatorID, 0);
-							// Update execution tables
-							if (IOperator::FExecuteFunction ExecuteFunc = Operator->GetExecuteFunction())
-							{
-								InGraphOperatorData.ExecuteTable.Insert(FExecuteEntry(OperatorID, *Operator, ExecuteFunc), 0);
-							}
-
-							if (IOperator::FPostExecuteFunction PostExecuteFunc = Operator->GetPostExecuteFunction())
-							{
-								InGraphOperatorData.PostExecuteTable.Insert(FPostExecuteEntry(OperatorID, *Operator, PostExecuteFunc), 0);
-							}
-
-							if (IOperator::FResetFunction ResetFunc = Operator->GetResetFunction())
-							{
-								InGraphOperatorData.ResetTable.Insert(FResetEntry(OperatorID, *Operator, ResetFunc), 0);
-							}
-						}
-						break;
-
-					case EExecutionOrderInsertLocation::Last:
-						{
-							InGraphOperatorData.OperatorOrder.Add(OperatorID);
-							// Update execution tables
-							if (IOperator::FExecuteFunction ExecuteFunc = Operator->GetExecuteFunction())
-							{
-								InGraphOperatorData.ExecuteTable.Add(FExecuteEntry(OperatorID, *Operator, ExecuteFunc));
-							}
-
-							if (IOperator::FPostExecuteFunction PostExecuteFunc = Operator->GetPostExecuteFunction())
-							{
-								InGraphOperatorData.PostExecuteTable.Add(FPostExecuteEntry(OperatorID, *Operator, PostExecuteFunc));
-							}
-
-							if (IOperator::FResetFunction ResetFunc = Operator->GetResetFunction())
-							{
-								InGraphOperatorData.ResetTable.Add(FResetEntry(OperatorID, *Operator, ResetFunc));
-							}
-						}
-						break;
-				}
-			}
+			InsertOperator(OperatorID, MoveTemp(OperatorInfo), InOutGraphOperatorData);
 
 			return EDynamicOperatorTransformQueueAction::Continue;
 		}
 
 		// Remove an operator from the graph
-		FRemoveOperator::FRemoveOperator(FOperatorID InOperatorID)
+		FRemoveOperator::FRemoveOperator(FOperatorID InOperatorID, TArray<FOperatorID> InOperatorsConnectedToInput)
 		: OperatorID(InOperatorID)
+		, OperatorsConnectedToInput(MoveTemp(InOperatorsConnectedToInput))
 		{
 		}
 
-		EDynamicOperatorTransformQueueAction FRemoveOperator::Transform(FDynamicGraphOperatorData& InGraphOperatorData) 
+		EDynamicOperatorTransformQueueAction FRemoveOperator::Transform(FDynamicGraphOperatorData& InOutGraphOperatorData) 
 		{
 			METASOUND_TRACE_CPUPROFILER_EVENT_SCOPE(Metasound::DynamicOperator::RemoveOperator)
 
-			InGraphOperatorData.OperatorOrder.Remove(OperatorID);
-			InGraphOperatorData.OperatorMap.Remove(OperatorID);
-
-			// Update execution tables
-			InGraphOperatorData.ExecuteTable.RemoveAll([&](const FExecuteEntry& InEntry) { return InEntry.OperatorID == OperatorID; });
-			InGraphOperatorData.PostExecuteTable.RemoveAll([&](const FPostExecuteEntry& InEntry) { return InEntry.OperatorID == OperatorID; });
-			InGraphOperatorData.ResetTable.RemoveAll([&](const FResetEntry& InEntry) { return InEntry.OperatorID == OperatorID; });
+			RemoveOperator(OperatorID, OperatorsConnectedToInput, InOutGraphOperatorData);
 
 			return EDynamicOperatorTransformQueueAction::Continue;
 		}
diff --git a/Engine/Plugins/Runtime/Metasound/Source/MetasoundGraphCore/Private/MetasoundDynamicOperator.h b/Engine/Plugins/Runtime/Metasound/Source/MetasoundGraphCore/Private/MetasoundDynamicOperator.h
--- a/Engine/Plugins/Runtime/Metasound/Source/MetasoundGraphCore/Private/MetasoundDynamicOperator.h
+++ b/Engine/Plugins/Runtime/Metasound/Source/MetasoundGraphCore/Private/MetasoundDynamicOperator.h
@@ -13,19 +13,26 @@
 #include "MetasoundVertexData.h"
 #include "Templates/UniquePtr.h"
 
+#ifndef METASOUND_DEBUG_DYNAMIC_TRANSACTOR
+#define METASOUND_DEBUG_DYNAMIC_TRANSACTOR !UE_BUILD_SHIPPING
+#endif
+
 namespace Metasound
 {
 	namespace DynamicGraph
 	{
 		struct FDynamicGraphOperatorData;
-		using FLiteralAssignmentFunction = void(*)(const FOperatorSettings& InOperatorSettings, const FLiteral& InLiteral, const FAnyDataReference& OutDataRef);
-
-		/** Where to place a new operator in the execution order. */
-		enum class EExecutionOrderInsertLocation : uint8
+		class FDynamicOperator;
+		class FDynamicOperatorTransactor;
+		
+#if METASOUND_DEBUG_DYNAMIC_TRANSACTOR
+		namespace Debug
 		{
-			First, //< New operator executes before all existing operators.
-			Last   //< New operator executes after all existing operators.
-		};
+			class FDynamicOperatorDebugger;
+		}
+#endif // if METASOUND_DEBUG_DYNAMIC_TRANSACTOR
+
+		using FLiteralAssignmentFunction = void(*)(const FOperatorSettings& InOperatorSettings, const FLiteral& InLiteral, const FAnyDataReference& OutDataRef);
 
 		/** Action to perform after a single transformation. */
 		enum class EDynamicOperatorTransformQueueAction : uint8
@@ -46,12 +53,12 @@ namespace Metasound
 		 * change it's topology. Changes are communicated to the dynamic operator
 		 * through a TransformationQueue. 
 		 */
-		class FDynamicOperator : public IOperator
+		class FDynamicOperator : public IOperator, public IDynamicGraphInPlaceBuildable
 		{
 		public:
 
 			FDynamicOperator(const FOperatorSettings& InSettings);
-			FDynamicOperator(DirectedGraphAlgo::FGraphOperatorData&& InGraphOperatorData, TSharedPtr<TSpscQueue<TUniquePtr<IDynamicOperatorTransform>>> TransformQueue, const FDynamicOperatorUpdateCallbacks& InOperatorUpdateCallbacks);
+			FDynamicOperator(const FOperatorSettings& InSettings, TSharedPtr<TSpscQueue<TUniquePtr<IDynamicOperatorTransform>>> TransformQueue, const FDynamicOperatorUpdateCallbacks& InOperatorUpdateCallbacks);
 
 			virtual void BindInputs(FInputVertexInterfaceData& InOutVertexData) override;
 			virtual void BindOutputs(FOutputVertexInterfaceData& InOutVertexData) override;
@@ -66,6 +73,12 @@ namespace Metasound
 			void FlushEnqueuedTransforms();
 
 		private:
+#if METASOUND_DEBUG_DYNAMIC_TRANSACTOR
+			friend class Debug::FDynamicOperatorDebugger;
+#endif // if METASOUND_DEBUG_DYNAMIC_TRANSACTOR
+
+			virtual FDynamicGraphOperatorData& GetDynamicGraphOperatorData() override;
+
 			void ApplyTransformsUntilFence();
 			void ApplyTransformsUntilFenceOrTimeout(double InTimeoutInSeconds);
 			void Execute();
@@ -90,25 +103,33 @@ namespace Metasound
 		};
 
 		/** A transform which determines the order of execution for operators. */
-		class FSetOperatorOrder : public IDynamicOperatorTransform
+		class FSetOperatorOrdinalsAndSort : public IDynamicOperatorTransform
 		{
 		public:
-			FSetOperatorOrder(TArray<FOperatorID> InOrder);
+			FSetOperatorOrdinalsAndSort(TMap<FOperatorID, int32> InOrdinals);
 			virtual EDynamicOperatorTransformQueueAction Transform(FDynamicGraphOperatorData& InGraphOperatorData) override;
 		private:
-			TArray<FOperatorID> Order;
+			TMap<FOperatorID, int32> Ordinals;
 		};
 
-		/** A transform which adds an operator. */
-		class FAddOperator : public IDynamicOperatorTransform
+		/** A transform which applies operator ordinal swaps. */
+		class FSwapOperatorOrdinalsAndSort : public IDynamicOperatorTransform
 		{
 		public:
+			FSwapOperatorOrdinalsAndSort(TArray<FOrdinalSwap> InSwaps);
+			virtual EDynamicOperatorTransformQueueAction Transform(FDynamicGraphOperatorData& InGraphOperatorData) override;
+		private:
+			TArray<FOrdinalSwap> Swaps;
+		};
 
-			FAddOperator(FOperatorID InOperatorID, EExecutionOrderInsertLocation InLocation, FOperatorInfo&& InInfo);
+		/** A transform which adds an operator. */
+		class FInsertOperator : public IDynamicOperatorTransform
+		{
+		public:
+			FInsertOperator(FOperatorID InOperatorID, FOperatorInfo InInfo);
 			virtual EDynamicOperatorTransformQueueAction Transform(FDynamicGraphOperatorData& InGraphOperatorData) override;
 		private:
 			FOperatorID OperatorID;
-			EExecutionOrderInsertLocation Location;
 			FOperatorInfo OperatorInfo;
 		};
 
@@ -116,10 +137,11 @@ namespace Metasound
 		class FRemoveOperator : public IDynamicOperatorTransform
 		{
 		public:
-			FRemoveOperator(FOperatorID InOperatorID);
+			FRemoveOperator(FOperatorID InOperatorID, TArray<FOperatorID> InOperatorsConnectedToInput);
 			virtual EDynamicOperatorTransformQueueAction Transform(FDynamicGraphOperatorData& InGraphOperatorData) override;
 		private:
 			FOperatorID OperatorID;
+			TArray<FOperatorID> OperatorsConnectedToInput;
 		};
 
 		/** A transform which exposes an input to the graph. */
diff --git a/Engine/Plugins/Runtime/Metasound/Source/MetasoundGraphCore/Private/MetasoundDynamicOperatorTransactor.cpp b/Engine/Plugins/Runtime/Metasound/Source/MetasoundGraphCore/Private/MetasoundDynamicOperatorTransactor.cpp
--- a/Engine/Plugins/Runtime/Metasound/Source/MetasoundGraphCore/Private/MetasoundDynamicOperatorTransactor.cpp
+++ b/Engine/Plugins/Runtime/Metasound/Source/MetasoundGraphCore/Private/MetasoundDynamicOperatorTransactor.cpp
@@ -7,6 +7,7 @@
 #include "MetasoundAudioBuffer.h"
 #include "MetasoundBuilderInterface.h"
 #include "MetasoundDataReference.h"
+#include "MetasoundDynamicGraphAlgo.h"
 #include "MetasoundDynamicOperator.h"
 #include "MetasoundGraph.h"
 #include "MetasoundGraphAlgo.h"
@@ -41,34 +42,228 @@ namespace Metasound
 			// Literal nodes always have output vertex with this name. 
 			static const FLazyName LiteralNodeOutputVertexName("Value");
 			
-			// Sorts the graph and determines order of operator execution.
-			TArray<FOperatorID> DetermineOperatorOrder(const IGraph& InGraph)
+			FString GetDebugNodeNameString(const INode& InNode)
+			{
+				return FString::Printf(TEXT("%s_v%d.%d"), *InNode.GetMetadata().ClassName.GetFullName().ToString(), InNode.GetMetadata().MajorVersion, InNode.GetMetadata().MinorVersion);
+			}
+
+			FString GetDebugNodeNameString(const FGuid& InNodeID, const INode& InNode)
+			{
+				return FString::Printf(TEXT("%s:%s"), *InNodeID.ToString(), *GetDebugNodeNameString(InNode));
+			}
+			
+
+
+		}
+
+#if METASOUND_DEBUG_DYNAMIC_TRANSACTOR
+		namespace Debug
+		{
+			// Represents an edge between two operators. Also supports sorting
+			// and special logic for when we do not know the vertex name.
+			// 
+			// This struct allows for easy comparison between two arrays of operator edges.
+			struct FOperatorEdge
 			{
-				/* determine new operator order. */
-				TArray<const INode*> NodeOrder;
+				FOperatorID FromOperator;
+				FOperatorID ToOperator;
+				FVertexName FromVertex;
+				FVertexName ToVertex;
 
-				bool bSuccess = DirectedGraphAlgo::DepthFirstTopologicalSort(InGraph, NodeOrder);
-				if (!bSuccess)
+				friend bool operator<(const FOperatorEdge& InLHS, const FOperatorEdge& InRHS)
 				{
-					UE_LOG(LogMetaSound, Error, TEXT("Cycles found in dynamic graph"));
+					if (InLHS.FromOperator == InRHS.FromOperator)
+					{
+						if (InLHS.ToOperator == InRHS.ToOperator)
+						{
+							if (InLHS.FromVertex.IsNone() || InRHS.ToVertex.IsNone())
+							{
+								// For scenarios where we don't have a vertex name, 
+								// assume that they are equivalent vertex names
+								return false;
+							} 
+							else
+							{
+								if (InLHS.FromVertex == InRHS.FromVertex)
+								{
+									return InLHS.ToVertex.FastLess(InRHS.ToVertex);
+								}
+								else
+								{
+									return InLHS.FromVertex.FastLess(InRHS.FromVertex);
+								}
+							}
+						}
+						else
+						{
+							return InLHS.ToOperator < InRHS.ToOperator;
+						}
+					}
+					else
+					{
+						return InLHS.FromOperator < InRHS.FromOperator;
+					}
 				}
+			};
 
-				TArray<FOperatorID> OperatorOrder;
-				Algo::Transform(NodeOrder, OperatorOrder, static_cast<FOperatorID(*)(const INode*)>(DirectedGraphAlgo::GetOperatorID)); //< Static cast to help deduce which overloaded version of GetOperatorID to call in Algo::Transform
+			/** This debugger compares the multiple graph representations that exist to ensure that they
+			  * are all the same. 
+			  * 
+			  * Note: These could be free functions, but it is easy to implement friendship on a class level
+			  *        as opposed to making multiple friend functions.
+			  */
+			class FDynamicOperatorDebugger
+			{
+			public:
+				// Get operator edges from an IGraph
+				static TArray<FOperatorEdge> FindGraphOperatorEdges(const IGraph& InGraph)
+				{
+					auto DataEdgeToOperatorEdge = [](const FDataEdge& InDataEdge) -> FOperatorEdge
+					{
+						return FOperatorEdge 
+						{
+							DirectedGraphAlgo::GetOperatorID(InDataEdge.From.Node),
+							DirectedGraphAlgo::GetOperatorID(InDataEdge.To.Node),
+							InDataEdge.From.Vertex.VertexName,
+							InDataEdge.To.Vertex.VertexName
+						};
+					};
+
+					TArray<FOperatorEdge> OperatorEdges;
+					Algo::Transform(InGraph.GetDataEdges(), OperatorEdges, DataEdgeToOperatorEdge);
+					return OperatorEdges;
+				}
+				
+				// Get operator edges from an incremental graph sorter
+				static TArray<FOperatorEdge> FindGraphOperatorEdges(const FDynamicGraphIncrementalSorter& InGraphSorter)
+				{
+					TArray<FOperatorEdge> OperatorEdges;
 
-				return OperatorOrder;
-			}
+					for (const TPair<FOperatorID, FDynamicGraphIncrementalSorter::FIncrementalSortOperatorInfo>& Pair : InGraphSorter.OperatorMap)
+					{
+						for (const FOperatorID& Output : Pair.Value.Outputs)
+						{
+							// We do not know what the vertex names are within the sorter
+							// so we leave them as NONE.
+							OperatorEdges.Add(FOperatorEdge{Pair.Key, Output}); 
+						}
+					}
 
-			FString GetDebugNodeNameString(const INode& InNode)
-			{
-				return FString::Printf(TEXT("%s_v%d.%d"), *InNode.GetMetadata().ClassName.GetFullName().ToString(), InNode.GetMetadata().MajorVersion, InNode.GetMetadata().MinorVersion);
-			}
+					return OperatorEdges;
+				}
 
-			FString GetDebugNodeNameString(const FGuid& InNodeID, const INode& InNode)
+				// Get operator eges from Graph Operator Data
+				static TArray<FOperatorEdge> FindGraphOperatorEdges(const DirectedGraphAlgo::FGraphOperatorData& InGraphData)
+				{
+					using namespace DirectedGraphAlgo;
+
+					TArray<FOperatorEdge> OperatorEdges;
+
+					for (const TPair<FOperatorID, FGraphOperatorData::FOperatorInfo>& OpInfoPair : InGraphData.OperatorMap)
+					{
+						for (const TPair<FVertexName, TArray<FGraphOperatorData::FVertexDestination>>& OutConnectPair : OpInfoPair.Value.OutputConnections)
+						{
+							for (const FGraphOperatorData::FVertexDestination& Destination : OutConnectPair.Value)
+							{
+								OperatorEdges.Add(FOperatorEdge{OpInfoPair.Key, Destination.OperatorID, OutConnectPair.Key, Destination.VertexName});
+							}
+						}
+					}
+
+					return OperatorEdges;
+				}
+
+				static void LogMissingEdge(const TCHAR* InWhereIsMissing, const FOperatorEdge& InEdge)
+				{
+					UE_LOG(LogMetaSound, Display, TEXT("%s Missing Edge %d %s -> %d %s"), InWhereIsMissing, InEdge.FromOperator, *InEdge.FromVertex.ToString(), InEdge.ToOperator, *InEdge.ToVertex.ToString());
+				}
+
+				// Returns true if arrays are equivalent, false otherwise. 
+				// Differences are logged.
+				// Input arrays must be sorted. 
+				static bool CompareAndLogEdgeArrays(const TCHAR* InSourceA, const TArray<FOperatorEdge>& InEdgesA, const TCHAR* InSourceB, const TArray<FOperatorEdge>& InEdgesB)
+				{
+					bool bEqual = true;
+
+					const FOperatorEdge* EdgeA = InEdgesA.GetData();
+					const FOperatorEdge* const EdgeAEnd = EdgeA + InEdgesA.Num();
+					const FOperatorEdge* EdgeB = InEdgesB.GetData();
+					const FOperatorEdge* const EdgeBEnd = EdgeB + InEdgesB.Num();
+
+					// Increment through both arrays simultaneously.
+					while ((EdgeA != EdgeAEnd) && (EdgeB != EdgeBEnd))
+					{
+						if (*EdgeA < *EdgeB)
+						{
+							LogMissingEdge(InSourceB, *EdgeA);
+							EdgeA++;
+							bEqual = false;
+						}
+						else if (*EdgeB < *EdgeA)
+						{
+							LogMissingEdge(InSourceA, *EdgeB);
+							EdgeB++;
+							bEqual = false;
+						}
+						else
+						{
+							// Matching
+							EdgeA++;
+							EdgeB++;
+						}
+					}
+
+					// Any left over edges are unmatched.
+					while (EdgeA != EdgeAEnd)
+					{
+						LogMissingEdge(InSourceB, *EdgeA);
+						EdgeA++;
+						bEqual = false;
+					}
+
+					while (EdgeB != EdgeBEnd)
+					{
+						LogMissingEdge(InSourceA, *EdgeB);
+						EdgeB++;
+						bEqual = false;
+					}
+
+					return bEqual;
+				}
+
+				// Returns true if the graph representations are equivalent, false otherwise.
+				// Differences are logged.
+				static bool CompareAndLogGraphRepresentationDiscrepancies(const FDynamicOperatorTransactor& InTransactor, const FDynamicOperator& InDynamicOperator)
+				{
+					using namespace Debug;
+
+					// Gather edges from various graph representations
+					TArray<FOperatorEdge> GraphEdges = FindGraphOperatorEdges(InTransactor.Graph);
+					TArray<FOperatorEdge> GraphSorterEdges = FindGraphOperatorEdges(InTransactor.GraphSorter);
+					TArray<FOperatorEdge> RuntimeEdges = FindGraphOperatorEdges(InDynamicOperator.DynamicOperatorData);
+
+					// Edges must be sorted before performing comparison
+					GraphEdges.Sort();
+					GraphSorterEdges.Sort();
+					RuntimeEdges.Sort();
+
+					// Run comparison
+					bool bEqualGraphAndSorter = CompareAndLogEdgeArrays(TEXT("Graph"), GraphEdges, TEXT("GraphSorter"), GraphSorterEdges);
+					bool bEqualGraphAndRuntime = CompareAndLogEdgeArrays(TEXT("Graph"), GraphEdges, TEXT("DynamicRuntime"), RuntimeEdges);
+					bool bEqualRuntimeAndSorter = CompareAndLogEdgeArrays(TEXT("DynamicRuntime"), RuntimeEdges, TEXT("GraphSorter"), GraphSorterEdges);
+
+					return bEqualGraphAndSorter && bEqualGraphAndRuntime && bEqualRuntimeAndSorter;
+				}
+			};
+
+			// Returns true if the graph representations are equivalent, false otherwise.
+			// Differences are logged.
+			bool CompareAndLogGraphRepresentationDiscrepancies(const FDynamicOperatorTransactor& InTransactor, const FDynamicOperator& InDynamicOperator)
 			{
-				return FString::Printf(TEXT("%s:%s"), *InNodeID.ToString(), *GetDebugNodeNameString(InNode));
+				return FDynamicOperatorDebugger::CompareAndLogGraphRepresentationDiscrepancies(InTransactor, InDynamicOperator);
 			}
 		}
+#endif // if METASOUND_DEBUG_DYNAMIC_TRANSACTOR
 
 		bool operator<(const FDynamicOperatorTransactor::FLiteralNodeID& InLHS, const FDynamicOperatorTransactor::FLiteralNodeID& InRHS)
 		{
@@ -86,11 +281,288 @@ namespace Metasound
 			}
 		}
 
+		FDynamicGraphIncrementalSorter::FDynamicGraphIncrementalSorter()
+		{
+		}
+
+		FDynamicGraphIncrementalSorter::FDynamicGraphIncrementalSorter(const IGraph& InGraph)
+		{
+			Init(InGraph);
+		}
+
+		/** Add a node to the graph. */
+		int32 FDynamicGraphIncrementalSorter::InsertOperator(FOperatorID InOperator, FDynamicGraphIncrementalSorter::EInsertLocation InLocation)
+		{
+			int32 NewOrdinal = ORDINAL_NONE;
+
+			switch (InLocation)
+			{
+				case EInsertLocation::First:
+					MinOrdinal--;
+					NewOrdinal = MinOrdinal;
+					break;
+
+				case EInsertLocation::Last:
+					MaxOrdinal++;
+					NewOrdinal = MaxOrdinal;
+					break;
+			}
+
+			check(NewOrdinal != ORDINAL_NONE);
+
+			if (FIncrementalSortOperatorInfo* Info = OperatorMap.Find(InOperator))
+			{
+				UE_LOG(LogMetaSound, Warning, TEXT("Attempt to add operator %d when operator already exists"), InOperator);
+				Info->Ordinal = NewOrdinal;
+			}
+			else
+			{	
+				OperatorMap.Add(InOperator, { NewOrdinal });
+			}
+
+			return NewOrdinal;
+		}
+
+		/** Remove a node from the graph. */
+		int32 FDynamicGraphIncrementalSorter::RemoveOperator(FOperatorID InOperatorID)
+		{
+			int32 RemovedOrdinal = ORDINAL_NONE;
+
+			if (FIncrementalSortOperatorInfo* Info = OperatorMap.Find(InOperatorID))
+			{
+				RemovedOrdinal = Info->Ordinal;
+
+				// Remove any remaining connections
+				for (FOperatorID InputOperatorID : Info->Inputs)
+				{
+					if (FIncrementalSortOperatorInfo* InputInfo = OperatorMap.Find(InputOperatorID))
+					{
+						InputInfo->Outputs.RemoveSwap(InOperatorID);
+					}
+				}
+
+				for (FOperatorID OutputOperatorID : Info->Outputs)
+				{
+					if (FIncrementalSortOperatorInfo* OutputInfo = OperatorMap.Find(OutputOperatorID))
+					{
+						OutputInfo->Inputs.RemoveSwap(InOperatorID);
+					}
+				}
+				
+				// Remove operator
+				OperatorMap.Remove(InOperatorID);
+			}
+
+			return RemovedOrdinal;
+		}
+
+		void FDynamicGraphIncrementalSorter::GenerateOrdinals(TMap<FOperatorID, int32>& OutOrdinals) const
+		{
+			for (const TPair<FOperatorID, FIncrementalSortOperatorInfo>& Entry : OperatorMap)
+			{
+				OutOrdinals.Emplace(Entry.Key, Entry.Value.Ordinal);
+			}
+		}
+
+		/** Add an edge to the graph, connecting two vertices from two 
+		 * nodes. 
+		 *
+		 * @param InFromOperatorID - Operator which contains the output vertex.
+		 * @param InToOperatorID - Operator which contains the input vertex.
+		 * @param OutOrdinalUpdates - Array to populate with ordinal updates to maintain topological sort of graph.
+		 */
+		void FDynamicGraphIncrementalSorter::AddDataEdge(FOperatorID InFromOperatorID, FOperatorID InToOperatorID, TArray<FOrdinalSwap>& OutOrdinalUpdates)
+		{
+			// Vertex names are not stored here, but if there are multiple edges 
+			// conneting two operators, then there will be multiple entries in the 
+			// FIncrementalSortOperatorInfo::Outputs & FIncrementalSortOperatorInfo::Inputs arrays. 	
+			FIncrementalSortOperatorInfo* FromInfo = OperatorMap.Find(InFromOperatorID);
+			if (!FromInfo)
+			{
+				UE_LOG(LogMetaSound, Error, TEXT("Could not find source operator ID %d for adding edge. Dynamic MetaSound may not render properly."), InFromOperatorID);
+				return;
+			}
+
+			FIncrementalSortOperatorInfo* ToInfo = OperatorMap.Find(InToOperatorID);
+			if (!ToInfo)
+			{
+				UE_LOG(LogMetaSound, Error, TEXT("Could not find destination operator ID %d for adding edge. Dynamic MetaSound may not render properly."), InToOperatorID);
+				return;
+				
+			}
+
+			// Add edge to operator info structs
+			FromInfo->Outputs.Add(InToOperatorID);
+			ToInfo->Inputs.Add(InFromOperatorID);
+
+			// Only sort if the "From" operator isn't executing before the "To" operator. 
+			if (FromInfo->Ordinal > ToInfo->Ordinal)
+			{
+				// Perform incremental sort
+				IncrementalTopologicalSortForNewEdge(InFromOperatorID, FromInfo->Ordinal, InToOperatorID, ToInfo->Ordinal, OutOrdinalUpdates);
+
+				// Apply sort changes internally.
+				for (const FOrdinalSwap& Update: OutOrdinalUpdates)
+				{
+					OperatorMap[Update.OperatorID].Ordinal = Update.NewOrdinal;
+				}
+			}
+		}
+
+		/** Remove the given data edge. */
+		void FDynamicGraphIncrementalSorter::RemoveDataEdge(FOperatorID InFromOperatorID, FOperatorID InToOperatorID)
+		{
+			// Vertex names are not stored here, but if there are multiple edges 
+			// conneting two operators, then there will be multiple entries in the 
+			// FIncrementalSortOperatorInfo::Outputs & FIncrementalSortOperatorInfo::Inputs arrays. 
+			if (FIncrementalSortOperatorInfo* FromInfo = OperatorMap.Find(InFromOperatorID))
+			{
+				FromInfo->Outputs.RemoveSingleSwap(InToOperatorID);
+			}
+			else
+			{
+				UE_LOG(LogMetaSound, Warning, TEXT("Could not find source operator ID %d for removing edge. Dynamic MetaSound may not render properly."), InFromOperatorID);
+			}
+
+			if (FIncrementalSortOperatorInfo* ToInfo = OperatorMap.Find(InToOperatorID))
+			{
+				ToInfo->Inputs.RemoveSingleSwap(InFromOperatorID);
+			}
+			else	
+			{
+				UE_LOG(LogMetaSound, Warning, TEXT("Could not find destination operator ID %d for removing edge. Dynamic MetaSound may not render properly."), InToOperatorID);
+			}
+		}
+
+
+		void FDynamicGraphIncrementalSorter::Init(const IGraph& InGraph)
+		{
+			/* determine new operator order. */
+			TArray<const INode*> NodeOrder;
+
+			bool bSuccess = DirectedGraphAlgo::DepthFirstTopologicalSort(InGraph, NodeOrder);
+			if (!bSuccess)
+			{
+				UE_LOG(LogMetaSound, Error, TEXT("Cycles found in dynamic graph"));
+			}
+
+			// Initialize ordinals in operator map
+			const int32 Num = NodeOrder.Num();
+			const INode** NodeData = NodeOrder.GetData();
+			for (int32 Ordinal = 0; Ordinal < Num; Ordinal++)
+			{
+				OperatorMap.Add(DirectedGraphAlgo::GetOperatorID(*NodeData[Ordinal]), FIncrementalSortOperatorInfo{ Ordinal });
+			}
+
+			// The next available ordinal for any operators added after initialization. 
+			MaxOrdinal = Num;
+
+			// Initialize edges in operator map
+			for (const FDataEdge& Edge : InGraph.GetDataEdges())
+			{
+				FOperatorID FromOperatorID = DirectedGraphAlgo::GetOperatorID(Edge.From.Node);
+				FOperatorID ToOperatorID = DirectedGraphAlgo::GetOperatorID(Edge.To.Node);
+
+				OperatorMap[FromOperatorID].Outputs.Add(ToOperatorID);
+				OperatorMap[ToOperatorID].Inputs.Add(FromOperatorID);
+			}
+		}
+
+		void FDynamicGraphIncrementalSorter::IncrementalTopologicalSortForNewEdge(FOperatorID InFromOperatorID, int32 InFromOrdinal, FOperatorID InToOperatorID, int32 InToOrdinal, TArray<FOrdinalSwap>& OutUpdates)
+		{
+			UE_CLOG(InToOrdinal > InFromOrdinal, LogMetaSound, Warning, TEXT("Operators are already in order. Only perform incremental sort if necessary."));
+
+
+			// Incremental topological sort maintains that the "FromOperator"
+			// is sorted before the "ToOperator", but does so in an incremental way 
+			// to avoid resorting the entire graph. To achieve this it:
+			
+			// 1. Searches for all the operators and their ordinals which need to be 
+			//    resorted..
+			// 2. Sorts the operators appropriately, but only using the set of ordinals 
+			//    already associated with the operators of interest.
+			
+			TArray<FOperatorID> SearchStack;
+			
+			// Find operators which need to be move before the "ToOperator"
+			// 
+			// Starting with the "FromOperator", find all operators which are:
+			//  - Directly or indirectly connected to the input of the "FromOperator"
+			//  AND
+			//  - Are set to execute after the "ToOperator". 
+			SearchStack.Add(InFromOperatorID);
+			while (SearchStack.Num())
+			{
+				FOperatorID CandidateOperatorID = SearchStack.Pop();
+				if (!OutUpdates.Contains(CandidateOperatorID))
+				{
+					const FIncrementalSortOperatorInfo& CandidateOperatorInfo = OperatorMap[CandidateOperatorID];
+					if (CandidateOperatorInfo.Ordinal > InToOrdinal)
+					{
+						OutUpdates.Add({CandidateOperatorID, CandidateOperatorInfo.Ordinal, CandidateOperatorInfo.Ordinal});
+						for (const FOperatorID& Connection : CandidateOperatorInfo.Inputs)
+						{
+							SearchStack.AddUnique(Connection);
+						}
+					}
+				}
+			}
+
+
+			const int32 NumOutputsLessThanOrdinal = OutUpdates.Num();
+			OutUpdates.Sort(FOrdinalSwap::OriginalOrdinalLessThan); // Sort by ascending original ordinal to maintain relative order
+
+			// Find operators which need to be moved after the "FromOperator"
+			// 
+			// Starting with the "ToOperator", find all operators which are:
+			//  - Directly or indirectly connected to the output of the "ToOperator"
+			//  AND
+			//  - Are set to execute before the "FromOperator". 
+			SearchStack.Add(InToOperatorID);
+			while (SearchStack.Num())
+			{
+				FOperatorID CandidateOperatorID = SearchStack.Pop();
+				if (!OutUpdates.Contains(CandidateOperatorID))
+				{
+					const FIncrementalSortOperatorInfo& CandidateOperatorInfo = OperatorMap[CandidateOperatorID];
+					if (CandidateOperatorInfo.Ordinal < InFromOrdinal)
+					{
+						OutUpdates.Add({CandidateOperatorID, CandidateOperatorInfo.Ordinal, CandidateOperatorInfo.Ordinal});
+						for (const FOperatorID& Connection : CandidateOperatorInfo.Outputs)
+						{
+							SearchStack.AddUnique(Connection);
+						}
+					}
+				}
+			}
+			const int32 NumInputsGreaterThanOrdinal = OutUpdates.Num() - NumOutputsLessThanOrdinal;
+
+			// Sort 2nd half of operators in ascending order w/o modifying order of 1st set of operators. 
+			// After this sort, the OutUpdates array will contain operators in the desired
+			// order but with incorrect "ordinal" values. 
+			Algo::Sort(TArrayView<FOrdinalSwap>(&OutUpdates[NumOutputsLessThanOrdinal], NumInputsGreaterThanOrdinal), FOrdinalSwap::OriginalOrdinalLessThan);
+
+			// Gather the available ordinals and sort them in ascending order
+			TArray<int32> AvailableOrdinals;
+			Algo::Transform(OutUpdates, AvailableOrdinals, [](const FOrdinalSwap& InOperatorOrdinal) { return InOperatorOrdinal.OriginalOrdinal; });
+			AvailableOrdinals.Sort();
+			
+			// Assign the sorted ordinals to the sorted operators. 
+			const int32 Num = OutUpdates.Num();
+			const int32* AvailableOrdinalData = AvailableOrdinals.GetData();
+			FOrdinalSwap* UpdateData = OutUpdates.GetData();
+
+			for (int32 i = 0; i < Num; i++)
+			{
+				UpdateData[i].NewOrdinal = AvailableOrdinalData[i];
+			}
+		}
+
 		FDynamicOperatorTransactor::FDynamicOperatorTransactor(const FGraph& InGraph)
 		: OperatorBuilder(DynamicOperatorTransactorPrivate::GetOperatorBuilderSettings())
 		, Graph(InGraph)
+		, GraphSorter(InGraph)
 		{
-			CurrentOperatorOrder = DynamicOperatorTransactorPrivate::DetermineOperatorOrder(Graph);
 		}
 
 		FDynamicOperatorTransactor::FDynamicOperatorTransactor()
@@ -104,6 +576,14 @@ namespace Metasound
 			TSharedRef<TSpscQueue<TUniquePtr<IDynamicOperatorTransform>>> Queue = MakeShared<TSpscQueue<TUniquePtr<IDynamicOperatorTransform>>>();
 			OperatorInfos.Add(FDynamicOperatorInfo{InOperatorSettings, InEnvironment, Queue});
 
+			TMap<FOperatorID, int32> OperatorOrdinals;
+			GraphSorter.GenerateOrdinals(OperatorOrdinals);
+
+			// All of these initial operations have to happen in one fell swoop for the 
+			// dynamic metasound to be setup correctly. We use an atomic transform
+			// to ensure they are all applied before generating audio.
+			TArray<TUniquePtr<IDynamicOperatorTransform>> AtomicTransforms;
+
 			// Unconnected nodes are intentionally skipped by the operator builder
 			// in order to reduce perf. In a dynamic operator, these nodes may be
 			// connected in the future. We queue them up to be added here so that
@@ -114,14 +594,27 @@ namespace Metasound
 			{
 				for (const TPair<FGuid, const INode*>& GuidAndNode : UnconnectedNodes)
 				{
-					// Only add to this queue because we do not know at what point
-					// the other queues and dynamic operators were created.
-					Queue->Enqueue(CreateAddOperatorTransform(*GuidAndNode.Value, EExecutionOrderInsertLocation::Last, InOperatorSettings, InEnvironment));
+					if (const int32* Ordinal = OperatorOrdinals.Find(DirectedGraphAlgo::GetOperatorID(GuidAndNode.Value)))
+					{
+						AtomicTransforms.Add(CreateInsertOperatorTransform(*GuidAndNode.Value, *Ordinal, InOperatorSettings, InEnvironment));
+					}
 				}
 			}
 
-			// Force order to be synchronized to internal execution order.
-			Queue->Enqueue(MakeUnique<FSetOperatorOrder>(CurrentOperatorOrder));
+			// When we create a new Dynamic Operator, the first thing that needs
+			// to synchronize is the initial set of operator ordinals. The ordinals
+			// in the DynamicOperator and the FDynamicGraphIncrementalSorter must
+			// be exactly equal so that subsequent modifications to the ordinals
+			// result in equal ordinals between the two objects.
+			//
+			// Note: The ordinals in FDynamicGraphIncrementalSorter are NOT 
+			// expected to be equal to the ordinals set in the FOperatorBuilder
+			// because they use different algorithms to determine order.
+			AtomicTransforms.Add(MakeUnique<FSetOperatorOrdinalsAndSort>(OperatorOrdinals));
+
+			// Only add to THIS queue because we do not know at what point
+			// the other queues and dynamic operators were created.
+			Queue->Enqueue(MakeUnique<FAtomicTransform>(MoveTemp(AtomicTransforms)));
 
 			return Queue;
 		}
@@ -140,8 +633,9 @@ namespace Metasound
 			TSharedRef<const INode> NodePtr(InNode.Release());
 
 			Graph.AddNode(InNodeID, NodePtr);
+			int32 OperatorOrdinal = GraphSorter.InsertOperator(GetOperatorID(*NodePtr), FDynamicGraphIncrementalSorter::EInsertLocation::Last);
 
-			EnqueueAddOperatorTransform(*NodePtr, EExecutionOrderInsertLocation::Last);
+			EnqueueInsertOperatorTransform(*NodePtr, OperatorOrdinal);
 		}
 
 		void FDynamicOperatorTransactor::RemoveNode(const FGuid& InNodeID)
@@ -162,18 +656,15 @@ namespace Metasound
 					}
 				}
 
+				constexpr bool bRemoveDataEdgesWithNode = true;
 				if (OutputsToFade.Num())
 				{
-					EnqueueFadeAndRemoveOperatorTransform(*Node, OutputsToFade);
+					FadeAndRemoveNodeInternal(*Node, OutputsToFade, bRemoveDataEdgesWithNode);
 				}
 				else
 				{
-					EnqueueRemoveOperatorTransform(*Node);
+					RemoveNodeInternal(*Node, bRemoveDataEdgesWithNode);
 				}
-
-				constexpr bool bRemoveDataEdgesWithNode = true;
-				bool bRemovedNodeFromGraph = Graph.RemoveNode(InNodeID, bRemoveDataEdgesWithNode);
-				check(bRemovedNodeFromGraph); //< Should always be true because we know the node exists in the Graph from `Graph.FindNode(...)`
 			}
 			else
 			{
@@ -208,7 +699,9 @@ namespace Metasound
 
 			const INode* FromNode = Graph.FindNode(InFromNodeID);
 			const INode* ToNode = Graph.FindNode(InToNodeID);
-			const INode* ReplacementLiteralNode = InReplacementLiteralNode.Get(); // Cache pointer because TUniquePtr<INode> will get moved
+			const FOperatorID FromOperatorID = DirectedGraphAlgo::GetOperatorID(FromNode);
+			const FOperatorID ToOperatorID = DirectedGraphAlgo::GetOperatorID(ToNode);
+			const FOperatorID LiteralOperatorID = DirectedGraphAlgo::GetOperatorID(InReplacementLiteralNode.Get());
 
 			if ((nullptr == ToNode) || (nullptr == FromNode))
 			{
@@ -235,25 +728,40 @@ namespace Metasound
 				UE_LOG(LogMetaSound, Error, TEXT("Failed to remove edge from %s:%s to %s:%s on internal graph."), *InFromNodeID.ToString(), *InFromVertex.ToString(), *InToNodeID.ToString(), *InToVertex.ToString());
 				return;
 			}
+			GraphSorter.RemoveDataEdge(FromOperatorID, ToOperatorID);
 
-			// Store literal node associated with the target of the literal value.
-			LiteralNodeMap.Add(FLiteralNodeID{ InToNodeID, InToVertex }, MoveTemp(InReplacementLiteralNode));
-			bSuccess = Graph.AddDataEdge(*ReplacementLiteralNode, DynamicOperatorTransactorPrivate::LiteralNodeOutputVertexName, *ToNode, InToVertex);
+
+			bSuccess = Graph.AddDataEdge(*InReplacementLiteralNode, DynamicOperatorTransactorPrivate::LiteralNodeOutputVertexName, *ToNode, InToVertex);
 			if (!bSuccess)
 			{
 				UE_LOG(LogMetaSound, Error, TEXT("Failed to add literal for %s:%s on internal graph."), *InToNodeID.ToString(), *InToVertex.ToString());
 				return;
 			}
 
+			INode& LiteralNodeRef = *InReplacementLiteralNode; // Store ref since pointer will be moved. 
+			// Store literal node associated with the target of the literal value.
+			LiteralNodeMap.Add(FLiteralNodeID{ InToNodeID, InToVertex }, MoveTemp(InReplacementLiteralNode));
+
+			// Put literals in the front of the execution stack to simplify updating 
+			// runtime instances. No need to sort the entire graph if we are just
+			// inserting something at the beginning of the execution stack. 
+			int32 LiteralOrdinal = GraphSorter.InsertOperator(LiteralOperatorID, FDynamicGraphIncrementalSorter::EInsertLocation::First);
+			TArray<FOrdinalSwap> OrdinalSwaps;
+			GraphSorter.AddDataEdge(LiteralOperatorID, ToOperatorID, OrdinalSwaps);
+
+			// There should not be any ordinal swaps because the literal operator was inserted before any 
+			// other operator and so will always have it's output data ready before the ToOperator executes. 
+			check(OrdinalSwaps.Num() == 0); 
+
+			// Immediately disconnect non-audio edges. 
 			if (ToNode->GetVertexInterface().GetInputVertex(InToVertex).DataTypeName == GetMetasoundDataTypeName<FAudioBuffer>())
 			{
 				// Handle audio edge removal with a fade out.
-				EnqueueFadeAndRemoveEdgeOperatorTransform(*FromNode, InFromVertex, *ToNode, InToVertex, *ReplacementLiteralNode);
+				EnqueueFadeAndRemoveEdgeOperatorTransform(*FromNode, InFromVertex, *ToNode, InToVertex, LiteralNodeRef, LiteralOrdinal);
 			}
 			else
 			{
-				// Immediately disconnect non-audio edges. 
-				EnqueueRemoveEdgeOperatorTransform(*FromNode, InFromVertex, *ToNode, InToVertex, *ReplacementLiteralNode);
+				EnqueueRemoveEdgeOperatorTransform(*FromNode, InFromVertex, *ToNode, InToVertex, LiteralNodeRef, LiteralOrdinal);
 			}
 		}
 
@@ -277,11 +785,11 @@ namespace Metasound
 			}
 
 			// Always insert new literal nodes first in execution order.
-			CurrentOperatorOrder.Insert(DirectedGraphAlgo::GetOperatorID(*LiteralNode), 0);
+			int32 LiteralOrdinal = GraphSorter.InsertOperator(DirectedGraphAlgo::GetOperatorID(LiteralNode), FDynamicGraphIncrementalSorter::EInsertLocation::First);
 
 			auto CreateAddNodeTransform = [&](const FOperatorSettings& InOperatorSettings, const FMetasoundEnvironment& InEnvironment) -> TUniquePtr<IDynamicOperatorTransform>
 			{
-				return CreateAddOperatorTransform(*LiteralNode, EExecutionOrderInsertLocation::First, InOperatorSettings, InEnvironment);
+				return CreateInsertOperatorTransform(*LiteralNode, LiteralOrdinal, InOperatorSettings, InEnvironment);
 			};
 
 			EnqueueTransformOnOperatorQueues(CreateAddNodeTransform);
@@ -416,107 +924,128 @@ namespace Metasound
 			 * Literal nodes are stored on the FDynamicOperatorTransactor and need to be 
 			 * disconnected and removed if they are no longer being used. 
 			 */
-			const FLiteralNodeID LiteralNodeKey{InToNodeID, InToVertex};
 			TUniquePtr<INode> PriorLiteralNode;
-			if (LiteralNodeMap.RemoveAndCopyValue(LiteralNodeKey, PriorLiteralNode))
+			bool bPriorLiteralNodeExists = LiteralNodeMap.RemoveAndCopyValue(FLiteralNodeID{InToNodeID, InToVertex}, PriorLiteralNode);
+
+			// Get relevant operator IDs.
+			const FOperatorID PriorLiteralOperatorID = DirectedGraphAlgo::GetOperatorID(PriorLiteralNode.Get());
+			const FOperatorID ToOperatorID = DirectedGraphAlgo::GetOperatorID(InToNode);
+			const FOperatorID FromOperatorID = DirectedGraphAlgo::GetOperatorID(InFromNode);
+
+			// Update edges on internal graph
+			if (bPriorLiteralNodeExists)
 			{
 				Graph.RemoveDataEdge(*PriorLiteralNode, DynamicOperatorTransactorPrivate::LiteralNodeOutputVertexName, InToNode, InToVertex);
+				GraphSorter.RemoveDataEdge(PriorLiteralOperatorID, ToOperatorID);
 			}
 
-			/* add edge to internal graph. */
 			Graph.AddDataEdge(InFromNode, InFromVertex, InToNode, InToVertex);
-			
+			TArray<FOrdinalSwap> OrdinalUpdates;
+			GraphSorter.AddDataEdge(FromOperatorID, ToOperatorID, OrdinalUpdates);
+		
+			if (bPriorLiteralNodeExists)
+			{
+				// The Graph does not maintain literal nodes so there is no need to remove the PriorLiteralNode from the Graph. Those are
+				// managed in the LiteralNodeMap. But, the GraphSorter DOES maintain literal operators, and so we have to still remove
+				// the literal operator from the GraphSorter.
+				GraphSorter.RemoveOperator(PriorLiteralOperatorID);
+			}
+
 			if (InputVertex->DataTypeName == GetMetasoundDataTypeName<FAudioBuffer>())
 			{
 				// If edge is audio, then the connection needs to be faded
-				EnqueueFadeAndAddEdgeOperatorTransform(InFromNode, InFromVertex, InToNode, InToVertex, PriorLiteralNode.Get());
+				EnqueueFadeAndAddEdgeOperatorTransform(InFromNode, InFromVertex, InToNode, InToVertex, PriorLiteralNode.Get(), OrdinalUpdates);
 			}
 			else
 			{
 				// If the edge is not audio, then no fading is performed. 
-				EnqueueAddEdgeOperatorTransform(InFromNode, InFromVertex, InToNode, InToVertex, PriorLiteralNode.Get());
+				EnqueueAddEdgeOperatorTransform(InFromNode, InFromVertex, InToNode, InToVertex, PriorLiteralNode.Get(), OrdinalUpdates);
 			}
 		}
 		
-		void FDynamicOperatorTransactor::EnqueueAddOperatorTransform(const INode& InNode, EExecutionOrderInsertLocation InLocation)
+		void FDynamicOperatorTransactor::EnqueueInsertOperatorTransform(const INode& InNode, int32 InOrdinal)
 		{
-			// Update the CurrentOperatorOrder based on the insert location of new node. 
-			// This avoids resorting of entire graph on rendering dynamic metasounds by
-			// simply appending to a specific location.
-			switch (InLocation)
-			{
-				case EExecutionOrderInsertLocation::First:
-					CurrentOperatorOrder.Insert(DirectedGraphAlgo::GetOperatorID(InNode), 0);
-					break;
-
-				case EExecutionOrderInsertLocation::Last:
-					CurrentOperatorOrder.Add(DirectedGraphAlgo::GetOperatorID(InNode));
-					break;
-			}
-
 			auto CreateAddNodeTransform = [&](const FOperatorSettings& InOperatorSettings, const FMetasoundEnvironment& InEnvironment) -> TUniquePtr<IDynamicOperatorTransform>
 			{
-				return CreateAddOperatorTransform(InNode, InLocation, InOperatorSettings, InEnvironment);
+				return CreateInsertOperatorTransform(InNode, InOrdinal, InOperatorSettings, InEnvironment);
 			};
 
 			EnqueueTransformOnOperatorQueues(CreateAddNodeTransform);
 		}
 
-		void FDynamicOperatorTransactor::EnqueueFadeAndRemoveOperatorTransform(const INode& InNode, TArrayView<const FVertexName> InOutputsToFade)
+		void FDynamicOperatorTransactor::RemoveNodeInternal(const INode& InNode, bool bInRemoveDataEdgesWithNode)
 		{
-			TArrayView<const FVertexName> InputsToFade; // We do not need to fade any inputs when removing a node.
-
-			EnqueueBeginFadeOperatorTransform(InNode, EAudioFadeType::FadeOut, InputsToFade, InOutputsToFade);
+			// Find any existing edges and remove them from the graph
+			TArray<FOperatorID> OperatorsConnectedToInput;
+			if (bInRemoveDataEdgesWithNode)
+			{
+				for (const FDataEdge& Edge : Graph.GetDataEdges())
+				{
+					if ((Edge.From.Node == &InNode) || (Edge.To.Node == &InNode))
+					{
+						// Remove edge from graph sorter
+						GraphSorter.RemoveDataEdge(DirectedGraphAlgo::GetOperatorID(Edge.From.Node), DirectedGraphAlgo::GetOperatorID(Edge.To.Node));
+
+						if (Edge.To.Node == &InNode)
+						{
+							// Track all incomding edges to remove them from dynamic operator's model.
+							OperatorsConnectedToInput.AddUnique(DirectedGraphAlgo::GetOperatorID(Edge.From.Node));
+						}
+					}
+				}
+			}
 
-			const FOperatorID OperatorID = DirectedGraphAlgo::GetOperatorID(InNode);
-			CurrentOperatorOrder.RemoveSingle(OperatorID);
+			
+			GraphSorter.RemoveOperator(DirectedGraphAlgo::GetOperatorID(InNode));
 
-			auto CreateEndFadeOutTransform = [&](const FOperatorSettings& InOperatorSettings, const FMetasoundEnvironment& InEnvironment) -> TUniquePtr<IDynamicOperatorTransform>
-			{
-				// We can skip the FEndAudioFadeTransform as an optimization here.
-				// FEndAudioFadeTransform removes the fading wrapper around the node,
-				// but since the node is being removed, we can remove the wrapper
-				// and the node with a single FRemoveOperator transform.
-				return MakeUnique<FRemoveOperator>(OperatorID);
-			};
+			// that do not exist will stay on the runtime model. 
+			// There's an open question about whether we should fade disconnections. We probably should, 
+			// but we lack the literal node's to do so. 
+			EnqueueRemoveOperatorTransform(InNode, OperatorsConnectedToInput);
 
-			EnqueueTransformOnOperatorQueues(CreateEndFadeOutTransform);
+			bool bRemovedNodeFromGraph = Graph.RemoveNode(InNode.GetInstanceID(), bInRemoveDataEdgesWithNode);
+			check(bRemovedNodeFromGraph); //< Should always be true because callers of RemoveNodeInternal() ensure that the node exists in graph.
 		}
 
-		void FDynamicOperatorTransactor::EnqueueAddEdgeOperatorTransform(const INode& InFromNode, const FVertexName& InFromVertex, const INode& InToNode, const FVertexName& InToVertex, const INode* InPriorLiteralNode)
+		void FDynamicOperatorTransactor::FadeAndRemoveNodeInternal(const INode& InNode, TArrayView<const FVertexName> InOutputsToFade, bool bInRemoveDataEdgesWithNode)
 		{
-			/* enqueue an update. */
-			FOperatorID FromOperatorID = DirectedGraphAlgo::GetOperatorID(InFromNode);
-			FOperatorID ToOperatorID = DirectedGraphAlgo::GetOperatorID(InToNode);
+			TArrayView<const FVertexName> InputsToFade; // We do not need to fade any inputs when removing a node.
 
-			if (InPriorLiteralNode)
-			{
-				CurrentOperatorOrder.RemoveSingle(DirectedGraphAlgo::GetOperatorID(InPriorLiteralNode));
-			}
+			EnqueueBeginFadeOperatorTransform(InNode, EAudioFadeType::FadeOut, InputsToFade, InOutputsToFade);
 
-			// Find order of operators after adding edge. 
-			TArray<FOperatorID> NewOperatorOrder = DynamicOperatorTransactorPrivate::DetermineOperatorOrder(Graph);
+			// We can skip the FEndAudioFadeTransform as an optimization here.
+			// FEndAudioFadeTransform removes the fading wrapper around the node,
+			// but since the node is being removed, we can remove the wrapper
+			// and the node with a single FRemoveOperator transform.
+			RemoveNodeInternal(InNode, bInRemoveDataEdgesWithNode);
+		}
 
-			// Only set new order if it's different than existing.
-			bool bSetNewOperatorOrder = NewOperatorOrder != CurrentOperatorOrder;
-			if (bSetNewOperatorOrder)
-			{
-				CurrentOperatorOrder = NewOperatorOrder;
-			}
+		void FDynamicOperatorTransactor::EnqueueAddEdgeOperatorTransform(const INode& InFromNode, const FVertexName& InFromVertex, const INode& InToNode, const FVertexName& InToVertex, const INode* InPriorLiteralNode, const TArray<FOrdinalSwap>& InOrdinalUpdates)
+		{
+			/* enqueue an update. */
+			const FOperatorID FromOperatorID = DirectedGraphAlgo::GetOperatorID(InFromNode);
+			const FOperatorID ToOperatorID = DirectedGraphAlgo::GetOperatorID(InToNode);
+			const FOperatorID PriorLiteralOperatorID = DirectedGraphAlgo::GetOperatorID(InPriorLiteralNode);
 
+			// Create transforms for runtime
 			auto CreateAddEdgeTransform = [&](const FOperatorSettings& InOperatorSettings, const FMetasoundEnvironment& InEnvironment) -> TUniquePtr<IDynamicOperatorTransform>
 			{
 				TArray<TUniquePtr<IDynamicOperatorTransform>> AtomicTransforms;
 
+				if (InOrdinalUpdates.Num() > 0)
+				{
+					AtomicTransforms.Add(MakeUnique<FSwapOperatorOrdinalsAndSort>(InOrdinalUpdates));
+				}
+
 				if (InPriorLiteralNode)
 				{
-					AtomicTransforms.Add(MakeUnique<FRemoveOperator>(DirectedGraphAlgo::GetOperatorID(InPriorLiteralNode)));
+					AtomicTransforms.Add(MakeUnique<FSwapOperatorConnection>(PriorLiteralOperatorID, DynamicOperatorTransactorPrivate::LiteralNodeOutputVertexName, FromOperatorID, InFromVertex, ToOperatorID, InToVertex));
+					AtomicTransforms.Add(MakeUnique<FRemoveOperator>(PriorLiteralOperatorID, TArray<FOperatorID>() /* OperatorsConnectedToInput */));
 				}
-				if (bSetNewOperatorOrder)
+				else
 				{
-					AtomicTransforms.Add(MakeUnique<FSetOperatorOrder>(NewOperatorOrder));
+					AtomicTransforms.Add(MakeUnique<FConnectOperators>(FromOperatorID, InFromVertex, ToOperatorID, InToVertex));
 				}
-				AtomicTransforms.Add(MakeUnique<FConnectOperators>(FromOperatorID, InFromVertex, ToOperatorID, InToVertex));
 
 				return MakeUnique<FAtomicTransform>(MoveTemp(AtomicTransforms));
 			};
@@ -524,47 +1053,38 @@ namespace Metasound
 			EnqueueTransformOnOperatorQueues(CreateAddEdgeTransform);
 		}
 
-		void FDynamicOperatorTransactor::EnqueueFadeAndAddEdgeOperatorTransform(const INode& InFromNode, const FVertexName& InFromVertex, const INode& InToNode, const FVertexName& InToVertex, const INode* InPriorLiteralNode)
+		void FDynamicOperatorTransactor::EnqueueFadeAndAddEdgeOperatorTransform(const INode& InFromNode, const FVertexName& InFromVertex, const INode& InToNode, const FVertexName& InToVertex, const INode* InPriorLiteralNode, const TArray<FOrdinalSwap>& InOrdinalUpdates)
 		{
-			FOperatorID FromOperatorID = DirectedGraphAlgo::GetOperatorID(InFromNode);
-			FOperatorID ToOperatorID = DirectedGraphAlgo::GetOperatorID(InToNode);
-			
+			const FOperatorID FromOperatorID = DirectedGraphAlgo::GetOperatorID(InFromNode);
+			const FOperatorID ToOperatorID = DirectedGraphAlgo::GetOperatorID(InToNode);
+			const FOperatorID PriorLiteralOperatorID = DirectedGraphAlgo::GetOperatorID(InPriorLiteralNode);
+
 			// Fade inputs on the receiving node when adding an edge. We don't fade the source node's outputs
 			// because those outputs could also be connected to other nodes which we do not want to fade. 
 			TArrayView<const FVertexName> InputsToFade(&InToVertex, 1);
 			TArrayView<const FVertexName> OutputsToFade;
 
-			if (InPriorLiteralNode)
-			{
-				CurrentOperatorOrder.RemoveSingle(DirectedGraphAlgo::GetOperatorID(InPriorLiteralNode));
-			}
-
-			// Find order of operators after removing literal and adding edge. 
-			TArray<FOperatorID> NewOperatorOrder = DynamicOperatorTransactorPrivate::DetermineOperatorOrder(Graph);
-
-			// Only set new order if it's different than existing.
-			bool bSetNewOperatorOrder = NewOperatorOrder != CurrentOperatorOrder;
-			if (bSetNewOperatorOrder)
-			{
-				CurrentOperatorOrder = NewOperatorOrder;
-			}
-
 			auto CreateBeginFadeAndAddEdgeTransform = [&](const FOperatorSettings& InOperatorSettings, const FMetasoundEnvironment& InEnvironment) -> TUniquePtr<IDynamicOperatorTransform>
 			{
 				TArray<TUniquePtr<IDynamicOperatorTransform>> AtomicTransforms;
 
+				
+				if (InOrdinalUpdates.Num() > 0)
+				{
+					AtomicTransforms.Add(MakeUnique<FSwapOperatorOrdinalsAndSort>(InOrdinalUpdates));
+				}
 				if (InPriorLiteralNode)
 				{
+					AtomicTransforms.Add(MakeUnique<FSwapOperatorConnection>(PriorLiteralOperatorID, DynamicOperatorTransactorPrivate::LiteralNodeOutputVertexName, FromOperatorID, InFromVertex, ToOperatorID, InToVertex));
 					// We assume that any FAudioBuffer created from a literal node is just silent audio. That means we can
 					// Remove that literal node without fading out the audio from the literal node connected to the destination node. 
-					AtomicTransforms.Add(MakeUnique<FRemoveOperator>(DirectedGraphAlgo::GetOperatorID(InPriorLiteralNode)));
+					AtomicTransforms.Add(MakeUnique<FRemoveOperator>(DirectedGraphAlgo::GetOperatorID(InPriorLiteralNode), TArray<FOperatorID>() /* OperatorsConnectedToInput */));
 				}
-				
-				if (bSetNewOperatorOrder)
+				else
 				{
-					AtomicTransforms.Add(MakeUnique<FSetOperatorOrder>(NewOperatorOrder));
+					AtomicTransforms.Add(MakeUnique<FConnectOperators>(FromOperatorID, InFromVertex, ToOperatorID, InToVertex));
 				}
-				AtomicTransforms.Add(MakeUnique<FConnectOperators>(FromOperatorID, InFromVertex, ToOperatorID, InToVertex));
+
 				AtomicTransforms.Add(MakeUnique<FBeginAudioFadeTransform>(ToOperatorID, EAudioFadeType::FadeIn, InputsToFade, OutputsToFade));
 				
 				// Fence must be the last transform since the fade must be performed
@@ -614,22 +1134,18 @@ namespace Metasound
 			EnqueueTransformOnOperatorQueues(CreateEndAudioFadeTransform);
 		}
 
-		void FDynamicOperatorTransactor::EnqueueRemoveOperatorTransform(const INode& InNode)
+		void FDynamicOperatorTransactor::EnqueueRemoveOperatorTransform(const INode& InNode, const TArray<FOperatorID>& InOperatorsConnectedToInput)
 		{
-			FOperatorID OperatorID = DirectedGraphAlgo::GetOperatorID(InNode);
-
-			// Removing operator does not require a re-sort because existing dependencies
-			// are still met due to DAG structure. 
-			CurrentOperatorOrder.RemoveSingle(OperatorID);
+			const FOperatorID OperatorID = DirectedGraphAlgo::GetOperatorID(InNode);
 
-			auto CreateRemoveNodeTransform = [&](const FOperatorSettings& InOperatorSettings, const FMetasoundEnvironment& InEnvironment) -> TUniquePtr<IDynamicOperatorTransform>
+			auto CreateRemoveNodeTransform = [&OperatorID, &InOperatorsConnectedToInput](const FOperatorSettings& InOperatorSettings, const FMetasoundEnvironment& InEnvironment) -> TUniquePtr<IDynamicOperatorTransform>
 			{
-				return MakeUnique<FRemoveOperator>(OperatorID);
+				return MakeUnique<FRemoveOperator>(OperatorID, InOperatorsConnectedToInput);
 			};
 			EnqueueTransformOnOperatorQueues(CreateRemoveNodeTransform);
 		}
 
-		void FDynamicOperatorTransactor::EnqueueRemoveEdgeOperatorTransform(const INode& InFromNode, const FVertexName& InFromVertex, const INode& InToNode, const FVertexName& InToVertex, const INode& InReplacementLiteralNode)
+		void FDynamicOperatorTransactor::EnqueueRemoveEdgeOperatorTransform(const INode& InFromNode, const FVertexName& InFromVertex, const INode& InToNode, const FVertexName& InToVertex, const INode& InReplacementLiteralNode, int32 InLiteralOrdinal)
 		{
 			using namespace DynamicOperatorTransactorPrivate;
 
@@ -637,15 +1153,10 @@ namespace Metasound
 			const FOperatorID ToOperatorID = DirectedGraphAlgo::GetOperatorID(InToNode);
 			const FOperatorID LiteralOperatorID = DirectedGraphAlgo::GetOperatorID(InReplacementLiteralNode);
 
-			// Put literals in the front of the execution stack to simplify updating 
-			// runtime instances. No need to sort the entire graph if we are just
-			// inserting something at the beginning of the execution stack. 
-			CurrentOperatorOrder.Insert(LiteralOperatorID, 0);
-
 			auto CreateRemoveEdgeTransform = [&](const FOperatorSettings& InOperatorSettings, const FMetasoundEnvironment& InEnvironment) -> TUniquePtr<IDynamicOperatorTransform>
 			{
 				// Add the literal node.
-				TUniquePtr<IDynamicOperatorTransform> AddNodeTransform = CreateAddOperatorTransform(InReplacementLiteralNode, EExecutionOrderInsertLocation::First, InOperatorSettings, InEnvironment);
+				TUniquePtr<IDynamicOperatorTransform> AddNodeTransform = CreateInsertOperatorTransform(InReplacementLiteralNode, InLiteralOrdinal, InOperatorSettings, InEnvironment);
 
 				// Swap prior connection with new connections.
 				TUniquePtr<IDynamicOperatorTransform> ConnectOperatorsTransform = MakeUnique<FSwapOperatorConnection>(FromOperatorID, InFromVertex, LiteralOperatorID, DynamicOperatorTransactorPrivate::LiteralNodeOutputVertexName, ToOperatorID, InToVertex);
@@ -671,7 +1182,7 @@ namespace Metasound
 			EnqueueTransformOnOperatorQueues(CreateRemoveEdgeTransform);
 		}
 
-		void FDynamicOperatorTransactor::EnqueueFadeAndRemoveEdgeOperatorTransform(const INode& InFromNode, const FVertexName& InFromVertex, const INode& InToNode, const FVertexName& InToVertex, const INode& InReplacementLiteralNode)
+		void FDynamicOperatorTransactor::EnqueueFadeAndRemoveEdgeOperatorTransform(const INode& InFromNode, const FVertexName& InFromVertex, const INode& InToNode, const FVertexName& InToVertex, const INode& InReplacementLiteralNode, int32 InLiteralOrdinal)
 		{
 			// Fade the input to the node getting disconnected rather than the output of the source node. The source node
 			// may be connected to other nodes and fading it's output would fade all the other connected nodes' inputs. 
@@ -688,15 +1199,13 @@ namespace Metasound
 			// If we ever find ourselves creating audio buffers with literals which 
 			// are anything other than silent buffers, we should rework this operation
 			// to do either a cross-fade, or an additional "fade in" to the new value. 
-			EnqueueRemoveEdgeOperatorTransform(InFromNode, InFromVertex, InToNode, InToVertex, InReplacementLiteralNode);
+			EnqueueRemoveEdgeOperatorTransform(InFromNode, InFromVertex, InToNode, InToVertex, InReplacementLiteralNode, InLiteralOrdinal);
 
 			// Remove fade operation. 
 			EnqueueEndFadeOperatorTransform(InToNode);
 		}
 
-
-
-		TUniquePtr<IDynamicOperatorTransform> FDynamicOperatorTransactor::CreateAddOperatorTransform(const INode& InNode, EExecutionOrderInsertLocation InLocation, const FOperatorSettings& InOperatorSettings, const FMetasoundEnvironment& InEnvironment) const
+		TUniquePtr<IDynamicOperatorTransform> FDynamicOperatorTransactor::CreateInsertOperatorTransform(const INode& InNode, int32 InOrdinal, const FOperatorSettings& InOperatorSettings, const FMetasoundEnvironment& InEnvironment) const
 		{
 			using namespace DynamicOperatorTransactorPrivate;
 
@@ -729,11 +1238,12 @@ namespace Metasound
 
 				FOperatorInfo OpInfo
 				{
+					InOrdinal,
 					MoveTemp(Operator),
 					MoveTemp(InterfaceData)
 				};
 
-				return MakeUnique<FAddOperator>(OperatorID, InLocation, MoveTemp(OpInfo));
+				return MakeUnique<FInsertOperator>(OperatorID, MoveTemp(OpInfo));
 			}
 			else
 			{
diff --git a/Engine/Plugins/Runtime/Metasound/Source/MetasoundGraphCore/Private/MetasoundGraphAlgoPrivate.h b/Engine/Plugins/Runtime/Metasound/Source/MetasoundGraphCore/Private/MetasoundGraphAlgoPrivate.h
--- a/Engine/Plugins/Runtime/Metasound/Source/MetasoundGraphCore/Private/MetasoundGraphAlgoPrivate.h
+++ b/Engine/Plugins/Runtime/Metasound/Source/MetasoundGraphCore/Private/MetasoundGraphAlgoPrivate.h
@@ -34,6 +34,11 @@ namespace Metasound
 
 			struct FOperatorInfo
 			{
+				// The ordinal associated with the operator. Ordinals can be used
+				// to determine ordering of operators, but shouldn't be considered
+				// as indexes.
+				int32 Ordinal;
+
 				TUniquePtr<IOperator> Operator;
 
 				// Vertex Data bound to operator.
@@ -50,10 +55,8 @@ namespace Metasound
 
 			// Vertex Data bound to the graph
 			FVertexInterfaceData VertexData;
-
-			// Sorted order of operators for execution.
-			TArray<FOperatorID> OperatorOrder;
-
+			
+			// Map of all operators in graph
 			TSortedMap<FOperatorID, FOperatorInfo> OperatorMap;
 
 			// Map with input vertex name as key, and OperatorID of input node as value.
@@ -62,5 +65,17 @@ namespace Metasound
 			// Map with output vertex name as key, and OperatorID of output node as value.
 			TSortedVertexNameMap<FOperatorID> OutputVertexMap;
 		};
+
+		/** FStaticGraphOperatorData contains FGraphOperatorData plus a static ordering
+		 * of nodes. This structure allows the FOpeatorBuilder to more efficiently inform
+		 * the FGraphOperator how the operators should order execution tables. 
+		 */
+		struct FStaticGraphOperatorData : FGraphOperatorData
+		{
+			using FGraphOperatorData::FGraphOperatorData; // Inherit constructor
+
+			// Order of node execution
+			TArray<const INode*> NodeOrder;
+		};
 	}
 }
diff --git a/Engine/Plugins/Runtime/Metasound/Source/MetasoundGraphCore/Private/MetasoundGraphOperator.cpp b/Engine/Plugins/Runtime/Metasound/Source/MetasoundGraphCore/Private/MetasoundGraphOperator.cpp
--- a/Engine/Plugins/Runtime/Metasound/Source/MetasoundGraphCore/Private/MetasoundGraphOperator.cpp
+++ b/Engine/Plugins/Runtime/Metasound/Source/MetasoundGraphCore/Private/MetasoundGraphOperator.cpp
@@ -78,13 +78,18 @@ namespace Metasound
 		}
 	}
 
-	FGraphOperator::FGraphOperator(TUniquePtr<DirectedGraphAlgo::FGraphOperatorData>&& InOperatorState)
+	FGraphOperator::FGraphOperator(TUniquePtr<DirectedGraphAlgo::FStaticGraphOperatorData> InOperatorState)
 	{
 		using namespace DirectedGraphAlgo;
 
 		// Append operators in order.
-		for (const FOperatorID OperatorID : InOperatorState->OperatorOrder)
+		for (const INode* Node: InOperatorState->NodeOrder)
 		{
+			// The Node pointer may not point to a valid memory because there is
+			// nothing ensuring that the node is still alive. GetOperatorID(...)
+			// simply uses the pointer address as the ID and does not access the 
+			// actual underlying INode. 
+			const FOperatorID OperatorID = GetOperatorID(Node);
 			AppendOperator(MoveTemp(InOperatorState->OperatorMap[OperatorID].Operator));
 		}
 
@@ -128,16 +133,6 @@ namespace Metasound
 		VertexData = InVertexData;
 	}
 
-	FDataReferenceCollection FGraphOperator::GetInputs() const
-	{
-		return VertexData.GetInputs().ToDataReferenceCollection();
-	}
-
-	FDataReferenceCollection FGraphOperator::GetOutputs() const
-	{
-		return VertexData.GetOutputs().ToDataReferenceCollection();
-	}
-
 	void FGraphOperator::BindInputs(FInputVertexInterfaceData& InInputVertexData)
 	{
 		if (!MetasoundGraphPrivate::IsSupportedVertexData(VertexData.GetInputs(), InInputVertexData))
diff --git a/Engine/Plugins/Runtime/Metasound/Source/MetasoundGraphCore/Private/MetasoundOperatorBuilder.cpp b/Engine/Plugins/Runtime/Metasound/Source/MetasoundGraphCore/Private/MetasoundOperatorBuilder.cpp
--- a/Engine/Plugins/Runtime/Metasound/Source/MetasoundGraphCore/Private/MetasoundOperatorBuilder.cpp
+++ b/Engine/Plugins/Runtime/Metasound/Source/MetasoundGraphCore/Private/MetasoundOperatorBuilder.cpp
@@ -33,23 +33,23 @@ namespace Metasound
 			const FOperatorSettings& Settings;
 			const FMetasoundEnvironment& Environment;
 			
+			DirectedGraphAlgo::FGraphOperatorData& GraphOperatorData;
 			FBuildResults& Results;
-			TUniquePtr<DirectedGraphAlgo::FGraphOperatorData> GraphOperatorData;
 
 			FBuildContext(
 				const IGraph& InGraph,
 				const FDirectedGraphAlgoAdapter& InAlgoAdapter,
 				const FOperatorSettings& InSettings,
 				const FMetasoundEnvironment& InEnvironment,
+				DirectedGraphAlgo::FGraphOperatorData& InGraphOperatorData,
 				FBuildResults& OutResults)
 			: Graph(InGraph)
 			, AlgoAdapter(InAlgoAdapter)
 			, Settings(InSettings)
 			, Environment(InEnvironment)
+			, GraphOperatorData(InGraphOperatorData)
 			, Results(OutResults)
-			, GraphOperatorData(MakeUnique<DirectedGraphAlgo::FGraphOperatorData>(InSettings))
 			{
-				GraphOperatorData->VertexData = FVertexInterfaceData(InGraph.GetVertexInterface());
 			}
 		};
 	}
@@ -79,32 +79,78 @@ namespace Metasound
 
 	TUniquePtr<IOperator> FOperatorBuilder::BuildGraphOperator(const FBuildGraphOperatorParams& InParams, FBuildResults& OutResults) const
 	{
-		TUniquePtr<DirectedGraphAlgo::FGraphOperatorData> GraphData = BuildGraphOperatorData(InParams, OutResults);
+		METASOUND_TRACE_CPUPROFILER_EVENT_SCOPE(Metasound::FOperatorBuilder::BuildGraphOperator);
 
-		if (GraphData.IsValid())
+		if (BuilderSettings.bEnableOperatorRebind)
 		{
-			// Create graph operator from collection of node operators.
-			return CreateGraphOperator(MoveTemp(GraphData));
+			return BuildRebindableGraphOperator(InParams, OutResults);
+		}
+		else
+		{
+			return BuildStaticGraphOperator(InParams, OutResults);
 		}
-
-		return TUniquePtr<IOperator>(nullptr);
 	}
 
 	TUniquePtr<IOperator> FOperatorBuilder::BuildDynamicGraphOperator(const FBuildDynamicGraphOperatorParams& InParams, FBuildResults& OutResults)
 	{
+		METASOUND_TRACE_CPUPROFILER_EVENT_SCOPE(Metasound::FOperatorBuilder::BuildDynamicGraphOperator);
 		using namespace DynamicGraph;
 
-		TUniquePtr<DirectedGraphAlgo::FGraphOperatorData> GraphData = BuildGraphOperatorData(InParams, OutResults);
+		TArray<const INode*> NodeOrder;
+		TUniquePtr<FDynamicOperator> GraphOperator = MakeUnique<FDynamicOperator>(InParams.OperatorSettings, InParams.TransformQueue, InParams.OperatorUpdateCallbacks);
+		bool bSuccess = BuildGraphOperatorData(InParams, GetDynamicGraphOperatorData(*GraphOperator), NodeOrder, OutResults);
+
+		if (bSuccess)
+		{
+			GetDynamicGraphOperatorData(*GraphOperator).InitTables();
+			return GraphOperator;
+		}
+
+		return TUniquePtr<IOperator>(nullptr);
+	}
+
+	TUniquePtr<IOperator> FOperatorBuilder::BuildRebindableGraphOperator(const FBuildGraphOperatorParams& InParams, FBuildResults& OutResults) const
+	{
+		METASOUND_TRACE_CPUPROFILER_EVENT_SCOPE(Metasound::FOperatorBuilder::BuildRebindableGraphOperator);
+
+		TArray<const INode*> NodeOrder;
+		TUniquePtr<FRebindableGraphOperator> GraphOperator = MakeUnique<FRebindableGraphOperator>(InParams.OperatorSettings);
+		bool bSuccess =  BuildGraphOperatorData(InParams, GetDynamicGraphOperatorData(*GraphOperator), NodeOrder, OutResults);
+
+		if (bSuccess)
+		{
+			GetDynamicGraphOperatorData(*GraphOperator).InitTables();
+			return GraphOperator;
+		}
+
+		return TUniquePtr<IOperator>(nullptr);
+	}
+
+	TUniquePtr<IOperator> FOperatorBuilder::BuildStaticGraphOperator(const FBuildGraphOperatorParams& InParams, FBuildResults& OutResults) const
+	{
+		METASOUND_TRACE_CPUPROFILER_EVENT_SCOPE(Metasound::FOperatorBuilder::BuildStaticGraphOperator);
+
+		TUniquePtr<DirectedGraphAlgo::FStaticGraphOperatorData> GraphOperatorData = MakeUnique<DirectedGraphAlgo::FStaticGraphOperatorData>(InParams.OperatorSettings);
+
+		bool bSuccess =  BuildGraphOperatorData(InParams, *GraphOperatorData, GraphOperatorData->NodeOrder, OutResults);
 
-		if (GraphData.IsValid())
+		if (bSuccess)
 		{
-			return MakeUnique<FDynamicOperator>(MoveTemp(*GraphData), InParams.TransformQueue, InParams.OperatorUpdateCallbacks);
+			// Create graph operator from collection of node operators.
+			return MakeUnique<FGraphOperator>(MoveTemp(GraphOperatorData));
 		}
 
 		return TUniquePtr<IOperator>(nullptr);
 	}
 
-	TUniquePtr<DirectedGraphAlgo::FGraphOperatorData> FOperatorBuilder::BuildGraphOperatorData(const FBuildGraphOperatorParams& InParams, FBuildResults& OutResults) const
+	DynamicGraph::FDynamicGraphOperatorData& FOperatorBuilder::GetDynamicGraphOperatorData(DynamicGraph::IDynamicGraphInPlaceBuildable& InBuildable) const
+	{
+		// This function exists as a convenience to avoid needing to do `static_cast<IDynamicGraphInPlaceBuildable&>(GraphOperator)->GetDynamicGraphOperatorData()` 
+		// all over the place. 
+		return InBuildable.GetDynamicGraphOperatorData();
+	}
+
+	bool FOperatorBuilder::BuildGraphOperatorData(const FBuildGraphOperatorParams& InParams, DirectedGraphAlgo::FGraphOperatorData& OutGraphOperatorData, TArray<const INode*>& OutNodeOrder, FBuildResults& OutResults) const
 	{
 		METASOUND_TRACE_CPUPROFILER_EVENT_SCOPE(Metasound::FOperatorBuilder::BuildGraphOperator);
 
@@ -144,7 +190,7 @@ namespace Metasound
 		// Possible early exit if edge validation fails.
 		if (BuildStatus > GetMaxErrorLevel())
 		{
-			return TUniquePtr<FGraphOperatorData>(nullptr);
+			return false;
 		}
 
 		// Create algo adapter view of graph to cache graph operations.
@@ -153,28 +199,28 @@ namespace Metasound
 		if (!AlgoAdapter.IsValid())
 		{
 			AddBuildError<FInternalError>(OutResults.Errors, __FILE__, __LINE__);
-			return TUniquePtr<FGraphOperatorData>(nullptr);
+			return false;
 		}
 
-		OperatorBuilder::FBuildContext BuildContext(InParams.Graph, *AlgoAdapter, InParams.OperatorSettings, InParams.Environment, OutResults);
-
-		TArray<const INode*> SortedNodes;
+		OperatorBuilder::FBuildContext BuildContext(InParams.Graph, *AlgoAdapter, InParams.OperatorSettings, InParams.Environment, OutGraphOperatorData, OutResults);
 
 		// Sort the nodes in a valid execution order
-		BuildStatus |= DepthFirstTopologicalSort(BuildContext, SortedNodes);
+		BuildStatus |= DepthFirstTopologicalSort(BuildContext, OutNodeOrder);
 
 		// TODO: Add FindReachableNodesFromVariables in Prune.
 		// TODO: will need to prune edges as well.
 		// Otherwise, subgraphs incorrectly get pruned.
-		// BuildStatus |= PruneNodges(BuildContext, SortedNodes);
+		// BuildStatus |= PruneNodges(BuildContext, OutNodeOrder);
 
 		// Check build status in case build routine should be exited early.
 		if (BuildStatus > GetMaxErrorLevel())
 		{
-			return TUniquePtr<FGraphOperatorData>(nullptr);
+			return false;
 		}
 
-		InitializeOperatorInfo(InParams.Graph, SortedNodes, *BuildContext.GraphOperatorData);
+		InitializeVertexInterfaceData(InParams.Graph, BuildContext.GraphOperatorData);
+
+		InitializeOperatorInfo(InParams.Graph, OutNodeOrder, BuildContext.GraphOperatorData);
 
 		// Assign external inputs to various vertex interfaces.
 		BuildStatus |= FOperatorBuilder::GatherExternalInputDataReferences(BuildContext, InParams.InputData);
@@ -182,31 +228,31 @@ namespace Metasound
 		// Check build status in case build routine should be exited early.
 		if (BuildStatus > GetMaxErrorLevel())
 		{
-			return TUniquePtr<FGraphOperatorData>(nullptr);
+			return false;
 		}
 
 		// Create node operators from factories.
-		BuildStatus |= CreateOperators(BuildContext, SortedNodes, InParams.InputData);
+		BuildStatus |= CreateOperators(BuildContext, OutNodeOrder, InParams.InputData);
 
 		if (BuildStatus > GetMaxErrorLevel())
 		{
-			return TUniquePtr<FGraphOperatorData>(nullptr);
+			return false;
 		}
 
 		if (BuilderSettings.bPopulateInternalDataReferences)
 		{
-			GatherInternalGraphDataReferences(BuildContext, SortedNodes, BuildContext.Results.InternalDataReferences);
+			GatherInternalGraphDataReferences(BuildContext, OutNodeOrder, BuildContext.Results.InternalDataReferences);
 		}
 
 		// Gather the inputs for the graph data. 
-		BuildStatus |= GatherGraphDataReferences(BuildContext, BuildContext.GraphOperatorData->VertexData);
+		BuildStatus |= GatherGraphDataReferences(BuildContext, BuildContext.GraphOperatorData.VertexData);
 
 		if (BuildStatus > GetMaxErrorLevel())
 		{
-			return TUniquePtr<FGraphOperatorData>(nullptr);
+			return false;
 		}
 
-		return MoveTemp(BuildContext.GraphOperatorData);
+		return true;
 	}
 
 	FOperatorBuilder::FBuildStatus FOperatorBuilder::DepthFirstTopologicalSort(OperatorBuilder::FBuildContext& InOutContext, TArray<const INode*>& OutNodes) const
@@ -322,6 +368,11 @@ namespace Metasound
 		return BuildStatus;
 	}
 
+	void FOperatorBuilder::InitializeVertexInterfaceData(const IGraph& InGraph, DirectedGraphAlgo::FGraphOperatorData& InOutGraphOperatorData) const
+	{
+		InOutGraphOperatorData.VertexData = FVertexInterfaceData(InGraph.GetVertexInterface());
+	}
+
 	void FOperatorBuilder::InitializeOperatorInfo(const IGraph& InGraph, TArray<const INode*>& InSortedNodes, DirectedGraphAlgo::FGraphOperatorData& InOutGraphOperatorData) const
 	{
 		METASOUND_TRACE_CPUPROFILER_EVENT_SCOPE(Metasound::FOperatorBuilder::InitializeOperatorInfo);
@@ -329,19 +380,18 @@ namespace Metasound
 
 		// Create FOperatorInfos from Nodes
 		TSortedMap<FOperatorID, FGraphOperatorData::FOperatorInfo>& OperatorMap = InOutGraphOperatorData.OperatorMap;
-		TArray<FOperatorID>& OperatorOrder = InOutGraphOperatorData.OperatorOrder;
 
 		const int32 NumNodes = InSortedNodes.Num();
 		OperatorMap.Reserve(OperatorMap.Num() + NumNodes);
-		OperatorOrder.Reserve(OperatorOrder.Num() + NumNodes);
 
 		{
 			METASOUND_TRACE_CPUPROFILER_EVENT_SCOPE(Metasound::FOperatorBuilder::InitializeOperatorInfo::Nodes);
+			int32 Ordinal = 0;
 			for (const INode* Node : InSortedNodes)
 			{
 				FOperatorID OperatorID = GetOperatorID(Node);
-				OperatorOrder.Add(OperatorID);
-				OperatorMap.Add(OperatorID, FGraphOperatorData::FOperatorInfo{nullptr, Node->GetVertexInterface()});
+				OperatorMap.Add(OperatorID, FGraphOperatorData::FOperatorInfo{Ordinal, nullptr, Node->GetVertexInterface()});
+				Ordinal++;
 			}
 		}
 
@@ -370,7 +420,7 @@ namespace Metasound
 			const FInputDataDestination& Destination = InputDestinationKV.Value;
 
 			FOperatorID OperatorID = GetOperatorID(Destination.Node);
-			FGraphOperatorData::FOperatorInfo& OperatorInfo = InOutContext.GraphOperatorData->OperatorMap.FindChecked(OperatorID);
+			FGraphOperatorData::FOperatorInfo& OperatorInfo = InOutContext.GraphOperatorData.OperatorMap.FindChecked(OperatorID);
 
 			if (const FAnyDataReference* DataReference = InExternalInputData.FindDataReference(Destination.Vertex.VertexName))
 			{
@@ -410,7 +460,7 @@ namespace Metasound
 		for (const INode* NodePtr : InNodes)
 		{
 			check(NodePtr);
-			if (const FGraphOperatorData::FOperatorInfo* OpInfo = InOutContext.GraphOperatorData->OperatorMap.Find(GetOperatorID(NodePtr)))
+			if (const FGraphOperatorData::FOperatorInfo* OpInfo = InOutContext.GraphOperatorData.OperatorMap.Find(GetOperatorID(NodePtr)))
 			{
 				OutNodeVertexData.Emplace(NodePtr->GetInstanceID(), OpInfo->VertexData.GetOutputs().ToDataReferenceCollection());
 			}
@@ -454,7 +504,7 @@ namespace Metasound
 		FBuildStatus BuildStatus;
 
 		// Create FOperatorInfos from Nodes
-		TSortedMap<FOperatorID, FGraphOperatorData::FOperatorInfo>& OperatorMap = InOutContext.GraphOperatorData->OperatorMap;
+		TSortedMap<FOperatorID, FGraphOperatorData::FOperatorInfo>& OperatorMap = InOutContext.GraphOperatorData.OperatorMap;
 
 		// Call operator factory for each node.
 		for (const INode* Node : InSortedNodes)
@@ -562,7 +612,7 @@ namespace Metasound
 			const FInputDataDestination& InputDestination = Element.Value;
 
 			const FOperatorID OperatorID = GetOperatorID(InputDestination.Node);
-			if (FGraphOperatorData::FOperatorInfo* OperatorInfo = InOutContext.GraphOperatorData->OperatorMap.Find(OperatorID))
+			if (FGraphOperatorData::FOperatorInfo* OperatorInfo = InOutContext.GraphOperatorData.OperatorMap.Find(OperatorID))
 			{
 				FInputVertexInterfaceData& NodeInputData = OperatorInfo->VertexData.GetInputs();
 				OperatorInfo->Operator->BindInputs(NodeInputData);
@@ -575,7 +625,7 @@ namespace Metasound
 						OutVertexData.GetInputs().SetVertex(InputDestination.Vertex.VertexName, *DataReference);
 					}
 				}
-				InOutContext.GraphOperatorData->InputVertexMap.Emplace(InputDestination.Vertex.VertexName, OperatorID);
+				InOutContext.GraphOperatorData.InputVertexMap.Emplace(InputDestination.Vertex.VertexName, OperatorID);
 			}
 
 			if (!bFoundDataReference)
@@ -592,7 +642,7 @@ namespace Metasound
 			const FOutputDataSource& OutputSource = Element.Value;
 
 			const FOperatorID OperatorID = GetOperatorID(OutputSource.Node);
-			if (const FGraphOperatorData::FOperatorInfo* OperatorInfo = InOutContext.GraphOperatorData->OperatorMap.Find(OperatorID))
+			if (const FGraphOperatorData::FOperatorInfo* OperatorInfo = InOutContext.GraphOperatorData.OperatorMap.Find(OperatorID))
 			{
 				const FOutputVertexInterfaceData& NodeOutputData = OperatorInfo->VertexData.GetOutputs();
 				if (const FAnyDataReference* DataReference = NodeOutputData.FindDataReference(OutputSource.Vertex.VertexName))
@@ -603,7 +653,7 @@ namespace Metasound
 						OutVertexData.GetOutputs().SetVertex(OutputSource.Vertex.VertexName, FAnyDataReference{*DataReference});
 					}
 				}
-				InOutContext.GraphOperatorData->OutputVertexMap.Emplace(OutputSource.Vertex.VertexName, OperatorID);
+				InOutContext.GraphOperatorData.OutputVertexMap.Emplace(OutputSource.Vertex.VertexName, OperatorID);
 			}
 
 			if (!bFoundDataReference)
@@ -616,24 +666,6 @@ namespace Metasound
 		return BuildStatus;
 	}
 
-	TUniquePtr<IOperator> FOperatorBuilder::CreateGraphOperator(TUniquePtr<DirectedGraphAlgo::FGraphOperatorData>&& InGraphOperatorData) const
-	{
-		METASOUND_TRACE_CPUPROFILER_EVENT_SCOPE(Metasound::FOperatorBuilder::CreateGraphOperator);
-
-		TUniquePtr<IOperator> GraphOperator;
-
-		if (BuilderSettings.bEnableOperatorRebind)
-		{
-			 GraphOperator = MakeUnique<FRebindableGraphOperator>(MoveTemp(*InGraphOperatorData));
-		}
-		else
-		{
-			 GraphOperator = MakeUnique<FGraphOperator>(MoveTemp(InGraphOperatorData));
-		}
-
-		return GraphOperator;
-	}
-
 	FOperatorBuilder::FBuildStatus::EStatus FOperatorBuilder::GetMaxErrorLevel() const
 	{
 		return BuilderSettings.bFailOnAnyError ? FBuildStatus::NoError : FBuildStatus::NonFatalError;
diff --git a/Engine/Plugins/Runtime/Metasound/Source/MetasoundGraphCore/Private/MetasoundRebindableGraphOperator.cpp b/Engine/Plugins/Runtime/Metasound/Source/MetasoundGraphCore/Private/MetasoundRebindableGraphOperator.cpp
--- a/Engine/Plugins/Runtime/Metasound/Source/MetasoundGraphCore/Private/MetasoundRebindableGraphOperator.cpp
+++ b/Engine/Plugins/Runtime/Metasound/Source/MetasoundGraphCore/Private/MetasoundRebindableGraphOperator.cpp
@@ -14,23 +14,11 @@
 
 namespace Metasound
 {
-	FRebindableGraphOperator::FRebindableGraphOperator(DirectedGraphAlgo::FGraphOperatorData&& InOperatorData)
-	: GraphOperatorData(MoveTemp(InOperatorData))
+	FRebindableGraphOperator::FRebindableGraphOperator(const FOperatorSettings& InOperatorSettings)
+	: GraphOperatorData(InOperatorSettings)
 	{
 	}
 
-	FDataReferenceCollection FRebindableGraphOperator::GetInputs() const
-	{
-		checkNoEntry();
-		return FDataReferenceCollection();
-	}
-
-	FDataReferenceCollection FRebindableGraphOperator::GetOutputs() const
-	{
-		checkNoEntry();
-		return FDataReferenceCollection();
-	}
-
 	// Bind the graph's interface data references to FVertexInterfaceData.
 	void FRebindableGraphOperator::BindInputs(FInputVertexInterfaceData& InOutVertexData)
 	{
@@ -86,6 +74,11 @@ namespace Metasound
 		check(InOperator);
 		static_cast<FRebindableGraphOperator*>(InOperator)->PostExecute();
 	}
+
+	DynamicGraph::FDynamicGraphOperatorData& FRebindableGraphOperator::GetDynamicGraphOperatorData()
+	{
+		return GraphOperatorData;
+	}
 }
 
 
diff --git a/Engine/Plugins/Runtime/Metasound/Source/MetasoundGraphCore/Private/MetasoundRebindableGraphOperator.h b/Engine/Plugins/Runtime/Metasound/Source/MetasoundGraphCore/Private/MetasoundRebindableGraphOperator.h
--- a/Engine/Plugins/Runtime/Metasound/Source/MetasoundGraphCore/Private/MetasoundRebindableGraphOperator.h
+++ b/Engine/Plugins/Runtime/Metasound/Source/MetasoundGraphCore/Private/MetasoundRebindableGraphOperator.h
@@ -13,15 +13,11 @@ namespace Metasound
 	class FOutputVertexInterfaceData;
 
 
-	class FRebindableGraphOperator : public TExecutableOperator<FRebindableGraphOperator>
+	class FRebindableGraphOperator : public TExecutableOperator<FRebindableGraphOperator>, public DynamicGraph::IDynamicGraphInPlaceBuildable
 	{
 	public:
 
-		FRebindableGraphOperator(DirectedGraphAlgo::FGraphOperatorData&& InOperatorState);
-
-		virtual FDataReferenceCollection GetInputs() const override;
-
-		virtual FDataReferenceCollection GetOutputs() const override;
+		FRebindableGraphOperator(const FOperatorSettings& InOperatorSettigns);
 
 		// Bind the graph's interface data references to FVertexInterfaceData.
 		virtual void BindInputs(FInputVertexInterfaceData& InOutVertexData) override;
@@ -36,6 +32,9 @@ namespace Metasound
 	private:
 		static void StaticPostExecute(IOperator* InOperator);
 
+		// Interface for FOperatorBuilder to access internal data structure.
+		virtual DynamicGraph::FDynamicGraphOperatorData& GetDynamicGraphOperatorData() override;
+
 		DynamicGraph::FDynamicGraphOperatorData GraphOperatorData;
 	};
 }
diff --git a/Engine/Plugins/Runtime/Metasound/Source/MetasoundGraphCore/Public/MetasoundDynamicOperatorTransactor.h b/Engine/Plugins/Runtime/Metasound/Source/MetasoundGraphCore/Public/MetasoundDynamicOperatorTransactor.h
--- a/Engine/Plugins/Runtime/Metasound/Source/MetasoundGraphCore/Public/MetasoundDynamicOperatorTransactor.h
+++ b/Engine/Plugins/Runtime/Metasound/Source/MetasoundGraphCore/Public/MetasoundDynamicOperatorTransactor.h
@@ -13,6 +13,10 @@
 #include "Templates/SharedPointer.h"
 #include "Templates/UniquePtr.h"
 
+#ifndef METASOUND_DEBUG_DYNAMIC_TRANSACTOR
+#define METASOUND_DEBUG_DYNAMIC_TRANSACTOR !UE_BUILD_SHIPPING
+#endif
+
 namespace Metasound
 {
 	class FOperatorSettings;
@@ -21,12 +25,35 @@ namespace Metasound
 
 	namespace DynamicGraph
 	{
+
 		// Forward declare
 		class IDynamicOperatorTransform;
 		class FDynamicOperator;
 		enum class EAudioFadeType : uint8;
-		enum class EExecutionOrderInsertLocation : uint8;
+		struct FDynamicGraphOperatorData;
+		class FDynamicOperatorTransactor;
+
+#if METASOUND_DEBUG_DYNAMIC_TRANSACTOR
+		namespace Debug
+		{
+			class FDynamicOperatorDebugger;
+			/** This can be used when debugging Dynamic Operator internals. The FDynamicOperatorTransactor and FDynamicOperator have three
+			 * graph representations. Maintaining three different graph representations can be error prone. This method can be used to
+			 * validate the graphs are equivalent, or help diagnose where they differ. 
+			 * 
+			 * The three graph representations are:
+			 * 1. An IGraph contains in the FDynamicOperatorTransactor
+			 * 2. A FDynamicGraphIncrementalSorter in the FDynamicOperatorTransactor
+			 * 3. A FDynamicGraphOperatorData in the FDynamicOperator.
+			 * 
+			 * This method returns true if all three graph representations are equivalent, and false otherwsie. Graph differences are logged to help
+			 * diagnose the issue. 
+			 */
+			bool CompareAndLogGraphRepresentationDiscrepancies(const FDynamicOperatorTransactor& InTransactor, const FDynamicOperator& InDynamicOperator);
+		}
+#endif // if METASOUND_DEBUG_DYNAMIC_TRANSACTOR
 
+		using FOperatorID = uintptr_t;
 		using FLiteralAssignmentFunction = void(*)(const FOperatorSettings& InOperatorSettings, const FLiteral& InLiteral, const FAnyDataReference& OutDataRef);
 		using FReferenceCreationFunction = TOptional<FAnyDataReference>(*)(const FOperatorSettings& InSettings, FName DataType, const FLiteral& InLiteral, EDataReferenceAccessType InAccessType);
 		using FOnInputVertexUpdated = TFunction<void(const FVertexName&, const FInputVertexInterfaceData&)>;
@@ -45,6 +72,104 @@ namespace Metasound
 			FOnOutputVertexUpdated OnOutputRemoved;
 		};
 
+		constexpr int32 ORDINAL_NONE = TNumericLimits<int32>::Max();
+
+		/** FOrdinalSwap represents a change in ordinal for an individual operator. */
+		struct FOrdinalSwap
+		{
+			FOperatorID OperatorID = 0;
+			int32 OriginalOrdinal = ORDINAL_NONE;
+			int32 NewOrdinal = ORDINAL_NONE;
+
+			friend bool operator==(const FOrdinalSwap& InOrdinalSwap, FOperatorID InOperatorID)
+			{
+				return InOrdinalSwap.OperatorID == InOperatorID;
+			}
+
+			static bool OriginalOrdinalLessThan(const FOrdinalSwap& InLHS, const FOrdinalSwap& InRHS)
+			{
+				return InLHS.OriginalOrdinal < InRHS.OriginalOrdinal;
+			}
+		};
+
+
+		/** FDynamicGraphIncrementalSorter implements a dynamic topological sorting algorithm which offers several 
+		  * optimizations over doing full sorts.
+		  * 
+		  * - It can detect if a sort is needed or not. If a sort is unneeded we can skip sorts.
+		  * - Sorts are generally not done on the entire set of operators. The sort is generally
+		  *   done on a subset of nodes related to the nodes being connected. 
+		  */
+		class FDynamicGraphIncrementalSorter
+		{
+		public:
+		
+			/** Where to insert a new operator. */
+			enum class EInsertLocation : uint8
+			{
+				First,
+				Last	
+			};
+
+			FDynamicGraphIncrementalSorter();
+
+			FDynamicGraphIncrementalSorter(const IGraph& InGraph);
+
+			/** Insert a node into the graph.
+			 * @return Ordinal of added operator.
+			 */
+			int32 InsertOperator(FOperatorID InOperator, EInsertLocation InLocation);
+
+			/** Remove a node from the graph.
+			 * 
+			 * @param OpeartorID - ID of operator to remove
+			 * @return Ordinal of removed operator.
+			 */
+			int32 RemoveOperator(FOperatorID InOperatorID);
+
+			/** Populate a TMap<> with an ordinal for every operator. */
+			void GenerateOrdinals(TMap<FOperatorID, int32>& OutOrdinals) const;
+
+			/** Add an edge to the graph, connecting two vertices from two 
+			 * nodes. 
+			 *
+			 * @param InFromOperatorID - Operator which contains the output vertex.
+			 * @param InToOperatorID - Operator which contains the input vertex.
+			 * @param OutOrdinalUpdates - Array to populate with ordinal updates to maintain topological sort of graph.
+			 */
+			void AddDataEdge(FOperatorID InFromOperatorID, FOperatorID InToOperatorID, TArray<FOrdinalSwap>& OutOrdinalUpdates);
+
+			/** Remove the given data edge. */
+			void RemoveDataEdge(FOperatorID InFromOperatorID, FOperatorID InToOperatorID);
+
+		private:
+
+#if METASOUND_DEBUG_DYNAMIC_TRANSACTOR
+			friend class Debug::FDynamicOperatorDebugger;
+#endif // if METASOUND_DEBUG_DYNAMIC_TRANSACTOR
+
+			struct FIncrementalSortOperatorInfo
+			{
+				int32 Ordinal;
+
+				// We track connections in order to determine whether a dependency
+				// exists between two operators. Because two operators can have
+				// multiple shared edges, FOperatorIDs may appear multiple times 
+				// in these arrays. 
+				TArray<FOperatorID> Inputs;
+				TArray<FOperatorID> Outputs;
+			};
+
+			void IncrementalTopologicalSortForNewEdge(FOperatorID InFromOperatorID, int32 InFromOrdinal, FOperatorID InToOperatorID, int32 InToOdinal, TArray<FOrdinalSwap>& OutUpdates);
+
+			void Init(const IGraph& InGraph);
+
+			int32 MaxOrdinal = 0;
+			int32 MinOrdinal = 0;
+			TMap<FOperatorID, FIncrementalSortOperatorInfo> OperatorMap;
+		};
+
+
 		/** The FDynamicOperatorTransactor is used for communicating with a dynamic
 		 * MetaSound operator.
 		 *
@@ -112,30 +237,37 @@ namespace Metasound
 			/** Return internal version of graph. */
 			const FGraph& GetGraph() const;
 
+
 		private:
-			using FOperatorID = uintptr_t;
 
-			void EnqueueAddOperatorTransform(const INode& InNode, EExecutionOrderInsertLocation InLocation);
-			void EnqueueFadeAndRemoveOperatorTransform(const INode& InNode, TArrayView<const FVertexName> InOutputsToFade);
-			void EnqueueRemoveOperatorTransform(const INode& InNode);
+#if METASOUND_DEBUG_DYNAMIC_TRANSACTOR
+			friend class Debug::FDynamicOperatorDebugger;
+#endif // if METASOUND_DEBUG_DYNAMIC_TRANSACTOR
+
+			void RemoveNodeInternal(const INode& InNode, bool bInRemoveDataEdgesWithNode);
+			void FadeAndRemoveNodeInternal(const INode& InNode, TArrayView<const FVertexName> InOutputsToFade, bool bInRemoveDataEdgesWithNode);
+
+			void EnqueueInsertOperatorTransform(const INode& InNode, int32 InOrdinal);
+			void EnqueueRemoveOperatorTransform(const INode& InNode, const TArray<FOperatorID>& InOperatorsConnectedToInput);
 			void EnqueueBeginFadeOperatorTransform(const INode& InNode, EAudioFadeType InFadeType, TArrayView<const FVertexName> InInputsToFade, TArrayView<const FVertexName> InOutputsToFade);
 			void EnqueueEndFadeOperatorTransform(const INode& InNode);
 
-			void EnqueueRemoveEdgeOperatorTransform(const INode& InFromNode, const FVertexName& InFromVertex, const INode& InToNode, const FVertexName& InToVertex, const INode& InReplacementLiteralNode);
-			void EnqueueFadeAndRemoveEdgeOperatorTransform(const INode& InFromNode, const FVertexName& InFromVertex, const INode& InToNode, const FVertexName& InToVertex, const INode& InReplacementLiteralNode);
-			void EnqueueAddEdgeOperatorTransform(const INode& InFromNode, const FVertexName& InFromVertex, const INode& InToNode, const FVertexName& InToVertex, const INode* InPriorLiteralNode);
-			void EnqueueFadeAndAddEdgeOperatorTransform(const INode& InFromNode, const FVertexName& InFromVertex, const INode& InToNode, const FVertexName& InToVertex, const INode* InPriorLiteralNode);
+			void EnqueueRemoveEdgeOperatorTransform(const INode& InFromNode, const FVertexName& InFromVertex, const INode& InToNode, const FVertexName& InToVertex, const INode& InReplacementLiteralNode, int32 InLiteralOrdinal);
+			void EnqueueFadeAndRemoveEdgeOperatorTransform(const INode& InFromNode, const FVertexName& InFromVertex, const INode& InToNode, const FVertexName& InToVertex, const INode& InReplacementLiteralNode, int32 InLiteralOrdinal);
+			void EnqueueAddEdgeOperatorTransform(const INode& InFromNode, const FVertexName& InFromVertex, const INode& InToNode, const FVertexName& InToVertex, const INode* InPriorLiteralNode, const TArray<FOrdinalSwap>& InOrdinalUpdates);
+			void EnqueueFadeAndAddEdgeOperatorTransform(const INode& InFromNode, const FVertexName& InFromVertex, const INode& InToNode, const FVertexName& InToVertex, const INode* InPriorLiteralNode, const TArray<FOrdinalSwap>& InOrdinalUpdates);
 
 			void AddDataEdgeInternal(const INode& InFromNode, const FVertexName& InFromVertex, const FGuid& InToNodeID, const INode& InToNode, const FVertexName& InToVertex);
 
 			using FCreateTransformFunctionRef = TFunctionRef<TUniquePtr<IDynamicOperatorTransform>(const FOperatorSettings& InOperatorSettings, const FMetasoundEnvironment& InEnvironment)>;
 
-			TUniquePtr<IDynamicOperatorTransform> CreateAddOperatorTransform(const INode& InNode, EExecutionOrderInsertLocation InLocation, const FOperatorSettings& InOperatorSettings, const FMetasoundEnvironment& InEnvironment) const;
+			TUniquePtr<IDynamicOperatorTransform> CreateInsertOperatorTransform(const INode& InNode, int32 InOrdinal, const FOperatorSettings& InOperatorSettings, const FMetasoundEnvironment& InEnvironment) const;
 
 			void EnqueueTransformOnOperatorQueues(FCreateTransformFunctionRef InFunc);
 
 			FOperatorBuilder OperatorBuilder;
 			FGraph Graph;
+			FDynamicGraphIncrementalSorter GraphSorter;
 
 			struct FDynamicOperatorInfo
 			{
@@ -145,7 +277,6 @@ namespace Metasound
 			};
 
 			TArray<FDynamicOperatorInfo> OperatorInfos;
-			TArray<FOperatorID> CurrentOperatorOrder;
 
 			struct FLiteralNodeID
 			{
diff --git a/Engine/Plugins/Runtime/Metasound/Source/MetasoundGraphCore/Public/MetasoundGraphOperator.h b/Engine/Plugins/Runtime/Metasound/Source/MetasoundGraphCore/Public/MetasoundGraphOperator.h
--- a/Engine/Plugins/Runtime/Metasound/Source/MetasoundGraphCore/Public/MetasoundGraphOperator.h
+++ b/Engine/Plugins/Runtime/Metasound/Source/MetasoundGraphCore/Public/MetasoundGraphOperator.h
@@ -15,7 +15,7 @@ namespace Metasound
 	// Forward declare
 	namespace DirectedGraphAlgo
 	{
-		struct FGraphOperatorData;
+		struct FStaticGraphOperatorData;
 	}
 
 	class METASOUNDGRAPHCORE_API FGraphOperator : public TExecutableOperator<FGraphOperator>
@@ -27,7 +27,7 @@ namespace Metasound
 			using FResetParams = IOperator::FResetParams;
 
 			FGraphOperator() = default;
-			FGraphOperator(TUniquePtr<DirectedGraphAlgo::FGraphOperatorData>&& InOperatorData);
+			FGraphOperator(TUniquePtr<DirectedGraphAlgo::FStaticGraphOperatorData> InOperatorData);
 
 			virtual ~FGraphOperator() = default;
 
@@ -38,10 +38,6 @@ namespace Metasound
 			// during calls to Bind(InOutVertexData).
 			void SetVertexInterfaceData(FVertexInterfaceData&& InVertexData);
 
-			virtual FDataReferenceCollection GetInputs() const override;
-
-			virtual FDataReferenceCollection GetOutputs() const override;
-
 			// Bind the graph's interface data references to FVertexInterfaceData.
 			virtual void BindInputs(FInputVertexInterfaceData& InOutVertexData) override;
 			virtual void BindOutputs(FOutputVertexInterfaceData& InOutVertexData) override;
diff --git a/Engine/Plugins/Runtime/Metasound/Source/MetasoundGraphCore/Public/MetasoundOperatorBuilder.h b/Engine/Plugins/Runtime/Metasound/Source/MetasoundGraphCore/Public/MetasoundOperatorBuilder.h
--- a/Engine/Plugins/Runtime/Metasound/Source/MetasoundGraphCore/Public/MetasoundOperatorBuilder.h
+++ b/Engine/Plugins/Runtime/Metasound/Source/MetasoundGraphCore/Public/MetasoundOperatorBuilder.h
@@ -31,6 +31,13 @@ namespace Metasound
 	{
 		class IDynamicOperatorTransform;
 		struct FDynamicOperatorUpdateCallbacks;
+		struct FDynamicGraphOperatorData;
+		class IDynamicGraphInPlaceBuildable;
+	}
+
+	namespace DirectedGraphAlgo
+	{
+		struct FGraphOperatorData;
 	}
 
 	// Parameters for building a dynamic graph operator consist of the parameters
@@ -77,8 +84,6 @@ namespace Metasound
 
 		private:
 
-			TUniquePtr<DirectedGraphAlgo::FGraphOperatorData> BuildGraphOperatorData(const FBuildGraphOperatorParams& InParams, FBuildResults& OutResults) const;
-
 			// Handles build status of current build operation.
 			struct FBuildStatus
 			{
@@ -135,6 +140,9 @@ namespace Metasound
 			// Prune unreachable nodes from InOutNodes
 			FBuildStatus PruneNodes(OperatorBuilder::FBuildContext& InOutContext, TArray<const INode*>& InOutNodes) const;
 
+			// Creates FVertexInterfaceData from the IGraph.
+			void InitializeVertexInterfaceData(const IGraph& InGraph, DirectedGraphAlgo::FGraphOperatorData& InOutGraphOperatorData) const;
+
 			// Initialize Operator Info
 			void InitializeOperatorInfo(const IGraph& InGraph ,TArray<const INode*>& InSortedNodes, DirectedGraphAlgo::FGraphOperatorData& InOutGraphOperatorData) const;
 
@@ -153,9 +161,18 @@ namespace Metasound
 			// Call the operator factories for the nodes
 			FBuildStatus CreateOperators(OperatorBuilder::FBuildContext& InOutContext, const TArray<const INode*>& InSortedNodes, const FInputVertexInterfaceData& InExternalInputData) const;
 
-			// Create the final graph operator from the provided build context.
-			TUniquePtr<IOperator> CreateGraphOperator(TUniquePtr<DirectedGraphAlgo::FGraphOperatorData>&& InGraphOperatorData) const;
+			// Creates a graph operator where the set of operators does not change, but the input data can be rebound.
+			TUniquePtr<IOperator> BuildRebindableGraphOperator(const FBuildGraphOperatorParams& InParams, FBuildResults& OutResults) const;
+
+			// Creates a graph operator where teh set of operators is constant and the input data cannot be rebound.
+			TUniquePtr<IOperator> BuildStaticGraphOperator(const FBuildGraphOperatorParams& InParams, FBuildResults& OutResults) const;
+
+			// Builds graph operator data structure. 
+			bool BuildGraphOperatorData(const FBuildGraphOperatorParams& InParams, DirectedGraphAlgo::FGraphOperatorData& OutGraphOperatorData, TArray<const INode*>& OutNodeOrder, FBuildResults& OutResults) const;
 
+			// Access function for retrieving FDynamicGraphOperatorData from an IDynamicGraphInPlaceBuildable.
+			DynamicGraph::FDynamicGraphOperatorData& GetDynamicGraphOperatorData(DynamicGraph::IDynamicGraphInPlaceBuildable& InBuildable) const;
+			
 			FBuildStatus::EStatus GetMaxErrorLevel() const;
 
 			FOperatorBuilderSettings BuilderSettings;
