Re-Try of First pass on AudioMaterialStyle Overriding
#rnx

diff --git a/Engine/Plugins/Runtime/AudioWidgets/Source/AudioWidgets/Private/AudioMaterialSlate/SAudioMaterialLabeledKnob.cpp b/Engine/Plugins/Runtime/AudioWidgets/Source/AudioWidgets/Private/AudioMaterialSlate/SAudioMaterialLabeledKnob.cpp
--- a/Engine/Plugins/Runtime/AudioWidgets/Source/AudioWidgets/Private/AudioMaterialSlate/SAudioMaterialLabeledKnob.cpp
+++ b/Engine/Plugins/Runtime/AudioWidgets/Source/AudioWidgets/Private/AudioMaterialSlate/SAudioMaterialLabeledKnob.cpp
@@ -46,6 +46,7 @@ void SAudioMaterialLabeledKnob::Construct(const FArguments& InArgs)
 	SAssignNew(Knob, SAudioMaterialKnob)
 		.Value(ValueAttribute)
 		.Owner(InArgs._Owner)
+		.AudioMaterialKnobStyle(Style)
 		.OnFloatValueChanged_Lambda([this](float Value)
 		{
 			ValueAttribute.Set(Value);
diff --git a/Engine/Plugins/Runtime/Metasound/Source/MetasoundEditor/Private/MetasoundEditorSettings.cpp b/Engine/Plugins/Runtime/Metasound/Source/MetasoundEditor/Private/MetasoundEditorSettings.cpp
--- a/Engine/Plugins/Runtime/Metasound/Source/MetasoundEditor/Private/MetasoundEditorSettings.cpp
+++ b/Engine/Plugins/Runtime/Metasound/Source/MetasoundEditor/Private/MetasoundEditorSettings.cpp
@@ -1,7 +1,10 @@
 // Copyright Epic Games, Inc. All Rights Reserved.
 #include "MetasoundEditorSettings.h"
 
+#include "AudioMaterialSlate/AudioMaterialSlateTypes.h"
+#include "AudioWidgetsStyle.h"
 #include "Styling/AppStyle.h"
+#include "Styling/SlateWidgetStyleAsset.h"
 #include "UObject/UnrealType.h"
 
 #include UE_INLINE_GENERATED_CPP_BY_NAME(MetasoundEditorSettings)
@@ -30,3 +33,15 @@ UMetasoundEditorSettings::UMetasoundEditorSettings(const FObjectInitializer& Obj
 	VariableNodeTitleColor = FLinearColor(0.211f, 0.513f, 0.035f);				// copper
 }
 
+const FAudioMaterialKnobStyle* UMetasoundEditorSettings::GetKnobStyle() const
+{
+	if (const UObject* Style = KnobStyleOverride.TryLoad())
+	{
+		if (const USlateWidgetStyleAsset* SlateWidgetStyleAsset = CastChecked<USlateWidgetStyleAsset>(Style))
+		{
+			return SlateWidgetStyleAsset->GetStyle<FAudioMaterialKnobStyle>();
+		}
+	}
+
+	return &FAudioWidgetsStyle::Get().GetWidgetStyle<FAudioMaterialKnobStyle>("AudioMaterialKnob.Style");
+}
diff --git a/Engine/Plugins/Runtime/Metasound/Source/MetasoundEditor/Private/MetasoundEditorSettings.h b/Engine/Plugins/Runtime/Metasound/Source/MetasoundEditor/Private/MetasoundEditorSettings.h
--- a/Engine/Plugins/Runtime/Metasound/Source/MetasoundEditor/Private/MetasoundEditorSettings.h
+++ b/Engine/Plugins/Runtime/Metasound/Source/MetasoundEditor/Private/MetasoundEditorSettings.h
@@ -5,9 +5,13 @@
 #include "UObject/NoExportTypes.h"
 #include "UObject/Object.h"
 #include "UObject/ObjectMacros.h"
+#include "UObject/SoftObjectPath.h"
 
 #include "MetasoundEditorSettings.generated.h"
 
+class USlateWidgetStyleAsset;
+struct FAudioMaterialKnobStyle;
+
 UENUM()
 enum class EMetasoundActiveAnalyzerEnvelopeDirection : uint8
 {
@@ -164,4 +168,10 @@ public:
 	/** Determines which details view to show in Metasounds Editor */
 	UPROPERTY(Transient)
 	EMetasoundActiveDetailView DetailView = EMetasoundActiveDetailView::General;
+	
+	UPROPERTY(EditAnywhere, config, Category = MetasoundStyling, meta = (AllowedClasses = "/Script/SlateCore.SlateWidgetStyleAsset", DisplayName = "Knob Style"))
+	FSoftObjectPath KnobStyleOverride;
+	
+	/** Get the AudioMaterialKnob Style. if KnobStyleOverride is not set, returns default style.*/
+	const FAudioMaterialKnobStyle* GetKnobStyle() const;
 };
diff --git a/Engine/Plugins/Runtime/Metasound/Source/MetasoundEditor/Private/SMetasoundGraphNode.cpp b/Engine/Plugins/Runtime/Metasound/Source/MetasoundEditor/Private/SMetasoundGraphNode.cpp
--- a/Engine/Plugins/Runtime/Metasound/Source/MetasoundEditor/Private/SMetasoundGraphNode.cpp
+++ b/Engine/Plugins/Runtime/Metasound/Source/MetasoundEditor/Private/SMetasoundGraphNode.cpp
@@ -23,6 +23,7 @@
 #include "MetasoundEditorGraphNode.h"
 #include "MetasoundEditorGraphSchema.h"
 #include "MetasoundEditorModule.h"
+#include "MetasoundEditorSettings.h"
 #include "MetasoundFrontendNodeTemplateRegistry.h"
 #include "MetasoundFrontendRegistries.h"
 #include "MetasoundTrace.h"
@@ -66,9 +67,9 @@ namespace Metasound
 		{
 			int32 UseAudioMaterialWidgets = 0;
 			FAutoConsoleVariableRef CVarUseAudioMaterialWidgets(
-				TEXT("au.MetaSound.Editor.EnableAudioMaterialNodes"),
+				TEXT("au.MetaSound.Editor.UseAudioMaterialSlates"),
 				UseAudioMaterialWidgets,
-				TEXT("Are new AudioMaterialWidgets used for visualization for the Metasound nodes, if implemented.\n")
+				TEXT("Are new AudioMaterialWidgets used for visualization in the Metasound editor, if implemented.\n")
 				TEXT("0: Disabled (default), !0: Enabled"),
 				ECVF_Default);
 		}// Metasound::Editor::GraphNodePrivate
@@ -858,8 +859,12 @@ namespace Metasound
 
 							if (Metasound::Editor::GraphNodePrivate::UseAudioMaterialWidgets)
 							{
+								const UMetasoundEditorSettings* MetasoundSettings = GetDefault<UMetasoundEditorSettings>();
+								check(MetasoundSettings)
+
 								SAssignNew(InputWidget, SAudioMaterialLabeledKnob)
 									.Owner(GraphMember->GetOwningGraph())
+									.Style(MetasoundSettings->GetKnobStyle())
 									.OnValueChanged_Lambda(OnValueChangedLambda)
 									.AudioUnitsValueType(DefaultFloat->WidgetUnitValueType)
 									.bUseLinearOutput(DefaultFloat->VolumeWidgetUseLinearOutput)
