Expose VertexInterface changes in OnGeneratorIOUpdated
#tests PIE, MetasoundDynamicGenerator test suite

diff --git a/Engine/Plugins/Runtime/Harmonix/Source/HarmonixMetasound/Private/HarmonixMetasound/Components/MetasoundMusicClockDriver.cpp b/Engine/Plugins/Runtime/Harmonix/Source/HarmonixMetasound/Private/HarmonixMetasound/Components/MetasoundMusicClockDriver.cpp
--- a/Engine/Plugins/Runtime/Harmonix/Source/HarmonixMetasound/Private/HarmonixMetasound/Components/MetasoundMusicClockDriver.cpp
+++ b/Engine/Plugins/Runtime/Harmonix/Source/HarmonixMetasound/Private/HarmonixMetasound/Components/MetasoundMusicClockDriver.cpp
@@ -169,7 +169,7 @@ bool FMetasoundMusicClockDriver::AttemptToConnectToAudioComponentsMetasound()
 	}
 	GeneratorAttachedCallbackHandle = CurrentGeneratorHandle->OnGeneratorHandleAttached.AddLambda([this](){OnGeneratorAttached();});
 	GeneratorDetachedCallbackHandle = CurrentGeneratorHandle->OnGeneratorHandleDetached.AddLambda([this](){OnGeneratorDetached();});
-	GeneratorIOUpdatedCallbackHandle = CurrentGeneratorHandle->OnIOUpdated.AddLambda([this](){OnGeneratorIOUpdated();});
+	GeneratorIOUpdatedCallbackHandle = CurrentGeneratorHandle->OnIOUpdatedWithChanges.AddLambda([this](const TArray<Metasound::FVertexInterfaceChange>& VertexInterfaceChanges){OnGeneratorIOUpdatedWithChanges(VertexInterfaceChanges);});
 	UMetasoundGeneratorHandle::FOnSetGraph::FDelegate OnSetGraph;
 	OnSetGraph.BindLambda([this](){OnGraphSet();});
 	GraphChangedCallbackHandle = CurrentGeneratorHandle->AddGraphSetCallback(MoveTemp(OnSetGraph));
@@ -185,7 +185,7 @@ void FMetasoundMusicClockDriver::DetachAllCallbacks()
 		GeneratorAttachedCallbackHandle.Reset();
 		CurrentGeneratorHandle->OnGeneratorHandleDetached.Remove(GeneratorDetachedCallbackHandle);
 		GeneratorDetachedCallbackHandle.Reset();
-		CurrentGeneratorHandle->OnIOUpdated.Remove(GeneratorIOUpdatedCallbackHandle);
+		CurrentGeneratorHandle->OnIOUpdatedWithChanges.Remove(GeneratorIOUpdatedCallbackHandle);
 		GeneratorIOUpdatedCallbackHandle.Reset();
 		CurrentGeneratorHandle->RemoveGraphSetCallback(GraphChangedCallbackHandle);
 		GraphChangedCallbackHandle.Reset();
@@ -203,10 +203,25 @@ void FMetasoundMusicClockDriver::OnGraphSet()
 }
 
 void FMetasoundMusicClockDriver::OnGeneratorIOUpdated()
+{
+	// Replaced by OnGeneratorIOUpdatedWithChanges, no-op but still here as it's protected, and child classes may call to it
+	UE_LOG(LogMusicClock, Warning, TEXT("Called to OnGeneratorIOUpdated, which is deprecated and no-op. Use OnGeneratorIOUpdatedWithChanges."), *MetasoundOutputName.ToString());
+}
+
+void FMetasoundMusicClockDriver::OnGeneratorIOUpdatedWithChanges(const TArray<Metasound::FVertexInterfaceChange>& VertexInterfaceChanges)
 {
 	// An output vertex update may have destroyed our Clock, reattach to the new one if it's there
-	// FUTURE: There isn't a way to listen to vertex changes on a specific Metasound output, so until then this will be called every time IO changes
-	TryToRegisterPlayCursor();
+	if (!MetasoundOutputName.IsNone())
+	{
+		for (auto& VertexInterfaceChange : VertexInterfaceChanges)
+		{
+			if (VertexInterfaceChange.VertexName.IsEqual(MetasoundOutputName))
+			{
+				TryToRegisterPlayCursor();
+				break;
+			}
+		}
+	}
 }
 
 void FMetasoundMusicClockDriver::OnGeneratorDetached()
diff --git a/Engine/Plugins/Runtime/Harmonix/Source/HarmonixMetasound/Private/HarmonixMetasound/Components/MetasoundMusicClockDriver.h b/Engine/Plugins/Runtime/Harmonix/Source/HarmonixMetasound/Private/HarmonixMetasound/Components/MetasoundMusicClockDriver.h
--- a/Engine/Plugins/Runtime/Harmonix/Source/HarmonixMetasound/Private/HarmonixMetasound/Components/MetasoundMusicClockDriver.h
+++ b/Engine/Plugins/Runtime/Harmonix/Source/HarmonixMetasound/Private/HarmonixMetasound/Components/MetasoundMusicClockDriver.h
@@ -28,7 +28,9 @@ protected:
 	void OnGeneratorAttached();
 	void OnGeneratorDetached();
 	void OnGraphSet();
+	UE_DEPRECATED(5.5, "Use OnGeneratorIOUpdatedWithChanges.")
 	void OnGeneratorIOUpdated();
+	void OnGeneratorIOUpdatedWithChanges(const TArray<Metasound::FVertexInterfaceChange>& VertexInterfaceChanges);
 
 private:
 	FName MetasoundOutputName;
diff --git a/Engine/Plugins/Runtime/Metasound/Source/MetasoundEngine/Private/MetasoundGeneratorHandle.cpp b/Engine/Plugins/Runtime/Metasound/Source/MetasoundEngine/Private/MetasoundGeneratorHandle.cpp
--- a/Engine/Plugins/Runtime/Metasound/Source/MetasoundEngine/Private/MetasoundGeneratorHandle.cpp
+++ b/Engine/Plugins/Runtime/Metasound/Source/MetasoundEngine/Private/MetasoundGeneratorHandle.cpp
@@ -82,6 +82,9 @@ namespace Metasound
 		return nullptr;
 	}
 
+	// Remove these PRAGMAs when cleaning up the deprecated OnGeneratorIOUpdated. 
+	// The curly brace line was throwing errors for usage of the deprecated member, presumably it was doing something to tear it down and complaining about it.
+	PRAGMA_DISABLE_DEPRECATION_WARNINGS
 	FMetasoundGeneratorHandle::~FMetasoundGeneratorHandle()
 	{
 		check(IsInGameThread());
@@ -100,6 +103,7 @@ namespace Metasound
 		// unset the generator and clean up
 		SetGenerator(nullptr);
 	}
+	PRAGMA_ENABLE_DEPRECATION_WARNINGS
 
 	bool FMetasoundGeneratorHandle::IsValid() const
 	{
@@ -304,7 +308,7 @@ namespace Metasound
 			}
 
 			// Vertex interface updated (Live Update support)
-			GeneratorVertexInterfaceChangedDelegateHandle = PinnedGenerator->OnVertexInterfaceDataUpdated.AddSP(
+			GeneratorVertexInterfaceChangedDelegateHandle = PinnedGenerator->OnVertexInterfaceDataUpdatedWithChanges.AddSP(
 				AsShared(),
 				&FMetasoundGeneratorHandle::HandleGeneratorVertexInterfaceChanged);
 		}
@@ -321,7 +325,7 @@ namespace Metasound
 		{
 			PinnedGenerator->OnOutputChanged.Remove(GeneratorOutputChangedDelegateHandle);
 			PinnedGenerator->RemoveGraphSetCallback(GeneratorGraphSetDelegateHandle);
-			PinnedGenerator->OnVertexInterfaceDataUpdated.Remove(GeneratorVertexInterfaceChangedDelegateHandle);
+			PinnedGenerator->OnVertexInterfaceDataUpdatedWithChanges.Remove(GeneratorVertexInterfaceChangedDelegateHandle);
 		}
 	}
 
@@ -635,23 +639,30 @@ namespace Metasound
 		});
 	}
 
-	void FMetasoundGeneratorHandle::HandleGeneratorVertexInterfaceChanged(FVertexInterfaceData)
+	void FMetasoundGeneratorHandle::HandleGeneratorVertexInterfaceChanged(const TArray<FVertexInterfaceChange>& VertexInterfaceChanges)
 	{
 		METASOUND_LLM_SCOPE;
 		METASOUND_TRACE_CPUPROFILER_EVENT_SCOPE(FMetasoundGeneratorHandle::HandleGeneratorVertexInterfaceChanged);
 
 		// Defer to the game thread. We grab a weak pointer in case this gets destroyed while we wait.
-		AsyncTask(ENamedThreads::GameThread, [WeakThis = AsWeak()]()
+		AsyncTask(ENamedThreads::GameThread, [WeakThis = AsWeak(), VertexInterfaceChanges]()
 		{
 			if (const TSharedPtr<FMetasoundGeneratorHandle> PinnedThis = WeakThis.Pin())
 			{
 				PinnedThis->SendParametersToGenerator();
 				PinnedThis->FixUpOutputWatchers();
 
+				PRAGMA_DISABLE_DEPRECATION_WARNINGS
 				if (PinnedThis->OnGeneratorIOUpdated.IsBound())
 				{
 					PinnedThis->OnGeneratorIOUpdated.Execute();
 				}
+				PRAGMA_ENABLE_DEPRECATION_WARNINGS
+
+				if (PinnedThis->OnGeneratorIOUpdatedWithChanges.IsBound())
+				{
+					PinnedThis->OnGeneratorIOUpdatedWithChanges.Execute(VertexInterfaceChanges);
+				}
 			}
 		});
 	}
@@ -848,10 +859,17 @@ bool UMetasoundGeneratorHandle::InitGeneratorHandle(TWeakObjectPtr<UAudioCompone
 		OnGeneratorsGraphChanged.Broadcast();
 	});
 
+	PRAGMA_DISABLE_DEPRECATION_WARNINGS
 	GeneratorHandle->OnGeneratorIOUpdated.BindLambda([this]()
 	{
 		OnIOUpdated.Broadcast();
 	});
+	PRAGMA_ENABLE_DEPRECATION_WARNINGS
+
+	GeneratorHandle->OnGeneratorIOUpdatedWithChanges.BindLambda([this](const TArray<Metasound::FVertexInterfaceChange>& VertexInterfaceChanges)
+	{
+		OnIOUpdatedWithChanges.Broadcast(VertexInterfaceChanges);
+	});
 
 	return true;
 }
diff --git a/Engine/Plugins/Runtime/Metasound/Source/MetasoundEngine/Public/MetasoundGeneratorHandle.h b/Engine/Plugins/Runtime/Metasound/Source/MetasoundEngine/Public/MetasoundGeneratorHandle.h
--- a/Engine/Plugins/Runtime/Metasound/Source/MetasoundEngine/Public/MetasoundGeneratorHandle.h
+++ b/Engine/Plugins/Runtime/Metasound/Source/MetasoundEngine/Public/MetasoundGeneratorHandle.h
@@ -19,6 +19,7 @@ class UMetaSoundSource;
 namespace Metasound
 {
 	class FMetasoundGenerator;
+	struct FVertexInterfaceChange;
 }
 
 DECLARE_DYNAMIC_DELEGATE_TwoParams(FOnMetasoundOutputValueChanged, FName, OutputName, const FMetaSoundOutput&, Output);
@@ -89,10 +90,18 @@ namespace Metasound
 		DECLARE_DELEGATE(FOnGeneratorIOUpdated);
 
 		/**
-		 * Fires on the game thread when the generator's graph inputs or outputs change
+		 * Fires on the game thread when the generator's graph inputs or outputs change.
 		 */
+		UE_DEPRECATED(5.5, "Use OnGeneratorIOUpdatedWithChanges.")
 		FOnGeneratorIOUpdated OnGeneratorIOUpdated;
 
+		DECLARE_DELEGATE_OneParam(FOnGeneratorIOUpdatedWithChanges, const TArray<FVertexInterfaceChange>&);
+
+		/**
+		 * Fires on the game thread when the generator's graph inputs or outputs change, and includes a list of those changes.
+		 */
+		FOnGeneratorIOUpdatedWithChanges OnGeneratorIOUpdatedWithChanges;
+
 		/**
 		 * Update the current parameter state on this handle and enqueue the changes on the generator.
 		 * 
@@ -243,7 +252,7 @@ namespace Metasound
 		void HandleGeneratorGraphSet();
 		FDelegateHandle GeneratorGraphSetDelegateHandle;
 
-		void HandleGeneratorVertexInterfaceChanged(FVertexInterfaceData VertexInterfaceData);
+		void HandleGeneratorVertexInterfaceChanged(const TArray<FVertexInterfaceChange>& VertexInterfaceData);
 		FDelegateHandle GeneratorVertexInterfaceChangedDelegateHandle;
 
 		void HandleOutputChanged(
@@ -450,8 +459,12 @@ public:
 	bool RemoveGraphSetCallback(const FDelegateHandle& Handle);
 
 	DECLARE_MULTICAST_DELEGATE(FOnIOUpdated)
+	UE_DEPRECATED(5.5, "Use OnIOUpdatedWithChanges.")
 	FOnIOUpdated OnIOUpdated;
 
+	DECLARE_MULTICAST_DELEGATE_OneParam(FOnIOUpdatedWithChanges, const TArray<Metasound::FVertexInterfaceChange>&)
+	FOnIOUpdatedWithChanges OnIOUpdatedWithChanges;
+
 	/**
 	 * Watch an output value.
 	 *
diff --git a/Engine/Plugins/Runtime/Metasound/Source/MetasoundEngineTest/Private/MetasoundDynamicGeneratorTests.cpp b/Engine/Plugins/Runtime/Metasound/Source/MetasoundEngineTest/Private/MetasoundDynamicGeneratorTests.cpp
--- a/Engine/Plugins/Runtime/Metasound/Source/MetasoundEngineTest/Private/MetasoundDynamicGeneratorTests.cpp
+++ b/Engine/Plugins/Runtime/Metasound/Source/MetasoundEngineTest/Private/MetasoundDynamicGeneratorTests.cpp
@@ -158,11 +158,18 @@ namespace Metasound::Test::Generator::Dynamic
 
 		// Register for vertex interface updates
 		FVertexInterfaceData LatestInterfaceData;
+		TArray<FVertexInterfaceChange> LatestInterfaceChanges;
+
 		GeneratorBuilder.Generator.OnVertexInterfaceDataUpdated.AddLambda([&LatestInterfaceData](FVertexInterfaceData VertexInterfaceData)
 		{
 			LatestInterfaceData = MoveTemp(VertexInterfaceData);
 		});
 
+		GeneratorBuilder.Generator.OnVertexInterfaceDataUpdatedWithChanges.AddLambda([&LatestInterfaceChanges](const TArray<FVertexInterfaceChange>& VertexInterfaceChanges)
+		{
+			LatestInterfaceChanges = VertexInterfaceChanges;
+		});
+
 		// Add an input
 		const FVertexName InputName = "SomeInput";
 		const FGuid InputGuid = FGuid::NewGuid();
@@ -180,6 +187,17 @@ namespace Metasound::Test::Generator::Dynamic
 			const float* Value = InputRef->GetValue<float>();
 			UTEST_NOT_NULL("Value exists", Value);
 			UTEST_EQUAL("Value is default", *Value, DefaultValue);
+
+			// Check that the change was tracked
+			UTEST_EQUAL("There are the three expected default IO plus the new input", LatestInterfaceChanges.Num(), 4);
+			const FVertexInterfaceChange* LastChange = LatestInterfaceChanges.FindByPredicate([InputName](const FVertexInterfaceChange& Other)
+				{
+					return Other.VertexName.IsEqual(InputName);
+				});
+			UTEST_NOT_NULL("Input addition is present in changes", LastChange);
+			UTEST_EQUAL("Input addition is for the right Vertex", LastChange->VertexName, InputName);
+			UTEST_EQUAL("Input addition is for the right Vertex type", LastChange->VertexType, EMetasoundFrontendClassType::Input);
+			UTEST_EQUAL("Input addition is the Added type", LastChange->ChangeType, Metasound::EVertexInterfaceChangeType::Added);
 		}
 
 		// Remove the input
@@ -191,6 +209,13 @@ namespace Metasound::Test::Generator::Dynamic
 			// Check that the input actually got removed
 			const FAnyDataReference* InputRef = LatestInterfaceData.GetInputs().FindDataReference(InputName);
 			UTEST_NULL("Vertex data does not contain input", InputRef);
+
+			// Check that the change was tracked
+			UTEST_EQUAL("Input removal is present in changes", LatestInterfaceChanges.Num(), 1);
+			const FVertexInterfaceChange LastChange = LatestInterfaceChanges.Last();
+			UTEST_EQUAL("Input removal is for the right Vertex", LastChange.VertexName, InputName);
+			UTEST_EQUAL("Input removal is for the right Vertex type", LastChange.VertexType, EMetasoundFrontendClassType::Input);
+			UTEST_EQUAL("Input removal is the Removed type", LastChange.ChangeType, Metasound::EVertexInterfaceChangeType::Removed);
 		}
 
 		// Add an output
@@ -205,6 +230,13 @@ namespace Metasound::Test::Generator::Dynamic
 			// check that the output actually got added
 			const FAnyDataReference* OutputRef = LatestInterfaceData.GetOutputs().FindDataReference(OutputName);
 			UTEST_NOT_NULL("Vertex data contains output", OutputRef);
+
+			// Check that the change was tracked
+			UTEST_EQUAL("Output addition is present in changes", LatestInterfaceChanges.Num(), 1);
+			const FVertexInterfaceChange LastChange = LatestInterfaceChanges.Last();
+			UTEST_EQUAL("Output addition is for the right Vertex", LastChange.VertexName, OutputName);
+			UTEST_EQUAL("Output addition is for the right Vertex type", LastChange.VertexType, EMetasoundFrontendClassType::Output);
+			UTEST_EQUAL("Output addition is the Removed type", LastChange.ChangeType, Metasound::EVertexInterfaceChangeType::Added);
 		}
 
 		// Remove the output
@@ -216,6 +248,13 @@ namespace Metasound::Test::Generator::Dynamic
 			// Check that the output actually got removed
 			const FAnyDataReference* OutputRef = LatestInterfaceData.GetOutputs().FindDataReference(OutputName);
 			UTEST_NULL("Vertex data does not contain output", OutputRef);
+
+			// Check that the change was tracked
+			UTEST_EQUAL("Output removal is present in changes", LatestInterfaceChanges.Num(), 1);
+			const FVertexInterfaceChange LastChange = LatestInterfaceChanges.Last();
+			UTEST_EQUAL("Output removal is for the right Vertex", LastChange.VertexName, OutputName);
+			UTEST_EQUAL("Output removal is for the right Vertex type", LastChange.VertexType, EMetasoundFrontendClassType::Output);
+			UTEST_EQUAL("Output removal is the Removed type", LastChange.ChangeType, Metasound::EVertexInterfaceChangeType::Removed);
 		}
 		
 		return true;
diff --git a/Engine/Plugins/Runtime/Metasound/Source/MetasoundGenerator/Private/MetasoundGenerator.cpp b/Engine/Plugins/Runtime/Metasound/Source/MetasoundGenerator/Private/MetasoundGenerator.cpp
--- a/Engine/Plugins/Runtime/Metasound/Source/MetasoundGenerator/Private/MetasoundGenerator.cpp
+++ b/Engine/Plugins/Runtime/Metasound/Source/MetasoundGenerator/Private/MetasoundGenerator.cpp
@@ -435,7 +435,20 @@ PRAGMA_ENABLE_DEPRECATION_WARNINGS
 
 		bIsWaitingForFirstGraph = false;
 		OnSetGraph.Broadcast();
+
+		for (const auto& Binding : VertexInterfaceData.GetInputs())
+		{
+			VertexInterfaceChangesSinceLastBroadcast.Add({ Binding.GetVertex().VertexName, EMetasoundFrontendClassType::Input, EVertexInterfaceChangeType::Added });
+		}
+
+		for (const auto& Binding : VertexInterfaceData.GetOutputs())
+		{
+			VertexInterfaceChangesSinceLastBroadcast.Add({ Binding.GetVertex().VertexName, EMetasoundFrontendClassType::Output, EVertexInterfaceChangeType::Added });
+		}
+
+		PRAGMA_DISABLE_DEPRECATION_WARNINGS
 		bVertexInterfaceHasChanged.store(true);
+		PRAGMA_ENABLE_DEPRECATION_WARNINGS
 	}
 
 	TUniquePtr<IOperator> FMetasoundGenerator::ReleaseGraphOperator()
@@ -450,6 +463,16 @@ PRAGMA_ENABLE_DEPRECATION_WARNINGS
 
 	void FMetasoundGenerator::ClearGraph()
 	{
+		for (const auto& Binding : VertexInterfaceData.GetInputs())
+		{
+			VertexInterfaceChangesSinceLastBroadcast.Add({ Binding.GetVertex().VertexName, EMetasoundFrontendClassType::Input, EVertexInterfaceChangeType::Removed });
+		}
+
+		for (const auto& Binding : VertexInterfaceData.GetOutputs())
+		{
+			VertexInterfaceChangesSinceLastBroadcast.Add({ Binding.GetVertex().VertexName, EMetasoundFrontendClassType::Output, EVertexInterfaceChangeType::Removed });
+		}
+
 		RootExecuter.ReleaseOperator();
 		VertexInterfaceData = FVertexInterfaceData();
 		GraphOutputAudio.Reset();
@@ -458,7 +481,9 @@ PRAGMA_ENABLE_DEPRECATION_WARNINGS
 		ParameterSetters.Reset();
 		ParameterPackSetters.Reset();
 
+		PRAGMA_DISABLE_DEPRECATION_WARNINGS
 		bVertexInterfaceHasChanged.store(true);
+		PRAGMA_ENABLE_DEPRECATION_WARNINGS
 	}
 
 	void FMetasoundGenerator::QueueParameterPack(TSharedPtr<FMetasoundParameterPackStorage> ParameterPack)
@@ -662,11 +687,14 @@ PRAGMA_ENABLE_DEPRECATION_WARNINGS
 		}
 
 		// If the vertex interface changed, notify listeners
+		if (!VertexInterfaceChangesSinceLastBroadcast.IsEmpty())
 		{
-			if (bVertexInterfaceHasChanged.exchange(false))
-			{
-				OnVertexInterfaceDataUpdated.Broadcast(VertexInterfaceData);
-			}
+			OnVertexInterfaceDataUpdated.Broadcast(VertexInterfaceData);
+			OnVertexInterfaceDataUpdatedWithChanges.Broadcast(VertexInterfaceChangesSinceLastBroadcast);
+			VertexInterfaceChangesSinceLastBroadcast.Empty();
+			PRAGMA_DISABLE_DEPRECATION_WARNINGS
+			bVertexInterfaceHasChanged.store(false);
+			PRAGMA_ENABLE_DEPRECATION_WARNINGS
 		}
 
 		return NumSamplesWritten;
@@ -1013,7 +1041,10 @@ PRAGMA_ENABLE_DEPRECATION_WARNINGS
 		VertexInterfaceData.GetInputs() = InGraphInputs;
 		GeneratorBuilder::AddParameterSetterIfWritable(InVertexName, VertexInterfaceData.GetInputs(), ParameterSetters, ParameterPackSetters);
 
+		VertexInterfaceChangesSinceLastBroadcast.Add({ InVertexName, EMetasoundFrontendClassType::Input, EVertexInterfaceChangeType::Added });
+		PRAGMA_DISABLE_DEPRECATION_WARNINGS
 		bVertexInterfaceHasChanged.store(true);
+		PRAGMA_ENABLE_DEPRECATION_WARNINGS
 	}
 
 	void FMetasoundDynamicGraphGenerator::OnInputRemoved(const FVertexName& InVertexName, const FInputVertexInterfaceData& InGraphInputs)
@@ -1021,7 +1052,10 @@ PRAGMA_ENABLE_DEPRECATION_WARNINGS
 		VertexInterfaceData.GetInputs() = InGraphInputs;
 		ParameterSetters.Remove(InVertexName);
 		ParameterPackSetters.Remove(InVertexName);
+		VertexInterfaceChangesSinceLastBroadcast.Add({ InVertexName, EMetasoundFrontendClassType::Input, EVertexInterfaceChangeType::Removed });
+		PRAGMA_DISABLE_DEPRECATION_WARNINGS
 		bVertexInterfaceHasChanged.store(true);
+		PRAGMA_ENABLE_DEPRECATION_WARNINGS
 	}
 
 	void FMetasoundDynamicGraphGenerator::OnOutputAdded(const FVertexName& InVertexName, const FOutputVertexInterfaceData& InGraphOutputs)
@@ -1033,7 +1067,11 @@ PRAGMA_ENABLE_DEPRECATION_WARNINGS
 		{
 			OnFinishedTriggerRef = InGraphOutputs.GetDataReadReference<FTrigger>(InVertexName);
 		}
+		
+		VertexInterfaceChangesSinceLastBroadcast.Add({ InVertexName, EMetasoundFrontendClassType::Output, EVertexInterfaceChangeType::Added });
+		PRAGMA_DISABLE_DEPRECATION_WARNINGS
 		bVertexInterfaceHasChanged.store(true);
+		PRAGMA_ENABLE_DEPRECATION_WARNINGS
 	}
 
 	void FMetasoundDynamicGraphGenerator::OnOutputUpdated(const FVertexName& InVertexName, const FOutputVertexInterfaceData& InGraphOutputs)
@@ -1052,7 +1090,10 @@ PRAGMA_ENABLE_DEPRECATION_WARNINGS
 			OnFinishedTriggerRef = InGraphOutputs.GetDataReadReference<FTrigger>(InVertexName);
 		}
 		
+		VertexInterfaceChangesSinceLastBroadcast.Add({ InVertexName, EMetasoundFrontendClassType::Output, EVertexInterfaceChangeType::Updated });
+		PRAGMA_DISABLE_DEPRECATION_WARNINGS
 		bVertexInterfaceHasChanged.store(true);
+		PRAGMA_ENABLE_DEPRECATION_WARNINGS
 	}
 
 	void FMetasoundDynamicGraphGenerator::OnOutputRemoved(const FVertexName& InVertexName, const FOutputVertexInterfaceData& InGraphOutputs)
@@ -1060,7 +1101,10 @@ PRAGMA_ENABLE_DEPRECATION_WARNINGS
 		MetasoundGeneratorPrivate::LogErrorIfIsAudioVertex(InVertexName, TEXT("The MetaSound Generator cannot dynamically remove audio outputs"));
 
 		VertexInterfaceData.GetOutputs() = InGraphOutputs;
+		VertexInterfaceChangesSinceLastBroadcast.Add({ InVertexName, EMetasoundFrontendClassType::Output, EVertexInterfaceChangeType::Removed });
+		PRAGMA_DISABLE_DEPRECATION_WARNINGS
 		bVertexInterfaceHasChanged.store(true);
+		PRAGMA_ENABLE_DEPRECATION_WARNINGS
 	}
 
 	TUniquePtr<IOperator> FMetasoundDynamicGraphGenerator::ReleaseGraphOperator()
diff --git a/Engine/Plugins/Runtime/Metasound/Source/MetasoundGenerator/Public/MetasoundGenerator.h b/Engine/Plugins/Runtime/Metasound/Source/MetasoundGenerator/Public/MetasoundGenerator.h
--- a/Engine/Plugins/Runtime/Metasound/Source/MetasoundGenerator/Public/MetasoundGenerator.h
+++ b/Engine/Plugins/Runtime/Metasound/Source/MetasoundGenerator/Public/MetasoundGenerator.h
@@ -95,6 +95,20 @@ namespace Metasound
 		static void Reset(FMetasoundGeneratorInitParams& InParams);
 	};
 
+	enum class EVertexInterfaceChangeType : uint8
+	{
+		Added,
+		Updated,
+		Removed
+	};
+
+	struct METASOUNDGENERATOR_API FVertexInterfaceChange
+	{
+		FVertexName VertexName;
+		EMetasoundFrontendClassType VertexType; // Input or Output
+		EVertexInterfaceChangeType ChangeType;
+	};
+
 	DECLARE_TS_MULTICAST_DELEGATE(FOnSetGraph);
 
 	class METASOUNDGENERATOR_API FMetasoundGenerator : public ISoundGenerator
@@ -191,6 +205,9 @@ namespace Metasound
 		DECLARE_TS_MULTICAST_DELEGATE_OneParam(FOnVertexInterfaceDataUpdated, FVertexInterfaceData);
 		FOnVertexInterfaceDataUpdated OnVertexInterfaceDataUpdated;
 
+		DECLARE_TS_MULTICAST_DELEGATE_OneParam(FOnVertexInterfaceDataUpdatedWithChanges, const TArray<FVertexInterfaceChange>&);
+		FOnVertexInterfaceDataUpdatedWithChanges OnVertexInterfaceDataUpdatedWithChanges;
+
 		/**
 		 * Add a vertex analyzer for a named output with the given address info.
 		 *
@@ -263,6 +280,7 @@ namespace Metasound
 		void ClearGraph();
 		bool UpdateGraphIfPending();
 
+		UE_DEPRECATED(5.5, "Use VertexInterfaceChangesSinceLastBroadcast to determine if changes have occurred.")
 		std::atomic<bool> bVertexInterfaceHasChanged{ false };
 
 	private:
@@ -306,6 +324,7 @@ namespace Metasound
 	protected:
 		FExecuter RootExecuter;
 		FVertexInterfaceData VertexInterfaceData;
+		TArray<FVertexInterfaceChange> VertexInterfaceChangesSinceLastBroadcast;
 
 		TArray<FAudioBufferReadRef> GraphOutputAudio;
 
