Fix Audio.Metasound.Generator.Dynamic.VertexInterfaceUpdated test
#tests Session Frontend Automation

diff --git a/Engine/Plugins/Runtime/Metasound/Source/MetasoundEngineTest/Private/MetasoundDynamicGeneratorTests.cpp b/Engine/Plugins/Runtime/Metasound/Source/MetasoundEngineTest/Private/MetasoundDynamicGeneratorTests.cpp
--- a/Engine/Plugins/Runtime/Metasound/Source/MetasoundEngineTest/Private/MetasoundDynamicGeneratorTests.cpp
+++ b/Engine/Plugins/Runtime/Metasound/Source/MetasoundEngineTest/Private/MetasoundDynamicGeneratorTests.cpp
@@ -189,15 +189,14 @@ namespace Metasound::Test::Generator::Dynamic
 			UTEST_EQUAL("Value is default", *Value, DefaultValue);
 
 			// Check that the change was tracked
-			UTEST_EQUAL("There are the three expected default IO plus the new input", LatestInterfaceChanges.Num(), 4);
-			const FVertexInterfaceChange* LastChange = LatestInterfaceChanges.FindByPredicate([InputName](const FVertexInterfaceChange& Other)
+			const TArray<FVertexInterfaceChange> InputChanges = LatestInterfaceChanges.FilterByPredicate([InputName](const FVertexInterfaceChange& Other)
 				{
 					return Other.VertexName.IsEqual(InputName);
 				});
-			UTEST_NOT_NULL("Input addition is present in changes", LastChange);
-			UTEST_EQUAL("Input addition is for the right Vertex", LastChange->VertexName, InputName);
-			UTEST_EQUAL("Input addition is for the right Vertex type", LastChange->VertexType, EMetasoundFrontendClassType::Input);
-			UTEST_EQUAL("Input addition is the Added type", LastChange->ChangeType, Metasound::EVertexInterfaceChangeType::Added);
+			UTEST_EQUAL("There is only one expected change with our Input", InputChanges.Num(), 1);
+			UTEST_EQUAL("Input addition is for the right Vertex", InputChanges[0].VertexName, InputName);
+			UTEST_EQUAL("Input addition is for the right Vertex type", InputChanges[0].VertexType, EMetasoundFrontendClassType::Input);
+			UTEST_EQUAL("Input addition is the Added type", InputChanges[0].ChangeType, Metasound::EVertexInterfaceChangeType::Added);
 		}
 
 		// Remove the input
