[MetaSound Feature Request] Add Get Name to Wave Asset
#rnx

diff --git a/Engine/Plugins/Runtime/Metasound/Source/MetasoundEngine/Private/MetasoundWaveInfoNode.cpp b/Engine/Plugins/Runtime/Metasound/Source/MetasoundEngine/Private/MetasoundWaveInfoNode.cpp
--- a/Engine/Plugins/Runtime/Metasound/Source/MetasoundEngine/Private/MetasoundWaveInfoNode.cpp
+++ b/Engine/Plugins/Runtime/Metasound/Source/MetasoundEngine/Private/MetasoundWaveInfoNode.cpp
@@ -30,12 +30,10 @@ namespace Metasound
 
 		// outputs
 		METASOUND_PARAM(ParamDurationSeconds, "Duration", "Duration of the wave asset in seconds");
+		METASOUND_PARAM(OutParamAssetName, "Name", "Name of the wave asset");
+		METASOUND_PARAM(OutParamAssetPath, "Path", "Full path of the wave asset");
 
-	} // namespace WaveInfoNodeParameterNames
-
-	using namespace WaveInfoNodeParameterNames;
-
-
+	} // namespace WaveInfoNodeParameterNames	
 
 	class FWaveInfoNodeOperator : public TExecutableOperator < FWaveInfoNodeOperator >
 	{
@@ -60,38 +58,38 @@ namespace Metasound
 
 		// output pins
 		FTimeWriteRef DurationSeconds;
+		FStringWriteRef NameOutput;
+		FStringWriteRef PathOutput;
 
 		// other
-		FOperatorSettings Settings;
+		FSoundWaveProxyPtr SoundWaveProxy;
 
 	}; // class FWaveInfoNodeOperator
 
-
-
-
 	// ctor
 	FWaveInfoNodeOperator::FWaveInfoNodeOperator(const FOperatorSettings& InSettings, const FWaveAssetReadRef& InWaveAsset)
 		: WaveAsset(InWaveAsset)
 		, DurationSeconds(FTimeWriteRef::CreateNew(0.0f))
-		, Settings(InSettings)
+		, NameOutput(FStringWriteRef::CreateNew(TEXT("")))
+		, PathOutput(FStringWriteRef::CreateNew(TEXT("")))
 	{
 		Execute();
 	}
 
-
 	const FNodeClassMetadata& FWaveInfoNodeOperator::GetNodeInfo()
 	{
 		auto InitNodeInfo = []() -> FNodeClassMetadata
 		{
 			FNodeClassMetadata Info;
-			Info.ClassName = { Metasound::EngineNodes::Namespace, TEXT("Get Wave Duration"), TEXT(" ") };
+			Info.ClassName = { Metasound::EngineNodes::Namespace, "Get Wave Duration", " " };
 			Info.MajorVersion = 1;
 			Info.MinorVersion = 0;
-			Info.DisplayName = METASOUND_LOCTEXT("MetasoundGetWaveDuration_ClassNodeDisplayName", "Get Wave Duration");
-			Info.Description = METASOUND_LOCTEXT("GetWaveDuration_NodeDescription", "Returns the duration of the input Wave asset (in seconds)"),
+			Info.DisplayName = METASOUND_LOCTEXT("MetasoundGetWaveInfo_ClassNodeDisplayName", "Get Wave Info");
+			Info.Description = METASOUND_LOCTEXT("GetWaveInfo_NodeDescription", "Returns the Info from the Wave Asset"),
 			Info.Author = PluginAuthor;
 			Info.PromptIfMissing = PluginNodeMissingPrompt;
 			Info.DefaultInterface = DeclareVertexInterface();
+			Info.CategoryHierarchy.Emplace(NodeCategories::Debug);
 
 			return Info;
 		};
@@ -101,7 +99,6 @@ namespace Metasound
 		return Info;
 	}
 
-
 	FVertexInterface FWaveInfoNodeOperator::DeclareVertexInterface()
 	{
 		using namespace WaveInfoNodeParameterNames;
@@ -110,14 +107,15 @@ namespace Metasound
 				TInputDataVertex<FWaveAsset>(METASOUND_GET_PARAM_NAME_AND_METADATA(ParamWaveAsset))
 			),
 			FOutputVertexInterface(
-				TOutputDataVertex<FTime>(METASOUND_GET_PARAM_NAME_AND_METADATA(ParamDurationSeconds))
+				TOutputDataVertex<FTime>(METASOUND_GET_PARAM_NAME_AND_METADATA(ParamDurationSeconds)),
+				TOutputDataVertex<FString>(METASOUND_GET_PARAM_NAME_AND_METADATA(OutParamAssetName)),
+				TOutputDataVertex<FString>(METASOUND_GET_PARAM_NAME_AND_METADATA(OutParamAssetPath))
 			)
 		);
 
 		return Interface;
 	}
 
-
 	TUniquePtr<IOperator> FWaveInfoNodeOperator::CreateOperator(const FBuildOperatorParams& InParams, FBuildResults& OutResults)
 	{
 		using namespace WaveInfoNodeParameterNames;
@@ -132,7 +130,6 @@ namespace Metasound
 	void FWaveInfoNodeOperator::BindInputs(FInputVertexInterfaceData& InOutVertexData)
 	{
 		using namespace WaveInfoNodeParameterNames;
-		FDataReferenceCollection InputDataReferences;
 		InOutVertexData.BindReadVertex(METASOUND_GET_PARAM_NAME(ParamWaveAsset), WaveAsset);	
 	}
 
@@ -140,8 +137,9 @@ namespace Metasound
 	{
 		// expose read access to our output buffer for other processors in the graph
 		using namespace WaveInfoNodeParameterNames;
-		FDataReferenceCollection OutputDataReferences;
 		InOutVertexData.BindReadVertex(METASOUND_GET_PARAM_NAME(ParamDurationSeconds), DurationSeconds);
+		InOutVertexData.BindReadVertex(METASOUND_GET_PARAM_NAME(OutParamAssetName), NameOutput);
+		InOutVertexData.BindReadVertex(METASOUND_GET_PARAM_NAME(OutParamAssetPath), PathOutput);
 	}
 
 	FDataReferenceCollection FWaveInfoNodeOperator::GetInputs() const
@@ -162,13 +160,34 @@ namespace Metasound
 
 	void FWaveInfoNodeOperator::Execute()
 	{
-		if ((*WaveAsset).IsSoundWaveValid())
-		{
-			*DurationSeconds = FTime::FromSeconds((*WaveAsset)->GetDuration());
-		}
-		else
+		FSoundWaveProxyPtr ProxyPtr = (*WaveAsset).GetSoundWaveProxy();
+
+		if (SoundWaveProxy != ProxyPtr)
 		{
-			*DurationSeconds = FTime::FromSeconds(0.0f);
+			SoundWaveProxy = ProxyPtr;
+
+			if (SoundWaveProxy.IsValid() && (*WaveAsset).IsSoundWaveValid())
+			{
+				*DurationSeconds = FTime::FromSeconds((*WaveAsset)->GetDuration());
+
+				FName AssetName = SoundWaveProxy->GetFName();
+				*NameOutput = AssetName.ToString();
+
+				FString FullPath;
+				SoundWaveProxy->GetPackageName().AppendString(FullPath);
+				if (!AssetName.IsNone())
+				{
+					FullPath += TEXT(".");
+					AssetName.AppendString(FullPath);
+				}
+				*PathOutput = FullPath;
+			}
+			else
+			{
+				*DurationSeconds = FTime::FromSeconds(0.0f);
+				*NameOutput = TEXT("");
+				*PathOutput = TEXT("");
+			}
 		}
 	}
 
@@ -177,8 +196,6 @@ namespace Metasound
 		Execute();
 	}
 
-
-
 	class FWaveInfoNode : public FNodeFacade
 	{
 	public:
@@ -195,7 +212,6 @@ namespace Metasound
 
 	};
 
-
 	METASOUND_REGISTER_NODE(FWaveInfoNode);
 
 } // namespace Metasound
