Fix for MetaSound RTTI failure on platforms which perform code folding.

diff --git a/Engine/Plugins/Runtime/Metasound/Source/MetasoundGraphCore/Public/MetasoundDataReference.h b/Engine/Plugins/Runtime/Metasound/Source/MetasoundGraphCore/Public/MetasoundDataReference.h
--- a/Engine/Plugins/Runtime/Metasound/Source/MetasoundGraphCore/Public/MetasoundDataReference.h
+++ b/Engine/Plugins/Runtime/Metasound/Source/MetasoundGraphCore/Public/MetasoundDataReference.h
@@ -27,7 +27,7 @@ namespace Metasound
 /** Convert a EDataReferenceAccessType to FString. */
 FString METASOUNDGRAPHCORE_API LexToString(Metasound::EDataReferenceAccessType InAccessType);
 
-using FMetasoundDataTypeId = void const*;
+using FMetasoundDataTypeId = void const* const;
 
 namespace Metasound
 {
@@ -84,7 +84,7 @@ namespace Metasound
 	template<typename DataType>
 	const void* const GetMetasoundDataTypeId()
 	{
-		return TDataReferenceTypeInfo<std::decay_t<DataType>>::TypeId;
+		return TDataReferenceTypeInfo<std::decay_t<DataType>>::GetTypeId();
 	}
 
 	/** Returns array type associated with the base datatype provided(ex. 'Float:Array' if 'Float' is provided) */
diff --git a/Engine/Plugins/Runtime/Metasound/Source/MetasoundGraphCore/Public/MetasoundDataReferenceMacro.h b/Engine/Plugins/Runtime/Metasound/Source/MetasoundGraphCore/Public/MetasoundDataReferenceMacro.h
--- a/Engine/Plugins/Runtime/Metasound/Source/MetasoundGraphCore/Public/MetasoundDataReferenceMacro.h
+++ b/Engine/Plugins/Runtime/Metasound/Source/MetasoundGraphCore/Public/MetasoundDataReferenceMacro.h
@@ -11,11 +11,11 @@
 	{ \
 		static ModuleApi const TCHAR* TypeName; \
 		static ModuleApi const FText& GetTypeDisplayText(); \
-		static ModuleApi const FMetasoundDataTypeId TypeId; \
+		static ModuleApi ::FMetasoundDataTypeId GetTypeId(); \
 		\
 		private: \
 		\
-		static const DataType* const TypePtr; \
+		static DataType* TypePtr; \
 	};
 
 #define DECLARE_METASOUND_DATA_REFERENCE_ALIAS_TYPES(DataType, DataTypeInfoTypeName, DataReadReferenceTypeName, DataWriteReferenceTypeName) \
@@ -39,14 +39,17 @@
 	DECLARE_METASOUND_DATA_REFERENCE_ALIAS_TYPES(DataType, DataTypeInfoTypeName, DataReadReferenceTypeName, DataWriteReferenceTypeName)
 
 #define DEFINE_METASOUND_DATA_REFERENCE_CORE_TYPE(DataType, DataTypeName, DataTypeLoctextKey) \
+	DataType* ::Metasound::TDataReferenceTypeInfo<DataType>::TypePtr = nullptr; \
 	const TCHAR* ::Metasound::TDataReferenceTypeInfo<DataType>::TypeName = TEXT(DataTypeName); \
 	const FText& ::Metasound::TDataReferenceTypeInfo<DataType>::GetTypeDisplayText() \
 	{ \
 		static const FText DisplayText = NSLOCTEXT("MetaSoundCore_DataReference", DataTypeLoctextKey, DataTypeName); \
 		return DisplayText; \
 	} \
-	const DataType* const ::Metasound::TDataReferenceTypeInfo<DataType>::TypePtr = nullptr; \
-	const void* const ::Metasound::TDataReferenceTypeInfo<DataType>::TypeId = static_cast<const FMetasoundDataTypeId>(&::Metasound::TDataReferenceTypeInfo<DataType>::TypePtr);
+	void const * const ::Metasound::TDataReferenceTypeInfo<DataType>::GetTypeId() \
+	{ \
+		return static_cast<void const * const>(&::Metasound::TDataReferenceTypeInfo<DataType>::TypePtr); \
+	}
 
 #define METASOUND_DATA_TYPE_NAME_VARIABLE_TYPE_SPECIFIER ":Variable"
 #define METASOUND_DATA_TYPE_NAME_ARRAY_TYPE_SPECIFIER ":Array"
diff --git a/Engine/Plugins/Runtime/Metasound/Source/MetasoundGraphCore/Public/MetasoundEnvironment.h b/Engine/Plugins/Runtime/Metasound/Source/MetasoundGraphCore/Public/MetasoundEnvironment.h
--- a/Engine/Plugins/Runtime/Metasound/Source/MetasoundGraphCore/Public/MetasoundEnvironment.h
+++ b/Engine/Plugins/Runtime/Metasound/Source/MetasoundGraphCore/Public/MetasoundEnvironment.h
@@ -16,19 +16,20 @@
 	template<> \
 	struct DLL_API ::Metasound::TEnvironmentVariableTypeInfo<VarType> \
 	{ \
-		typedef VarType Type; \
-		static const FMetasoundEnvironmentVariableTypeId TypeId; \
-		\
 	private: \
 		\
-		static const VarType* const TypePtr; \
+		static VarType* TypePtr; \
+	public: \
+		typedef VarType Type; \
+		static FMetasoundEnvironmentVariableTypeId GetTypeId(); \
 	};
 
 #define DEFINE_METASOUND_ENVIRONMENT_VARIABLE_TYPE(VarType) \
-	const VarType* const ::Metasound::TEnvironmentVariableTypeInfo<VarType>::TypePtr = nullptr; \
-	const void* const ::Metasound::TEnvironmentVariableTypeInfo<VarType>::TypeId = static_cast<const FMetasoundEnvironmentVariableTypeId>(&::Metasound::TEnvironmentVariableTypeInfo<VarType>::TypePtr);
+	VarType* ::Metasound::TEnvironmentVariableTypeInfo<VarType>::TypePtr = nullptr; \
+	void const* const ::Metasound::TEnvironmentVariableTypeInfo<VarType>::GetTypeId() { return static_cast<FMetasoundEnvironmentVariableTypeId>(&::Metasound::TEnvironmentVariableTypeInfo<VarType>::TypePtr); }
+
 	
-using FMetasoundEnvironmentVariableTypeId = void const*;
+using FMetasoundEnvironmentVariableTypeId = void const* const;
 
 namespace Metasound
 {
@@ -62,7 +63,7 @@ namespace Metasound
 	{
 		using FTypeInfo = TEnvironmentVariableTypeInfo< std::decay_t<Type> >;
 
-		return FTypeInfo::TypeId;
+		return FTypeInfo::GetTypeId();
 	}
 
 	/** Interface for a metasound environment variable which supports
