UE-171773 - MetaSounds get last index helper node for arrays
#rnx

diff --git a/Engine/Plugins/Runtime/Metasound/Source/MetasoundFrontend/Public/MetasoundArrayNodes.h b/Engine/Plugins/Runtime/Metasound/Source/MetasoundFrontend/Public/MetasoundArrayNodes.h
--- a/Engine/Plugins/Runtime/Metasound/Source/MetasoundFrontend/Public/MetasoundArrayNodes.h
+++ b/Engine/Plugins/Runtime/Metasound/Source/MetasoundFrontend/Public/MetasoundArrayNodes.h
@@ -69,6 +69,7 @@ namespace Metasound
 		METASOUND_PARAM(OutputArrayConcat, "Array", "Array after concatenation.")
 		METASOUND_PARAM(OutputArraySet, "Array", "Array after setting.")
 		METASOUND_PARAM(OutputArraySubset, "Array", "Subset of input array.")
+		METASOUND_PARAM(OutputLastIndex, "Last Index", "Last index of the array.")
 	};
 
 	/** TArrayNumOperator gets the number of elements in an Array. The operator
@@ -868,5 +869,111 @@ namespace Metasound
 
 		virtual ~TArraySubsetNode() = default;
 	};
+
+	/** TArrayLastIndex gets last index of an array. */
+	template<typename ArrayType>
+	class TArrayLastIndexOperator : public TExecutableOperator<TArrayLastIndexOperator<ArrayType>>
+	{
+	public:
+		using FArrayDataReadReference = TDataReadReference<ArrayType>;
+
+		// Declare the vertex interface
+		static const FVertexInterface& GetDefaultInterface()
+		{
+			using namespace ArrayNodeVertexNames;
+
+			static const FVertexInterface DefaultInterface(
+				FInputVertexInterface(
+					TInputDataVertex<ArrayType>(METASOUND_GET_PARAM_NAME_AND_METADATA(InputArray))
+				),
+				FOutputVertexInterface(
+					TOutputDataVertex<int32>(METASOUND_GET_PARAM_NAME_AND_METADATA(OutputLastIndex))
+				)
+			);
+
+			return DefaultInterface;
+		}
+
+		static const FNodeClassMetadata& GetNodeInfo()
+		{
+			auto CreateNodeClassMetadata = []() -> FNodeClassMetadata
+				{
+					const FName DataTypeName = GetMetasoundDataTypeName<ArrayType>();
+					const FName OperatorName = TEXT("GetLastIndex");
+					const FText NodeDisplayName = METASOUND_LOCTEXT_FORMAT("ArrayOpArrayLastIndexDisplayNamePattern", "Get Last Index ({0})", GetMetasoundDataTypeDisplayText<ArrayType>());
+					const FText NodeDescription = METASOUND_LOCTEXT("ArrayOpArrayLastIndexDescription", "Last index of the array");
+					const FVertexInterface NodeInterface = GetDefaultInterface();
+
+					return MetasoundArrayNodesPrivate::CreateArrayNodeClassMetadata(DataTypeName, OperatorName, NodeDisplayName, NodeDescription, NodeInterface);
+				};
+
+			static const FNodeClassMetadata Metadata = CreateNodeClassMetadata();
+
+			return Metadata;
+		}
+
+		static TUniquePtr<IOperator> CreateOperator(const FBuildOperatorParams& InParams, FBuildResults& OutResults)
+		{
+			using namespace ArrayNodeVertexNames;
+			using namespace MetasoundArrayNodesPrivate;
+
+			const FInputVertexInterfaceData& InputData = InParams.InputData;
+
+			// Get the input array or construct an empty one. 
+			FArrayDataReadReference Array = InputData.GetOrCreateDefaultDataReadReference<ArrayType>(METASOUND_GET_PARAM_NAME(InputArray), InParams.OperatorSettings);
+
+			return MakeUnique<TArrayLastIndexOperator>(Array);
+		}
+
+		TArrayLastIndexOperator(FArrayDataReadReference InArray)
+			: Array(InArray)
+			, LastIndex(TDataWriteReference<int32>::CreateNew())
+		{
+			// Initialize value for downstream nodes.
+			*LastIndex = Array->Num() - 1;
+		}
+
+		virtual ~TArrayLastIndexOperator() = default;
+
+		virtual void BindInputs(FInputVertexInterfaceData& InVertexData) override
+		{
+			using namespace ArrayNodeVertexNames;
+			InVertexData.BindReadVertex(METASOUND_GET_PARAM_NAME(InputArray), Array);
+		}
+
+		virtual void BindOutputs(FOutputVertexInterfaceData& InVertexData) override
+		{
+			using namespace ArrayNodeVertexNames;
+
+			InVertexData.BindReadVertex(METASOUND_GET_PARAM_NAME(OutputLastIndex), LastIndex);
+		}
+
+		void Execute()
+		{
+			*LastIndex = Array->Num() - 1;
+		}
+
+		void Reset(const IOperator::FResetParams& InParams)
+		{
+			Execute();
+		}
+
+	private:
+
+		FArrayDataReadReference Array;
+		TDataWriteReference<int32> LastIndex;
+	};
+
+	template<typename ArrayType>
+	class TArrayLastIndexNode : public FNodeFacade
+	{
+	public:
+		TArrayLastIndexNode(const FNodeInitData& InInitData)
+			: FNodeFacade(InInitData.InstanceName, InInitData.InstanceID, TFacadeOperatorClass<TArrayLastIndexOperator<ArrayType>>())
+		{
+		}
+
+		virtual ~TArrayLastIndexNode() = default;
+	};
 } // namespace Metasound
 #undef LOCTEXT_NAMESPACE
diff --git a/Engine/Plugins/Runtime/Metasound/Source/MetasoundFrontend/Public/MetasoundArrayNodesRegistration.h b/Engine/Plugins/Runtime/Metasound/Source/MetasoundFrontend/Public/MetasoundArrayNodesRegistration.h
--- a/Engine/Plugins/Runtime/Metasound/Source/MetasoundFrontend/Public/MetasoundArrayNodesRegistration.h
+++ b/Engine/Plugins/Runtime/Metasound/Source/MetasoundFrontend/Public/MetasoundArrayNodesRegistration.h
@@ -165,6 +165,19 @@ namespace Metasound
 			// No op if not supported
 			return true;
 		}
+
+		template<typename ArrayType, typename std::enable_if<TArrayNodeSupport<ArrayType>::bIsArrayNumSupported, bool>::type = true>
+		bool RegisterArrayLastIndexNode()
+		{
+			return ensureAlways(RegisterNodeWithFrontend<Metasound::TArrayLastIndexNode<ArrayType>>());
+		}
+
+		template<typename ArrayType, typename std::enable_if<!TArrayNodeSupport<ArrayType>::bIsArrayNumSupported, bool>::type = true>
+		bool RegisterArrayLastIndexNode()
+		{
+			// No op if not supported
+			return true;
+		}
 	}
 
 	/** Registers all available array nodes which can be instantiated for the given
@@ -183,6 +196,7 @@ namespace Metasound
 		bSuccess = bSuccess && RegisterArrayConcatNode<ArrayType>();
 		bSuccess = bSuccess && RegisterArrayShuffleNode<ArrayType>();
 		bSuccess = bSuccess && RegisterArrayRandomGetNode<ArrayType>();
+		bSuccess = bSuccess && RegisterArrayLastIndexNode<ArrayType>();
 		return bSuccess;
 	}
 }
