Version Metasound Document to include all ed data and make all Metasound EdGraph data transient
- Add input template nodes
- Add comment node data to document
- Keep references to member literal data (i.e. knob/slider ranges) in document metadata to ensure continued serialization and flexibility to add more editor-only fields and literal metadata
- Misc builder API updates, bug fixes and migration of controllers to builder API in anticipation of pages
- Sunset non-deterministic guid cvar
#tests Standard Automated Audio Tests, EngineTests, Offline QA Smoke pass, CPR, etc., extensive MetaSound Editor use, -game MetaSound qa levels, AudioUnitTests

diff --git a/Engine/Plugins/Runtime/Metasound/Source/MetasoundEditor/Private/MetasoundAssetDefinitions.cpp b/Engine/Plugins/Runtime/Metasound/Source/MetasoundEditor/Private/MetasoundAssetDefinitions.cpp
--- a/Engine/Plugins/Runtime/Metasound/Source/MetasoundEditor/Private/MetasoundAssetDefinitions.cpp
+++ b/Engine/Plugins/Runtime/Metasound/Source/MetasoundEditor/Private/MetasoundAssetDefinitions.cpp
@@ -52,7 +52,7 @@ static bool GetIsPreset(const FSoftObjectPath& InSourcePath)
 		const FMetasoundAssetBase* MetaSoundAsset = IMetasoundUObjectRegistry::Get().GetObjectAsAssetBase(Object);
 		if (MetaSoundAsset)
 		{
-			bIsPreset = MetaSoundAsset->GetDocumentChecked().RootGraph.PresetOptions.bIsPreset;
+			bIsPreset = MetaSoundAsset->GetConstDocumentChecked().RootGraph.PresetOptions.bIsPreset;
 		}
 	}
 	// Otherwise, try to pull from asset registry, but avoid load as this call
diff --git a/Engine/Plugins/Runtime/Metasound/Source/MetasoundEditor/Private/MetasoundDetailCustomization.cpp b/Engine/Plugins/Runtime/Metasound/Source/MetasoundEditor/Private/MetasoundDetailCustomization.cpp
--- a/Engine/Plugins/Runtime/Metasound/Source/MetasoundEditor/Private/MetasoundDetailCustomization.cpp
+++ b/Engine/Plugins/Runtime/Metasound/Source/MetasoundEditor/Private/MetasoundDetailCustomization.cpp
@@ -556,7 +556,7 @@ namespace Metasound
 					return false;
 				};
 
-				const TSet<FMetasoundFrontendVersion>& ImplementedInterfaces = MetaSoundAsset->GetDocumentChecked().Interfaces;
+				const TSet<FMetasoundFrontendVersion>& ImplementedInterfaces = MetaSoundAsset->GetConstDocumentChecked().Interfaces;
 				Algo::TransformIf(ImplementedInterfaces, ImplementedInterfaceNames, CanAddOrRemoveInterface, GetVersionName);
 
 				TArray<FMetasoundFrontendInterface> Interfaces = Frontend::ISearchEngine::Get().FindAllInterfaces();
diff --git a/Engine/Plugins/Runtime/Metasound/Source/MetasoundEditor/Private/MetasoundEditor.cpp b/Engine/Plugins/Runtime/Metasound/Source/MetasoundEditor/Private/MetasoundEditor.cpp
--- a/Engine/Plugins/Runtime/Metasound/Source/MetasoundEditor/Private/MetasoundEditor.cpp
+++ b/Engine/Plugins/Runtime/Metasound/Source/MetasoundEditor/Private/MetasoundEditor.cpp
@@ -38,8 +38,10 @@
 #include "MetasoundBuilderSubsystem.h"
 #include "MetasoundDocumentInterface.h"
 #include "MetasoundEditorCommands.h"
+#include "MetasoundEditorDocumentClipboardUtils.h"
 #include "MetasoundEditorGraph.h"
 #include "MetasoundEditorGraphBuilder.h"
+#include "MetasoundEditorGraphCommentNode.h"
 #include "MetasoundEditorGraphInputNode.h"
 #include "MetasoundEditorGraphSchema.h"
 #include "MetasoundEditorGraphValidation.h"
@@ -134,14 +136,6 @@ namespace Metasound
 			{
 			}
 
-			FMetasoundAssetBase& GetMetasoundAssetChecked() const
-			{
-				UObject* Object = CastChecked<UMetasoundEditorGraph>(Graph)->GetMetasound();
-				FMetasoundAssetBase* MetasoundAsset = IMetasoundUObjectRegistry::Get().GetObjectAsAssetBase(Object);
-				check(MetasoundAsset);
-				return *MetasoundAsset;
-			}
-
 			UMetasoundEditorGraphMember* GetGraphMember() const
 			{
 				UMetasoundEditorGraph* MetasoundGraph = CastChecked<UMetasoundEditorGraph>(Graph);
@@ -157,9 +151,13 @@ namespace Metasound
 				return NAME_None;
 			}
 
-			Frontend::FGraphHandle GetGraphHandle() const
+			Frontend::FConstGraphHandle GetConstGraphHandle() const
 			{
-				return GetMetasoundAssetChecked().GetRootGraphHandle();
+				const UObject* Object = CastChecked<const UMetasoundEditorGraph>(Graph)->GetMetasound();
+				const FMetasoundAssetBase* MetasoundAsset = IMetasoundUObjectRegistry::Get().GetObjectAsAssetBase(Object);
+				check(MetasoundAsset);
+
+				return MetasoundAsset->GetRootGraphHandle();
 			}
 
 			// FEdGraphSchemaAction interface
@@ -175,7 +173,7 @@ namespace Metasound
 
 			virtual int32 GetReorderIndexInContainer() const override
 			{
-				TArray<Frontend::FConstNodeHandle> InputHandles = GetGraphHandle()->GetConstInputNodes();
+				TArray<Frontend::FConstNodeHandle> InputHandles = GetConstGraphHandle()->GetConstInputNodes();
 				return InputHandles.IndexOfByPredicate([this](const Frontend::FConstNodeHandle& NodeHandle)
 				{
 					return NodeHandle->GetID() == MemberID;
@@ -227,9 +225,14 @@ namespace Metasound
 					MetasoundGraph->Modify();
 					Input->Modify();
 
-					if (UMetasoundEditorGraphInputNode* NewGraphNode = FGraphBuilder::AddInputNode(ParentMetasound, Input->GetNodeHandle(), InGraphPosition))
+					FNodeHandle InputHandle = Input->GetNodeHandle();
+					FConstNodeHandle TemplateNodeHandle = FGraphBuilder::AddInputTemplateNodeHandle(ParentMetasound, InputHandle);
+					if (UMetasoundEditorGraphNode* NewGraphNode = FGraphBuilder::AddInputNode(ParentMetasound, TemplateNodeHandle))
 					{
 						NewGraphNode->Modify();
+						NewGraphNode->UpdateFrontendNodeLocation(InGraphPosition);
+						NewGraphNode->SyncLocationFromFrontendNode();
+
 						FGraphBuilder::RegisterGraphWithFrontend(ParentMetasound);
 						TSharedPtr<FEditor> MetasoundEditor = FGraphBuilder::GetEditorForGraph(*MetasoundGraph);
 						if (MetasoundEditor.IsValid())
@@ -250,10 +253,12 @@ namespace Metasound
 						MetasoundGraph->Modify();
 						Output->Modify();
 
-						FNodeHandle OutputHandle = Output->GetNodeHandle();
-						if (UMetasoundEditorGraphOutputNode* NewGraphNode = FGraphBuilder::AddOutputNode(ParentMetasound, OutputHandle, InGraphPosition))
+						FConstNodeHandle OutputHandle = Output->GetConstNodeHandle();
+						if (UMetasoundEditorGraphOutputNode* NewGraphNode = FGraphBuilder::AddOutputNode(ParentMetasound, OutputHandle))
 						{
 							NewGraphNode->Modify();
+							NewGraphNode->UpdateFrontendNodeLocation(InGraphPosition);
+							NewGraphNode->SyncLocationFromFrontendNode();
 							FGraphBuilder::RegisterGraphWithFrontend(ParentMetasound);
 							TSharedPtr<FEditor> MetasoundEditor = FGraphBuilder::GetEditorForGraph(*MetasoundGraph);
 							if (MetasoundEditor.IsValid())
@@ -343,10 +348,12 @@ namespace Metasound
 					}
 
 					const FNodeClassName ClassName = VariableClass.Metadata.GetClassName().ToNodeClassName();
-					FNodeHandle NodeHandle = FGraphBuilder::AddVariableNodeHandle(ParentMetasound, Variable->GetVariableID(), ClassName);
-					if (UMetasoundEditorGraphVariableNode* NewGraphNode = FGraphBuilder::AddVariableNode(ParentMetasound, NodeHandle, InGraphPosition))
+					FConstNodeHandle NodeHandle = FGraphBuilder::AddVariableNodeHandle(ParentMetasound, Variable->GetVariableID(), ClassName);
+					if (UMetasoundEditorGraphVariableNode* NewGraphNode = FGraphBuilder::AddVariableNode(ParentMetasound, NodeHandle))
 					{
 						NewGraphNode->Modify();
+						NewGraphNode->UpdateFrontendNodeLocation(InGraphPosition);
+						NewGraphNode->SyncLocationFromFrontendNode();
 						FGraphBuilder::RegisterGraphWithFrontend(ParentMetasound);
 						TSharedPtr<FEditor> MetasoundEditor = FGraphBuilder::GetEditorForGraph(*MetasoundGraph);
 						if (MetasoundEditor.IsValid())
@@ -838,6 +845,19 @@ namespace Metasound
 			Metasound = ObjectToEdit;
 			Metasound->SetFlags(RF_Transactional);
 
+			// Typically sounds are versioned on load of the asset. There are certain instances where an asset is not versioned on reload.
+			// This is here as a hack around that fact to force versioning the document on load prior to the editor synchronizing and building
+			// the editor graph if an asset is reloaded while the asset editor was open.
+			if (FMetasoundAssetBase* MetaSoundAsset = IMetasoundUObjectRegistry::Get().GetObjectAsAssetBase(Metasound))
+			{
+				if (MetaSoundAsset->VersionAsset())
+				{
+					MetaSoundAsset->SetVersionedOnLoad();
+				}
+			}
+
+			Builder.Reset(&UMetaSoundBuilderSubsystem::GetChecked().AttachBuilderToAssetChecked(*ObjectToEdit));
+
 			GEditor->RegisterForUndo(this);
 
 			FGraphEditorCommands::Register();
@@ -1031,7 +1051,7 @@ namespace Metasound
 
 				const FMetasoundAssetBase* MetaSoundAsset = IMetasoundUObjectRegistry::Get().GetObjectAsAssetBase(Metasound);
 				check(MetaSoundAsset);
-				if (MetaSoundAsset->GetDocumentChecked().RootGraph.PresetOptions.bIsPreset)
+				if (MetaSoundAsset->GetConstDocumentChecked().RootGraph.PresetOptions.bIsPreset)
 				{
 					IconName += TEXT(".Preset");
 				}
@@ -1119,14 +1139,14 @@ namespace Metasound
 			if (MetasoundGraphEditor.IsValid())
 			{
 				UMetasoundEditorGraph& MetaSoundGraph = GetMetaSoundGraphChecked();
-				const bool bVersionedOnLoad = MetaSoundGraph.GetVersionedOnLoad();
-				if (bVersionedOnLoad)
+				FMetasoundAssetBase* MetaSoundAsset = IMetasoundUObjectRegistry::Get().GetObjectAsAssetBase(Metasound);
+				check(MetaSoundAsset);
+
+				if (MetaSoundAsset->GetVersionedOnLoad())
 				{
-					MetaSoundGraph.ClearVersionedOnLoad();
-					const FMetasoundAssetBase* MetaSoundAsset = IMetasoundUObjectRegistry::Get().GetObjectAsAssetBase(Metasound);
-					check(MetaSoundAsset);
+					MetaSoundAsset->ClearVersionedOnLoad();
 
-					const FString VersionString = MetaSoundAsset->GetDocumentChecked().Metadata.Version.Number.ToString();
+					const FString VersionString = MetaSoundAsset->GetConstDocumentChecked().Metadata.Version.Number.ToString();
 					FText Msg = FText::Format(LOCTEXT("MetaSoundDocumentVersioned", "Document versioned to '{0}' on load."), FText::FromString(VersionString));
 					FNotificationInfo Info(Msg);
 					Info.bFireAndForget = true;
@@ -1963,15 +1983,9 @@ namespace Metasound
 
 			if (Metasound)
 			{
-				FMetasoundAssetBase* MetasoundAsset = IMetasoundUObjectRegistry::Get().GetObjectAsAssetBase(Metasound);
-				check(MetasoundAsset);
-
-				const UClass* MetaSoundClass = Metasound->GetClass();
-				check(MetaSoundClass);
-
-				TScriptInterface<IMetaSoundDocumentInterface> DocInterface = MetasoundAsset->GetOwningAsset();
-				FMetaSoundFrontendDocumentBuilder Builder(DocInterface);
-				Builder.ConvertFromPreset();
+				EMetaSoundBuilderResult Result = EMetaSoundBuilderResult::Failed;
+				Builder->ConvertFromPreset(Result);
+				ensure(Result == EMetaSoundBuilderResult::Succeeded);
 
 				// Hack until toolbar is polished up & corner text properly dynamically updates
 				if (UAssetEditorSubsystem* AssetEditorSubsystem = GEditor->GetEditorSubsystem<UAssetEditorSubsystem>())
@@ -2069,12 +2083,11 @@ namespace Metasound
 					FExecuteAction::CreateLambda([this]() { MetasoundGraphEditor->SelectAllNodes(); }));
 
 				GraphEditorCommands->MapAction(FGenericCommands::Get().Copy,
-					FExecuteAction::CreateSP(this, &FEditor::CopySelectedNodes),
-					FCanExecuteAction::CreateSP(this, &FEditor::CanCopyNodes));
+					FExecuteAction::CreateSP(this, &FEditor::CopySelectedNodes));
 
 				GraphEditorCommands->MapAction(FGenericCommands::Get().Cut,
 					FExecuteAction::CreateSP(this, &FEditor::CutSelectedNodes),
-					FCanExecuteAction::CreateLambda([this]() { return CanCopyNodes() && CanDeleteNodes(); }));
+					FCanExecuteAction::CreateLambda([this]() { return CanDeleteNodes(); }));
 
 				GraphEditorCommands->MapAction(FGenericCommands::Get().Paste,
 					FExecuteAction::CreateLambda([this]() { PasteNodes(); }),
@@ -2229,7 +2242,32 @@ namespace Metasound
 					Metasound->Modify();
 					Graph.Modify();
 					GraphMember->Modify();
-					Graph.RemoveMember(*GraphMember);
+
+					const bool bRemovedMetadata = Builder->ClearMemberMetadata(GraphMember->GetMemberID());
+					if (bRemovedMetadata)
+					{
+						const FName MemberName = GraphMember->GetMemberName();
+						EMetaSoundBuilderResult Result = EMetaSoundBuilderResult::Failed;
+						if (GraphMember->IsA<UMetasoundEditorGraphInput>())
+						{
+							Builder->RemoveGraphInput(MemberName, Result);
+							ensure(Result == EMetaSoundBuilderResult::Succeeded);
+						}
+						else if (GraphMember->IsA<UMetasoundEditorGraphOutput>())
+						{
+							Builder->RemoveGraphOutput(MemberName, Result);
+							ensure(Result == EMetaSoundBuilderResult::Succeeded);
+						}
+						// TODO: Move to builder API
+						else if (UMetasoundEditorGraphVariable* Variable = Cast<UMetasoundEditorGraphVariable>(GraphMember))
+						{
+							const FGuid VariableID = Variable->GetVariableID();
+							if (VariableID.IsValid())
+							{
+								ensure(Graph.GetGraphHandle()->RemoveVariable(VariableID));
+							}
+						}
+					}
 				}
 
 				if (NextToSelect)
@@ -2240,12 +2278,6 @@ namespace Metasound
 						SetSelection(GraphMembersToSelect);
 					}
 				}
-				else
-				{
-					FMetasoundAssetBase* MetasoundAsset = IMetasoundUObjectRegistry::Get().GetObjectAsAssetBase(Metasound);
-					check(MetasoundAsset);
-					MetasoundAsset->GetModifyContext().AddMemberIDsModified({ MemberID });
-				}
 			}
 
 			FGraphBuilder::RegisterGraphWithFrontend(*Metasound);
@@ -2285,13 +2317,13 @@ namespace Metasound
 			Graph->Modify();
 			for (UObject* NodeObj : SelectedNodes)
 			{
-				// Some nodes may not be metasound nodes (ex. comments and perhaps aliases eventually), but can be safely deleted.
-				if (UEdGraphNode* Node = Cast<UEdGraphNode>(NodeObj))
+				if (UMetasoundEditorGraphNode* Node = Cast<UMetasoundEditorGraphNode>(NodeObj))
 				{
-					if (!FGraphBuilder::DeleteNode(*Node))
-					{
-						MetasoundGraphEditor->SetNodeSelection(Node, true /* bSelect */);
-					}
+					Node->RemoveFromDocument();
+				}
+				else if (UMetasoundEditorGraphCommentNode* CommentNode = Cast<UMetasoundEditorGraphCommentNode>(NodeObj))
+				{
+					CommentNode->RemoveFromDocument();
 				}
 			}
 		}
@@ -2389,27 +2421,25 @@ namespace Metasound
 			}
 		}
 
-		void FEditor::CopySelectedNodes() const
+		void FEditor::ExportNodesToText(FString& OutText) const
 		{
-			FString NodeString;
-			const FGraphPanelSelectionSet SelectedNodes = MetasoundGraphEditor->GetSelectedNodes();
-			FEdGraphUtilities::ExportNodesToText(SelectedNodes, NodeString);
-			FPlatformApplicationMisc::ClipboardCopy(*NodeString);
-		}
-
-		bool FEditor::CanCopyNodes() const
-		{
-			// If any of the nodes can be duplicated then allow copying
 			const FGraphPanelSelectionSet& SelectedNodes = MetasoundGraphEditor->GetSelectedNodes();
-			for (FGraphPanelSelectionSet::TConstIterator SelectedIter(SelectedNodes); SelectedIter; ++SelectedIter)
+			for (UObject* Object : SelectedNodes)
 			{
-				UEdGraphNode* Node = Cast<UEdGraphNode>(*SelectedIter);
-				if (Node && Node->CanDuplicateNode())
+				if (UMetasoundEditorGraphNode* Node = Cast<UMetasoundEditorGraphNode>(Object))
 				{
-					return true;
+					Node->CacheBreadcrumb();
 				}
 			}
-			return false;
+
+			FEdGraphUtilities::ExportNodesToText(SelectedNodes, OutText);
+		}
+
+		void FEditor::CopySelectedNodes() const
+		{
+			FString NodeString;
+			ExportNodesToText(NodeString);
+			FPlatformApplicationMisc::ClipboardCopy(*NodeString);
 		}
 
 		bool FEditor::CanDuplicateNodes() const
@@ -2460,7 +2490,7 @@ namespace Metasound
 				UEdGraphNode* Node = Cast<UEdGraphNode>(*SelectedIter);
 				if (Node && Node->CanUserDeleteNode())
 				{
-					if (Cast<UEdGraphNode_Comment>(Node) || IsGraphEditable())
+					if (IsGraphEditable())
 					{
 						return true;
 					}
@@ -2511,6 +2541,10 @@ namespace Metasound
 							return true;
 						}
 					}
+					else
+					{
+						return true;
+					}
 				}
 			}
 			return false;
@@ -2518,8 +2552,7 @@ namespace Metasound
 
 		void FEditor::DuplicateNodes()
 		{
-			const FGraphPanelSelectionSet SelectedNodes = MetasoundGraphEditor->GetSelectedNodes();
-			FEdGraphUtilities::ExportNodesToText(SelectedNodes, NodeTextToPaste);
+			ExportNodesToText(NodeTextToPaste);
 			PasteNodes(nullptr, LOCTEXT("MetaSoundEditorDuplicate", "Duplicate MetaSound Node(s)"));
 		}
 
@@ -2543,297 +2576,34 @@ namespace Metasound
 				Location = MetasoundGraphEditor->GetPasteLocation();
 			}
 
-			FMetasoundAssetBase* MetasoundAsset = IMetasoundUObjectRegistry::Get().GetObjectAsAssetBase(Metasound);
-			check(MetasoundAsset);
-
-			UMetasoundEditorGraph& Graph = GetMetaSoundGraphChecked();
-
-			const FScopedTransaction Transaction(InTransactionText);
-			Metasound->Modify();
-			Graph.Modify();
+			FDocumentPasteNotifications Notifications;
+			TArray<UEdGraphNode*> PastedNodes = FDocumentClipboardUtils::PasteClipboardString(InTransactionText, NodeTextToPaste, Location, *Metasound, Notifications);
 
 			// Clear the selection set (newly pasted stuff will be selected)
-			MetasoundGraphEditor->ClearSelectionSet();
-
-			TSet<UEdGraphNode*> PastedGraphNodes;
-			FEdGraphUtilities::ImportNodesFromText(&Graph, NodeTextToPaste, PastedGraphNodes);
-
-			NodeTextToPaste.Empty();
-
-			bool bNotifyReferenceLoop = false;
-			bool bNotifyMultipleVariableSetters = false;
-
-			TArray<UEdGraphNode*> NodesToRemove;
-			for (UEdGraphNode* GraphNode : PastedGraphNodes)
+			if (!PastedNodes.IsEmpty())
 			{
-				GraphNode->CreateNewGuid();
-				if (UMetasoundEditorGraphExternalNode* ExternalNode = Cast<UMetasoundEditorGraphExternalNode>(GraphNode))
-				{
-					FMetasoundFrontendClassMetadata LookupMetadata;
-					LookupMetadata.SetClassName(ExternalNode->GetClassName());
-					LookupMetadata.SetType(EMetasoundFrontendClassType::External);
-					const FNodeRegistryKey PastedRegistryKey = FNodeRegistryKey(LookupMetadata);
-					if (const FSoftObjectPath* AssetPath = IMetaSoundAssetManager::GetChecked().FindObjectPathFromKey(PastedRegistryKey))
-					{
-						if (MetasoundAsset->AddingReferenceCausesLoop(*AssetPath))
-						{
-							FMetasoundFrontendClass MetaSoundClass;
-							FMetasoundFrontendRegistryContainer::Get()->FindFrontendClassFromRegistered(PastedRegistryKey, MetaSoundClass);
-							FString FriendlyClassName = MetaSoundClass.Metadata.GetDisplayName().ToString();
-							if (FriendlyClassName.IsEmpty())
-							{
-								FriendlyClassName = MetaSoundClass.Metadata.GetClassName().ToString();
-							}
-							UE_LOG(LogMetaSound, Warning, TEXT("Failed to paste node with class '%s'.  Class would introduce cyclic asset dependency."), *FriendlyClassName);
-							bNotifyReferenceLoop = true;
-							NodesToRemove.Add(GraphNode);
-						}
-						else
-						{
-							FNodeHandle NewHandle = FGraphBuilder::AddNodeHandle(*Metasound, *ExternalNode);
-							if (!NewHandle->IsValid())
-							{
-								NodesToRemove.Add(GraphNode);
-							}
-						}
-					}
-					else
-					{
-						FMetasoundFrontendClass FrontendClass;
-						const FMetasoundFrontendClassName ClassName = ExternalNode->GetClassName();
-						if (const INodeTemplate* NodeTemplate = INodeTemplateRegistry::Get().FindTemplate(ClassName))
-						{
-							bool bIsValid = false;
-							if (!ExternalNode->Pins.IsEmpty())
-							{
-								FNodeTemplateGenerateInterfaceParams Params;
-								for (const UEdGraphPin* Pin : ExternalNode->Pins)
-								{
-									check(Pin);
-
-									FName PinDataType = FGraphBuilder::GetPinDataType(Pin);
-									if (Pin->Direction == EGPD_Input)
-									{
-										Params.InputsToConnect.Add(PinDataType);
-									}
-									else // EGPD_Output
-									{
-										Params.OutputsToConnect.Add(PinDataType);
-									}
-								}
-
-								FNodeHandle NodeHandle = MetasoundAsset->GetRootGraphHandle()->AddTemplateNode(*NodeTemplate, MoveTemp(Params));
-								FGraphBuilder::InitGraphNodeIDFromNodeHandle(NodeHandle, ExternalNode);
-
-								bIsValid = NodeHandle->IsValid();
-							}
-
-							if (!bIsValid)
-							{
-								NodesToRemove.Add(GraphNode);
-							}
-						}
-						else
-						{
-							FMetasoundFrontendClass ExternalClass;
-							if (ISearchEngine::Get().FindClassWithHighestVersion(ClassName, ExternalClass))
-							{
-								FNodeHandle NewHandle = FGraphBuilder::AddNodeHandle(*Metasound, *ExternalNode);
-								if (!NewHandle->IsValid())
-								{
-									NodesToRemove.Add(GraphNode);
-								}
-							}
-							else
-							{
-								NodesToRemove.Add(GraphNode);
-								UE_LOG(LogMetaSound, Warning, TEXT("Cannot add pasted node with class '%s': Node class not found"), *ClassName.ToString());
-							}
-						}
-					}
-				}
-				else if (UMetasoundEditorGraphInputNode* InputNode = Cast<UMetasoundEditorGraphInputNode>(GraphNode))
-				{
-					if (!InputNode->Input || !Graph.ContainsInput(*InputNode->Input))
-					{
-						NodesToRemove.Add(GraphNode);
-					}
-				}
-				else if (UMetasoundEditorGraphOutputNode* OutputNode = Cast<UMetasoundEditorGraphOutputNode>(GraphNode))
-				{
-					if (OutputNode->Output && Graph.ContainsOutput(*OutputNode->Output))
-					{
-						auto NodeMatches = [OutputNodeID = OutputNode->GetNodeID()](const TObjectPtr<UEdGraphNode>& EdNode)
-						{
-							if (UMetasoundEditorGraphOutputNode* OutputNode = Cast<UMetasoundEditorGraphOutputNode>(EdNode))
-							{
-								return OutputNodeID == OutputNode->GetNodeID();
-							}
-							return false;
-						};
-
-						// Can only have one output reference node
-						if (Graph.Nodes.ContainsByPredicate(NodeMatches))
-						{
-							NodesToRemove.Add(GraphNode);
-						}
-					}
-					else
-					{
-						NodesToRemove.Add(GraphNode);
-					}
-				}
-				else if (UMetasoundEditorGraphVariableNode* VariableNode = Cast<UMetasoundEditorGraphVariableNode>(GraphNode))
-				{
-					// Can only have one setter node
-					if (const UMetasoundEditorGraphVariable* Variable = VariableNode->Variable)
-					{
-						if (Graph.ContainsVariable(*Variable))
-						{
-							FConstVariableHandle VariableHandle = Variable->GetConstVariableHandle();
-							if (VariableHandle->IsValid())
-							{
-								FConstNodeHandle VariableMutatorNodeHandle = VariableHandle->FindMutatorNode();
-								if (VariableNode->GetNodeID() == VariableMutatorNodeHandle->GetID())
-								{
-									bNotifyMultipleVariableSetters = true;
-									NodesToRemove.Add(GraphNode);
-								}
-								else
-								{
-									// Add new variable node
-									const FNodeClassName NodeClassName = VariableNode->GetClassName().ToNodeClassName();
-									FNodeHandle NodeHandle = FGraphBuilder::AddVariableNodeHandle(*Metasound, Variable->GetVariableID(), NodeClassName, VariableNode);
-									if (!NodeHandle->IsValid())
-									{
-										NodesToRemove.Add(GraphNode);
-									}
-								}
-							}
-						}
-						else
-						{
-							NodesToRemove.Add(GraphNode);
-						}
-					}
-					else
-					{
-						NodesToRemove.Add(GraphNode);
-					}
-				}
-				else if (!GraphNode->IsA<UEdGraphNode_Comment>())
-				{
-					checkNoEntry();
-				}
-			}
-
-			// Remove nodes failed to import before attempting to connect/place
-			// in frontend graph.
-			for (UEdGraphNode* Node : NodesToRemove)
-			{
-				Graph.RemoveNode(Node);
-				PastedGraphNodes.Remove(Node);
-			}
-
-			// Find average midpoint of nodes and offset subgraph accordingly
-			FVector2D AvgNodePosition = FVector2D::ZeroVector;
-			for (UEdGraphNode* Node : PastedGraphNodes)
-			{
-				AvgNodePosition.X += Node->NodePosX;
-				AvgNodePosition.Y += Node->NodePosY;
-			}
-
-			if (!PastedGraphNodes.IsEmpty())
-			{
-				float InvNumNodes = 1.0f / PastedGraphNodes.Num();
-				AvgNodePosition.X *= InvNumNodes;
-				AvgNodePosition.Y *= InvNumNodes;
-			}
-
-			for (UEdGraphNode* GraphNode : PastedGraphNodes)
-			{
-				GraphNode->NodePosX = (GraphNode->NodePosX - AvgNodePosition.X) + Location.X;
-				GraphNode->NodePosY = (GraphNode->NodePosY - AvgNodePosition.Y) + Location.Y;
+				MetasoundGraphEditor->ClearSelectionSet();
 
-				GraphNode->SnapToGrid(SNodePanel::GetSnapGridSize());
-				if (UMetasoundEditorGraphNode* MetasoundGraphNode = Cast<UMetasoundEditorGraphNode>(GraphNode))
+				// Select the newly pasted stuff
+				for (UEdGraphNode* GraphNode : PastedNodes)
 				{
-					FNodeHandle NodeHandle = MetasoundGraphNode->GetNodeHandle();
-					if (ensure(NodeHandle->IsValid()))
-					{
-						const FVector2D NewNodeLocation = FVector2D(GraphNode->NodePosX, GraphNode->NodePosY);
-						FMetasoundFrontendNodeStyle NodeStyle = NodeHandle->GetNodeStyle();
-						NodeStyle.Display.Locations.FindOrAdd(MetasoundGraphNode->NodeGuid) = NewNodeLocation;
-						NodeHandle->SetNodeStyle(NodeStyle);
-					}
+					MetasoundGraphEditor->SetNodeSelection(GraphNode, true);
 				}
-			}
 
-			for (UEdGraphNode* GraphNode : PastedGraphNodes)
-			{
-				UMetasoundEditorGraphNode* MetasoundNode = Cast<UMetasoundEditorGraphNode>(GraphNode);
-				if (!MetasoundNode)
+				if (Notifications.bPastedNodesCreateLoop)
 				{
-					continue;
+					NotifyNodePasteFailure_ReferenceLoop();
 				}
 
-				for (UEdGraphPin* Pin : GraphNode->Pins)
+				if (Notifications.bPastedNodesAddMultipleVariableSetters)
 				{
-					if (Pin->Direction != EGPD_Input)
-					{
-						continue;
-					}
-
-					FInputHandle InputHandle = FGraphBuilder::GetInputHandleFromPin(Pin);
-					if (InputHandle->IsValid() && InputHandle->GetDataType() != GetMetasoundDataTypeName<FTrigger>())
-					{
-						FMetasoundFrontendLiteral LiteralValue;
-						if (FGraphBuilder::GetPinLiteral(*Pin, LiteralValue))
-						{
-							if (const FMetasoundFrontendLiteral* ClassDefault = InputHandle->GetClassDefaultLiteral())
-							{
-								// Check equivalence with class default and don't set if they are equal. Copied node
-								// pin has no information to indicate whether or not the literal was already set.
-								if (!LiteralValue.IsEqual(*ClassDefault))
-								{
-									InputHandle->SetLiteral(LiteralValue);
-								}
-							}
-							else
-							{
-								InputHandle->SetLiteral(LiteralValue);
-							}
-						}
-					}
-
-					for (UEdGraphPin* LinkedPin : Pin->LinkedTo)
-					{
-						if (UMetasoundEditorGraphNode* Node = Cast<UMetasoundEditorGraphNode>(LinkedPin->GetOwningNode()))
-						{
-							FGraphBuilder::ConnectNodes(*Pin, *LinkedPin, false /* bConnectEdPins */);
-						}
-					}
+					NotifyNodePasteFailure_MultipleVariableSetters();
 				}
-			}
-
-			// Select the newly pasted stuff
-			for (UEdGraphNode* GraphNode : PastedGraphNodes)
-			{
-				MetasoundGraphEditor->SetNodeSelection(GraphNode, true);
-			}
 
-
-			if (bNotifyReferenceLoop)
-			{
-				NotifyNodePasteFailure_ReferenceLoop();
-			}
-
-			if (bNotifyMultipleVariableSetters)
-			{
-				NotifyNodePasteFailure_MultipleVariableSetters();
+				MetasoundGraphEditor->NotifyGraphChanged();
 			}
 
-			MetasoundGraphEditor->NotifyGraphChanged();
+			NodeTextToPaste.Empty();
 		}
 
 		bool FEditor::CanRenameSelectedNodes() const
@@ -3073,7 +2843,8 @@ namespace Metasound
 				FDocumentHandle DocumentHandle = Graph.GetDocumentHandle();
 				DocumentHandle->RemoveUnreferencedDependencies();
 				DocumentHandle->SynchronizeDependencyMetadata();
-				Graph.GetModifyContext().SetDocumentModified();
+				FMetasoundFrontendDocumentModifyContext& ModifyContext = FGraphBuilder::GetOutermostMetaSoundChecked(Graph).GetModifyContext();
+				ModifyContext.SetDocumentModified();
 			}
 		}
 
@@ -3450,7 +3221,7 @@ namespace Metasound
 			check(MetasoundAsset);
 
 			// Capture after synchronizing as the modification state may be modified therein
-			const FMetasoundFrontendDocumentModifyContext& ModifyContext = MetasoundAsset->GetModifyContext();
+			const FMetasoundFrontendDocumentModifyContext& ModifyContext = MetasoundAsset->GetConstModifyContext();
 			const bool bForceRefreshViews = ModifyContext.GetForceRefreshViews();
 			const TSet<FName>& InterfacesModified = ModifyContext.GetInterfacesModified();
 			const TSet<FGuid>& MembersModified = ModifyContext.GetMemberIDsModified();
@@ -3582,14 +3353,17 @@ namespace Metasound
 
 		bool FEditor::IsGraphEditable() const
 		{
-			if (!Metasound)
+			using namespace Metasound::Editor;
+
+			if (Metasound)
 			{
-				return false;
+				const FMetasoundAssetBase* MetasoundAsset = IMetasoundUObjectRegistry::Get().GetObjectAsAssetBase(Metasound);
+				check(MetasoundAsset);
+				const FMetasoundFrontendDocument& Document = MetasoundAsset->GetConstDocumentChecked();
+				return Document.RootGraph.Graph.Style.bIsGraphEditable;
 			}
 
-			FMetasoundAssetBase* MetasoundAsset = IMetasoundUObjectRegistry::Get().GetObjectAsAssetBase(Metasound);
-			check(MetasoundAsset);
-			return MetasoundAsset->GetRootGraphHandle()->GetGraphStyle().bIsGraphEditable;
+			return false;
 		}
 
 		void FEditor::ClearSelectionAndSelectNode(UEdGraphNode* Node)
@@ -3733,7 +3507,7 @@ namespace Metasound
 				return FReply::Unhandled();
 			}
 
-			FGraphBuilder::RegisterGraphWithFrontend(*Metasound);
+			FGraphBuilder::RegisterGraphWithFrontend(*Metasound, true);
 
 			if (GraphMembersMenu.IsValid())
 			{
diff --git a/Engine/Plugins/Runtime/Metasound/Source/MetasoundEditor/Private/MetasoundEditor.h b/Engine/Plugins/Runtime/Metasound/Source/MetasoundEditor/Private/MetasoundEditor.h
--- a/Engine/Plugins/Runtime/Metasound/Source/MetasoundEditor/Private/MetasoundEditor.h
+++ b/Engine/Plugins/Runtime/Metasound/Source/MetasoundEditor/Private/MetasoundEditor.h
@@ -299,9 +299,6 @@ namespace Metasound
 			/** Copy the currently selected nodes */
 			void CopySelectedNodes() const;
 
-			/** Whether copying the currently selected node(s) is permissible */
-			bool CanCopyNodes() const;
-
 			/** Whether or not the currently selected node(s) can be duplicated */
 			bool CanDuplicateNodes() const;
 
@@ -345,6 +342,8 @@ namespace Metasound
 
 			void SetPreviewID(uint32 InPreviewID);
 
+			void ExportNodesToText(FString& OutText) const;
+
 			/** FNotifyHook interface */
 			virtual void NotifyPostChange(const FPropertyChangedEvent& PropertyChangedEvent, FProperty* PropertyThatChanged) override;
 
@@ -459,8 +458,12 @@ namespace Metasound
 			/** Command list for this editor */
 			TSharedPtr<FUICommandList> GraphEditorCommands;
 
-			/** The Metasound asset being edited */
-			TObjectPtr<UObject> Metasound = nullptr;
+			/** TODO: Deprecate in favor of just accessing the MetaSound document via
+			  * the builder below */
+			TObjectPtr<UObject> Metasound;
+
+			/** Pointer to builder being actively used to mutate MetaSound asset */
+			TStrongObjectPtr<UMetaSoundBuilderBase> Builder;
 
 			/** Whether or not metasound being edited is valid */
 			bool bPassedValidation = true;
diff --git a/Engine/Plugins/Runtime/Metasound/Source/MetasoundEditor/Private/MetasoundEditorDocumentClipboardUtils.cpp b/Engine/Plugins/Runtime/Metasound/Source/MetasoundEditor/Private/MetasoundEditorDocumentClipboardUtils.cpp
new file mode 100644
--- /dev/null
+++ b/Engine/Plugins/Runtime/Metasound/Source/MetasoundEditor/Private/MetasoundEditorDocumentClipboardUtils.cpp
@@ -0,0 +1,437 @@
+// Copyright Epic Games, Inc. All Rights Reserved.
+
+#include "MetasoundEditorDocumentClipboardUtils.h"
+
+#include "Algo/Transform.h"
+#include "EdGraphUtilities.h"
+#include "Logging/TokenizedMessage.h"
+#include "MetasoundBuilderSubsystem.h"
+#include "MetasoundFrontendSearchEngine.h"
+#include "ScopedTransaction.h"
+
+
+namespace Metasound::Editor
+{
+	void FDocumentClipboardUtils::ProcessPastedInputNodes(FMetasoundAssetBase& OutAsset, TArray<UMetasoundEditorGraphNode*>& OutPastedNodes)
+	{
+		using namespace Frontend;
+
+		TMap<FName, TObjectPtr<UMetasoundEditorGraphInput>> MappedGeneratedInputNames;
+		UMetasoundEditorGraph& Graph = *CastChecked<UMetasoundEditorGraph>(&OutAsset.GetGraphChecked());
+		for (int32 Index = OutPastedNodes.Num() - 1; Index >= 0; --Index)
+		{
+			UMetasoundEditorGraphInputNode* InputNode = Cast<UMetasoundEditorGraphInputNode>(OutPastedNodes[Index]);
+			if (!InputNode)
+			{
+				continue;
+			}
+
+			InputNode->CreateNewGuid();
+			TObjectPtr<UMetasoundEditorGraphInput>& Input = InputNode->Input;
+			if (!Input || Input->GetOwningGraph() != &Graph)
+			{
+				Input = nullptr;
+
+				bool bNameMatchFound = false;
+				const FMetasoundEditorGraphVertexNodeBreadcrumb& Breadcrumb = InputNode->GetBreadcrumb();
+				Graph.IterateInputs([&Input, &InputNode, &Graph, &Breadcrumb, &bNameMatchFound](UMetasoundEditorGraphInput& TestInput)
+					{
+						FConstNodeHandle InputHandle = TestInput.GetConstNodeHandle();
+						FConstOutputHandle TestOutput = InputHandle->GetConstOutputs().Last();
+						const bool bTypeMatches = TestOutput->GetDataType() == Breadcrumb.DataType;
+						const bool bAccessMatches = TestOutput->GetVertexAccessType() == Breadcrumb.AccessType;
+						const bool bNameMatches = InputHandle->GetNodeName() == Breadcrumb.MemberName;
+						bNameMatchFound |= bNameMatches;
+						if (bTypeMatches && bAccessMatches && bNameMatches)
+						{
+							Input = &TestInput;
+						}
+					});
+
+				if (!Input)
+				{
+					FDataTypeRegistryInfo Info;
+					if (!Input && IDataTypeRegistry::Get().GetDataTypeInfo(Breadcrumb.DataType, Info))
+					{
+						const FName InputName = Breadcrumb.MemberName;
+						if (TObjectPtr<UMetasoundEditorGraphInput>* InputNodeHandle = MappedGeneratedInputNames.Find(InputName))
+						{
+							Input = *InputNodeHandle;
+						}
+						else
+						{
+							FCreateNodeVertexParams VertexParams;
+							VertexParams.DataType = Breadcrumb.DataType;
+							FNodeHandle NewNodeHandle = FGraphBuilder::AddInputNodeHandle(*OutAsset.GetOwningAsset(), VertexParams, &Breadcrumb.DefaultLiteral, &Breadcrumb.MemberName);
+							MappedGeneratedInputNames.Add(Breadcrumb.MemberName, Input);
+							Input = Graph.FindOrAddInput(NewNodeHandle);
+						}
+					}
+				}
+			}
+
+			if (!Input)
+			{
+				constexpr bool bAllowShrinking = false;
+				Graph.RemoveNode(InputNode);
+				OutPastedNodes.RemoveAtSwap(Index, 1, EAllowShrinking::No);
+			}
+			else
+			{
+				FNodeHandle InputNodeHandle = Input->GetNodeHandle();
+				FConstNodeHandle InputTemplateNode = FGraphBuilder::AddInputTemplateNodeHandle(*OutAsset.GetOwningAsset(), InputNodeHandle);
+				InputNode->NodeID = InputTemplateNode->GetID();
+			}
+		}
+	}
+
+	void FDocumentClipboardUtils::ProcessPastedOutputNodes(FMetasoundAssetBase& OutAsset, TArray<UMetasoundEditorGraphNode*>& OutPastedNodes)
+	{
+		using namespace Frontend;
+
+		UMetasoundEditorGraph& Graph = *CastChecked<UMetasoundEditorGraph>(&OutAsset.GetGraphChecked());
+		for (int32 Index = OutPastedNodes.Num() - 1; Index >= 0; --Index)
+		{
+			UMetasoundEditorGraphOutputNode* OutputNode = Cast<UMetasoundEditorGraphOutputNode>(OutPastedNodes[Index]);
+			if (!OutputNode)
+			{
+				continue;
+			}
+
+			OutputNode->CreateNewGuid();
+
+			// TODO: Add logic to dynamically add output if missing
+			if (OutputNode->Output && Graph.ContainsOutput(*OutputNode->Output))
+			{
+				auto NodeMatches = [OutputNodeID = OutputNode->GetNodeID()](const TObjectPtr<UEdGraphNode>& EdNode)
+				{
+					if (UMetasoundEditorGraphOutputNode* OutputNode = Cast<UMetasoundEditorGraphOutputNode>(EdNode))
+					{
+						return OutputNodeID == OutputNode->GetNodeID();
+					}
+					return false;
+				};
+
+				// Can only have one output reference node
+				if (TObjectPtr<UEdGraphNode>* MatchingOutput = Graph.Nodes.FindByPredicate(NodeMatches))
+				{
+					Graph.RemoveNode(OutputNode);
+					OutPastedNodes.RemoveAtSwap(Index, 1, EAllowShrinking::No);
+				}
+			}
+			else
+			{
+				Graph.RemoveNode(OutputNode);
+				OutPastedNodes.RemoveAtSwap(Index, 1, EAllowShrinking::No);
+			}
+		}
+	}
+
+	void FDocumentClipboardUtils::ProcessPastedVariableNodes(FMetasoundAssetBase& OutAsset, TArray<UMetasoundEditorGraphNode*>& OutPastedNodes, FDocumentPasteNotifications& OutNotifications)
+	{
+		using namespace Frontend;
+
+		OutNotifications.bPastedNodesAddMultipleVariableSetters = false;
+
+		UMetasoundEditorGraph& Graph = *CastChecked<UMetasoundEditorGraph>(&OutAsset.GetGraphChecked());
+		TMap<FName, TObjectPtr<UMetasoundEditorGraphVariable>> MappedGeneratedVariableNames;
+		Graph.IterateVariables([&MappedGeneratedVariableNames](UMetasoundEditorGraphVariable& Variable)
+		{
+			MappedGeneratedVariableNames.Add(Variable.GetMemberName(), &Variable);
+		});
+
+		for (int32 Index = OutPastedNodes.Num() - 1; Index >= 0; --Index)
+		{
+			UMetasoundEditorGraphVariableNode* VariableNode = Cast<UMetasoundEditorGraphVariableNode>(OutPastedNodes[Index]);
+			if (!VariableNode)
+			{
+				continue;
+			}
+
+			VariableNode->CreateNewGuid();
+
+			// Can only have one setter node
+			TObjectPtr<UMetasoundEditorGraphVariable>& Variable = VariableNode->Variable;
+			if (!Variable || Variable->GetOwningGraph() != &Graph)
+			{
+				const FMetasoundEditorGraphMemberNodeBreadcrumb& Breadcrumb = VariableNode->Breadcrumb;
+				const FName BaseName = Breadcrumb.MemberName;
+				TObjectPtr<UMetasoundEditorGraphVariable> CachedVariable = MappedGeneratedVariableNames.FindRef(BaseName);
+				Frontend::FVariableHandle FrontendVariable = IVariableController::GetInvalidHandle();
+				if (CachedVariable)
+				{
+					FrontendVariable = CachedVariable->GetVariableHandle();
+				}
+
+				if (FrontendVariable->IsValid() && FrontendVariable->GetDataType() == Breadcrumb.DataType)
+				{
+					Variable = CachedVariable;
+				}
+				else
+				{
+					FrontendVariable = FGraphBuilder::AddVariableHandle(*OutAsset.GetOwningAsset(), Breadcrumb.DataType);
+					Variable = Graph.FindOrAddVariable(FrontendVariable);
+					check(Variable);
+
+					constexpr bool bPostSubTransaction = false;
+					const FName VariableName = FGraphBuilder::GenerateUniqueNameByClassType(*OutAsset.GetOwningAsset(), EMetasoundFrontendClassType::Variable, BaseName.ToString());
+					Variable->SetMemberName(VariableName, bPostSubTransaction);
+					VariableNode->CacheBreadcrumb(); // Name of referenced variable/node state has changed so make sure up-to-date in case breadcrumb is used later
+					MappedGeneratedVariableNames.Add(BaseName, Variable);
+				}
+			}
+
+			FConstVariableHandle VariableHandle = Variable->GetConstVariableHandle();
+			if (VariableHandle->IsValid())
+			{
+				FConstNodeHandle VariableMutatorNodeHandle = VariableHandle->FindMutatorNode();
+				if (VariableNode->GetNodeID() == VariableMutatorNodeHandle->GetID())
+				{
+					OutNotifications.bPastedNodesAddMultipleVariableSetters = true;
+					OutPastedNodes.RemoveAtSwap(Index, 1, EAllowShrinking::No);
+					Graph.RemoveNode(VariableNode);
+				}
+				else
+				{
+					// Add new variable node
+					const FNodeClassName NodeClassName = VariableNode->GetClassName().ToNodeClassName();
+					FNodeHandle NodeHandle = FGraphBuilder::AddVariableNodeHandle(*OutAsset.GetOwningAsset(), Variable->GetVariableID(), NodeClassName, VariableNode);
+					if (!NodeHandle->IsValid())
+					{
+						OutPastedNodes.RemoveAtSwap(Index, 1, EAllowShrinking::No);
+						Graph.RemoveNode(VariableNode);
+					}
+				}
+			}
+		}
+	}
+
+	void FDocumentClipboardUtils::ProcessPastedExternalNodes(FMetasoundAssetBase& OutAsset, TArray<UMetasoundEditorGraphNode*>& OutPastedNodes, FDocumentPasteNotifications& OutNotifications)
+	{
+		using namespace Frontend;
+
+		OutNotifications.bPastedNodesCreateLoop = false;
+
+		UMetasoundEditorGraph& Graph = *CastChecked<UMetasoundEditorGraph>(&OutAsset.GetGraphChecked());
+		for (int32 Index = OutPastedNodes.Num() - 1; Index >= 0; --Index)
+		{
+			UMetasoundEditorGraphExternalNode* ExternalNode = Cast<UMetasoundEditorGraphExternalNode>(OutPastedNodes[Index]);
+			if (!ExternalNode)
+			{
+				continue;
+			}
+
+			ExternalNode->CreateNewGuid();
+			FMetasoundFrontendClassMetadata LookupMetadata;
+			const FMetasoundEditorGraphNodeBreadcrumb& Breadcrumb = ExternalNode->GetBreadcrumb();
+			LookupMetadata.SetClassName(Breadcrumb.ClassName);
+			LookupMetadata.SetType(EMetasoundFrontendClassType::External);
+			const FNodeRegistryKey PastedRegistryKey(LookupMetadata);
+			UObject& MetaSound = *OutAsset.GetOwningAsset();
+
+			if (const FSoftObjectPath* AssetPath = IMetaSoundAssetManager::GetChecked().FindObjectPathFromKey(PastedRegistryKey))
+			{
+				if (OutAsset.AddingReferenceCausesLoop(*AssetPath))
+				{
+					FMetasoundFrontendClass MetaSoundClass;
+					FMetasoundFrontendRegistryContainer::Get()->FindFrontendClassFromRegistered(PastedRegistryKey, MetaSoundClass);
+					FString FriendlyClassName = MetaSoundClass.Metadata.GetDisplayName().ToString();
+					if (FriendlyClassName.IsEmpty())
+					{
+						FriendlyClassName = MetaSoundClass.Metadata.GetClassName().ToString();
+					}
+					UE_LOG(LogMetaSound, Warning, TEXT("Failed to paste node with class '%s'.  Class would introduce cyclic asset dependency."), *FriendlyClassName);
+					OutNotifications.bPastedNodesCreateLoop = true;
+					OutPastedNodes.RemoveAtSwap(Index, 1, EAllowShrinking::No);
+					Graph.RemoveNode(ExternalNode);
+				}
+				else
+				{
+					FNodeHandle NodeHandle = FGraphBuilder::AddExternalNodeHandle(MetaSound, Breadcrumb.ClassName);
+					ExternalNode->NodeID = NodeHandle->GetID();
+					if (!NodeHandle->IsValid())
+					{
+						OutPastedNodes.RemoveAtSwap(Index, 1, EAllowShrinking::No);
+						Graph.RemoveNode(ExternalNode);
+					}
+				}
+			}
+			else
+			{
+				FMetasoundFrontendClass FrontendClass;
+				if (const INodeTemplate* Template = INodeTemplateRegistry::Get().FindTemplate(Breadcrumb.ClassName))
+				{
+					FNodeHandle NodeHandle = FGraphBuilder::AddExternalNodeHandle(MetaSound, Breadcrumb.ClassName);
+					ExternalNode->NodeID = NodeHandle->GetID();
+					if (!NodeHandle->IsValid())
+					{
+						OutPastedNodes.RemoveAtSwap(Index, 1, EAllowShrinking::No);
+						Graph.RemoveNode(ExternalNode);
+					}
+				}
+				else
+				{
+					FMetasoundFrontendClass ExternalClass;
+					if (ISearchEngine::Get().FindClassWithHighestVersion(Breadcrumb.ClassName, ExternalClass))
+					{
+						FConstNodeHandle NodeHandle = FGraphBuilder::AddExternalNodeHandle(MetaSound, Breadcrumb.ClassName);
+						ExternalNode->NodeID = NodeHandle->GetID();
+						if (!NodeHandle->IsValid())
+						{
+							OutPastedNodes.RemoveAtSwap(Index, 1, EAllowShrinking::No);
+							Graph.RemoveNode(ExternalNode);
+						}
+					}
+					else
+					{
+						OutPastedNodes.RemoveAtSwap(Index, 1, EAllowShrinking::No);
+						Graph.RemoveNode(ExternalNode);
+						UE_LOG(LogMetaSound, Warning, TEXT("Cannot add pasted node with class '%s': Node class not found"), *Breadcrumb.ClassName.ToString());
+					}
+				}
+			}
+		}
+	}
+
+	void FDocumentClipboardUtils::ProcessPastedCommentNodes(FMetasoundAssetBase& OutAsset, const TArrayView<UMetasoundEditorGraphCommentNode*> CommentNodes)
+	{
+		using namespace Frontend;
+
+		UMetasoundEditorGraph& Graph = *CastChecked<UMetasoundEditorGraph>(&OutAsset.GetGraphChecked());
+		UMetaSoundBuilderBase& Builder = UMetaSoundBuilderSubsystem::GetChecked().AttachBuilderToAssetChecked(*OutAsset.GetOwningAsset());
+		for (UMetasoundEditorGraphCommentNode* CommentNode : CommentNodes)
+		{
+			CommentNode->CreateNewGuid();
+			FMetaSoundFrontendGraphComment& NewComment = Builder.FindOrAddGraphComment(CommentNode->GetCommentID());
+			UMetasoundEditorGraphCommentNode::ConvertToFrontendComment(*CommentNode, NewComment);
+		}
+	}
+
+	void FDocumentClipboardUtils::ProcessPastedNodePositions(FMetasoundAssetBase& OutAsset, const FVector2D& InLocation, TArray<UMetasoundEditorGraphNode*>& OutPastedNodes)
+	{
+		using namespace Frontend;
+
+		// Find average midpoint of nodes and offset subgraph accordingly
+		FVector2D AvgNodePosition = FVector2D::ZeroVector;
+		for (UEdGraphNode* Node : OutPastedNodes)
+		{
+			AvgNodePosition.X += Node->NodePosX;
+			AvgNodePosition.Y += Node->NodePosY;
+		}
+
+		if (!OutPastedNodes.IsEmpty())
+		{
+			float InvNumNodes = 1.0f / OutPastedNodes.Num();
+			AvgNodePosition.X *= InvNumNodes;
+			AvgNodePosition.Y *= InvNumNodes;
+		}
+
+		for (UEdGraphNode* GraphNode : OutPastedNodes)
+		{
+			GraphNode->NodePosX = (GraphNode->NodePosX - AvgNodePosition.X) + InLocation.X;
+			GraphNode->NodePosY = (GraphNode->NodePosY - AvgNodePosition.Y) + InLocation.Y;
+
+			GraphNode->SnapToGrid(SNodePanel::GetSnapGridSize());
+			if (UMetasoundEditorGraphNode* MetasoundGraphNode = Cast<UMetasoundEditorGraphNode>(GraphNode))
+			{
+				FNodeHandle NodeHandle = MetasoundGraphNode->GetNodeHandle();
+				if (ensure(NodeHandle->IsValid()))
+				{
+					const FVector2D NewNodeLocation = FVector2D(GraphNode->NodePosX, GraphNode->NodePosY);
+					FMetasoundFrontendNodeStyle NodeStyle = NodeHandle->GetNodeStyle();
+					NodeStyle.Display.Locations.FindOrAdd(MetasoundGraphNode->NodeGuid) = NewNodeLocation;
+					NodeHandle->SetNodeStyle(NodeStyle);
+				}
+			}
+		}
+	}
+
+	void FDocumentClipboardUtils::ProcessPastedNodeConnections(FMetasoundAssetBase& OutAsset, TArray<UMetasoundEditorGraphNode*>& OutPastedNodes)
+	{
+		using namespace Frontend;
+
+		for (UEdGraphNode* GraphNode : OutPastedNodes)
+		{
+			for (UEdGraphPin* Pin : GraphNode->Pins)
+			{
+				if (Pin->Direction != EGPD_Input)
+				{
+					continue;
+				}
+
+				FInputHandle InputHandle = FGraphBuilder::GetInputHandleFromPin(Pin);
+				if (InputHandle->IsValid() && InputHandle->GetDataType() != GetMetasoundDataTypeName<FTrigger>())
+				{
+					FMetasoundFrontendLiteral LiteralValue;
+					if (FGraphBuilder::GetPinLiteral(*Pin, LiteralValue))
+					{
+						if (const FMetasoundFrontendLiteral* ClassDefault = InputHandle->GetClassDefaultLiteral())
+						{
+							// Check equivalence with class default and don't set if they are equal. Copied node
+							// pin has no information to indicate whether or not the literal was already set.
+							if (!LiteralValue.IsEqual(*ClassDefault))
+							{
+								InputHandle->SetLiteral(LiteralValue);
+							}
+						}
+						else
+						{
+							InputHandle->SetLiteral(LiteralValue);
+						}
+					}
+				}
+
+				for (UEdGraphPin* LinkedPin : Pin->LinkedTo)
+				{
+					if (UMetasoundEditorGraphNode* Node = Cast<UMetasoundEditorGraphNode>(LinkedPin->GetOwningNode()))
+					{
+						FGraphBuilder::ConnectNodes(*Pin, *LinkedPin, false /* bConnectEdPins */);
+					}
+				}
+			}
+		}
+	}
+
+	TArray<UEdGraphNode*> FDocumentClipboardUtils::PasteClipboardString(const FText& InTransactionText, const FString& InClipboardString, const FVector2D& InLocation, UObject& OutMetaSound, FDocumentPasteNotifications& OutNotifications)
+	{
+		FMetasoundAssetBase* Asset = IMetasoundUObjectRegistry::Get().GetObjectAsAssetBase(&OutMetaSound);
+		check(Asset);
+
+		TArray<UMetasoundEditorGraphCommentNode*> PastedCommentNodes;
+		TArray<UMetasoundEditorGraphNode*> PastedGraphNodes;
+		{
+			TSet<UEdGraphNode*> PastedNodeSet;
+			FEdGraphUtilities::ImportNodesFromText(Asset->GetGraph(), InClipboardString, PastedNodeSet);
+
+			auto CastToMetaSoundNode = [](UEdGraphNode* Node) { return Cast<UMetasoundEditorGraphNode>(Node); };
+			Algo::TransformIf(PastedNodeSet, PastedGraphNodes, CastToMetaSoundNode, CastToMetaSoundNode);
+
+			auto CastToCommentNode = [](UEdGraphNode* Node) { return Cast<UMetasoundEditorGraphCommentNode>(Node); };
+			Algo::TransformIf(PastedNodeSet, PastedCommentNodes, CastToCommentNode, CastToCommentNode);
+		}
+
+		TArray<UEdGraphNode*> PastedNodes;
+		if (PastedGraphNodes.IsEmpty() && PastedCommentNodes.IsEmpty())
+		{
+			return PastedNodes;
+		}
+
+		const FScopedTransaction Transaction(InTransactionText);
+
+		OutMetaSound.Modify();
+		Asset->GetGraphChecked().Modify();
+
+		ProcessPastedCommentNodes(*Asset, PastedCommentNodes);
+
+		ProcessPastedInputNodes(*Asset, PastedGraphNodes);
+		ProcessPastedOutputNodes(*Asset, PastedGraphNodes);
+		ProcessPastedVariableNodes(*Asset, PastedGraphNodes, OutNotifications);
+		ProcessPastedExternalNodes(*Asset, PastedGraphNodes, OutNotifications);
+		ProcessPastedNodePositions(*Asset, InLocation, PastedGraphNodes);
+		ProcessPastedNodeConnections(*Asset, PastedGraphNodes);
+
+		PastedNodes.Append(MoveTemp(PastedGraphNodes));
+		PastedNodes.Append(MoveTemp(PastedCommentNodes));
+		return PastedNodes;
+	}
+} // namespace Metasound::Editor
diff --git a/Engine/Plugins/Runtime/Metasound/Source/MetasoundEditor/Private/MetasoundEditorDocumentClipboardUtils.h b/Engine/Plugins/Runtime/Metasound/Source/MetasoundEditor/Private/MetasoundEditorDocumentClipboardUtils.h
new file mode 100644
--- /dev/null
+++ b/Engine/Plugins/Runtime/Metasound/Source/MetasoundEditor/Private/MetasoundEditorDocumentClipboardUtils.h
@@ -0,0 +1,43 @@
+// Copyright Epic Games, Inc. All Rights Reserved.
+#pragma once
+
+#include "MetasoundEditorGraph.h"
+#include "MetasoundEditorGraphBuilder.h"
+#include "MetasoundEditorGraphCommentNode.h"
+#include "MetasoundEditorGraphInputNode.h"
+#include "MetasoundEditorGraphNode.h"
+#include "MetasoundFrontend.h"
+#include "MetasoundFrontendController.h"
+#include "MetasoundFrontendDocument.h"
+#include "UObject/NameTypes.h"
+#include "UObject/NoExportTypes.h"
+
+
+namespace Metasound::Editor
+{
+	struct FDocumentPasteNotifications
+	{
+		bool bPastedNodesAddMultipleVariableSetters = false;
+		bool bPastedNodesCreateLoop = false;
+	};
+
+	struct FDocumentClipboardUtils
+	{
+		static TArray<UEdGraphNode*> PasteClipboardString(const FText& InTransactionText, const FString& InClipboardString, const FVector2D& InLocation, UObject& OutMetaSound, FDocumentPasteNotifications& OutNotifications);
+
+	private:
+		static void ProcessPastedCommentNodes(FMetasoundAssetBase& OutAsset, const TArrayView<UMetasoundEditorGraphCommentNode*> CommentNodes);
+
+		static void ProcessPastedExternalNodes(FMetasoundAssetBase& OutAsset, TArray<UMetasoundEditorGraphNode*>& OutPastedNodes, FDocumentPasteNotifications& OutNotifications);
+
+		static void ProcessPastedInputNodes(FMetasoundAssetBase& OutAsset, TArray<UMetasoundEditorGraphNode*>& OutPastedNodes);
+
+		static void ProcessPastedNodePositions(FMetasoundAssetBase& OutAsset, const FVector2D& InLocation, TArray<UMetasoundEditorGraphNode*>& OutPastedNodes);
+
+		static void ProcessPastedNodeConnections(FMetasoundAssetBase& OutAsset, TArray<UMetasoundEditorGraphNode*>& OutPastedNodes);
+
+		static void ProcessPastedOutputNodes(FMetasoundAssetBase& OutAsset, TArray<UMetasoundEditorGraphNode*>& OutPastedNodes);
+
+		static void ProcessPastedVariableNodes(FMetasoundAssetBase& OutAsset, TArray<UMetasoundEditorGraphNode*>& OutPastedNodes, FDocumentPasteNotifications& OutNotifications);
+	};
+} // namespace Metasound::Editor
