Version Metasound Document to include all ed data and make all Metasound EdGraph data transient

diff --git a/Engine/Plugins/Runtime/Metasound/Source/MetasoundEditor/Private/MetasoundEditorGraph.cpp b/Engine/Plugins/Runtime/Metasound/Source/MetasoundEditor/Private/MetasoundEditorGraph.cpp
--- a/Engine/Plugins/Runtime/Metasound/Source/MetasoundEditor/Private/MetasoundEditorGraph.cpp
+++ b/Engine/Plugins/Runtime/Metasound/Source/MetasoundEditor/Private/MetasoundEditorGraph.cpp
@@ -8,14 +8,20 @@
 #include "EdGraph/EdGraphNode.h"
 #include "Interfaces/ITargetPlatform.h"
 #include "MetasoundAssetBase.h"
+#include "MetasoundBuilderSubsystem.h"
+#include "MetasoundEditorGraphCommentNode.h"
 #include "MetasoundEditorGraphInputNode.h"
 #include "MetasoundEditorGraphMemberDefaults.h"
 #include "MetasoundEditorGraphNode.h"
 #include "MetasoundEditorGraphValidation.h"
 #include "MetasoundEditorModule.h"
+#include "MetasoundEditorSubsystem.h"
+#include "MetasoundFrontendDocumentBuilder.h"
+#include "MetasoundFrontendNodeTemplateRegistry.h"
 #include "MetasoundUObjectRegistry.h"
 #include "MetasoundVariableNodes.h"
 #include "MetasoundVertex.h"
+#include "NodeTemplates/MetasoundFrontendNodeTemplateInput.h"
 #include "ScopedTransaction.h"
 #include "UObject/NameTypes.h"
 #include "UObject/UnrealType.h"
@@ -29,75 +35,89 @@
 #define INVALID_PARAMETER_NAME_CHARACTERS TEXT("\"',\n\r\t") METASOUND_ANALYZER_PATH_SEPARATOR
 
 
-namespace Metasound
+namespace Metasound::Editor
 {
-	namespace Editor
+	namespace GraphPrivate
 	{
-		namespace GraphPrivate
+		FName GetUniqueTransientMemberName()
 		{
-			void OnLiteralChanged(UMetasoundEditorGraphMember& InMember, EPropertyChangeType::Type InChangeType)
-			{
-				constexpr bool bPostTransaction = false;
-				InMember.UpdateFrontendDefaultLiteral(bPostTransaction);
+			// Use unique instance ID to avoid copy/paste logic resolving invalid relationship between graphs.
+			// Equality is properly resolved based on associated Frontend node's Name, TypeName, & AccessType.
+			// This can bloat the name table within a editor given session, but ed graph is not serialized so
+			// purely for editing.
+			return *FString::Printf(TEXT("Member_%s"), *FGuid::NewGuid().ToString());
+		}
+
+		void OnLiteralChanged(UMetasoundEditorGraphMember& InMember, EPropertyChangeType::Type InChangeType)
+		{
+			constexpr bool bPostTransaction = false;
+			InMember.UpdateFrontendDefaultLiteral(bPostTransaction);
 
-				const bool bCommitChange = InChangeType != EPropertyChangeType::Interactive;
-				if (bCommitChange)
+			const bool bCommitChange = InChangeType != EPropertyChangeType::Interactive;
+			if (bCommitChange)
+			{
+				if (UObject* MetaSound = InMember.GetOutermostObject())
 				{
-					if (UObject* MetaSound = InMember.GetOutermostObject())
+					FGraphBuilder::RegisterGraphWithFrontend(*MetaSound);
+					if (FMetasoundAssetBase* MetaSoundAsset = IMetasoundUObjectRegistry::Get().GetObjectAsAssetBase(MetaSound))
 					{
-						FGraphBuilder::RegisterGraphWithFrontend(*MetaSound);
-						if (FMetasoundAssetBase* MetaSoundAsset = IMetasoundUObjectRegistry::Get().GetObjectAsAssetBase(MetaSound))
-						{
-							MetaSoundAsset->GetModifyContext().AddMemberIDsModified({ InMember.GetMemberID() });
-						}
+						MetaSoundAsset->GetModifyContext().AddMemberIDsModified({ InMember.GetMemberID() });
 					}
 				}
 			}
+		}
 
-			void SetLiteralOrClearIfMatchesDefault(Frontend::FInputHandle& InInputHandle, const FMetasoundFrontendLiteral& InDefaultLiteral)
-			{
-				using namespace Frontend;
+		void SetLiteralOrClearIfMatchesDefault(Frontend::FInputHandle& InInputHandle, const FMetasoundFrontendLiteral& InDefaultLiteral)
+		{
+			using namespace Frontend;
 
-				// Avoids member literal setting the node literal if its not required (which in turn
-				// avoids 'Reset To Default' action from being enabled when the default is equal)
-				bool bClearLiteral = false;
-				if (const FMetasoundFrontendLiteral* ClassDefault = InInputHandle->GetClassDefaultLiteral())
-				{
-					bClearLiteral = ClassDefault->IsEqual(InDefaultLiteral);
-				}
+			// Avoids member literal setting the node literal if its not required (which in turn
+			// avoids 'Reset To Default' action from being enabled when the default is equal)
+			bool bClearLiteral = false;
+			if (const FMetasoundFrontendLiteral* ClassDefault = InInputHandle->GetClassDefaultLiteral())
+			{
+				bClearLiteral = ClassDefault->IsEqual(InDefaultLiteral);
+			}
 
-				if (!bClearLiteral)
-				{
-					FMetasoundFrontendLiteral DefaultTypeLiteral;
-					DefaultTypeLiteral.SetFromLiteral(IDataTypeRegistry::Get().CreateDefaultLiteral(InInputHandle->GetDataType()));
-					bClearLiteral = InDefaultLiteral.IsEqual(DefaultTypeLiteral);
-				}
+			if (!bClearLiteral)
+			{
+				FMetasoundFrontendLiteral DefaultTypeLiteral;
+				DefaultTypeLiteral.SetFromLiteral(IDataTypeRegistry::Get().CreateDefaultLiteral(InInputHandle->GetDataType()));
+				bClearLiteral = InDefaultLiteral.IsEqual(DefaultTypeLiteral);
+			}
 
-				if (bClearLiteral)
-				{
-					InInputHandle->ClearLiteral();
-				}
-				else
-				{
-					InInputHandle->SetLiteral(InDefaultLiteral);
-				}
+			if (bClearLiteral)
+			{
+				InInputHandle->ClearLiteral();
+			}
+			else
+			{
+				InInputHandle->SetLiteral(InDefaultLiteral);
 			}
-		} // namespace GraphPrivate
-	} // namespace Editor
-} // namespace Metasound
+		}
+	} // namespace GraphPrivate
+} // namespace Metasound::Editor
+
 
 UMetasoundEditorGraph* UMetasoundEditorGraphMember::GetOwningGraph()
 {
+	// Due to a prior document migration that enables ed graphs to be built from the frontend document exclusively,
+	// MetaSound objects may contain more than one editor graph, so must check outer rather than accessing the
+	// transient graph from the FMetaSoundAssetBase layer.
 	return Cast<UMetasoundEditorGraph>(GetOuter());
 }
 
 const UMetasoundEditorGraph* UMetasoundEditorGraphMember::GetOwningGraph() const
 {
+	// Due to a prior document migration that enables ed graphs to be built from the frontend document exclusively,
+	// MetaSound objects may contain more than one editor graph, so must check outer rather than accessing the
+	// transient graph from the FMetaSoundAssetBase layer.
 	return Cast<const UMetasoundEditorGraph>(GetOuter());
 }
 
 void UMetasoundEditorGraphMember::InitializeLiteral()
 {
+	using namespace Metasound;
 	using namespace Metasound::Editor;
 	using namespace Metasound::Frontend;
 
@@ -114,17 +134,18 @@ void UMetasoundEditorGraphMember::InitializeLiteral()
 
 	if (!Literal || Literal->GetClass() != LiteralClass)
 	{
-		UMetasoundEditorGraphMemberDefaultLiteral* NewLiteral = NewObject<UMetasoundEditorGraphMemberDefaultLiteral>(this, LiteralClass, FName(), RF_Transactional);
-
-		// Set default widget type from editor settings for float inputs 
-		if (Cast<UMetasoundEditorGraphInput>(this) && LiteralClass == UMetasoundEditorGraphMemberDefaultFloat::StaticClass())
+		const bool bIsNew = UMetaSoundEditorSubsystem::GetChecked().BindMemberMetadata(*this, LiteralClass);
+		if (bIsNew)
 		{
-			if (const UMetasoundEditorSettings* EditorSettings = ::GetDefault<UMetasoundEditorSettings>())
+			// Set default widget type from editor settings for float inputs 
+			if (IsA<UMetasoundEditorGraphInput>() && LiteralClass == UMetasoundEditorGraphMemberDefaultFloat::StaticClass())
 			{
-				Cast<UMetasoundEditorGraphMemberDefaultFloat>(NewLiteral)->WidgetType = EditorSettings->DefaultInputWidgetType;
+				if (const UMetasoundEditorSettings* EditorSettings = ::GetDefault<UMetasoundEditorSettings>())
+				{
+					Cast<UMetasoundEditorGraphMemberDefaultFloat>(Literal)->WidgetType = EditorSettings->DefaultInputWidgetType;
+				}
 			}
 		}
-		Literal = NewLiteral;
 	}
 }
 
@@ -225,7 +246,7 @@ void UMetasoundEditorGraphVertex::SetDescription(const FText& InDescription, boo
 			NodeHandle->SetDescription(InDescription);
 
 			Graph->RegisterGraphWithFrontend();
-			Graph->GetModifyContext().AddMemberIDsModified({ GetMemberID() });
+			FGraphBuilder::GetOutermostMetaSoundChecked(*Graph).GetModifyContext().AddMemberIDsModified({ GetMemberID() });
 		}
 	}
 }
@@ -283,7 +304,7 @@ void UMetasoundEditorGraphVertex::SetMemberName(const FName& InNewName, bool bPo
 	}
 
 	Graph->RegisterGraphWithFrontend();
-	Graph->GetModifyContext().AddMemberIDsModified({ GetMemberID() });
+	FGraphBuilder::GetOutermostMetaSoundChecked(*Graph).GetModifyContext().AddMemberIDsModified({ GetMemberID() });
 }
 
 FText UMetasoundEditorGraphVertex::GetDisplayName() const
@@ -294,9 +315,10 @@ FText UMetasoundEditorGraphVertex::GetDisplayName() const
 
 void UMetasoundEditorGraphVertex::SetDisplayName(const FText& InNewName, bool bPostTransaction)
 {
+	using namespace Metasound::Editor;
 	using namespace Metasound::Frontend;
 
-	UMetasoundEditorGraph* Graph = Cast<UMetasoundEditorGraph>(GetOuter());
+	UMetasoundEditorGraph* Graph = GetOwningGraph();
 	if (!ensure(Graph))
 	{
 		return;
@@ -330,206 +352,17 @@ void UMetasoundEditorGraphVertex::SetDisplayName(const FText& InNewName, bool bP
 	}
 
 	Graph->RegisterGraphWithFrontend();
-	Graph->GetModifyContext().AddMemberIDsModified({ GetMemberID() });
-}
-
-void UMetasoundEditorGraphVertex::SetDataType(FName InNewType, bool bPostTransaction)
-{
-	using namespace Metasound::Editor;
-	using namespace Metasound::Frontend;
-
-	if (InNewType == GetDataType())
-	{
-		return;
-	}
-
-	UMetasoundEditorGraph* Graph = GetOwningGraph();
-	if (!ensure(Graph))
-	{
-		return;
-	}
-
-	const FScopedTransaction Transaction(LOCTEXT("SetGraphVertexType", "Set MetaSound GraphVertex Type"), bPostTransaction);
-	Graph->GetMetasoundChecked().Modify();
-	Graph->Modify();
-	Modify();
-
-	// 1. Cache current editor input node reference positions & delete nodes.
-	TArray<UMetasoundEditorGraphMemberNode*> InputNodes = GetNodes();
-	TArray<FVector2D> NodeLocations;
-	for (UMetasoundEditorGraphMemberNode* Node : InputNodes)
-	{
-		if (ensure(Node))
-		{
-			NodeLocations.Add(FVector2D(Node->NodePosX, Node->NodePosY));
-		}
-	}
-
-	// 2. Cache the old version's Frontend data.
-	FNodeHandle NodeHandle = GetNodeHandle();
-	const FName NodeName = NodeHandle->GetNodeName();
-	const FText NodeDisplayName = NodeHandle->GetDisplayName();
-	const FText NodeDescription = NodeHandle->GetDescription();
-	const int32 SortOrderIndex = GetSortOrderIndex();
-
-	// Remove the current nodes and vertex
-	Graph->RemoveMemberNodes(*this);
-	Graph->RemoveFrontendMember(*this);
-
-	// 4. Add the new input node with the same identifier data but new datatype.
-	UObject& Metasound = Graph->GetMetasoundChecked();
-
-	FCreateNodeVertexParams VertexParams;
-	VertexParams.DataType = InNewType;
-
-	FNodeHandle NewNodeHandle = AddNodeHandle(NodeName, VertexParams);
-	NewNodeHandle->SetNodeName(NodeName);
-	NewNodeHandle->SetDisplayName(NodeDisplayName);
-	NewNodeHandle->SetDescription(NodeDescription);
-	FGraphHandle GraphHandle = Graph->GetGraphHandle();
-	if (IsA<UMetasoundEditorGraphInput>())
-	{
-		GraphHandle->SetSortOrderIndexForInput(NodeName, SortOrderIndex);
-	}
-	else if (IsA<UMetasoundEditorGraphOutput>())
-	{
-		GraphHandle->SetSortOrderIndexForOutput(NodeName, SortOrderIndex);
-	}
-
-	if (!ensure(NewNodeHandle->IsValid()))
-	{
-		return;
-	}
-
-	ClassName = NewNodeHandle->GetClassMetadata().GetClassName();
-	NodeID = NewNodeHandle->GetID();
-
-	TypeName = InNewType;
-	InitializeLiteral();
-
-	// 6. Create new node references in the same locations as the old locations
-	for (FVector2D Location : NodeLocations)
-	{
-		FGraphBuilder::AddNode(Metasound, NewNodeHandle, Location, false /* bInSelectNewNode */);
-	}
-
-	Graph->GetModifyContext().AddMemberIDsModified({ GetMemberID() });
-	Graph->RegisterGraphWithFrontend();
-}
-
-void UMetasoundEditorGraphVertex::SetVertexAccessType(EMetasoundFrontendVertexAccessType InNewAccessType, bool bPostTransaction)
-{
-	using namespace Metasound::Editor;
-	using namespace Metasound::Frontend;
-
-	ensureAlwaysMsgf(InNewAccessType != EMetasoundFrontendVertexAccessType::Unset, TEXT("Vertices do not support ambiguous access type."));
-
-	if (InNewAccessType == GetVertexAccessType())
-	{
-		return;
-	}
-
-	UMetasoundEditorGraph* Graph = GetOwningGraph();
-	if (!ensure(Graph))
-	{
-		return;
-	}
-
-	const FScopedTransaction Transaction(LOCTEXT("SetVertexConstructorPinState", "Set MetaSound Vertex Constructor Pin State"), bPostTransaction);
-	Graph->GetMetasoundChecked().Modify();
-	Graph->Modify();
-	Modify();
-
-	// Logic is based on UMetasoundEditorGraphMember::SetDataType
-	// 1. Cache current editor node reference positions & linked pins
-	TArray<UMetasoundEditorGraphMemberNode*> Nodes = GetNodes();
-	TArray<FVector2D> NodeLocations;
-	// Map of node index within NodeLocations to pins linked to that node
-	TMap<int32, TArray<UEdGraphPin*>> LinkedPins;
-	for (int32 i = 0; i < Nodes.Num(); ++i)
-	{
-		const UMetasoundEditorGraphMemberNode* Node = Nodes[i];
-		if (ensure(Node))
-		{
-			NodeLocations.Add(FVector2D(Node->NodePosX, Node->NodePosY));
-
-			ensure(Node->Pins.Num() == 1);
-			if (UEdGraphPin* NodePin = Node->Pins.Last()) 
-			{
-				if (NodePin->LinkedTo.Num() > 0)
-				{
-					LinkedPins.Add(i, NodePin->LinkedTo);
-				}
-			}
-		}
-	}
-
-	// 2. Cache the old version's Frontend data.
-	FNodeHandle NodeHandle = GetNodeHandle();
-	const FName NodeName = NodeHandle->GetNodeName();
-	const FText NodeDisplayName = NodeHandle->GetDisplayName();
-	const FText NodeDescription = NodeHandle->GetDescription();
-	const FName NodeType = GetDataType();
-	const int32 SortOrderIndex = GetSortOrderIndex();
-	
-	// 3. Remove the current nodes and vertex
-	Graph->RemoveMemberNodes(*this);
-	Graph->RemoveFrontendMember(*this);
-
-	// 4. Add the new input node with the same identifier data but new access type.
-	FCreateNodeVertexParams VertexParams;
-	VertexParams.DataType = NodeType;
-	VertexParams.AccessType = InNewAccessType;
-	UObject& Metasound = Graph->GetMetasoundChecked();
-	
-	FNodeHandle NewNodeHandle = AddNodeHandle(NodeName, VertexParams);
-	NewNodeHandle->SetNodeName(NodeName);
-	NewNodeHandle->SetDisplayName(NodeDisplayName);
-	NewNodeHandle->SetDescription(NodeDescription);
-	FGraphHandle GraphHandle = Graph->GetGraphHandle();
-	if (IsA<UMetasoundEditorGraphInput>())
-	{
-		GraphHandle->SetSortOrderIndexForInput(NodeName, SortOrderIndex);
-	}
-	else if (IsA<UMetasoundEditorGraphOutput>())
-	{
-		GraphHandle->SetSortOrderIndexForOutput(NodeName, SortOrderIndex);
-	}
-	
-	if (!ensure(NewNodeHandle->IsValid()))
-	{
-		return;
-	}
-
-	ClassName = NewNodeHandle->GetClassMetadata().GetClassName();
-	NodeID = NewNodeHandle->GetID();
-
-	constexpr bool bUpdateLiteralPostTransaction = false;
-	UpdateFrontendDefaultLiteral(bUpdateLiteralPostTransaction);
-
-	// 6. Create new node references in the same locations as the old locations
-	// and reconnect linked pins 
-	for (int32 i = 0; i < NodeLocations.Num(); ++i)
-	{
-		UMetasoundEditorGraphNode* NewEdGraphNode = FGraphBuilder::AddNode(Metasound, NewNodeHandle, NodeLocations[i], false /* bInSelectNewNode */);
-		
-		if (LinkedPins.Contains(i))
-		{
-			for (UEdGraphPin* LinkedPin : LinkedPins[i])
-			{
-				NewEdGraphNode->AutowireNewNode(LinkedPin);
-			}
-		}
-	}
-
-	Graph->RegisterGraphWithFrontend();
+	FGraphBuilder::GetOutermostMetaSoundChecked(*Graph).GetModifyContext().AddMemberIDsModified({ GetMemberID() });
 }
 
 Metasound::Frontend::FNodeHandle UMetasoundEditorGraphVertex::GetNodeHandle()
 {
 	using namespace Metasound;
 
-	UObject* Object = CastChecked<UMetasoundEditorGraph>(GetOuter())->GetMetasound();
+	UMetasoundEditorGraph* Graph = GetOwningGraph();
+	check(Graph);
+
+	UObject* Object = Graph->GetMetasound();
 	if (!ensure(Object))
 	{
 		return Frontend::INodeController::GetInvalidHandle();
@@ -545,16 +378,8 @@ Metasound::Frontend::FConstNodeHandle UMetasoundEditorGraphVertex::GetConstNodeH
 {
 	using namespace Metasound;
 
-	const UObject* Object = CastChecked<UMetasoundEditorGraph>(GetOuter())->GetMetasound();
-	if (!ensure(Object))
-	{
-		return Frontend::INodeController::GetInvalidHandle();
-	}
-
-	const FMetasoundAssetBase* MetasoundAsset = IMetasoundUObjectRegistry::Get().GetObjectAsAssetBase(Object);
-	check(MetasoundAsset);
-
-	return MetasoundAsset->GetRootGraphHandle()->GetNodeWithID(NodeID);
+	const FMetasoundAssetBase& MetaSound = Editor::FGraphBuilder::GetOutermostConstMetaSoundChecked(*this);
+	return MetaSound.GetRootGraphHandle()->GetNodeWithID(NodeID);
 }
 
 const FMetasoundFrontendVersion& UMetasoundEditorGraphVertex::GetInterfaceVersion() const
@@ -647,10 +472,37 @@ bool UMetasoundEditorGraphVertex::CanRename(const FText& InNewText, FText& OutEr
 	return bIsNameValid;
 }
 
+const UMetasoundEditorGraphMember* UMetasoundEditorGraphMemberDefaultLiteral::FindMember() const
+{
+	using namespace Metasound;
+
+	const FMetasoundAssetBase& MetaSound = Editor::FGraphBuilder::GetOutermostConstMetaSoundChecked(*this);
+	if (const UMetasoundEditorGraph* Graph = Cast<UMetasoundEditorGraph>(MetaSound.GetGraph()))
+	{
+		return Graph->FindMember(MemberID);
+	}
+
+	return nullptr;
+}
+
+UMetasoundEditorGraphMember* UMetasoundEditorGraphMemberDefaultLiteral::FindMember()
+{
+	using namespace Metasound;
+
+	FMetasoundAssetBase* MetasoundAsset = IMetasoundUObjectRegistry::Get().GetObjectAsAssetBase(GetOutermostObject());
+	check(MetasoundAsset);
+	if (UMetasoundEditorGraph* Graph = Cast<UMetasoundEditorGraph>(MetasoundAsset->GetGraph()))
+	{
+		return Graph->FindMember(MemberID);
+	}
+
+	return nullptr;
+}
+
 #if WITH_EDITOR
 void UMetasoundEditorGraphMemberDefaultLiteral::PostEditChangeProperty(FPropertyChangedEvent& InPropertyChangedEvent)
 {
-	if (UMetasoundEditorGraphMember* Member = Cast<UMetasoundEditorGraphMember>(GetOuter()))
+	if (UMetasoundEditorGraphMember* Member = FindMember())
 	{
 		Metasound::Editor::GraphPrivate::OnLiteralChanged(*Member, InPropertyChangedEvent.ChangeType);
 	}
@@ -658,7 +510,7 @@ void UMetasoundEditorGraphMemberDefaultLiteral::PostEditChangeProperty(FProperty
 
 void UMetasoundEditorGraphMemberDefaultLiteral::PostEditChangeChainProperty(FPropertyChangedChainEvent& InPropertyChangedEvent)
 {
-	if (UMetasoundEditorGraphMember* Member = Cast<UMetasoundEditorGraphMember>(GetOuter()))
+	if (UMetasoundEditorGraphMember* Member = FindMember())
 	{
 		Metasound::Editor::GraphPrivate::OnLiteralChanged(*Member, InPropertyChangedEvent.ChangeType);
 	}
@@ -666,6 +518,8 @@ void UMetasoundEditorGraphMemberDefaultLiteral::PostEditChangeChainProperty(FPro
 
 void UMetasoundEditorGraphMemberDefaultLiteral::PostEditUndo()
 {
+	using namespace Metasound::Editor;
+
 	Super::PostEditUndo();
 
 	if (!IsValid(this))
@@ -674,12 +528,10 @@ void UMetasoundEditorGraphMemberDefaultLiteral::PostEditUndo()
 	}
 
 	constexpr bool bPostTransaction = false;
-	UMetasoundEditorGraphMember* Member = CastChecked<UMetasoundEditorGraphMember>(GetOuter());
-	Member->UpdateFrontendDefaultLiteral(bPostTransaction);
-
-	if (UMetasoundEditorGraph* Graph = Member->GetOwningGraph())
+	if (UMetasoundEditorGraphMember* Member = FindMember())
 	{
-		Graph->GetModifyContext().SetDocumentModified();
+		Member->UpdateFrontendDefaultLiteral(bPostTransaction);
+		FGraphBuilder::GetOutermostMetaSoundChecked(*this).GetModifyContext().SetDocumentModified();
 	}
 }
 #endif // WITH_EDITOR
@@ -689,19 +541,12 @@ Metasound::Editor::ENodeSection UMetasoundEditorGraphInput::GetSectionID() const
 	return Metasound::Editor::ENodeSection::Inputs;
 }
 
-Metasound::Frontend::FNodeHandle UMetasoundEditorGraphInput::AddNodeHandle(const FName& InName, FName InDataType)
-{
-	UE_LOG(LogMetaSound, Error, TEXT("UMetasoundEditorGraphInput::AddNodeHandle with these parameters is no longer supported and should not be called. Use the one with FCreateNodeVertexParams instead."));
-
-	return Metasound::Frontend::INodeController::GetInvalidHandle();
-}
-
 Metasound::Frontend::FNodeHandle UMetasoundEditorGraphInput::AddNodeHandle(const FName& InName, const Metasound::Editor::FCreateNodeVertexParams& InParams)
 {
 	using namespace Metasound::Editor;
 	using namespace Metasound::Frontend;
 
-	UMetasoundEditorGraph* Graph = Cast<UMetasoundEditorGraph>(GetOuter());
+	UMetasoundEditorGraph* Graph = GetOwningGraph();
 	if (!ensure(Graph))
 	{
 		return Metasound::Frontend::INodeController::GetInvalidHandle();
@@ -712,6 +557,18 @@ Metasound::Frontend::FNodeHandle UMetasoundEditorGraphInput::AddNodeHandle(const
 	return NewNodeHandle;
 }
 
+UMetasoundEditorGraphNode* UMetasoundEditorGraphInput::AddNode(Metasound::Frontend::FNodeHandle InNodeHandle, bool bInSelectNewNode)
+{
+	using namespace Metasound::Editor;
+	using namespace Metasound::Frontend;
+
+	UMetasoundEditorGraph* Graph = GetOwningGraph();
+	check(Graph);
+
+	FConstNodeHandle TemplateNodeHandle = FGraphBuilder::AddInputTemplateNodeHandle(Graph->GetMetasoundChecked(), InNodeHandle);
+	return FGraphBuilder::AddInputNode(Graph->GetMetasoundChecked(), TemplateNodeHandle, bInSelectNewNode);
+}
+
 const FText& UMetasoundEditorGraphInput::GetGraphMemberLabel() const
 {
 	static const FText Label = LOCTEXT("GraphMemberLabel_Input", "Input");
@@ -730,6 +587,29 @@ int32 UMetasoundEditorGraphInput::GetSortOrderIndex() const
 	return GraphHandle->GetSortOrderIndexForInput(NodeName);
 }
 
+TArray<UMetasoundEditorGraphMemberNode*> UMetasoundEditorGraphInput::GetNodes() const
+{
+	TArray<UMetasoundEditorGraphMemberNode*> Nodes;
+
+	const UMetasoundEditorGraph* Graph = GetOwningGraph();
+	if (ensure(Graph))
+	{
+		TArray<UMetasoundEditorGraphInputNode*> InputNodes;
+		Graph->GetNodesOfClassEx<UMetasoundEditorGraphInputNode>(InputNodes);
+		Algo::TransformIf(InputNodes, Nodes,
+			[this](const UMetasoundEditorGraphMemberNode* Node)
+			{
+				return Node->GetMember() == this;
+			},
+			[](UMetasoundEditorGraphMemberNode* Node)
+			{
+				return Node;
+			});
+	}
+
+	return Nodes;
+}
+
 void UMetasoundEditorGraphInput::SetSortOrderIndex(int32 InSortOrderIndex)
 {
 	using namespace Metasound::Editor;
@@ -743,7 +623,7 @@ void UMetasoundEditorGraphInput::SetSortOrderIndex(int32 InSortOrderIndex)
 	const Metasound::FVertexName& NodeName = NodeHandle->GetNodeName();
 
 	GraphHandle->SetSortOrderIndexForInput(NodeName, InSortOrderIndex);
-	MetaSoundGraph->GetModifyContext().AddMemberIDsModified({ GetMemberID() });
+	FGraphBuilder::GetOutermostMetaSoundChecked(*MetaSoundGraph).GetModifyContext().AddMemberIDsModified({ GetMemberID() });
 }
 
 void UMetasoundEditorGraphInput::ResetToClassDefault()
@@ -776,9 +656,60 @@ void UMetasoundEditorGraphInput::ResetToClassDefault()
 		}
 	}
 
-	if (UMetasoundEditorGraph* Graph = GetOwningGraph())
+	FGraphBuilder::GetOutermostMetaSoundChecked(*MetaSoundGraph).GetModifyContext().AddMemberIDsModified({ GetMemberID() });
+}
+
+void UMetasoundEditorGraphInput::SetDataType(FName InNewType, bool bPostTransaction)
+{
+	using namespace Metasound::Editor;
+	using namespace Metasound::Frontend;
+
+	if (InNewType != GetDataType())
+	{
+		if (UMetasoundEditorGraph* Graph = GetOwningGraph())
+		{
+			UObject& MetaSound = Graph->GetMetasoundChecked();
+			const FScopedTransaction Transaction(LOCTEXT("SetGraphInputData", "Set MetaSound Graph Input DataType"), bPostTransaction);
+			MetaSound.Modify();
+			Graph->Modify();
+			Modify();
+
+			EMetaSoundBuilderResult Result = EMetaSoundBuilderResult::Failed;
+			Graph->GetBuilderChecked().SetGraphInputDataType(GetMemberName(), InNewType, Result);
+			ensure(Result == EMetaSoundBuilderResult::Succeeded);
+
+			// Cached TypeName here must be set prior to re-initializing literal below
+			TypeName = InNewType;
+			ClassName = GetConstNodeHandle()->GetClassMetadata().GetClassName();
+
+			InitializeLiteral();
+
+			FGraphBuilder::RegisterGraphWithFrontend(MetaSound);
+		}
+	}
+}
+
+void UMetasoundEditorGraphInput::SetVertexAccessType(EMetasoundFrontendVertexAccessType InNewAccessType, bool bPostTransaction)
+{
+	using namespace Metasound::Editor;
+	using namespace Metasound::Frontend;
+
+	if (InNewAccessType != GetVertexAccessType())
 	{
-		Graph->GetModifyContext().AddMemberIDsModified({ GetMemberID() });
+		if (UMetasoundEditorGraph* Graph = GetOwningGraph())
+		{
+			const FScopedTransaction Transaction(LOCTEXT("SetGraphInputVertexType", "Set MetaSound Graph Input Vertex Access Type"), bPostTransaction);
+			Graph->GetMetasoundChecked().Modify();
+			Graph->Modify();
+			Modify();
+
+			EMetaSoundBuilderResult Result = EMetaSoundBuilderResult::Failed;
+			Graph->GetBuilderChecked().SetGraphInputAccessType(GetMemberName(), InNewAccessType, Result);
+			ensure(Result == EMetaSoundBuilderResult::Succeeded);
+
+			ClassName = GetConstNodeHandle()->GetClassMetadata().GetClassName();
+			Graph->RegisterGraphWithFrontend();
+		}
 	}
 }
 
@@ -859,19 +790,12 @@ EMetasoundFrontendVertexAccessType UMetasoundEditorGraphInput::GetVertexAccessTy
 	return EMetasoundFrontendVertexAccessType::Reference;
 }
 
-Metasound::Frontend::FNodeHandle UMetasoundEditorGraphOutput::AddNodeHandle(const FName& InName, FName InDataType)
-{
-	UE_LOG(LogMetaSound, Error, TEXT("UMetasoundEditorGraphOutput::AddNodeHandle with these parameters is no longer supported and should not be called. Use the one with FCreateNodeVertexParams instead."));
-
-	return Metasound::Frontend::INodeController::GetInvalidHandle();
-}
-
 Metasound::Frontend::FNodeHandle UMetasoundEditorGraphOutput::AddNodeHandle(const FName& InName, const Metasound::Editor::FCreateNodeVertexParams& InParams)
 {
 	using namespace Metasound::Editor;
 	using namespace Metasound::Frontend;
 
-	UMetasoundEditorGraph* Graph = Cast<UMetasoundEditorGraph>(GetOuter());
+	UMetasoundEditorGraph* Graph = GetOwningGraph();
 	if (!ensure(Graph))
 	{
 		return Metasound::Frontend::INodeController::GetInvalidHandle();
@@ -882,6 +806,14 @@ Metasound::Frontend::FNodeHandle UMetasoundEditorGraphOutput::AddNodeHandle(cons
 	return NewNodeHandle;
 }
 
+UMetasoundEditorGraphNode* UMetasoundEditorGraphOutput::AddNode(Metasound::Frontend::FNodeHandle InNodeHandle, bool bInSelectNewNode)
+{
+	using namespace Metasound::Editor;
+	UMetasoundEditorGraph* Graph = GetOwningGraph();
+	check(Graph);
+	return FGraphBuilder::AddOutputNode(Graph->GetMetasoundChecked(), InNodeHandle, bInSelectNewNode);
+}
+
 int32 UMetasoundEditorGraphOutput::GetSortOrderIndex() const
 {
 	using namespace Metasound::Editor;
@@ -907,7 +839,7 @@ void UMetasoundEditorGraphOutput::SetSortOrderIndex(int32 InSortOrderIndex)
 	const Metasound::FVertexName& NodeName = NodeHandle->GetNodeName();
 
 	GraphHandle->SetSortOrderIndexForOutput(NodeName, InSortOrderIndex);
-	MetaSoundGraph->GetModifyContext().AddMemberIDsModified({ GetMemberID()});
+	FGraphBuilder::GetOutermostMetaSoundChecked(*MetaSoundGraph).GetModifyContext().AddMemberIDsModified({ GetMemberID() });
 }
 
 const FText& UMetasoundEditorGraphOutput::GetGraphMemberLabel() const
@@ -919,6 +851,7 @@ const FText& UMetasoundEditorGraphOutput::GetGraphMemberLabel() const
 void UMetasoundEditorGraphOutput::ResetToClassDefault()
 {
 	using namespace Metasound;
+	using namespace Metasound::Editor;
 	using namespace Metasound::Frontend;
 
 	TArray<UMetasoundEditorGraphMemberNode*> Nodes = GetNodes();
@@ -937,9 +870,60 @@ void UMetasoundEditorGraphOutput::ResetToClassDefault()
 	Literal->Modify();
 	Literal->SetFromLiteral(DefaultLiteral);
 
-	if (UMetasoundEditorGraph* Graph = GetOwningGraph())
+	FGraphBuilder::GetOutermostMetaSoundChecked(*this).GetModifyContext().AddMemberIDsModified({ GetMemberID() });
+}
+
+void UMetasoundEditorGraphOutput::SetDataType(FName InNewType, bool bPostTransaction)
+{
+	using namespace Metasound::Editor;
+	using namespace Metasound::Frontend;
+
+	if (InNewType != GetDataType())
 	{
-		Graph->GetModifyContext().AddMemberIDsModified({ GetMemberID() });
+		if (UMetasoundEditorGraph* Graph = GetOwningGraph())
+		{
+			const FScopedTransaction Transaction(LOCTEXT("SetGraphOutputData", "Set MetaSound Graph Output DataType"), bPostTransaction);
+			Graph->GetMetasoundChecked().Modify();
+			Graph->Modify();
+			Modify();
+
+			EMetaSoundBuilderResult Result = EMetaSoundBuilderResult::Failed;
+			Graph->GetBuilderChecked().SetGraphOutputDataType(GetMemberName(), InNewType, Result);
+			ensure(Result == EMetaSoundBuilderResult::Succeeded);
+
+			// Cached TypeName here must be set prior to re-initializing literal below
+			TypeName = InNewType;
+			ClassName = GetConstNodeHandle()->GetClassMetadata().GetClassName();
+
+			InitializeLiteral();
+
+			Graph->RegisterGraphWithFrontend();
+		}
+	}
+}
+
+void UMetasoundEditorGraphOutput::SetVertexAccessType(EMetasoundFrontendVertexAccessType InNewAccessType, bool bPostTransaction)
+{
+	using namespace Metasound::Editor;
+	using namespace Metasound::Frontend;
+
+	if (InNewAccessType != GetVertexAccessType())
+	{
+		UMetasoundEditorGraph* Graph = GetOwningGraph();
+		if (ensure(Graph))
+		{
+			const FScopedTransaction Transaction(LOCTEXT("SetGraphOutputAccessType", "Set MetaSound Graph Output Access Type"), bPostTransaction);
+			Graph->GetMetasoundChecked().Modify();
+			Graph->Modify();
+			Modify();
+
+			EMetaSoundBuilderResult Result = EMetaSoundBuilderResult::Failed;
+			Graph->GetBuilderChecked().SetGraphOutputAccessType(GetMemberName(), InNewAccessType, Result);
+			ensure(Result == EMetaSoundBuilderResult::Succeeded);
+
+			ClassName = GetConstNodeHandle()->GetClassMetadata().GetClassName();
+			Graph->RegisterGraphWithFrontend();
+		}
 	}
 }
 
@@ -1028,36 +1012,30 @@ Metasound::Frontend::FVariableHandle UMetasoundEditorGraphVariable::GetVariableH
 {
 	using namespace Metasound;
 
-	UObject* Object = CastChecked<UMetasoundEditorGraph>(GetOuter())->GetMetasound();
-	if (!ensure(Object))
-	{
-		return Frontend::IVariableController::GetInvalidHandle();
-	}
-
-	FMetasoundAssetBase* MetasoundAsset = IMetasoundUObjectRegistry::Get().GetObjectAsAssetBase(Object);
-	check(MetasoundAsset);
-
-	return MetasoundAsset->GetRootGraphHandle()->FindVariable(VariableID);
+	FMetasoundAssetBase& MetasoundAsset = Editor::FGraphBuilder::GetOutermostMetaSoundChecked(*this);
+	return MetasoundAsset.GetRootGraphHandle()->FindVariable(VariableID);
 }
 
 Metasound::Frontend::FConstVariableHandle UMetasoundEditorGraphVariable::GetConstVariableHandle() const
 {
 	using namespace Metasound;
 
-	const UObject* Object = CastChecked<const UMetasoundEditorGraph>(GetOuter())->GetMetasound();
-	if (!ensure(Object))
-	{
-		return Frontend::IVariableController::GetInvalidHandle();
-	}
+	const FMetasoundAssetBase& MetaSound = Editor::FGraphBuilder::GetOutermostConstMetaSoundChecked(*this);
+	return MetaSound.GetRootGraphHandle()->FindVariable(VariableID);
+}
 
-	const FMetasoundAssetBase* MetasoundAsset = IMetasoundUObjectRegistry::Get().GetObjectAsAssetBase(Object);
-	check(MetasoundAsset);
+UMetasoundEditorGraphNode* UMetasoundEditorGraphVariable::AddNode(Metasound::Frontend::FNodeHandle InNodeHandle, bool bInSelectNewNode)
+{
+	using namespace Metasound::Editor;
 
-	return MetasoundAsset->GetRootGraphHandle()->FindVariable(VariableID);
+	FMetasoundAssetBase& MetaSound = FGraphBuilder::GetOutermostMetaSoundChecked(*this);
+	return FGraphBuilder::AddVariableNode(*MetaSound.GetOwningAsset(), InNodeHandle, bInSelectNewNode);
 }
 
 void UMetasoundEditorGraphVariable::SetMemberName(const FName& InNewName, bool bPostTransaction)
 {
+	using namespace Metasound::Editor;
+
 	UMetasoundEditorGraph* Graph = GetOwningGraph();
 	if (!ensure(Graph))
 	{
@@ -1073,7 +1051,7 @@ void UMetasoundEditorGraphVariable::SetMemberName(const FName& InNewName, bool b
 	GetVariableHandle()->SetName(InNewName);
 
 	Graph->RegisterGraphWithFrontend();
-	Graph->GetModifyContext().AddMemberIDsModified({ GetMemberID() });
+	FGraphBuilder::GetOutermostMetaSoundChecked(*this).GetModifyContext().AddMemberIDsModified({ GetMemberID() });
 }
 
 FGuid UMetasoundEditorGraphVariable::GetMemberID() const 
@@ -1191,6 +1169,7 @@ FText UMetasoundEditorGraphVariable::GetDisplayName() const
 
 void UMetasoundEditorGraphVariable::SetDisplayName(const FText& InNewName, bool bPostTransaction)
 {
+	using namespace Metasound::Editor;
 	using namespace Metasound::Frontend;
 
 	UMetasoundEditorGraph* Graph = GetOwningGraph();
@@ -1209,7 +1188,7 @@ void UMetasoundEditorGraphVariable::SetDisplayName(const FText& InNewName, bool
 	FVariableHandle VariableHandle = GetVariableHandle();
 	VariableHandle->SetDisplayName(InNewName);
 
-	Graph->GetModifyContext().AddMemberIDsModified({ GetMemberID() });
+	FGraphBuilder::GetOutermostMetaSoundChecked(*this).GetModifyContext().AddMemberIDsModified({ GetMemberID() });
 }
 
 void UMetasoundEditorGraphVariable::SetDataType(FName InNewType, bool bPostTransaction)
@@ -1256,7 +1235,7 @@ void UMetasoundEditorGraphVariable::SetDataType(FName InNewType, bool bPostTrans
 
 
 		// Remove the current variable
-		Graph->RemoveMemberNodes(*this);
+		FGraphBuilder::DeleteMemberNodes(*this);
 		FGraphHandle FrontendGraph = Graph->GetGraphHandle();
 		FrontendGraph->RemoveVariable(VariableID);
 		VariableID = FGuid();
@@ -1379,21 +1358,27 @@ void UMetasoundEditorGraphVariable::AddVariableNodes(UObject& InMetasound, Metas
 		bool bMutatorNodeAlreadyExists = GetConstVariableHandle()->FindMutatorNode()->IsValid();
 		if (ensure(!bMutatorNodeAlreadyExists))
 		{
-			FNodeHandle MutatorFrontendNode = InFrontendGraph->FindOrAddVariableMutatorNode(VariableID);
-			FGraphBuilder::AddNode(InMetasound, MutatorFrontendNode, *InNodeLocs.MutatorLocation, false /* bInSelectNewNode */);
+			FConstNodeHandle MutatorFrontendNode = InFrontendGraph->FindOrAddVariableMutatorNode(VariableID);
+			UMetasoundEditorGraphNode* NewGraphNode = FGraphBuilder::AddVariableNode(InMetasound, MutatorFrontendNode, false /* bInSelectNewNode */);
+			NewGraphNode->UpdateFrontendNodeLocation(*InNodeLocs.MutatorLocation);
+			NewGraphNode->SyncLocationFromFrontendNode();
 		}
 	}
 
 	for (const FVector2D& Location : InNodeLocs.AccessorLocations)
 	{
-		FNodeHandle AccessorFrontendNode = InFrontendGraph->AddVariableAccessorNode(VariableID);
-		FGraphBuilder::AddNode(InMetasound, AccessorFrontendNode, Location, false /* bInSelectNewNode */);
+		FConstNodeHandle AccessorFrontendNode = InFrontendGraph->AddVariableAccessorNode(VariableID);
+		UMetasoundEditorGraphNode* NewGraphNode = FGraphBuilder::AddVariableNode(InMetasound, AccessorFrontendNode, false /* bInSelectNewNode */);
+		NewGraphNode->UpdateFrontendNodeLocation(Location);
+		NewGraphNode->SyncLocationFromFrontendNode();
 	}
 
 	for (const FVector2D& Location : InNodeLocs.DeferredAccessorLocations)
 	{
-		FNodeHandle DeferredAccessorFrontendNode = InFrontendGraph->AddVariableDeferredAccessorNode(VariableID);
-		FGraphBuilder::AddNode(InMetasound, DeferredAccessorFrontendNode, Location, false /* bInSelectNewNode */);
+		FConstNodeHandle DeferredAccessorFrontendNode = InFrontendGraph->AddVariableDeferredAccessorNode(VariableID);
+		UMetasoundEditorGraphNode* NewGraphNode = FGraphBuilder::AddVariableNode(InMetasound, DeferredAccessorFrontendNode, false /* bInSelectNewNode */);
+		NewGraphNode->UpdateFrontendNodeLocation(Location);
+		NewGraphNode->SyncLocationFromFrontendNode();
 	}
 }
 
@@ -1405,8 +1390,9 @@ const FGuid& UMetasoundEditorGraphVariable::GetVariableID() const
 void UMetasoundEditorGraphVariable::ResetToClassDefault()
 {
 	using namespace Metasound;
+	using namespace Metasound::Editor;
 	using namespace Metasound::Frontend;
-	using namespace Metasound::VariableNames; 
+	using namespace Metasound::VariableNames;
 
 	FMetasoundFrontendLiteral DefaultLiteral;
 	DefaultLiteral.SetFromLiteral(IDataTypeRegistry::Get().CreateDefaultLiteral(GetDataType()));
@@ -1427,10 +1413,7 @@ void UMetasoundEditorGraphVariable::ResetToClassDefault()
 		}
 	}
 
-	if (UMetasoundEditorGraph* Graph = GetOwningGraph())
-	{
-		Graph->GetModifyContext().AddMemberIDsModified({ GetMemberID() });
-	}
+	FGraphBuilder::GetOutermostMetaSoundChecked(*this).GetModifyContext().AddMemberIDsModified({ GetMemberID() });
 }
 
 void UMetasoundEditorGraphVariable::UpdateFrontendDefaultLiteral(bool bPostTransaction)
@@ -1477,39 +1460,7 @@ void UMetasoundEditorGraphVariable::UpdateFrontendDefaultLiteral(bool bPostTrans
 
 UMetasoundEditorGraphInputNode* UMetasoundEditorGraph::CreateInputNode(Metasound::Frontend::FNodeHandle InNodeHandle, bool bInSelectNewNode)
 {
-	using namespace Metasound::Editor;
-	using namespace Metasound::Frontend;
-
-	TArray<FConstOutputHandle> NodeOutputs = InNodeHandle->GetConstOutputs();
-	if (!ensure(!NodeOutputs.IsEmpty()))
-	{
-		return nullptr;
-	}
-
-	if (!ensure(InNodeHandle->GetClassMetadata().GetType() == EMetasoundFrontendClassType::Input))
-	{
-		return nullptr;
-	}
-
-	UEdGraphNode* NewEdGraphNode = CreateNode(UMetasoundEditorGraphInputNode::StaticClass(), bInSelectNewNode);
-	UMetasoundEditorGraphInputNode* NewInputNode = CastChecked<UMetasoundEditorGraphInputNode>(NewEdGraphNode);
-	if (ensure(NewInputNode))
-	{
-		NewInputNode->CreateNewGuid();
-		NewInputNode->PostPlacedNewNode();
-
-		NewInputNode->Input = FindOrAddInput(InNodeHandle);
-
-		if (NewInputNode->Pins.IsEmpty())
-		{
-			NewInputNode->AllocateDefaultPins();
-		}
-
-		NewInputNode->CacheTitle();
-
-		return NewInputNode;
-	}
-	
+	checkNoEntry();
 	return nullptr;
 }
 
@@ -1556,14 +1507,19 @@ void UMetasoundEditorGraph::PreSave(FObjectPreSaveContext InSaveContext)
 	Super::PreSave(InSaveContext);
 }
 
+UMetaSoundBuilderBase& UMetasoundEditorGraph::GetBuilderChecked()
+{
+	return UMetaSoundBuilderSubsystem::GetChecked().AttachBuilderToAssetChecked(*GetOutermostObject());
+}
+
 UObject* UMetasoundEditorGraph::GetMetasound()
 {
-	return GetOuter();
+	return GetOutermostObject();
 }
 
 const UObject* UMetasoundEditorGraph::GetMetasound() const
 {
-	return GetOuter();
+	return GetOutermostObject();
 }
 
 UObject& UMetasoundEditorGraph::GetMetasoundChecked()
@@ -1584,51 +1540,12 @@ void UMetasoundEditorGraph::RegisterGraphWithFrontend()
 {
 	using namespace Metasound::Editor;
 
-	if (UObject* ParentMetasound = GetOuter())
+	if (UObject* ParentMetasound = GetOutermostObject())
 	{
 		FGraphBuilder::RegisterGraphWithFrontend(*ParentMetasound);
 	}
 }
 
-void UMetasoundEditorGraph::ClearVersionedOnLoad()
-{
-	bVersionedOnLoad = false;
-}
-
-bool UMetasoundEditorGraph::GetVersionedOnLoad() const
-{
-	return bVersionedOnLoad;
-}
-
-void UMetasoundEditorGraph::SetVersionedOnLoad()
-{
-	bVersionedOnLoad = true;
-}
-
-FMetasoundFrontendDocumentModifyContext& UMetasoundEditorGraph::GetModifyContext()
-{
-	using namespace Metasound;
-
-	UObject* ParentMetasound = GetOuter();
-	check(ParentMetasound);
-
-	FMetasoundAssetBase* MetasoundAsset = IMetasoundUObjectRegistry::Get().GetObjectAsAssetBase(ParentMetasound);
-	check(MetasoundAsset);
-	return MetasoundAsset->GetModifyContext();
-}
-
-const FMetasoundFrontendDocumentModifyContext& UMetasoundEditorGraph::GetModifyContext() const
-{
-	using namespace Metasound;
-
-	UObject* ParentMetasound = GetOuter();
-	check(ParentMetasound);
-
-	const FMetasoundAssetBase* MetasoundAsset = IMetasoundUObjectRegistry::Get().GetObjectAsAssetBase(ParentMetasound);
-	check(MetasoundAsset);
-	return MetasoundAsset->GetModifyContext();
-}
-
 UMetasoundEditorGraphInput* UMetasoundEditorGraph::FindInput(FGuid InNodeID) const
 {
 	const TObjectPtr<UMetasoundEditorGraphInput>* Input = Inputs.FindByPredicate([InNodeID](const TObjectPtr<UMetasoundEditorGraphInput>& InInput)
@@ -1658,12 +1575,12 @@ UMetasoundEditorGraphInput* UMetasoundEditorGraph::FindInput(FName InName) const
 	return Input ? Input->Get() : nullptr;
 }
 
-UMetasoundEditorGraphInput* UMetasoundEditorGraph::FindOrAddInput(Metasound::Frontend::FNodeHandle InNodeHandle)
+UMetasoundEditorGraphInput* UMetasoundEditorGraph::FindOrAddInput(Metasound::Frontend::FConstNodeHandle InNodeHandle)
 {
 	using namespace Metasound::Editor;
 	using namespace Metasound::Frontend;
 
-	FGraphHandle Graph = InNodeHandle->GetOwningGraph();
+	FConstGraphHandle Graph = InNodeHandle->GetOwningGraph();
 
 	FName TypeName;
 	FGuid VertexID;
@@ -1682,7 +1599,7 @@ UMetasoundEditorGraphInput* UMetasoundEditorGraph::FindOrAddInput(Metasound::Fro
 		return Input;
 	}
 
-	UMetasoundEditorGraphInput* NewInput = NewObject<UMetasoundEditorGraphInput>(this, FName(), RF_Transactional);
+	UMetasoundEditorGraphInput* NewInput = NewObject<UMetasoundEditorGraphInput>(this, GraphPrivate::GetUniqueTransientMemberName(), RF_Transactional);
 	if (ensure(NewInput))
 	{
 		FMetasoundFrontendLiteral DefaultLiteral = Graph->GetDefaultInput(VertexID);
@@ -1722,12 +1639,12 @@ UMetasoundEditorGraphOutput* UMetasoundEditorGraph::FindOutput(FName InName) con
 	return Output ? Output->Get() : nullptr;
 }
 
-UMetasoundEditorGraphOutput* UMetasoundEditorGraph::FindOrAddOutput(Metasound::Frontend::FNodeHandle InNodeHandle)
+UMetasoundEditorGraphOutput* UMetasoundEditorGraph::FindOrAddOutput(Metasound::Frontend::FConstNodeHandle InNodeHandle)
 {
 	using namespace Metasound::Editor;
 	using namespace Metasound::Frontend;
 
-	FGraphHandle Graph = InNodeHandle->GetOwningGraph();
+	FConstGraphHandle Graph = InNodeHandle->GetOwningGraph();
 
 	FName TypeName;
 	FGuid VertexID;
@@ -1746,7 +1663,7 @@ UMetasoundEditorGraphOutput* UMetasoundEditorGraph::FindOrAddOutput(Metasound::F
 		return Output;
 	}
 
-	UMetasoundEditorGraphOutput* NewOutput = NewObject<UMetasoundEditorGraphOutput>(this, FName(), RF_Transactional);
+	UMetasoundEditorGraphOutput* NewOutput = NewObject<UMetasoundEditorGraphOutput>(this, GraphPrivate::GetUniqueTransientMemberName(), RF_Transactional);
 	if (ensure(NewOutput))
 	{
 		FMetasoundFrontendLiteral DefaultLiteral;
@@ -1783,7 +1700,6 @@ UMetasoundEditorGraphVariable* UMetasoundEditorGraph::FindOrAddVariable(const Me
 
 	FName TypeName = InVariableHandle->GetDataType();
 	const FGuid VariableID = InVariableHandle->GetID();
-	const FGuid NodeID = InVariableHandle->GetID();
 
 	if (TObjectPtr<UMetasoundEditorGraphVariable> EditorVariable = FindVariable(VariableID))
 	{
@@ -1791,11 +1707,11 @@ UMetasoundEditorGraphVariable* UMetasoundEditorGraph::FindOrAddVariable(const Me
 		return EditorVariable;
 	}
 
-	UMetasoundEditorGraphVariable* NewVariable = NewObject<UMetasoundEditorGraphVariable>(this, FName(), RF_Transactional);
+	UMetasoundEditorGraphVariable* NewVariable = NewObject<UMetasoundEditorGraphVariable>(this, GraphPrivate::GetUniqueTransientMemberName(), RF_Transactional);
 	if (ensure(NewVariable))
 	{
 		const FMetasoundFrontendLiteral DefaultLiteral = InVariableHandle->GetLiteral();
-		NewVariable->InitMember(InVariableHandle->GetDataType(), DefaultLiteral, InVariableHandle->GetID());
+		NewVariable->InitMember(InVariableHandle->GetDataType(), DefaultLiteral, VariableID);
 		Variables.Add(NewVariable);
 		return NewVariable;
 	}
@@ -1803,19 +1719,19 @@ UMetasoundEditorGraphVariable* UMetasoundEditorGraph::FindOrAddVariable(const Me
 	return nullptr;
 }
 
-UMetasoundEditorGraphMember* UMetasoundEditorGraph::FindMember(FGuid InNodeID) const
+UMetasoundEditorGraphMember* UMetasoundEditorGraph::FindMember(FGuid InMemberID) const
 {
-	if (UMetasoundEditorGraphOutput* Output = FindOutput(InNodeID))
+	if (UMetasoundEditorGraphOutput* Output = FindOutput(InMemberID))
 	{
 		return Output;
 	}
 
-	if (UMetasoundEditorGraphInput* Input = FindInput(InNodeID))
+	if (UMetasoundEditorGraphInput* Input = FindInput(InMemberID))
 	{
 		return Input;
 	}
 
-	return FindVariable(InNodeID); // Note: ID is a VariableID in this case. 
+	return FindVariable(InMemberID);
 }
 
 UMetasoundEditorGraphMember* UMetasoundEditorGraph::FindAdjacentMember(const UMetasoundEditorGraphMember& InMember)
@@ -1887,6 +1803,168 @@ bool UMetasoundEditorGraph::ContainsVariable(const UMetasoundEditorGraphVariable
 	return Variables.Contains(&InVariable);
 }
 
+void UMetasoundEditorGraph::MigrateEditorDocumentData(FMetaSoundFrontendDocumentBuilder& OutBuilder)
+{
+	using namespace Metasound::Editor;
+	using namespace Metasound::Frontend;
+
+	// 1. Add node comments to style for non-comment nodes (comment nodes processed separately below)
+	TArray<UMetasoundEditorGraphNode*> AllMetaSoundNodes;
+	GetNodesOfClass(AllMetaSoundNodes);
+
+	constexpr bool bIsVisible = true;
+	for (UMetasoundEditorGraphNode* Node : AllMetaSoundNodes)
+	{
+		// Comment nodes are migrated in the next loop
+		if (!Node->IsA<UEdGraphNode_Comment>() && !Node->NodeComment.IsEmpty())
+		{
+			const FGuid& NodeID = Node->GetNodeID();
+			if (NodeID.IsValid())
+			{
+				OutBuilder.SetNodeComment(NodeID, MoveTemp(Node->NodeComment));
+				OutBuilder.SetNodeCommentVisible(NodeID, Node->bCommentBubblePinned);
+			}
+		}
+	}
+
+	const INodeTemplate* Template = INodeTemplateRegistry::Get().FindTemplate(FInputNodeTemplate::GetRegistryKey());
+	checkf(Template, TEXT("Failed to find InputNodeTemplate, which is required for migrating editor document data"));
+
+	// 2. Move inputs to input template nodes, using connection data within the ed graph as a way to inform
+	// which template node should effectively represent which input template node and own which connections.
+	// Cache literals in the literal metadata map to ensure data is serialized appropriately.
+	IterateInputs([this, &OutBuilder, Template](UMetasoundEditorGraphInput& Input)
+	{
+#if WITH_EDITOR
+		if (GEditor) // Have to check if valid as its unavailable in standalone editor builds`
+		{
+			if (UMetaSoundEditorSubsystem* EditorSubsystem = GEditor->GetEditorSubsystem<UMetaSoundEditorSubsystem>())
+			{
+				if (UMetasoundEditorGraphMemberDefaultLiteral* DefaultLiteral = Input.GetLiteral())
+				{
+					TSubclassOf<UMetasoundEditorGraphMemberDefaultLiteral> SubClass = DefaultLiteral->GetClass();
+					EditorSubsystem->BindMemberMetadata(Input, SubClass, DefaultLiteral);
+				}
+			}
+		}
+#endif // WITH_EDITOR
+
+		// Cache data to be used in edge swapping below, being careful to not reference the input node pointer
+		// or vertex reference within the lower inner loop as the underlying node array may be reallocated by
+		// template nodes being added.
+		const FName InputName = Input.GetMemberName();
+		const FMetasoundFrontendNode* InputNode = OutBuilder.FindGraphInputNode(InputName);
+
+		// Potentially not used input, which is perfectly valid so early out
+		if (!InputNode)
+		{
+			return;
+		}
+
+		const FMetasoundFrontendVertex InputNodeOutputVertex = InputNode->Interface.Outputs.Last();
+
+		const FGuid InputNodeID = InputNode->GetID();
+		FMetasoundFrontendEdge EdgeToRemove { InputNodeID, InputNodeOutputVertex.VertexID };
+
+		TArray<UMetasoundEditorGraphMemberNode*> Nodes = Input.GetNodes();
+		for (UMetasoundEditorGraphMemberNode* EdNode : Nodes)
+		{
+			FNodeTemplateGenerateInterfaceParams Params { { InputNodeOutputVertex.TypeName }, { } };
+			const FMetasoundFrontendNode* TemplateNode = OutBuilder.AddNodeByTemplate(*Template, MoveTemp(Params));
+			check(TemplateNode);
+
+			const FGuid& TemplateNodeID = TemplateNode->GetID();
+			const FGuid& TemplateInputID = TemplateNode->Interface.Inputs.Last().VertexID;
+			const FGuid& TemplateOutputID = TemplateNode->Interface.Outputs.Last().VertexID;
+			OutBuilder.SetNodeLocation(TemplateNodeID, FVector2D(EdNode->NodePosX, EdNode->NodePosY));
+
+			// Transform comment to template from input node
+			OutBuilder.SetNodeComment(TemplateNodeID, MoveTemp(EdNode->NodeComment));
+			OutBuilder.SetNodeCommentVisible(TemplateNodeID, EdNode->bCommentBubblePinned);
+
+			// Add edge between input node and new template node
+			OutBuilder.AddEdge(FMetasoundFrontendEdge
+			{
+				InputNodeID,
+				InputNodeOutputVertex.VertexID,
+				TemplateNodeID,
+				TemplateInputID
+			});
+
+			for (const UEdGraphPin* Pin : EdNode->Pins)
+			{
+				if (Pin && Pin->Direction == EGPD_Output)
+				{
+					for (const UEdGraphPin* Linked : Pin->LinkedTo)
+					{
+						if (Linked)
+						{
+							const UMetasoundEditorGraphNode* ConnectedNode = CastChecked<const UMetasoundEditorGraphNode>(Linked->GetOwningNode());
+							const FMetasoundFrontendVertex* ConnectedInput = OutBuilder.FindNodeInput(ConnectedNode->GetNodeID(), Linked->GetFName());
+
+							// Swap connection from input node to connected node to now be from template node to connected node
+							EdgeToRemove.ToNodeID = ConnectedNode->GetNodeID(),
+							EdgeToRemove.ToVertexID = ConnectedInput->VertexID;
+							OutBuilder.RemoveEdge(EdgeToRemove);
+
+							OutBuilder.AddEdge(FMetasoundFrontendEdge
+							{
+								TemplateNodeID,
+								TemplateOutputID,
+								ConnectedNode->GetNodeID(),
+								ConnectedInput->VertexID
+							});
+						}
+					}
+				}
+			}
+		}
+	});
+
+	// 4. Add comment nodes as builder graph comments to frontend document
+	TArray<UEdGraphNode_Comment*> CommentNodes;
+	GetNodesOfClass(CommentNodes);
+	for (UEdGraphNode_Comment* Node : CommentNodes)
+	{
+		FMetaSoundFrontendGraphComment& NewComment = OutBuilder.FindOrAddGraphComment(FGuid::NewGuid());
+		UMetasoundEditorGraphCommentNode::ConvertToFrontendComment(*Node, NewComment);
+	}
+
+	// 5. Remove input locations and ensure that all other nodes only have at most one
+	// location represented in the style/editor graph (0 is acceptable as some member
+	// node types (eg. variables) may not contain locations and that's ok).
+	const FMetasoundFrontendDocument& Document = static_cast<const FMetaSoundFrontendDocumentBuilder&>(OutBuilder).GetDocument();
+	const TArray<FMetasoundFrontendNode>& GraphNodes = Document.RootGraph.Graph.Nodes;
+	for (const FMetasoundFrontendNode& Node : GraphNodes)
+	{
+		if (const FMetasoundFrontendClass* Class = OutBuilder.FindDependency(Node.ClassID))
+		{
+			// Inputs no longer have locational data as input template nodes provide that
+			if (Class->Metadata.GetType() == EMetasoundFrontendClassType::Input)
+			{
+				OutBuilder.RemoveNodeLocation(Node.GetID());
+			}
+			else
+			{
+				const TMap<FGuid, FVector2D>& Locations = Node.Style.Display.Locations;
+				if (Locations.Num() > 1)
+				{
+					TPair<FGuid, FVector2D> DefaultLocation;
+					for (const TPair<FGuid, FVector2D>& Pair : Locations)
+					{
+						DefaultLocation = Pair;
+						break;
+					}
+
+					// Remove first in case there are multiple locations and the editor guid may be different
+					OutBuilder.RemoveNodeLocation(Node.GetID());
+					OutBuilder.SetNodeLocation(Node.GetID(), DefaultLocation.Value, &DefaultLocation.Key);
+				}
+			}
+		}
+	}
+}
+
 void UMetasoundEditorGraph::SetPreviewID(uint32 InPreviewID)
 {
 	PreviewID = InPreviewID;
@@ -1990,6 +2068,7 @@ void UMetasoundEditorGraph::ValidateInternal(Metasound::Editor::FGraphValidation
 	bool bNodeIdFound = false;
 	for (UMetasoundEditorGraphNode* Node : NodesToValidate)
 	{
+		Node->CacheBreadcrumb();
 		FGraphNodeValidationResult NodeResult(*Node);
 
 		// Validate there is only 1 editor node per guid 
@@ -2009,98 +2088,4 @@ void UMetasoundEditorGraph::ValidateInternal(Metasound::Editor::FGraphValidation
 		OutResults.NodeResults.Add(MoveTemp(NodeResult));
 	}
 }
-
-bool UMetasoundEditorGraph::RemoveMember(UMetasoundEditorGraphMember& InGraphMember)
-{
-	bool bSuccess = RemoveMemberNodes(InGraphMember);
-	int32 NumRemoved = 0;
-	if (UMetasoundEditorGraphInput* Input = Cast<UMetasoundEditorGraphInput>(&InGraphMember))
-	{
-		bSuccess = RemoveFrontendInput(*Input);
-		NumRemoved = Inputs.Remove(Input);
-	}
-	else if (UMetasoundEditorGraphOutput* Output = Cast<UMetasoundEditorGraphOutput>(&InGraphMember))
-	{
-		bSuccess = RemoveFrontendOutput(*Output);
-		NumRemoved = Outputs.Remove(Output);
-	}
-	else if (UMetasoundEditorGraphVariable* Variable = Cast<UMetasoundEditorGraphVariable>(&InGraphMember))
-	{
-		bSuccess = RemoveFrontendVariable(*Variable);
-		NumRemoved = Variables.Remove(Variable);
-	}
-	else
-	{
-		bSuccess = false;
-	}
-
-	bSuccess = bSuccess && (NumRemoved > 0);
-
-	return bSuccess;
-}
-
-bool UMetasoundEditorGraph::RemoveMemberNodes(UMetasoundEditorGraphMember& InGraphMember)
-{
-	bool bSuccess = true;
-	for (UMetasoundEditorGraphNode* Node : InGraphMember.GetNodes())
-	{
-		if (ensure(Node))
-		{
-			bSuccess &= Metasound::Editor::FGraphBuilder::DeleteNode(*Node);
-		}
-	}
-	return bSuccess;
-}
-
-bool UMetasoundEditorGraph::RemoveFrontendMember(UMetasoundEditorGraphMember& InMember)
-{
-	if (UMetasoundEditorGraphInput* Input = Cast<UMetasoundEditorGraphInput>(&InMember))
-	{
-		return RemoveFrontendInput(*Input);
-	}
-	else if (UMetasoundEditorGraphOutput* Output = Cast<UMetasoundEditorGraphOutput>(&InMember))
-	{
-		return RemoveFrontendOutput(*Output);
-	}
-	else if (UMetasoundEditorGraphVariable* Variable = Cast<UMetasoundEditorGraphVariable>(&InMember))
-	{
-		return RemoveFrontendVariable(*Variable);
-	}
-
-	return false;
-}
-
-bool UMetasoundEditorGraph::RemoveFrontendInput(UMetasoundEditorGraphInput& Input)
-{
-	using namespace Metasound::Frontend;
-
-	FGraphHandle Graph = GetGraphHandle();
-	FConstNodeHandle InputNode = Graph->GetNodeWithID(Input.NodeID);
-	return Graph->RemoveInputVertex(InputNode->GetNodeName());
-}
-
-bool UMetasoundEditorGraph::RemoveFrontendOutput(UMetasoundEditorGraphOutput& Output)
-{
-	using namespace Metasound::Frontend;
-
-	FGraphHandle Graph = GetGraphHandle();
-	FConstNodeHandle OutputNode = Graph->GetNodeWithID(Output.NodeID);
-	return Graph->RemoveOutputVertex(OutputNode->GetNodeName());
-}
-
-bool UMetasoundEditorGraph::RemoveFrontendVariable(UMetasoundEditorGraphVariable& Variable)
-{
-	FGuid VariableID = Variable.GetVariableID();
-
-	// If the UMetasoundEditorGraphVariable is being deleted via an undo action, then the VariableID
-	// will be invalid and the frontend variable will already have been cleaned up.
-	if (VariableID.IsValid())
-	{
-		return GetGraphHandle()->RemoveVariable(VariableID);
-	}
-
-	return true;
-}
-
 #undef LOCTEXT_NAMESPACE
-
