Version Metasound Document to include all ed data and make all Metasound EdGraph data transient

diff --git a/Engine/Plugins/Runtime/Metasound/Source/MetasoundEditor/Private/MetasoundEditorGraphBuilder.cpp b/Engine/Plugins/Runtime/Metasound/Source/MetasoundEditor/Private/MetasoundEditorGraphBuilder.cpp
--- a/Engine/Plugins/Runtime/Metasound/Source/MetasoundEditor/Private/MetasoundEditorGraphBuilder.cpp
+++ b/Engine/Plugins/Runtime/Metasound/Source/MetasoundEditor/Private/MetasoundEditorGraphBuilder.cpp
@@ -1,7 +1,6 @@
 // Copyright Epic Games, Inc. All Rights Reserved.
 #include "MetasoundEditorGraphBuilder.h"
 
-#include "MetasoundEditorGraphBuilder.h"
 #include "Algo/AnyOf.h"
 #include "Algo/Sort.h"
 #include "Algo/Transform.h"
@@ -11,12 +10,15 @@
 #include "EdGraph/EdGraphPin.h"
 #include "Framework/Notifications/NotificationManager.h"
 #include "GraphEditor.h"
+#include "GraphEditorSettings.h"
 #include "Kismet/KismetSystemLibrary.h"
 #include "Metasound.h"
 #include "MetasoundAssetBase.h"
 #include "MetasoundAudioBuffer.h"
+#include "MetasoundBuilderSubsystem.h"
 #include "MetasoundEditor.h"
 #include "MetasoundEditorGraph.h"
+#include "MetasoundEditorGraphCommentNode.h"
 #include "MetasoundEditorGraphInputNode.h"
 #include "MetasoundEditorGraphMemberDefaults.h"
 #include "MetasoundEditorGraphNode.h"
@@ -42,6 +44,7 @@
 #include "MetasoundVertex.h"
 #include "MetasoundWaveTable.h"
 #include "Modules/ModuleManager.h"
+#include "NodeTemplates/MetasoundFrontendNodeTemplateInput.h"
 #include "NodeTemplates/MetasoundFrontendNodeTemplateReroute.h"
 #include "Templates/Tuple.h"
 #include "Toolkits/ToolkitManager.h"
@@ -57,16 +60,34 @@ namespace Metasound
 	{
 		namespace GraphBuilderPrivate
 		{
-			void DeleteNode(UObject& InMetaSound, Frontend::FNodeHandle InNodeHandle)
+			template <
+				typename TEdGraphNode,
+				typename FInitMetaSoundNodeFunc = TFunctionRef<void(UMetasoundEditorGraph&, TEdGraphNode&)>>
+			TEdGraphNode* AddNode(UObject& InMetaSound, FInitMetaSoundNodeFunc InitNodeFunc, bool bInSelectNewNode)
 			{
-				if (InNodeHandle->IsValid())
+				FMetasoundAssetBase* MetaSoundAsset = IMetasoundUObjectRegistry::Get().GetObjectAsAssetBase(&InMetaSound);
+				check(MetaSoundAsset);
+				UEdGraph& Graph = MetaSoundAsset->GetGraphChecked();
+				FGraphNodeCreator<TEdGraphNode> NodeCreator(Graph);
+				if (TEdGraphNode* NewGraphNode = NodeCreator.CreateNode(bInSelectNewNode))
 				{
-					Frontend::FGraphHandle GraphHandle = InNodeHandle->GetOwningGraph();
-					if (GraphHandle->IsValid())
-					{
-						GraphHandle->RemoveNode(*InNodeHandle);
-					}
+					// Required to happen prior to caching title and syncing location in case underlying type requires
+					// additional logic to initialize state in order to cache the title/sync the location
+					UMetasoundEditorGraph* MetasoundGraph = CastChecked<UMetasoundEditorGraph>(&Graph);
+					check(MetasoundGraph);
+					InitNodeFunc(*MetasoundGraph, *NewGraphNode);
+
+					NodeCreator.Finalize();
+
+					NewGraphNode->CacheTitle();
+
+					// Override finalized EdGraphNode ID if a location is found.
+					constexpr bool bUpdateEditorNodeID = true;
+					NewGraphNode->SyncLocationFromFrontendNode(bUpdateEditorNodeID);
+					return NewGraphNode;
 				}
+
+				return nullptr;
 			}
 
 			FName GenerateUniqueName(const TArray<FName>& InExistingNames, const FString& InBaseName)
@@ -99,7 +120,7 @@ namespace Metasound
 				}
 			};
 
-			bool SynchronizeGraphRecursively(UObject& InMetaSound, bool bEditorGraphModified = false)
+			void SynchronizeGraphRecursively(UObject& InMetaSound)
 			{
 				using namespace Frontend;
 
@@ -113,26 +134,25 @@ namespace Metasound
 				for (FMetasoundAssetBase* Reference : References)
 				{
 					check(Reference);
-					bEditorGraphModified |= SynchronizeGraphRecursively(*Reference->GetOwningAsset(), bEditorGraphModified);
+					SynchronizeGraphRecursively(*Reference->GetOwningAsset());
 				}
 
-				if (!MetaSoundAsset->GetModifyContext().GetDocumentModified())
+				if (!MetaSoundAsset->GetConstModifyContext().GetDocumentModified())
 				{
-					return bEditorGraphModified;
+					return;
 				}
 
 				// If no graph is set, MetaSound has been created outside of asset factory, so initialize it here.
 				if (!MetaSoundAsset->GetGraph())
 				{
-					UMetaSoundEditorSubsystem::GetChecked().InitAsset(InMetaSound);
+					UMetaSoundEditorSubsystem::GetChecked().InitEdGraph(InMetaSound);
 				}
 
-				bEditorGraphModified |= FGraphBuilder::SynchronizeGraphMembers(InMetaSound);
-				bEditorGraphModified |= FGraphBuilder::SynchronizeNodeMembers(InMetaSound);
-				bEditorGraphModified |= FGraphBuilder::SynchronizeNodes(InMetaSound);
-				bEditorGraphModified |= FGraphBuilder::SynchronizeConnections(InMetaSound);
-
-				return bEditorGraphModified;
+				FGraphBuilder::SynchronizeComments(InMetaSound);
+				FGraphBuilder::SynchronizeGraphMembers(InMetaSound);
+				FGraphBuilder::SynchronizeOutputNodes(InMetaSound);
+				FGraphBuilder::SynchronizeNodes(InMetaSound);
+				FGraphBuilder::SynchronizeConnections(InMetaSound);
 			}
 		} // namespace GraphBuilderPrivate
 
@@ -232,6 +252,22 @@ namespace Metasound
 			return true;
 		}
 
+		UMetasoundEditorGraphCommentNode* FGraphBuilder::CreateCommentNode(UObject& InMetaSound, bool bInSelectNewNode, FGuid InCommentID)
+		{
+			FMetasoundAssetBase* MetaSoundAsset = IMetasoundUObjectRegistry::Get().GetObjectAsAssetBase(&InMetaSound);
+			check(MetaSoundAsset);
+			UEdGraph& Graph = MetaSoundAsset->GetGraphChecked();
+			FGraphNodeCreator<UMetasoundEditorGraphCommentNode> NodeCreator(Graph);
+			if (UMetasoundEditorGraphCommentNode* NewNode = NodeCreator.CreateNode(bInSelectNewNode))
+			{
+				NewNode->CommentID = InCommentID;
+				NodeCreator.Finalize();
+				return NewNode;
+			}
+
+			return nullptr;
+		}
+
 		FText FGraphBuilder::GetDisplayName(const FMetasoundFrontendClassMetadata& InClassMetadata, FName InNodeName, bool bInIncludeNamespace)
 		{
 			using namespace Frontend;
@@ -372,6 +408,7 @@ namespace Metasound
 					return OwningNode->GetNodeName();
 				}
 
+				case EMetasoundFrontendClassType::Template:
 				default:
 				{
 					return InFrontendOutput.GetName();
@@ -406,6 +443,7 @@ namespace Metasound
 					return OwningNode->GetNodeName();
 				}
 
+				case EMetasoundFrontendClassType::Template:
 				default:
 				{
 					return InFrontendInput.GetName();
@@ -413,158 +451,157 @@ namespace Metasound
 			}
 		}
 
-		UMetasoundEditorGraphExternalNode* FGraphBuilder::AddExternalNode(UObject& InMetaSound, Frontend::FNodeHandle& InNodeHandle, FVector2D InLocation, bool bInSelectNewNode)
+		UMetasoundEditorGraphExternalNode* FGraphBuilder::AddTemplateNode(UObject& InMetaSound, const FGuid& InNodeID, const FMetasoundFrontendClassMetadata& InMetadata, bool bInSelectNewNode)
 		{
-			using namespace Frontend;
-
-			UMetasoundEditorGraphExternalNode* NewGraphNode = nullptr;
-
-			const EMetasoundFrontendClassType ClassType = InNodeHandle->GetClassMetadata().GetType();
-			const bool bIsExternalNode = ClassType == EMetasoundFrontendClassType::External;
-			const bool bIsTemplateNode = ClassType == EMetasoundFrontendClassType::Template;
-			if (!ensure(bIsExternalNode || bIsTemplateNode))
+			const EMetasoundFrontendClassType ClassType = InMetadata.GetType();
+			if (ensureMsgf(ClassType == EMetasoundFrontendClassType::Template, TEXT("Cannot call 'AddTemplateNode' with node of class type '%s'."), LexToString(ClassType)))
 			{
-				return nullptr;
+				auto InitNodeFunc = [&InMetadata, &InNodeID](UMetasoundEditorGraph&, UMetasoundEditorGraphExternalNode& NewGraphNode)
+				{
+					NewGraphNode.NodeID = InNodeID;
+					NewGraphNode.ClassName = InMetadata.GetClassName();
+				};
+				return GraphBuilderPrivate::AddNode<UMetasoundEditorGraphExternalNode>(InMetaSound, InitNodeFunc, bInSelectNewNode);
 			}
 
-			FMetasoundAssetBase* MetaSoundAsset = IMetasoundUObjectRegistry::Get().GetObjectAsAssetBase(&InMetaSound);
-			check(MetaSoundAsset);
-			UEdGraph& Graph = MetaSoundAsset->GetGraphChecked();
-			FGraphNodeCreator<UMetasoundEditorGraphExternalNode> NodeCreator(Graph);
+			return nullptr;
+		}
 
-			NewGraphNode = NodeCreator.CreateNode(bInSelectNewNode);
-			if (ensure(NewGraphNode))
+		UMetasoundEditorGraphExternalNode* FGraphBuilder::AddExternalNode(UObject& InMetaSound, const FGuid& InNodeID, const FMetasoundFrontendClassMetadata& InMetadata, bool bInSelectNewNode)
+		{
+			using namespace Frontend;
+
+			const EMetasoundFrontendClassType ClassType = InMetadata.GetType();
+			if (ensureMsgf(ClassType == EMetasoundFrontendClassType::External, TEXT("Cannot call 'AddExternalNode' with node of class type '%s'."), LexToString(ClassType)))
 			{
-				const FNodeRegistryKey RegistryKey = FNodeRegistryKey(InNodeHandle->GetClassMetadata());
-				NewGraphNode->bIsClassNative = FMetasoundFrontendRegistryContainer::Get()->IsNodeNative(RegistryKey);
-				NewGraphNode->ClassName = InNodeHandle->GetClassMetadata().GetClassName();
-				NewGraphNode->CacheTitle();
+				auto InitNodeFunc = [&InMetadata, &InNodeID](UMetasoundEditorGraph& MetasoundGraph, UMetasoundEditorGraphExternalNode& NewGraphNode)
+				{
+					const FNodeRegistryKey RegistryKey(InMetadata);
+					NewGraphNode.bIsClassNative = FMetasoundFrontendRegistryContainer::Get()->IsNodeNative(RegistryKey);
+					NewGraphNode.NodeID = InNodeID;
+					NewGraphNode.ClassName = InMetadata.GetClassName();
 
-				NodeCreator.Finalize();
-				InitGraphNode(InNodeHandle, NewGraphNode, InMetaSound);
-				NewGraphNode->SetNodeLocation(InLocation);
+					if (!NewGraphNode.bIsClassNative)
+					{
+						FMetasoundAssetBase* MetaSoundAsset = IMetasoundUObjectRegistry::Get().GetObjectAsAssetBase(MetasoundGraph.GetMetasound());
+						check(MetaSoundAsset);
+						MetaSoundAsset->RebuildReferencedAssetClasses();
+					}
+				};
 
-				// Adding external node may introduce referenced asset so rebuild referenced keys.
-				MetaSoundAsset->RebuildReferencedAssetClasses();
+				return GraphBuilderPrivate::AddNode<UMetasoundEditorGraphExternalNode>(InMetaSound, InitNodeFunc, bInSelectNewNode);
 			}
 
-			return NewGraphNode;
+			return nullptr;
 		}
 
-		UMetasoundEditorGraphExternalNode* FGraphBuilder::AddExternalNode(UObject& InMetaSound, const FMetasoundFrontendClassMetadata& InMetadata, FVector2D InLocation, bool bInSelectNewNode)
+		UMetasoundEditorGraphExternalNode* FGraphBuilder::AddExternalNode(UObject& InMetaSound, const FMetasoundFrontendClassMetadata& InMetadata, bool bInSelectNewNode)
 		{
 			FMetasoundAssetBase* MetaSoundAsset = IMetasoundUObjectRegistry::Get().GetObjectAsAssetBase(&InMetaSound);
 			check(MetaSoundAsset);
 
-			Frontend::FNodeHandle NodeHandle = MetaSoundAsset->GetRootGraphHandle()->AddNode(InMetadata);
-			return AddExternalNode(InMetaSound, NodeHandle, InLocation, bInSelectNewNode);
+			Frontend::FConstNodeHandle NodeHandle = MetaSoundAsset->GetRootGraphHandle()->AddNode(InMetadata);
+			return AddExternalNode(InMetaSound, NodeHandle->GetID(), NodeHandle->GetClassMetadata(), bInSelectNewNode);
 		}
 
-		UMetasoundEditorGraphVariableNode* FGraphBuilder::AddVariableNode(UObject& InMetaSound, Frontend::FNodeHandle& InNodeHandle, FVector2D InLocation, bool bInSelectNewNode)
+		Frontend::FNodeHandle FGraphBuilder::AddExternalNodeHandle(UObject& InMetaSound, const FMetasoundFrontendClassName& InClassName)
 		{
 			using namespace Frontend;
 
-			EMetasoundFrontendClassType ClassType = InNodeHandle->GetClassMetadata().GetType();
-			const bool bIsSupportedClassType = (ClassType == EMetasoundFrontendClassType::VariableAccessor) 
-				|| (ClassType == EMetasoundFrontendClassType::VariableDeferredAccessor)
-				|| (ClassType == EMetasoundFrontendClassType::VariableMutator);
-
-			if (!ensure(bIsSupportedClassType))
+			FMetasoundFrontendClass FrontendClass;
+			bool bDidFindClassWithName = ISearchEngine::Get().FindClassWithHighestVersion(InClassName, FrontendClass);
+			if (ensure(bDidFindClassWithName))
 			{
-				return nullptr;
+				FMetasoundAssetBase* MetaSoundAsset = IMetasoundUObjectRegistry::Get().GetObjectAsAssetBase(&InMetaSound);
+				check(MetaSoundAsset);
+				return MetaSoundAsset->GetRootGraphHandle()->AddNode(FrontendClass.Metadata);
 			}
 
-			FConstVariableHandle FrontendVariable = InNodeHandle->GetOwningGraph()->FindVariableContainingNode(InNodeHandle->GetID());
-			if (!ensure(FrontendVariable->IsValid()))
-			{
-				return nullptr;
-			}
+			return INodeController::GetInvalidHandle();
+		}
 
-			UMetasoundEditorGraphVariableNode* NewGraphNode = nullptr;
-			FMetasoundAssetBase* MetaSoundAsset = IMetasoundUObjectRegistry::Get().GetObjectAsAssetBase(&InMetaSound);
-			if (ensure(nullptr != MetaSoundAsset))
+		UMetasoundEditorGraphVariableNode* FGraphBuilder::AddVariableNode(UObject& InMetaSound, const Frontend::FConstNodeHandle& InNodeHandle, bool bInSelectNewNode)
+		{
+			using namespace Frontend;
+
+			const FMetasoundFrontendClassMetadata& Metadata = InNodeHandle->GetClassMetadata();
+			EMetasoundFrontendClassType ClassType = Metadata.GetType();
+			const bool bIsSupportedClassType = (ClassType == EMetasoundFrontendClassType::VariableAccessor) 
+				|| (ClassType == EMetasoundFrontendClassType::VariableDeferredAccessor)
+				|| (ClassType == EMetasoundFrontendClassType::VariableMutator);
+
+			if (ensureMsgf(bIsSupportedClassType, TEXT("Cannot call 'AddVariableNode' with node of class type '%s'"), LexToString(ClassType)))
 			{
-				if (UMetasoundEditorGraph* MetasoundGraph = CastChecked<UMetasoundEditorGraph>(MetaSoundAsset->GetGraph()))
+				FConstVariableHandle FrontendVariable = InNodeHandle->GetOwningGraph()->FindVariableContainingNode(InNodeHandle->GetID());
+				if (ensure(FrontendVariable->IsValid()))
 				{
-					FGraphNodeCreator<UMetasoundEditorGraphVariableNode> NodeCreator(*MetasoundGraph);
-
-					NewGraphNode = NodeCreator.CreateNode(bInSelectNewNode);
-					if (ensure(NewGraphNode))
+					auto InitNodeFunc = [&FrontendVariable, &InNodeHandle, &Metadata](UMetasoundEditorGraph& MetasoundGraph, UMetasoundEditorGraphVariableNode& NewGraphNode)
 					{
-						NewGraphNode->ClassName = InNodeHandle->GetClassMetadata().GetClassName();
-						NewGraphNode->ClassType = ClassType;
-						InitGraphNode(InNodeHandle, NewGraphNode, InMetaSound);
-						NodeCreator.Finalize();
-
-						UMetasoundEditorGraphVariable* Variable = MetasoundGraph->FindOrAddVariable(FrontendVariable);
+						UMetasoundEditorGraphVariable* Variable = MetasoundGraph.FindOrAddVariable(FrontendVariable);
 						if (ensure(Variable))
 						{
-							NewGraphNode->Variable = Variable;
-
-							// Ensures the variable node value is synced with the editor literal value should it be set
-							constexpr bool bPostTransaction = false;
-							Variable->UpdateFrontendDefaultLiteral(bPostTransaction);
+							NewGraphNode.Variable = Variable;
+							NewGraphNode.NodeID = InNodeHandle->GetID();
+							NewGraphNode.ClassName = Metadata.GetClassName();
+							NewGraphNode.ClassType = Metadata.GetType();
 						}
+					};
 
-						MetasoundGraph->GetModifyContext().AddNodeIDsModified({ NewGraphNode->GetNodeID() });
-						NewGraphNode->SetNodeLocation(InLocation);
-					}
+					return GraphBuilderPrivate::AddNode<UMetasoundEditorGraphVariableNode>(InMetaSound, InitNodeFunc, bInSelectNewNode);
 				}
 			}
 
-			return NewGraphNode;
+			return nullptr;
 		}
 
-		UMetasoundEditorGraphOutputNode* FGraphBuilder::AddOutputNode(UObject& InMetaSound, Frontend::FNodeHandle& InNodeHandle, FVector2D InLocation, bool bInSelectNewNode)
+		UMetasoundEditorGraphInputNode* FGraphBuilder::AddInputNode(UObject& InMetaSound, const Frontend::FConstNodeHandle& InNodeHandle, bool bInSelectNewNode)
 		{
 			using namespace Frontend;
 
-			UMetasoundEditorGraphOutputNode* NewGraphNode = nullptr;
-			if (!ensure(InNodeHandle->GetClassMetadata().GetType() == EMetasoundFrontendClassType::Output))
+			const FMetasoundFrontendClassMetadata& InputMetadata = InNodeHandle->GetClassMetadata();
+			if (ensureMsgf(InputMetadata.GetType() == EMetasoundFrontendClassType::Template, TEXT("Cannot call 'AddInputNode' with node of class type '%s': Must be input template."), LexToString(InputMetadata.GetType())))
 			{
-				return nullptr;
-			}
-
-			FMetasoundAssetBase* MetaSoundAsset = IMetasoundUObjectRegistry::Get().GetObjectAsAssetBase(&InMetaSound);
-			check(MetaSoundAsset);
-			UEdGraph& Graph = MetaSoundAsset->GetGraphChecked();
-			FGraphNodeCreator<UMetasoundEditorGraphOutputNode> NodeCreator(Graph);
-
-			NewGraphNode = NodeCreator.CreateNode(bInSelectNewNode);
-			if (ensure(NewGraphNode))
-			{
-				UMetasoundEditorGraph* MetasoundGraph = CastChecked<UMetasoundEditorGraph>(&Graph);
-
-				UMetasoundEditorGraphOutput* Output = MetasoundGraph->FindOrAddOutput(InNodeHandle);
-				if (ensure(Output))
+				if (ensureMsgf(InputMetadata.GetClassName() == FInputNodeTemplate::ClassName, TEXT("Cannot call 'AddInputNode with node handle that is not of class '%s'"), *InputMetadata.GetClassName().ToString()))
 				{
-					NewGraphNode->Output = Output;
-					NodeCreator.Finalize();
-					InitGraphNode(InNodeHandle, NewGraphNode, InMetaSound);
-
-					// Ensures the output node value is synced with the editor literal value should it be set
-					constexpr bool bPostTransaction = false;
-					Output->UpdateFrontendDefaultLiteral(bPostTransaction);
+					auto InitNodeFunc = [&InNodeHandle](UMetasoundEditorGraph& MetasoundGraph, UMetasoundEditorGraphInputNode& NewGraphNode)
+					{
+						FConstInputHandle TemplateInput = InNodeHandle->GetConstInputs().Last();
+						FConstOutputHandle InputNodeOutput = TemplateInput->GetConnectedOutput();
+						UMetasoundEditorGraphInput* Input = MetasoundGraph.FindOrAddInput(InputNodeOutput->GetOwningNode());
+						if (ensure(Input))
+						{
+							NewGraphNode.Input = Input;
+							NewGraphNode.NodeID = InNodeHandle->GetID();
+						}
+					};
 
-					MetasoundGraph->GetModifyContext().AddNodeIDsModified({ NewGraphNode->GetNodeID() });
+					return GraphBuilderPrivate::AddNode<UMetasoundEditorGraphInputNode>(InMetaSound, InitNodeFunc, bInSelectNewNode);
 				}
-
-				NewGraphNode->CacheTitle();
-				NewGraphNode->SetNodeLocation(InLocation);
 			}
 
-			return NewGraphNode;
+			return nullptr;
 		}
 
-		void FGraphBuilder::InitGraphNode(Frontend::FNodeHandle& InNodeHandle, UMetasoundEditorGraphNode* NewGraphNode, UObject& InMetaSound)
+		UMetasoundEditorGraphOutputNode* FGraphBuilder::AddOutputNode(UObject& InMetaSound, const Frontend::FConstNodeHandle& InNodeHandle, bool bInSelectNewNode)
 		{
-			NewGraphNode->SetNodeID(InNodeHandle->GetID());
-			RebuildNodePins(*NewGraphNode);
-		}
+			using namespace Frontend;
 
-		void FGraphBuilder::InitGraphNodeIDFromNodeHandle(const Frontend::FConstNodeHandle& InNodeHandle, UMetasoundEditorGraphNode* NewGraphNode)
-		{
-			NewGraphNode->SetNodeID(InNodeHandle->GetID());
+			const EMetasoundFrontendClassType ClassType = InNodeHandle->GetClassMetadata().GetType();
+			if (ensureMsgf(ClassType == EMetasoundFrontendClassType::Output, TEXT("Cannot call 'AddOutputNode' with node of class type '%s'"), LexToString(ClassType)))
+			{
+				auto InitNodeFunc = [&InNodeHandle](UMetasoundEditorGraph& MetasoundGraph, UMetasoundEditorGraphOutputNode& NewGraphNode)
+				{
+					UMetasoundEditorGraphOutput* Output = MetasoundGraph.FindOrAddOutput(InNodeHandle);
+					if (ensure(Output))
+					{
+						NewGraphNode.Output = Output;
+					}
+				};
+
+				return GraphBuilderPrivate::AddNode<UMetasoundEditorGraphOutputNode>(InMetaSound, InitNodeFunc, bInSelectNewNode);
+			}
+
+			return nullptr;
 		}
 
 		FGraphValidationResults FGraphBuilder::ValidateGraph(UObject& InMetaSound)
@@ -976,7 +1013,26 @@ namespace Metasound
 
 		Frontend::FConstOutputHandle FGraphBuilder::GetConstOutputHandleFromPin(const UEdGraphPin* InPin)
 		{
-			return GetOutputHandleFromPin(InPin);
+			using namespace Frontend;
+			using namespace VariableNames;
+
+			if (InPin && ensure(InPin->Direction == EGPD_Output))
+			{
+				if (const UMetasoundEditorGraphVariableNode* EdVariableNode = Cast<UMetasoundEditorGraphVariableNode>(InPin->GetOwningNode()))
+				{
+					// UEdGraphPins on variable nodes use the variable's name for display
+					// purposes instead of the underlying vertex's name. The frontend vertices
+					// of a variable node have consistent names no matter what the 
+					// variable is named.
+					return EdVariableNode->GetConstNodeHandle()->GetConstOutputWithVertexName(METASOUND_GET_PARAM_NAME(OutputData));
+				}
+				else if (const UMetasoundEditorGraphNode* EdNode = CastChecked<UMetasoundEditorGraphNode>(InPin->GetOwningNode()))
+				{
+					return EdNode->GetConstNodeHandle()->GetConstOutputWithVertexName(InPin->GetFName());
+				}
+			}
+
+			return IOutputController::GetInvalidHandle();
 		}
 
 		UEdGraphPin* FGraphBuilder::FindReroutedOutputPin(UEdGraphPin* OutputPin)
@@ -987,7 +1043,7 @@ namespace Metasound
 			{
 				if (UMetasoundEditorGraphExternalNode* ExternalNode = Cast<UMetasoundEditorGraphExternalNode>(OutputPin->GetOwningNode()))
 				{
-					if (ExternalNode->GetClassName() == FRerouteNodeTemplate::ClassName)
+					if (ExternalNode->GetBreadcrumb().ClassName == FRerouteNodeTemplate::ClassName)
 					{
 						auto IsInput = [](UEdGraphPin* Pin) { check(Pin); return Pin->Direction == EGPD_Input; };
 						if (UEdGraphPin* RerouteInput = *ExternalNode->Pins.FindByPredicate(IsInput))
@@ -1014,7 +1070,7 @@ namespace Metasound
 			{
 				if (UMetasoundEditorGraphExternalNode* ExternalNode = Cast<UMetasoundEditorGraphExternalNode>(OutputPin->GetOwningNode()))
 				{
-					if (ExternalNode->GetClassName() == FRerouteNodeTemplate::ClassName)
+					if (ExternalNode->GetBreadcrumb().ClassName == FRerouteNodeTemplate::ClassName)
 					{
 						auto IsInput = [](const UEdGraphPin* Pin) { check(Pin); return Pin->Direction == EGPD_Input; };
 						if (const UEdGraphPin* RerouteInput = *ExternalNode->Pins.FindByPredicate(IsInput))
@@ -1041,7 +1097,7 @@ namespace Metasound
 			{
 				if (UMetasoundEditorGraphExternalNode* ExternalNode = Cast<UMetasoundEditorGraphExternalNode>(OutputPin->GetOwningNode()))
 				{
-					if (ExternalNode->GetClassName() == FRerouteNodeTemplate::ClassName)
+					if (ExternalNode->GetBreadcrumb().ClassName == FRerouteNodeTemplate::ClassName)
 					{
 						auto IsInput = [](UEdGraphPin* Pin) { check(Pin); return Pin->Direction == EGPD_Input; };
 						if (UEdGraphPin* RerouteInput = *ExternalNode->Pins.FindByPredicate(IsInput))
@@ -1062,9 +1118,33 @@ namespace Metasound
 			return IOutputController::GetInvalidHandle();
 		}
 
-		Frontend::FConstOutputHandle FGraphBuilder::FindReroutedConstOutputHandleFromPin(const UEdGraphPin* InPin)
+		Frontend::FConstOutputHandle FGraphBuilder::FindReroutedConstOutputHandleFromPin(const UEdGraphPin* OutputPin)
 		{
-			return FindReroutedOutputHandleFromPin(InPin);
+			using namespace Frontend;
+
+			if (OutputPin)
+			{
+				if (const UMetasoundEditorGraphExternalNode* ExternalNode = Cast<UMetasoundEditorGraphExternalNode>(OutputPin->GetOwningNode()))
+				{
+					if (ExternalNode->GetBreadcrumb().ClassName == FRerouteNodeTemplate::ClassName)
+					{
+						auto IsInput = [](UEdGraphPin* Pin) { check(Pin); return Pin->Direction == EGPD_Input; };
+						if (const UEdGraphPin* RerouteInput = *ExternalNode->Pins.FindByPredicate(IsInput))
+						{
+							const TArray<UEdGraphPin*>& LinkedTo = RerouteInput->LinkedTo;
+							if (!LinkedTo.IsEmpty())
+							{
+								const UEdGraphPin* ReroutedOutput = LinkedTo.Last();
+								return FindReroutedConstOutputHandleFromPin(ReroutedOutput);
+							}
+						}
+					}
+				}
+
+				return GetConstOutputHandleFromPin(OutputPin);
+			}
+
+			return IOutputController::GetInvalidHandle();
 		}
 
 		void FGraphBuilder::FindReroutedInputPins(UEdGraphPin* InPinToCheck, TArray<UEdGraphPin*>& InOutInputPins)
@@ -1075,7 +1155,7 @@ namespace Metasound
 			{
 				if (UMetasoundEditorGraphExternalNode* ExternalNode = Cast<UMetasoundEditorGraphExternalNode>(InPinToCheck->GetOwningNode()))
 				{
-					if (ExternalNode->GetClassName() == FRerouteNodeTemplate::ClassName)
+					if (ExternalNode->GetBreadcrumb().ClassName == FRerouteNodeTemplate::ClassName)
 					{
 						for (UEdGraphPin* Pin : ExternalNode->Pins)
 						{
@@ -1111,60 +1191,43 @@ namespace Metasound
 			return Algo::AnyOf(EditorNodes, [](const UMetasoundEditorGraphNode* Node) { return Node->ErrorType == EMessageSeverity::Error; });
 		}
 
-		bool FGraphBuilder::SynchronizeNodeLocation(const Frontend::FConstNodeHandle& InNode, UMetasoundEditorGraphNode& OutGraphNode)
+		Frontend::FNodeHandle FGraphBuilder::AddInputTemplateNodeHandle(UObject& InMetaSound, Frontend::FNodeHandle& InputNodeHandle)
 		{
-			bool bModified = false;
+			using namespace Frontend;
 
-			const FMetasoundFrontendNodeStyle& Style = InNode->GetNodeStyle();
+			FMetasoundAssetBase* MetaSoundAsset = IMetasoundUObjectRegistry::Get().GetObjectAsAssetBase(&InMetaSound);
+			check(MetaSoundAsset);
 
-			const FVector2D* Location = Style.Display.Locations.Find(OutGraphNode.NodeGuid);
-			if (!Location)
+			UMetasoundEditorGraph* Graph = Cast<UMetasoundEditorGraph>(MetaSoundAsset->GetGraph());
+			if (!ensure(Graph))
 			{
-				// If no specific location found, use default location if provided (zero guid
-				// for example, provided by preset defaults.)
-				Location = Style.Display.Locations.Find({ });
+				return INodeController::GetInvalidHandle();
 			}
 
-			if (Location)
+			TArray<FConstOutputHandle> NodeOutputs = InputNodeHandle->GetConstOutputs();
+			if (!ensure(!NodeOutputs.IsEmpty()))
 			{
-				const int32 LocX = FMath::TruncToInt(Location->X);
-				const int32 LocY = FMath::TruncToInt(Location->Y);
-				const bool bXChanged = static_cast<bool>(LocX - OutGraphNode.NodePosX);
-				const bool bYChanged = static_cast<bool>(LocY - OutGraphNode.NodePosY);
-				if (bXChanged || bYChanged)
-				{
-					OutGraphNode.NodePosX = LocX;
-					OutGraphNode.NodePosY = LocY;
-					bModified = true;
-				}
+				return INodeController::GetInvalidHandle();
 			}
 
-			return bModified;
-		}
-
-		UMetasoundEditorGraphInputNode* FGraphBuilder::AddInputNode(UObject& InMetaSound, Frontend::FNodeHandle InNodeHandle, FVector2D InLocation, bool bInSelectNewNode)
-		{
-			using namespace Frontend;
-
-			FMetasoundAssetBase* MetaSoundAsset = IMetasoundUObjectRegistry::Get().GetObjectAsAssetBase(&InMetaSound);
-			check(MetaSoundAsset);
-
-			UMetasoundEditorGraph* MetasoundGraph = Cast<UMetasoundEditorGraph>(MetaSoundAsset->GetGraph());
-			if (!ensure(MetasoundGraph))
+			if (!ensure(InputNodeHandle->GetClassMetadata().GetType() == EMetasoundFrontendClassType::Input))
 			{
-				return nullptr;
+				return INodeController::GetInvalidHandle();
 			}
 
-			UMetasoundEditorGraphInputNode* NewGraphNode = MetasoundGraph->CreateInputNode(InNodeHandle, bInSelectNewNode);
-			if (ensure(NewGraphNode))
-			{
-				NewGraphNode->SetNodeLocation(InLocation);
-				RebuildNodePins(*NewGraphNode);
-				MetasoundGraph->GetModifyContext().AddNodeIDsModified({ NewGraphNode->GetNodeID() });
-				return NewGraphNode;
-			}
+			const INodeTemplate* InputTemplate = INodeTemplateRegistry::Get().FindTemplate(FInputNodeTemplate::ClassName);
+			check(InputTemplate);
 
-			return nullptr;
+			const FName DataType = NodeOutputs.Last()->GetDataType();
+			FNodeTemplateGenerateInterfaceParams Params { { DataType }, { } };
+
+			FGraphHandle GraphHandle = Graph->GetGraphHandle();
+			FNodeHandle TemplateNodeHandle = GraphHandle->AddTemplateNode(*InputTemplate, MoveTemp(Params));
+
+			TemplateNodeHandle->GetInputs().Last()->Connect(*InputNodeHandle->GetOutputs().Last());
+			FMetasoundFrontendDocumentModifyContext& ModifyContext = GetOutermostMetaSoundChecked(*Graph).GetModifyContext();
+			ModifyContext.AddNodeIDsModified({ TemplateNodeHandle->GetID() });
+			return TemplateNodeHandle;
 		}
 
 		bool FGraphBuilder::GetPinLiteral(UEdGraphPin& InInputPin, FMetasoundFrontendLiteral& OutDefaultLiteral)
@@ -1173,7 +1236,7 @@ namespace Metasound
 
 			IMetasoundEditorModule& EditorModule = FModuleManager::GetModuleChecked<IMetasoundEditorModule>("MetaSoundEditor");
 
-			FInputHandle InputHandle = GetInputHandleFromPin(&InInputPin);
+			FConstInputHandle InputHandle = GetConstInputHandleFromPin(&InInputPin);
 			if (!ensure(InputHandle->IsValid()))
 			{
 				return false;
@@ -1326,78 +1389,6 @@ namespace Metasound
 			return true;
 		}
 
-		Frontend::FNodeHandle FGraphBuilder::AddNodeHandle(UObject& InMetaSound, UMetasoundEditorGraphNode& InGraphNode)
-		{
-			using namespace Frontend;
-
-			FNodeHandle NodeHandle = INodeController::GetInvalidHandle();
-			if (UMetasoundEditorGraphInputNode* InputNode = Cast<UMetasoundEditorGraphInputNode>(&InGraphNode))
-			{
-				const TArray<UEdGraphPin*>& Pins = InGraphNode.GetAllPins();
-				const UEdGraphPin* Pin = Pins.IsEmpty() ? nullptr : Pins[0];
-				if (ensure(Pin) && ensure(Pin->Direction == EGPD_Output))
-				{
-					UMetasoundEditorGraphInput* Input = InputNode->Input;
-					if (ensure(Input))
-					{
-						const FName PinName = Pin->GetFName();
-						FCreateNodeVertexParams VertexParams;
-						VertexParams.DataType = Input->GetDataType();
-
-						NodeHandle = AddInputNodeHandle(InMetaSound, VertexParams, nullptr, &PinName);
-						NodeHandle->SetDescription(InGraphNode.GetTooltipText());
-					}
-				}
-			}
-
-			else if (UMetasoundEditorGraphOutputNode* OutputNode = Cast<UMetasoundEditorGraphOutputNode>(&InGraphNode))
-			{
-				const TArray<UEdGraphPin*>& Pins = InGraphNode.GetAllPins();
-				const UEdGraphPin* Pin = Pins.IsEmpty() ? nullptr : Pins[0];
-				if (ensure(Pin) && ensure(Pin->Direction == EGPD_Input))
-				{
-					UMetasoundEditorGraphOutput* Output = OutputNode->Output;
-					if (ensure(Output))
-					{
-						const FName PinName = Pin->GetFName();
-						FCreateNodeVertexParams VertexParams;
-						VertexParams.DataType = Output->GetDataType();
-
-						NodeHandle = FGraphBuilder::AddOutputNodeHandle(InMetaSound, VertexParams, &PinName);
-						NodeHandle->SetDescription(InGraphNode.GetTooltipText());
-					}
-				}
-			}
-			else if (UMetasoundEditorGraphVariableNode* VariableNode = Cast<UMetasoundEditorGraphVariableNode>(&InGraphNode))
-			{
-				NodeHandle = FGraphBuilder::AddVariableNodeHandle(InMetaSound, VariableNode->Variable->GetVariableID(), VariableNode->GetClassName().ToNodeClassName());
-			}
-			else if (UMetasoundEditorGraphExternalNode* ExternalNode = Cast<UMetasoundEditorGraphExternalNode>(&InGraphNode))
-			{
-				FMetasoundFrontendClass FrontendClass;
-				bool bDidFindClassWithName = ISearchEngine::Get().FindClassWithHighestVersion(ExternalNode->ClassName.ToNodeClassName(), FrontendClass);
-				if (ensure(bDidFindClassWithName))
-				{
-					FMetasoundAssetBase* MetaSoundAsset = IMetasoundUObjectRegistry::Get().GetObjectAsAssetBase(&InMetaSound);
-					check(MetaSoundAsset);
-
-					Frontend::FNodeHandle NewNode = MetaSoundAsset->GetRootGraphHandle()->AddNode(FrontendClass.Metadata);
-					ExternalNode->SetNodeID(NewNode->GetID());
-
-					NodeHandle = NewNode;
-				}
-			}
-
-			if (NodeHandle->IsValid())
-			{
-				FMetasoundFrontendNodeStyle Style = NodeHandle->GetNodeStyle();
-				Style.Display.Locations.Add(InGraphNode.NodeGuid, FVector2D(InGraphNode.NodePosX, InGraphNode.NodePosY));
-				NodeHandle->SetNodeStyle(Style);
-			}
-
-			return NodeHandle;
-		}
-
 		Frontend::FNodeHandle FGraphBuilder::AddInputNodeHandle(UObject& InMetaSound, const FCreateNodeVertexParams& InParams, const FMetasoundFrontendLiteral* InDefaultValue, const FName* InNameBase)
 		{
 			FMetasoundAssetBase* MetaSoundAsset = IMetasoundUObjectRegistry::Get().GetObjectAsAssetBase(&InMetaSound);
@@ -1457,6 +1448,20 @@ namespace Metasound
 			return GraphBuilderPrivate::GenerateUniqueName(ExistingVariableNames, InBaseName);
 		}
 
+		FMetasoundAssetBase& FGraphBuilder::GetOutermostMetaSoundChecked(UObject& InSubObject)
+		{
+			FMetasoundAssetBase* MetaSoundAsset = IMetasoundUObjectRegistry::Get().GetObjectAsAssetBase(InSubObject.GetOutermostObject());
+			check(MetaSoundAsset);
+			return *MetaSoundAsset;
+		}
+
+		const FMetasoundAssetBase& FGraphBuilder::GetOutermostConstMetaSoundChecked(const UObject& InSubObject)
+		{
+			const FMetasoundAssetBase* MetaSoundAsset = IMetasoundUObjectRegistry::Get().GetObjectAsAssetBase(InSubObject.GetOutermostObject());
+			check(MetaSoundAsset);
+			return *MetaSoundAsset;
+		}
+
 		Frontend::FVariableHandle FGraphBuilder::AddVariableHandle(UObject& InMetaSound, const FName& InTypeName)
 		{
 			using namespace Frontend;
@@ -1505,24 +1510,24 @@ namespace Metasound
 							break;
 
 						case EMetasoundFrontendClassType::VariableMutator:
+						{
+							FConstVariableHandle Variable = Graph->FindVariable(InVariableID);
+							FConstNodeHandle ExistingMutator = Variable->FindMutatorNode();
+							if (!ExistingMutator->IsValid())
 							{
-								FConstVariableHandle Variable = Graph->FindVariable(InVariableID);
-								FConstNodeHandle ExistingMutator = Variable->FindMutatorNode();
-								if (!ExistingMutator->IsValid())
-								{
-									FrontendNode = Graph->FindOrAddVariableMutatorNode(InVariableID);
-								}
-								else
-								{
-									UE_LOG(LogMetaSound, Error, TEXT("Cannot add node because \"%s\" already exists for variable \"%s\""), *ExistingMutator->GetDisplayName().ToString(), *Variable->GetDisplayName().ToString());
-								}
+								FrontendNode = Graph->FindOrAddVariableMutatorNode(InVariableID);
 							}
-							break;
-
-						default:
+							else
 							{
-								checkNoEntry();
+								UE_LOG(LogMetaSound, Error, TEXT("Cannot add node because \"%s\" already exists for variable \"%s\""), *ExistingMutator->GetDisplayName().ToString(), *Variable->GetDisplayName().ToString());
 							}
+						}
+						break;
+
+						default:
+						{
+							checkNoEntry();
+						}
 					}
 				}
 			}
@@ -1531,60 +1536,13 @@ namespace Metasound
 			{
 				InVariableNode->ClassName = FrontendNode->GetClassMetadata().GetClassName();
 				InVariableNode->ClassType = FrontendNode->GetClassMetadata().GetType();
-				InVariableNode->SetNodeID(FrontendNode->GetID());
+				InVariableNode->NodeID = FrontendNode->GetID();
+				MetaSoundAsset->GetModifyContext().AddNodeIDsModified({ InVariableNode->GetNodeID() });
 			}
 
 			return FrontendNode;
 		}
 
-		UMetasoundEditorGraphNode* FGraphBuilder::AddNode(UObject& InMetaSound, Frontend::FNodeHandle InNodeHandle, FVector2D InLocation, bool bInSelectNewNode)
-		{
-			switch (InNodeHandle->GetClassMetadata().GetType())
-			{
-				case EMetasoundFrontendClassType::Input:
-				{
-					return CastChecked<UMetasoundEditorGraphNode>(AddInputNode(InMetaSound, InNodeHandle, InLocation, bInSelectNewNode));
-				}
-				break;
-
-				case EMetasoundFrontendClassType::External:
-				case EMetasoundFrontendClassType::Template:
-				{
-					return CastChecked<UMetasoundEditorGraphNode>(AddExternalNode(InMetaSound, InNodeHandle, InLocation, bInSelectNewNode));
-				}
-				break;
-
-				case EMetasoundFrontendClassType::Output:
-				{
-					return CastChecked<UMetasoundEditorGraphNode>(AddOutputNode(InMetaSound, InNodeHandle, InLocation, bInSelectNewNode));
-				}
-				break;
-
-				case EMetasoundFrontendClassType::VariableMutator:
-				case EMetasoundFrontendClassType::VariableAccessor:
-				case EMetasoundFrontendClassType::VariableDeferredAccessor:
-				case EMetasoundFrontendClassType::Variable:
-				{
-					return CastChecked<UMetasoundEditorGraphNode>(AddVariableNode(InMetaSound, InNodeHandle, InLocation, bInSelectNewNode));
-				}
-				break;
-
-				case EMetasoundFrontendClassType::Invalid:
-				case EMetasoundFrontendClassType::Graph:
-				
-				case EMetasoundFrontendClassType::Literal: // Not yet supported in editor
-				
-				default:
-				{
-					checkNoEntry();
-					static_assert(static_cast<int32>(EMetasoundFrontendClassType::Invalid) == 10, "Possible missing FMetasoundFrontendClassType case coverage");
-				}
-				break;
-			}
-
-			return nullptr;
-		}
-
 		bool FGraphBuilder::ConnectNodes(UEdGraphPin& InInputPin, UEdGraphPin& InOutputPin, bool bInConnectEdPins)
 		{
 			using namespace Frontend;
@@ -1684,196 +1642,45 @@ namespace Metasound
 			MetaSoundAsset->GetModifyContext().SetDocumentModified();
 		}
 
-		void FGraphBuilder::InitMetaSound(UObject& InMetaSound, const FString& InAuthor)
+		bool FGraphBuilder::DeleteMemberNodes(const UMetasoundEditorGraphMember& InGraphMember)
 		{
-			using namespace Frontend;
-			using namespace GraphBuilderPrivate;
-
-			FMetasoundAssetBase* MetaSoundAsset = IMetasoundUObjectRegistry::Get().GetObjectAsAssetBase(&InMetaSound);
-			check(MetaSoundAsset);
-			FMetasoundFrontendDocument& Doc = MetaSoundAsset->GetDocumentChecked();
-
-			// 1. Set default class Metadata
-			Doc.RootGraph.Metadata.SetClassName(FMetasoundFrontendClassName(FName(), *FGuid::NewGuid().ToString(), FName()));
-			Doc.RootGraph.Metadata.SetVersion({ 1, 0 });
-			Doc.RootGraph.Metadata.SetType(EMetasoundFrontendClassType::Graph);
-			Doc.RootGraph.Metadata.SetAuthor(InAuthor);
-
-			if (Doc.RootGraph.Metadata.GetType() != EMetasoundFrontendClassType::Graph)
+			bool bSuccess = true;
+			const TArray<UMetasoundEditorGraphMemberNode*> Nodes = InGraphMember.GetNodes();
+			for (UMetasoundEditorGraphMemberNode* Node : Nodes)
 			{
-				UE_LOG(LogMetaSound, Display, TEXT("Forcing class type to EMetasoundFrontendClassType::Graph on root graph metadata"));
-				Doc.RootGraph.Metadata.SetType(EMetasoundFrontendClassType::Graph);
-			}
-
-			// 2. Set default doc version Metadata
-			Doc.Metadata.Version.Number = FMetasoundFrontendDocument::GetMaxVersion();
-
-			// 3. Add default interfaces
-			{
-				UClass* AssetClass = InMetaSound.GetClass();
-				check(AssetClass);
-
-				TArray<FMetasoundFrontendVersion> InitVersions = ISearchEngine::Get().FindUClassDefaultInterfaceVersions(AssetClass->GetClassPathName());
-				FModifyRootGraphInterfaces({ }, InitVersions).Transform(Doc);
-			}
-
-			FGraphHandle GraphHandle = MetaSoundAsset->GetRootGraphHandle();
-			FVector2D InputNodeLocation = FVector2D::ZeroVector;
-			FVector2D ExternalNodeLocation = InputNodeLocation + DisplayStyle::NodeLayout::DefaultOffsetX;
-			FVector2D OutputNodeLocation = ExternalNodeLocation + DisplayStyle::NodeLayout::DefaultOffsetX;
-
-			TArray<FNodeHandle> NodeHandles = GraphHandle->GetNodes();
-			for (FNodeHandle& NodeHandle : NodeHandles)
-			{
-				const EMetasoundFrontendClassType NodeType = NodeHandle->GetClassMetadata().GetType();
-				FVector2D NewLocation;
-				if (NodeType == EMetasoundFrontendClassType::Input)
+				if (ensure(Node))
 				{
-					NewLocation = InputNodeLocation;
-					InputNodeLocation += DisplayStyle::NodeLayout::DefaultOffsetY;
+					bSuccess &= Metasound::Editor::FGraphBuilder::DeleteNode(*Node);
 				}
-				else if (NodeType == EMetasoundFrontendClassType::Output)
-				{
-					NewLocation = OutputNodeLocation;
-					OutputNodeLocation += DisplayStyle::NodeLayout::DefaultOffsetY;
-				}
-				else
-				{
-					NewLocation = ExternalNodeLocation;
-					ExternalNodeLocation += DisplayStyle::NodeLayout::DefaultOffsetY;
-				}
-				FMetasoundFrontendNodeStyle Style = NodeHandle->GetNodeStyle();
-				// TODO: Find consistent location for controlling node locations.
-				// Currently it is split between MetasoundEditor and MetasoundFrontend modules.
-				Style.Display.Locations = {{FGuid::NewGuid(), NewLocation}};
-				NodeHandle->SetNodeStyle(Style);
 			}
+			return bSuccess;
 		}
 
-		void FGraphBuilder::InitMetaSoundPreset(UObject& InMetaSoundReferenced, UObject& InMetaSoundPreset)
+		bool FGraphBuilder::DeleteNode(UEdGraphNode& InNode, bool bRemoveUnusedDependencies)
 		{
 			using namespace Frontend;
-			using namespace GraphBuilderPrivate;
-
-			FMetasoundAssetBase* PresetAsset = IMetasoundUObjectRegistry::Get().GetObjectAsAssetBase(&InMetaSoundPreset);
-			check(PresetAsset);
-
-			// Mark preset as auto-update and non-editable
-			FGraphHandle PresetGraphHandle = PresetAsset->GetRootGraphHandle();
-			FMetasoundFrontendGraphStyle Style = PresetGraphHandle->GetGraphStyle();
-			Style.bIsGraphEditable = false;
-			PresetGraphHandle->SetGraphStyle(Style);
-
-			// Mark all inputs as inherited by default
-			TSet<FName> InputsInheritingDefault;
-			Algo::Transform(PresetGraphHandle->GetInputNodes(), InputsInheritingDefault, [](FConstNodeHandle NodeHandle)
-			{
-				return NodeHandle->GetNodeName();
-			});
-			PresetGraphHandle->SetInputsInheritingDefault(MoveTemp(InputsInheritingDefault));
-
-			FGraphBuilder::RegisterGraphWithFrontend(InMetaSoundReferenced);
-
-			const FMetasoundAssetBase* ReferencedAsset = IMetasoundUObjectRegistry::Get().GetObjectAsAssetBase(&InMetaSoundReferenced);
-			check(ReferencedAsset);
-
-			FRebuildPresetRootGraph(ReferencedAsset->GetDocumentHandle()).Transform(PresetAsset->GetDocumentHandle());
-
-			PresetAsset->ConformObjectDataToInterfaces();
-		}
-
-		bool FGraphBuilder::DeleteNode(UEdGraphNode& InNode)
-		{
-			using namespace Frontend;
-
-			if (!InNode.CanUserDeleteNode())
-			{
-				return false;
-			}
 
-			// If node isn't a MetasoundEditorGraphNode, just remove and return (ex. comment nodes)
-			UMetasoundEditorGraphNode* Node = Cast<UMetasoundEditorGraphNode>(&InNode);
 			UMetasoundEditorGraph* Graph = CastChecked<UMetasoundEditorGraph>(InNode.GetGraph());
-			if (!Node)
-			{
-				Graph->RemoveNode(&InNode);
-				return true;
-			}
-
+			UObject& MetaSound = Graph->GetMetasoundChecked();
+			UMetaSoundBuilderBase& Builder = UMetaSoundBuilderSubsystem::GetChecked().AttachBuilderToAssetChecked(MetaSound);
 
-			// Remove connects only to pins associated with this EdGraph node
-			// only (Iterate pins and not Frontend representation to preserve
-			// other input/output EditorGraph reference node associations)
-			Node->IteratePins([](UEdGraphPin& Pin, int32 Index)
+			if (UMetasoundEditorGraphCommentNode* Node = Cast<UMetasoundEditorGraphCommentNode>(&InNode))
 			{
-				// Only add literal inputs for output pins as adding when disconnecting
-				// inputs would immediately orphan them on EditorGraph node removal below.
-				const bool bAddLiteralInputs = Pin.Direction == EGPD_Output;
-				FGraphBuilder::DisconnectPinVertex(Pin, bAddLiteralInputs);
-			});
-
-			FNodeHandle NodeHandle = Node->GetNodeHandle();
-			Frontend::FGraphHandle GraphHandle = NodeHandle->GetOwningGraph();
-
-			auto RemoveNodeLocation = [](FNodeHandle InNodeHandle, const FGuid& InNodeGuid)
-			{
-				FMetasoundFrontendNodeStyle Style = InNodeHandle->GetNodeStyle();
-				Style.Display.Locations.Remove(InNodeGuid);
-				InNodeHandle->SetNodeStyle(Style);
-			};	
-
-			auto RemoveNodeHandle = [] (FGraphHandle InGraphHandle, FNodeHandle InNodeHandle)
-			{
-				if (ensure(InGraphHandle->RemoveNode(*InNodeHandle)))
+				const bool bFrontendCommentRemoved = Builder.RemoveGraphComment(Node->CommentID);
+				if (ensure(bFrontendCommentRemoved))
 				{
-					InGraphHandle->GetOwningDocument()->RemoveUnreferencedDependencies();
+					return ensure(Graph->RemoveNode(&InNode));
 				}
-			};
+			}
 
-			if (GraphHandle->IsValid())
+			UMetasoundEditorGraphNode* Node = Cast<UMetasoundEditorGraphNode>(&InNode);
+			if (ensure(Node))
 			{
-				const EMetasoundFrontendClassType ClassType = NodeHandle->GetClassMetadata().GetType();
-				switch (ClassType)
-				{
-					// NodeHandle does not get removed in these cases as EdGraph Inputs/Outputs
-					// Frontend node is represented by the editor graph as a respective member
-					// (not a node) on the MetasoundGraph. Therefore, just the editor position
-					// data is removed.
-					case EMetasoundFrontendClassType::Output:
-					case EMetasoundFrontendClassType::Input:
-					{
-						RemoveNodeLocation(NodeHandle, InNode.NodeGuid);
-					}
-					break;
-					
-					// NodeHandle is only removed for variable accessors if the editor graph
-					// no longer contains nodes representing the given accessor on the MetasoundGraph.
-					// Therefore, just the editor position data is removed unless no location remains
-					// on the Frontend node.
-					case EMetasoundFrontendClassType::VariableAccessor:
-					case EMetasoundFrontendClassType::VariableDeferredAccessor:
-					{
-						RemoveNodeLocation(NodeHandle, InNode.NodeGuid);
-						if (NodeHandle->GetNodeStyle().Display.Locations.IsEmpty())
-						{
-							RemoveNodeHandle(GraphHandle, NodeHandle);
-						}
-					}
-					break;
-
-					case EMetasoundFrontendClassType::Graph:
-					case EMetasoundFrontendClassType::Literal:
-					case EMetasoundFrontendClassType::VariableMutator:
-					case EMetasoundFrontendClassType::Variable:
-					case EMetasoundFrontendClassType::External:
-					case EMetasoundFrontendClassType::Template:
-					default:
-					{
-						static_assert(static_cast<int32>(EMetasoundFrontendClassType::Invalid) == 10, "Possible missing MetasoundFrontendClassType switch case coverage.");
-						RemoveNodeHandle(GraphHandle, NodeHandle);
-					}
-					break;
-				}
+//				Need to split out delete ed vs. frontend impl which will happen in subsequent change.
+//				For now, just ignore removal failure as a hack as this gets called from ed graph sync.
+				EMetaSoundBuilderResult Result = EMetaSoundBuilderResult::Failed;
+				Builder.RemoveNode({ Node->GetNodeID() }, Result, bRemoveUnusedDependencies);
+// 				ensure(Result == EMetaSoundBuilderResult::Succeeded);
 			}
 
 			return ensure(Graph->RemoveNode(&InNode));
@@ -1888,17 +1695,10 @@ namespace Metasound
 				InGraphNode.RemovePin(InGraphNode.Pins[i]);
 			}
 
-			// TODO: Make this a utility in Frontend (ClearInputLiterals())
-			FNodeHandle NodeHandle = InGraphNode.GetNodeHandle();
-			TArray<FInputHandle> Inputs = NodeHandle->GetInputs();
-			for (FInputHandle& Input : Inputs)
-			{
-				NodeHandle->ClearInputLiteral(Input->GetID());
-			}
-
-			TArray<FInputHandle> InputHandles = NodeHandle->GetInputs();
-			NodeHandle->GetInputStyle().SortDefaults(InputHandles, [](const Frontend::FInputHandle& Handle) { return FGraphBuilder::GetDisplayName(*Handle); });
-			for (const FInputHandle& InputHandle : InputHandles)
+			FConstNodeHandle NodeHandle = InGraphNode.GetNodeHandle();
+			TArray<FConstInputHandle> InputHandles = NodeHandle->GetConstInputs();
+			NodeHandle->GetInputStyle().SortDefaults(InputHandles, [](const Frontend::FConstInputHandle& Handle) { return FGraphBuilder::GetDisplayName(*Handle); });
+			for (const FConstInputHandle& InputHandle : InputHandles)
 			{
 				// Only add pins of the node if the connection is user modifiable. 
 				// Connections which the user cannot modify are controlled elsewhere.
@@ -1908,9 +1708,9 @@ namespace Metasound
 				}
 			}
 
-			TArray<FOutputHandle> OutputHandles = NodeHandle->GetOutputs();
-			NodeHandle->GetOutputStyle().SortDefaults(OutputHandles, [](const Frontend::FOutputHandle& Handle) { return FGraphBuilder::GetDisplayName(*Handle); });
-			for (const FOutputHandle& OutputHandle : OutputHandles)
+			TArray<FConstOutputHandle> OutputHandles = NodeHandle->GetConstOutputs();
+			NodeHandle->GetOutputStyle().SortDefaults(OutputHandles, [](const Frontend::FConstOutputHandle& Handle) { return FGraphBuilder::GetDisplayName(*Handle); });
+			for (const FConstOutputHandle& OutputHandle : OutputHandles)
 			{
 				// Only add pins of the node if the connection is user modifiable. 
 				// Connections which the user cannot modify are controlled elsewhere.
@@ -2025,7 +1825,7 @@ namespace Metasound
 				return false;
 			}
 
-			Frontend::FInputHandle PinInputHandle = GetInputHandleFromPin(&InEditorPin);
+			Frontend::FConstInputHandle PinInputHandle = GetConstInputHandleFromPin(&InEditorPin);
 			if (PinInputHandle->GetID() == InInputHandle->GetID())
 			{
 				return true;
@@ -2041,7 +1841,7 @@ namespace Metasound
 				return false;
 			}
 
-			Frontend::FOutputHandle PinOutputHandle = GetOutputHandleFromPin(&InEditorPin);
+			Frontend::FConstOutputHandle PinOutputHandle = GetConstOutputHandleFromPin(&InEditorPin);
 			if (PinOutputHandle->GetID() == InOutputHandle->GetID())
 			{
 				return true;
@@ -2121,7 +1921,7 @@ namespace Metasound
 		{
 			using namespace Metasound::Frontend;
 
-			if (InAssetBase.GetModifyContext().GetDocumentModified())
+			if (InAssetBase.GetConstModifyContext().GetDocumentModified())
 			{
 				return true;
 			}
@@ -2163,26 +1963,97 @@ namespace Metasound
 			return false;
 		}
 
-		bool FGraphBuilder::SynchronizeConnections(UObject& InMetaSound)
+		bool FGraphBuilder::SynchronizeComments(UObject& InMetaSound)
 		{
-			using namespace Frontend;
-
-			bool bIsGraphDirty = false;
+			bool bModified = false;
 
 			FMetasoundAssetBase* MetaSoundAsset = IMetasoundUObjectRegistry::Get().GetObjectAsAssetBase(&InMetaSound);
 			check(MetaSoundAsset);
 
-			FGraphHandle GraphHandle = MetaSoundAsset->GetRootGraphHandle();
-
 			UMetasoundEditorGraph* EditorGraph = CastChecked<UMetasoundEditorGraph>(MetaSoundAsset->GetGraph());
+			check(EditorGraph);
+
+			TMap<FGuid, UMetasoundEditorGraphCommentNode*> CommentIDToEdNode;
+			Algo::TransformIf(EditorGraph->Nodes, CommentIDToEdNode,
+				[](const UEdGraphNode* Node)
+				{
+					return Node->IsA<UMetasoundEditorGraphCommentNode>();
+				},
+				[](UEdGraphNode* Node)
+				{
+					UMetasoundEditorGraphCommentNode* CommentNode = CastChecked<UMetasoundEditorGraphCommentNode>(Node);
+					return TPair<FGuid, UMetasoundEditorGraphCommentNode*>(CommentNode->GetCommentID(), CommentNode);
+				});
+
+			const FMetasoundFrontendDocument& Document = MetaSoundAsset->GetConstDocumentChecked();
+			const TMap<FGuid, FMetaSoundFrontendGraphComment>& Comments = Document.RootGraph.Graph.Style.Comments;
+			for (const TPair<FGuid, FMetaSoundFrontendGraphComment>& Pair : Comments)
+			{
+				UMetasoundEditorGraphCommentNode* CommentNode = nullptr;
+				if (!CommentIDToEdNode.RemoveAndCopyValue(Pair.Key, CommentNode))
+				{
+					bModified = true;
+					constexpr bool bSelectNode = false;
+
+					// Can't use spawn node action because it modifies the transaction stack, so just generate from CDO and backport frontend data accordingly.
+					UClass* CommentClass = UMetasoundEditorGraphCommentNode::StaticClass();
+					check(CommentClass);
+					UMetasoundEditorGraphCommentNode* CommentCDO = CommentClass->GetDefaultObject<UMetasoundEditorGraphCommentNode>();
+					check(CommentCDO);
+
+					const UGraphEditorSettings* GraphEditorSettings = GetDefault<UGraphEditorSettings>();
+					check(GraphEditorSettings);
+
+					UMetasoundEditorGraphCommentNode* NewNode = CastChecked<UMetasoundEditorGraphCommentNode>(DuplicateObject<UEdGraphNode>(CommentCDO, EditorGraph));
+					NewNode->SetFlags(RF_Transactional);
+					EditorGraph->AddNode(NewNode, false, false);
+					NewNode->CreateNewGuid();
+					NewNode->bCommentBubbleVisible_InDetailsPanel = GraphEditorSettings->bShowCommentBubbleWhenZoomedOut;
+
+					// Pull position, color, etc. from the existing frontend data.
+					NewNode->CommentID = Pair.Key;
+					UMetasoundEditorGraphCommentNode::ConvertFromFrontendComment(Pair.Value, *NewNode);
+				}
+			}
+
+			// Remaining items are stale, so they are removed
+			bModified |= !CommentIDToEdNode.IsEmpty();
+			for (const TPair<FGuid, UMetasoundEditorGraphCommentNode*>& Comment : CommentIDToEdNode)
+			{
+				constexpr bool bMarkDirty = false;
+				constexpr bool bBreakAllLinks = true;
+				EditorGraph->RemoveNode(Comment.Value, bBreakAllLinks, bMarkDirty);
+			}
+
+			return bModified;
+		}
+
+		bool FGraphBuilder::SynchronizeConnections(UObject& InMetaSound)
+		{
+			using namespace Frontend;
+
+			bool bIsGraphDirty = false;
 
 			TArray<UMetasoundEditorGraphNode*> EditorNodes;
-			EditorGraph->GetNodesOfClass(EditorNodes);
+			{
+				FMetasoundAssetBase* MetaSoundAsset = IMetasoundUObjectRegistry::Get().GetObjectAsAssetBase(&InMetaSound);
+				check(MetaSoundAsset);
+				UMetasoundEditorGraph* EditorGraph = CastChecked<UMetasoundEditorGraph>(MetaSoundAsset->GetGraph());
+				EditorGraph->GetNodesOfClass(EditorNodes);
+			}
 
-			TMap<FGuid, TArray<UMetasoundEditorGraphNode*>> EditorNodesByFrontendID;
+			TMap<FGuid, UMetasoundEditorGraphNode*> EditorNodesByFrontendID;
 			for (UMetasoundEditorGraphNode* EditorNode : EditorNodes)
 			{
-				EditorNodesByFrontendID.FindOrAdd(EditorNode->GetNodeID()).Add(EditorNode);
+				const FGuid& NodeID = EditorNode->GetNodeID();
+				if (EditorNodesByFrontendID.Contains(NodeID))
+				{
+					UE_LOG(LogMetasoundEditor, Error, TEXT("Multiple editor nodes associated with FrontendDocument node with ID '%s'"), *NodeID.ToString());
+				}
+				else
+				{
+					EditorNodesByFrontendID.Add(NodeID, EditorNode);
+				}
 			}
 
 			// Iterate through all nodes in metasound editor graph and synchronize connections.
@@ -2190,7 +2061,7 @@ namespace Metasound
 			{
 				bool bIsNodeDirty = false;
 
-				FConstNodeHandle Node = EditorNode->GetNodeHandle();
+				FConstNodeHandle Node = EditorNode->GetConstNodeHandle();
 
 				TArray<UEdGraphPin*> Pins = EditorNode->GetAllPins();
 				TArray<FConstInputHandle> NodeInputs = Node->GetConstInputs();
@@ -2236,25 +2107,30 @@ namespace Metasound
 						else if (!IsMatchingOutputHandleAndPin(OutputHandle, *MatchingPin->LinkedTo[0]))
 						{
 							// The wrong link exists.
-							MatchingPin->BreakAllPinLinks();
+							constexpr bool bNotifyNodes = false;
+							constexpr bool bMarkDirty = false;
+							MatchingPin->BreakAllPinLinks(bNotifyNodes, bMarkDirty);
 							bAddLink = true;
 						}
 
 						if (bAddLink)
 						{
 							const FGuid NodeID = OutputHandle->GetOwningNodeID();
-							TArray<UMetasoundEditorGraphNode*>* OutputEditorNode = EditorNodesByFrontendID.Find(NodeID);
-							if (ensure(OutputEditorNode))
+							const FText& OwningNodeName = EditorNode->GetDisplayName();
+							UMetasoundEditorGraphNode* OutputEditorNode = EditorNodesByFrontendID.FindRef(NodeID);
+							if (OutputEditorNode)
 							{
-								if (ensure(!OutputEditorNode->IsEmpty()))
-								{
-									UEdGraphPin* OutputPin = (*OutputEditorNode)[0]->FindPinChecked(OutputHandle->GetName(), EEdGraphPinDirection::EGPD_Output);
-									const FText& OwningNodeName = EditorNode->GetDisplayName();
+								UEdGraphPin* OutputPin = OutputEditorNode->FindPinChecked(OutputHandle->GetName(), EEdGraphPinDirection::EGPD_Output);
 
-									UE_LOG(LogMetasoundEditor, Verbose, TEXT("Synchronizing Node '%s' Connection: Linking Pin '%s' to '%s'"), *OwningNodeName.ToString(), *MatchingPin->GetName(), *OutputPin->GetName());
-									MatchingPin->MakeLinkTo(OutputPin);
-									bIsNodeDirty = true;
-								}
+								UE_LOG(LogMetasoundEditor, Verbose, TEXT("Synchronizing Node '%s' Connection: Linking Pin '%s' to '%s'"), *OwningNodeName.ToString(), *MatchingPin->GetName(), *OutputPin->GetName());
+
+								constexpr bool bAlwaysMarkDirty = false;
+								MatchingPin->MakeLinkTo(OutputPin, bAlwaysMarkDirty);
+								bIsNodeDirty = true;
+							}
+							else
+							{
+								UE_LOG(LogMetasoundEditor, Warning, TEXT("Failed to synchronize Frontend Node '%s' Connection: Pin '%s'"), *OwningNodeName.ToString(), *MatchingPin->GetName());
 							}
 						}
 					}
@@ -2263,7 +2139,9 @@ namespace Metasound
 						// No link should exist.
 						if (!MatchingPin->LinkedTo.IsEmpty())
 						{
-							MatchingPin->BreakAllPinLinks();
+							constexpr bool bNotifyNodes = false;
+							constexpr bool bMarkDirty = false;
+							MatchingPin->BreakAllPinLinks(bNotifyNodes, bMarkDirty);
 							const FText OwningNodeName = EditorNode->GetDisplayName();
 							const FText InputName = FGraphBuilder::GetDisplayName(*NodeInput);
 							UE_LOG(LogMetasoundEditor, Verbose, TEXT("Synchronizing Node '%s' Connection: Breaking all pin links to '%s'"), *OwningNodeName.ToString(), *InputName.ToString());
@@ -2291,7 +2169,7 @@ namespace Metasound
 
 			if (RecurseGetDocumentModified(*MetaSoundAsset))
 			{
-				TArray<FMetasoundAssetBase*> EditedReferencingMetaSounds;
+				TSet<FMetasoundAssetBase*> EditedReferencingMetaSounds;
 				if (GEditor)
 				{
 					TArray<UObject*> EditedAssets = GEditor->GetEditorSubsystem<UAssetEditorSubsystem>()->GetAllEditedAssets();
@@ -2313,7 +2191,7 @@ namespace Metasound
 				if (EditedReferencingMetaSounds.IsEmpty())
 				{
 					MetaSoundAsset->CacheRegistryMetadata();
-					const bool bEditorGraphModified = GraphBuilderPrivate::SynchronizeGraphRecursively(InMetaSound);
+					GraphBuilderPrivate::SynchronizeGraphRecursively(InMetaSound);
 					GraphBuilderPrivate::RecurseClearDocumentModified(*MetaSoundAsset);
 				}
 				else
@@ -2333,9 +2211,9 @@ namespace Metasound
 			return false;
 		}
 
-		bool FGraphBuilder::SynchronizeNodeMembers(UObject& InMetaSound)
+		bool FGraphBuilder::SynchronizeOutputNodes(UObject& InMetaSound)
 		{
-			TRACE_CPUPROFILER_EVENT_SCOPE(Metasound::Editor::FGraphBuilder::SynchronizeNodeMembers);
+			TRACE_CPUPROFILER_EVENT_SCOPE(Metasound::Editor::FGraphBuilder::SynchronizeOutputNodes);
 
 			using namespace Frontend;
 
@@ -2343,41 +2221,9 @@ namespace Metasound
 
 			FMetasoundAssetBase* MetaSoundAsset = IMetasoundUObjectRegistry::Get().GetObjectAsAssetBase(&InMetaSound);
 			check(MetaSoundAsset);
-			FGraphHandle GraphHandle = MetaSoundAsset->GetRootGraphHandle();
+			FConstGraphHandle GraphHandle = MetaSoundAsset->GetRootGraphHandle();
 			UMetasoundEditorGraph* EditorGraph = CastChecked<UMetasoundEditorGraph>(MetaSoundAsset->GetGraph());
 
-			TArray<UMetasoundEditorGraphInputNode*> InputNodes;
-			EditorGraph->GetNodesOfClassEx<UMetasoundEditorGraphInputNode>(InputNodes);
-			for (UMetasoundEditorGraphInputNode* Node : InputNodes)
-			{
-				check(Node);
-				FConstNodeHandle NodeHandle = Node->GetConstNodeHandle();
-				if (!NodeHandle->IsValid())
-				{
-					for (UEdGraphPin* Pin : Node->Pins)
-					{
-						check(Pin);
-
-						FConstClassInputAccessPtr ClassInputPtr = GraphHandle->FindClassInputWithName(Pin->PinName);
-						if (const FMetasoundFrontendClassInput* Input = ClassInputPtr.Get())
-						{
-							const FGuid& InitialID = Node->GetNodeID();
-							if (Node->GetNodeHandle()->GetID() != Input->NodeID)
-							{
-								Node->SetNodeID(Input->NodeID);
-
-								// Requery handle as the id has been fixed up
-								NodeHandle = Node->GetConstNodeHandle();
-								FText InputDisplayName = Node->GetDisplayName();
-								UE_LOG(LogMetasoundEditor, Verbose, TEXT("Editor Input Node '%s' interface versioned"), *InputDisplayName.ToString());
-
-								bEditorGraphModified = true;
-							}
-						}
-					}
-				}
-			}
-
 			TArray<UMetasoundEditorGraphOutputNode*> OutputNodes;
 			EditorGraph->GetNodesOfClassEx<UMetasoundEditorGraphOutputNode>(OutputNodes);
 			for (UMetasoundEditorGraphOutputNode* Node : OutputNodes)
@@ -2411,30 +2257,22 @@ namespace Metasound
 			return bEditorGraphModified;
 		}
 
-		bool FGraphBuilder::SynchronizeNodes(UObject& InMetaSound)
+		void FGraphBuilder::SynchronizeNodes(UObject& InMetaSound)
 		{
 			TRACE_CPUPROFILER_EVENT_SCOPE(Metasound::Editor::FGraphBuilder::SynchronizeNodes);
 
 			using namespace Frontend;
 
-			bool bEditorGraphModified = false;
-
 			// Get all external nodes from Frontend graph.  Input and output references will only be added/synchronized
 			// if required when synchronizing connections (as they are not required to inhabit editor graph).
 			FMetasoundAssetBase* MetaSoundAsset = IMetasoundUObjectRegistry::Get().GetObjectAsAssetBase(&InMetaSound);
 			check(MetaSoundAsset);
-			FGraphHandle GraphHandle = MetaSoundAsset->GetRootGraphHandle();
-			TArray<FNodeHandle> FrontendNodes = GraphHandle->GetNodes();
+			FConstGraphHandle GraphHandle = static_cast<const FMetasoundAssetBase*>(MetaSoundAsset)->GetRootGraphHandle();
+			TArray<FConstNodeHandle> FrontendNodes = GraphHandle->GetConstNodes();
 			UMetasoundEditorGraph* EditorGraph = CastChecked<UMetasoundEditorGraph>(MetaSoundAsset->GetGraph());
 			TArray<UMetasoundEditorGraphNode*> EditorNodes;
 			EditorGraph->GetNodesOfClass(EditorNodes);
 
-			TMap<FGuid, UMetasoundEditorGraphNode*> EditorNodesByEdNodeGuid;
-			for (UMetasoundEditorGraphNode* Node : EditorNodes)
-			{
-				EditorNodesByEdNodeGuid.Add(Node->NodeGuid, Node);
-			}
-
 			// Find existing array of editor nodes associated with Frontend node
 			struct FAssociatedNodes
 			{
@@ -2443,6 +2281,13 @@ namespace Metasound
 			};
 			TMap<FGuid, FAssociatedNodes> AssociatedNodes;
 
+			// Input nodes use template input nodes to visually distinguish location, so ignore them.
+			FrontendNodes.RemoveAllSwap([](const FConstNodeHandle& Node)
+			{
+				return Node->GetClassMetadata().GetType() == EMetasoundFrontendClassType::Input;
+			},
+			EAllowShrinking::No);
+
 			// Reverse iterate so paired nodes can safely be removed from the array.
 			for (int32 i = FrontendNodes.Num() - 1; i >= 0; i--)
 			{
@@ -2464,7 +2309,8 @@ namespace Metasound
 							AssociatedNodeData.Node = Node;
 						}
 
-						bEditorGraphModified |= SynchronizeNodeLocation(Node, *EditorNode);
+						EditorNode->SyncLocationFromFrontendNode();
+						EditorNode->SyncCommentFromFrontendNode();
 						AssociatedNodeData.EditorNodes.Add(EditorNode);
 						EditorNodes.RemoveAtSwap(j, EAllowShrinking::No);
 					}
@@ -2484,27 +2330,78 @@ namespace Metasound
 			// exist before attempting to synchronize connections.
 			for (UMetasoundEditorGraphNode* EditorNode : EditorNodes)
 			{
-				bEditorGraphModified |= EditorGraph->RemoveNode(EditorNode);
+				constexpr bool bMarkDirty = false;
+				constexpr bool bBreakAllLinks = true;
+				EditorGraph->RemoveNode(EditorNode, bBreakAllLinks, bMarkDirty);
 			}
 
 			// Add missing editor nodes marked as visible.
-			for (const FNodeHandle& Node : FrontendNodes)
+			for (const FConstNodeHandle& Node : FrontendNodes)
 			{
 				const FMetasoundFrontendNodeStyle& CurrentStyle = Node->GetNodeStyle();
-				if (CurrentStyle.Display.Locations.IsEmpty())
+				for (const TPair<FGuid, FVector2D>& Location : CurrentStyle.Display.Locations)
 				{
-					continue;
-				}
+					UMetasoundEditorGraphNode* NewGraphNode = nullptr;
+					const FMetasoundFrontendClassMetadata& ClassMetadata = Node->GetClassMetadata();
+					switch (ClassMetadata.GetType())
+					{
+						case EMetasoundFrontendClassType::External:
+						{
+							NewGraphNode = CastChecked<UMetasoundEditorGraphNode>(AddExternalNode(InMetaSound, Node->GetID(), ClassMetadata, false));
+						}
+						break;
 
-				FMetasoundFrontendNodeStyle NewStyle = CurrentStyle;
-				bEditorGraphModified = true;
+						case EMetasoundFrontendClassType::Template:
+						{
+							if (Node->GetClassMetadata().GetClassName() == FInputNodeTemplate::ClassName)
+							{
+								NewGraphNode = CastChecked<UMetasoundEditorGraphNode>(AddInputNode(InMetaSound, Node, false));
+							}
+							else
+							{
+								NewGraphNode = CastChecked<UMetasoundEditorGraphNode>(AddTemplateNode(InMetaSound, Node->GetID(), Node->GetClassMetadata(), false));
+							}
+						}
+						break;
 
-				TArray<UMetasoundEditorGraphNode*> AddedNodes;
-				for (const TPair<FGuid, FVector2D>& Location : NewStyle.Display.Locations)
-				{
-					UMetasoundEditorGraphNode* NewNode = AddNode(InMetaSound, Node, Location.Value, false /* bInSelectNewNode */);
-					if (ensure(NewNode))
+						case EMetasoundFrontendClassType::Output:
+						{
+							NewGraphNode = CastChecked<UMetasoundEditorGraphNode>(AddOutputNode(InMetaSound, Node, false));
+						}
+						break;
+
+						case EMetasoundFrontendClassType::VariableMutator:
+						case EMetasoundFrontendClassType::VariableAccessor:
+						case EMetasoundFrontendClassType::VariableDeferredAccessor:
+						case EMetasoundFrontendClassType::Variable:
+						{
+							NewGraphNode = CastChecked<UMetasoundEditorGraphNode>(AddVariableNode(InMetaSound, Node, false));
+						}
+						break;
+
+						case EMetasoundFrontendClassType::Invalid:
+
+						// Class type needs to be deprecated 
+						case EMetasoundFrontendClassType::Graph:
+
+						// Not supported in editor
+						case EMetasoundFrontendClassType::Literal:
+
+						// Since MetaSound Document v1.12 update, the editor uses template input nodes, so no direct node representation of inputs no longer exists
+						case EMetasoundFrontendClassType::Input: 
+						default:
+						{
+							checkNoEntry();
+							static_assert(static_cast<int32>(EMetasoundFrontendClassType::Invalid) == 10, "Possible missing FMetasoundFrontendClassType case coverage");
+						}
+						break;
+					}
+
+					if (ensure(NewGraphNode))
 					{
+						ensureMsgf(NewGraphNode->NodeGuid.IsValid(), TEXT("New editor NodeGuid must be valid."));
+						ensureMsgf(NewGraphNode->NodeGuid == Location.Key, TEXT("New editor NodeGuid must match location key"));
+
 						FAssociatedNodes& AssociatedNodeData = AssociatedNodes.FindOrAdd(Node->GetID());
 						if (AssociatedNodeData.Node->IsValid())
 						{
@@ -2515,17 +2412,9 @@ namespace Metasound
 							AssociatedNodeData.Node = Node;
 						}
 
-						AddedNodes.Add(NewNode);
-						AssociatedNodeData.EditorNodes.Add(NewNode);
+						AssociatedNodeData.EditorNodes.Add(NewGraphNode);
 					}
 				}
-
-				NewStyle.Display.Locations.Reset();
-				for (UMetasoundEditorGraphNode* EditorNode : AddedNodes)
-				{
-					NewStyle.Display.Locations.Add(EditorNode->NodeGuid, FVector2D(EditorNode->NodePosX, EditorNode->NodePosY));
-				}
-				Node->SetNodeStyle(NewStyle);
 			}
 
 			// Synchronize pins on node associations.
@@ -2533,11 +2422,9 @@ namespace Metasound
 			{
 				for (UMetasoundEditorGraphNode* EditorNode : IdNodePair.Value.EditorNodes)
 				{
-					bEditorGraphModified |= SynchronizeNodePins(*EditorNode, IdNodePair.Value.Node);
+					SynchronizeNodePins(*EditorNode, IdNodePair.Value.Node);
 				}
 			}
-
-			return bEditorGraphModified;
 		}
 
 		bool FGraphBuilder::SynchronizeNodePins(UMetasoundEditorGraphNode& InEditorNode, Frontend::FConstNodeHandle InNode, bool bRemoveUnusedPins, bool bLogChanges)
@@ -2719,7 +2606,7 @@ namespace Metasound
 
 			const FString OldValue = InPin.DefaultValue;
 
-			FInputHandle InputHandle = GetInputHandleFromPin(&InPin);
+			FConstInputHandle InputHandle = GetConstInputHandleFromPin(&InPin);
 			if (const FMetasoundFrontendLiteral* NodeDefaultLiteral = InputHandle->GetLiteral())
 			{
 				InPin.DefaultValue = NodeDefaultLiteral->ToString();
@@ -2772,6 +2659,27 @@ namespace Metasound
 				bEditorGraphModified = true;
 			}, EMetasoundFrontendClassType::Input);
 
+			// Remove stale inputs.
+			{
+				for (int32 Index = Graph->Inputs.Num() - 1; Index >= 0; --Index)
+				{
+					const UMetasoundEditorGraphInput* Input = Graph->Inputs[Index];
+					check(Input);
+					if (!Inputs.Contains(Input))
+					{
+						TArray<UMetasoundEditorGraphMemberNode*> InputNodes = Input->GetNodes();
+						for (UMetasoundEditorGraphNode* Node : InputNodes)
+						{
+							FGraphBuilder::DeleteMemberNodes(*Input);
+						}
+
+						bEditorGraphModified = true;
+						UE_LOG(LogMetasoundEditor, Verbose, TEXT("Synchronizing Outputs: Removing stale output '%s'."), *Input->GetName());
+						Graph->Inputs.RemoveAtSwap(Index, EAllowShrinking::No);
+					}
+				}
+			}
+
 			// Collect all editor graph outputs with corresponding frontend outputs. 
 			GraphHandle->IterateNodes([&](FNodeHandle NodeHandle)
 			{
@@ -2789,30 +2697,19 @@ namespace Metasound
 				bEditorGraphModified = true;
 			}, EMetasoundFrontendClassType::Output);
 
-			// Collect editor inputs and outputs to remove which have no corresponding frontend input or output.
-			TArray<UMetasoundEditorGraphMember*> ToRemove;
-			Graph->IterateInputs([&](UMetasoundEditorGraphInput& Input)
+			// Remove stale outputs.
 			{
-				if (!Inputs.Contains(&Input))
+				for (int32 Index = Graph->Outputs.Num() - 1; Index >= 0; --Index)
 				{
-					UE_LOG(LogMetasoundEditor, Verbose, TEXT("Synchronizing Inputs: Removing stale input '%s'."), *Input.GetName());
-					ToRemove.Add(&Input);
-				}
-			});
-			Graph->IterateOutputs([&](UMetasoundEditorGraphOutput& Output)
-			{
-				if (!Outputs.Contains(&Output))
-				{
-					UE_LOG(LogMetasoundEditor, Verbose, TEXT("Synchronizing Outputs: Removing stale output '%s'."), *Output.GetName());
-					ToRemove.Add(&Output);
+					const UMetasoundEditorGraphOutput* Output = Graph->Outputs[Index];
+					check(Output);
+					if (!Outputs.Contains(Output))
+					{
+						bEditorGraphModified = true;
+						UE_LOG(LogMetasoundEditor, Verbose, TEXT("Synchronizing Outputs: Removing stale output '%s'."), *Output->GetName());
+						Graph->Outputs.RemoveAtSwap(Index, EAllowShrinking::No);
+					}
 				}
-			});
-
-			// Remove stale inputs and outputs.
-			bEditorGraphModified |= !ToRemove.IsEmpty();
-			for (UMetasoundEditorGraphMember* GraphMember: ToRemove)
-			{
-				Graph->RemoveMember(*GraphMember);
 			}
 
 			UMetasoundEditorGraphMember* Member = nullptr;
@@ -2873,6 +2770,10 @@ namespace Metasound
 							}
 						}
 					}
+					else
+					{
+						Input->InitializeLiteral();
+					}
 				}
 			}, EMetasoundFrontendClassType::Input);
 
@@ -2883,7 +2784,7 @@ namespace Metasound
 				{
 					SynchronizeMemberDataType(*Output);
 
-					// Fix up corrupted assets with no literal set 
+					// Fix up members with no literal set
 					if (!Output->GetLiteral())
 					{
 						Output->InitializeLiteral();
@@ -2891,10 +2792,32 @@ namespace Metasound
 				}
 			}, EMetasoundFrontendClassType::Output);
 
+			// Enforce all variables exists in the edgraph
+			TSet<FGuid> VariableIDs;
+			Algo::Transform(Graph->Variables, VariableIDs, [](const TObjectPtr<UMetasoundEditorGraphVariable>& Variable) { return Variable->GetVariableID(); });
+			TArray<FVariableHandle> FrontendVariables = GraphHandle->GetVariables();
+			for (const FVariableHandle& VariableHandle : FrontendVariables)
+			{
+				const FGuid VariableID = VariableHandle->GetID();
+				if (VariableIDs.Remove(VariableID) == 0)
+				{
+					UMetasoundEditorGraphVariable* NewVariable = NewObject<UMetasoundEditorGraphVariable>(Graph, FName(), RF_Transactional);
+					if (ensure(NewVariable))
+					{
+						const FMetasoundFrontendLiteral DefaultLiteral = VariableHandle->GetLiteral();
+						NewVariable->InitMember(VariableHandle->GetDataType(), DefaultLiteral, VariableID);
+						Graph->Variables.Add(NewVariable);
+					}
+				}
+			}
+
 			// Remove empty entries
 			bEditorGraphModified |= Graph->Inputs.RemoveAllSwap([](const TObjectPtr<UMetasoundEditorGraphInput>& Input) { return !Input; }) > 0;
 			bEditorGraphModified |= Graph->Outputs.RemoveAllSwap([](const TObjectPtr<UMetasoundEditorGraphOutput>& Output) { return !Output; }) > 0;
-			bEditorGraphModified |= Graph->Variables.RemoveAllSwap([](const TObjectPtr<UMetasoundEditorGraphVariable>& Variable) { return !Variable; }) > 0;
+			bEditorGraphModified |= Graph->Variables.RemoveAllSwap([&VariableIDs](const TObjectPtr<UMetasoundEditorGraphVariable>& Variable)
+			{
+				return !Variable || VariableIDs.Contains(Variable->GetVariableID());
+			}) > 0;
 
 			return bEditorGraphModified;
 		}
diff --git a/Engine/Plugins/Runtime/Metasound/Source/MetasoundEditor/Private/MetasoundEditorGraphBuilder.h b/Engine/Plugins/Runtime/Metasound/Source/MetasoundEditor/Private/MetasoundEditorGraphBuilder.h
--- a/Engine/Plugins/Runtime/Metasound/Source/MetasoundEditor/Private/MetasoundEditorGraphBuilder.h
+++ b/Engine/Plugins/Runtime/Metasound/Source/MetasoundEditor/Private/MetasoundEditorGraphBuilder.h
@@ -20,6 +20,7 @@ class UEdGraph;
 class UEdGraphNode;
 class UEdGraphPin;
 class UMetaSoundPatch;
+class UMetasoundEditorGraphCommentNode;
 class UMetasoundEditorGraphExternalNode;
 class UMetasoundEditorGraphNode;
 class UMetasoundEditorGraphInputNode;
@@ -61,12 +62,6 @@ namespace Metasound
 			static const FText FunctionMenuName;
 			static const FText GraphMenuName;
 
-			static void InitGraphNode(Frontend::FNodeHandle& InNodeHandle, UMetasoundEditorGraphNode* NewGraphNode, UObject& InMetaSound);
-			static void InitGraphNodeIDFromNodeHandle(const Frontend::FConstNodeHandle& InNodeHandle, UMetasoundEditorGraphNode* NewGraphNode);
-
-			// Adds an EdGraph node to mirror the provided FNodeHandle.
-			static UMetasoundEditorGraphNode* AddNode(UObject& InMetaSound, Frontend::FNodeHandle InNodeHandle, FVector2D InLocation, bool bInSelectNewNode = true);
-
 			// Convenience functions for retrieving the editor for the given UObject
 			static TSharedPtr<FEditor> GetEditorForMetasound(const UObject& InMetaSound);
 			static TSharedPtr<FEditor> GetEditorForGraph(const UEdGraph& InEdGraph);
@@ -123,12 +118,11 @@ namespace Metasound
 			// Returns the PinName for an IInputController.
 			static FName GetPinName(const Frontend::IInputController& InFrontendInput);
 
-			// Adds a node handle to mirror the provided graph node and binds to it.  Does *NOT* mirror existing EdGraph connections
-			// nor does it remove existing bound Frontend Node (if set) from associated Frontend Graph.
-			static Frontend::FNodeHandle AddNodeHandle(UObject& InMetaSound, UMetasoundEditorGraphNode& InGraphNode);
+			// Adds a new EdGraph comment node associated with the given MetaSoundFrontendGraph comment ID
+			static UMetasoundEditorGraphCommentNode* CreateCommentNode(UObject& InMetaSound, bool bInSelectNewNode = true, FGuid InCommentID = FGuid::NewGuid());
 
-			// Adds a corresponding UMetasoundEditorGraphInputNode for the provided node handle.
-			static UMetasoundEditorGraphInputNode* AddInputNode(UObject& InMetaSound, Frontend::FNodeHandle InNodeHandle, FVector2D InLocation, bool bInSelectNewNode = true);
+			// Adds an input template node to the given MetaSound object's document model and returns an associated node handle
+			static Frontend::FNodeHandle AddInputTemplateNodeHandle(UObject& InMetaSound, Frontend::FNodeHandle& InputNodeHandle);
 
 			// Generates FNodeHandle for the given external node data. Does not bind or create EdGraph representation of given node.
 			static Frontend::FNodeHandle AddInputNodeHandle(
@@ -138,17 +132,25 @@ namespace Metasound
 				const FName* InNameBase = nullptr);
 
 			// Adds a corresponding UMetasoundEditorGraphExternalNode for the provided node handle.
-			static UMetasoundEditorGraphExternalNode* AddExternalNode(UObject& InMetaSound, Frontend::FNodeHandle& InNodeHandle, FVector2D InLocation, bool bInSelectNewNode = true);
+			static UMetasoundEditorGraphExternalNode* AddExternalNode(UObject& InMetaSound, const FGuid& InNodeID, const FMetasoundFrontendClassMetadata& InMetadata, bool bInSelectNewNode = true);
+
+			// Adds an editor graph node that corresponds with an instance of a node that is defined by an external MetaSound node class.
+			static UMetasoundEditorGraphExternalNode* AddTemplateNode(UObject& InMetaSound, const FGuid& InNodeID, const FMetasoundFrontendClassMetadata& InMetadata, bool bInSelectNewNode = true);
 
 			// Adds an externally-defined node with the given class info to both the editor and document graphs.
 			// Generates analogous FNodeHandle.
-			static UMetasoundEditorGraphExternalNode* AddExternalNode(UObject& InMetaSound, const FMetasoundFrontendClassMetadata& InMetadata, FVector2D InLocation, bool bInSelectNewNode = true);
+			static UMetasoundEditorGraphExternalNode* AddExternalNode(UObject& InMetaSound, const FMetasoundFrontendClassMetadata& InMetadata, bool bInSelectNewNode = true);
+
+			static Frontend::FNodeHandle AddExternalNodeHandle(UObject& InMetaSound, const FMetasoundFrontendClassName& InClassName);
 
 			// Adds an variable node with the given node handle to the editor graph.
-			static UMetasoundEditorGraphVariableNode* AddVariableNode(UObject& InMetaSound, Frontend::FNodeHandle& InNodeHandle, FVector2D InLocation, bool bInSelectNewNode = true);
+			static UMetasoundEditorGraphVariableNode* AddVariableNode(UObject& InMetaSound, const Frontend::FConstNodeHandle& InNodeHandle, bool bInSelectNewNode = true);
+
+			// Adds an input node to the editor graph that corresponds to the provided input template node handle.
+			static UMetasoundEditorGraphInputNode* AddInputNode(UObject& InMetaSound, const Frontend::FConstNodeHandle& InInputTemplateNode, bool bInSelectNewNode = true);
 
 			// Adds an output node to the editor graph that corresponds to the provided node handle.
-			static UMetasoundEditorGraphOutputNode* AddOutputNode(UObject& InMetaSound, Frontend::FNodeHandle& InNodeHandle, FVector2D InLocation, bool bInSelectNewNode = true);
+			static UMetasoundEditorGraphOutputNode* AddOutputNode(UObject& InMetaSound, const Frontend::FConstNodeHandle& InNodeHandle, bool bInSelectNewNode = true);
 
 			// Generates analogous FNodeHandle for the given internal node data. Does not bind nor create EdGraph representation of given node.
 			static Frontend::FNodeHandle AddOutputNodeHandle(UObject& InMetaSound, const FCreateNodeVertexParams& InParams, const FName* InNameBase = nullptr);
@@ -156,6 +158,10 @@ namespace Metasound
 			// Create a unique name for the variable.
 			static FName GenerateUniqueVariableName(const Frontend::FConstGraphHandle& InFrontendGraph, const FString& InBaseName);
 
+			// Convenience method for walking to the outermost object and transforming to a base MetaSound
+			static FMetasoundAssetBase& GetOutermostMetaSoundChecked(UObject& InSubObject);
+			static const FMetasoundAssetBase& GetOutermostConstMetaSoundChecked(const UObject& InSubObject);
+
 			// Adds a frontend variable to the root graph of the MetaSound
 			//
 			// @param InMetaSound - FMetasoundAssetBase derived UObject.
@@ -227,20 +233,14 @@ namespace Metasound
 			// Retrieves the proper pin color for the given PinType
 			static FLinearColor GetPinCategoryColor(const FEdGraphPinType& PinType);
 
-			// Initializes MetaSound with default inputs & outputs.
-			UE_DEPRECATED(5.3, "Use FMetaSoundFrontendDocumentBuilder::InitDocument instead.")
-			static void InitMetaSound(UObject& InMetaSound, const FString& InAuthor);
-
-			// Initializes a MetaSound Preset using the provided ReferencedMetaSound asset's
-			// root graph as the sole, encapsulated topology.
-			UE_DEPRECATED(5.3, "Use FMetaSoundFrontendDocumentBuilder::ConvertToPreset instead.")
-			static void InitMetaSoundPreset(UObject& InMetaSoundReferenced, UObject& InMetaSoundPreset);
-
 			// Rebuilds all editor node pins based on the provided node handle's class definition.
 			static void RebuildNodePins(UMetasoundEditorGraphNode& InGraphNode);
 
+			// Deletes all nodes from a parent graph that are linked to a given member.
+			static bool DeleteMemberNodes(const UMetasoundEditorGraphMember& InGraphMember);
+
 			// Deletes both the editor graph & frontend nodes from respective graphs
-			static bool DeleteNode(UEdGraphNode& InNode);
+			static bool DeleteNode(UEdGraphNode& InNode, bool bRemoveUnusedDependencies = true);
 
 			// Adds an Input UEdGraphPin to a UMetasoundEditorGraphNode
 			static UEdGraphPin* AddPinToNode(UMetasoundEditorGraphNode& InEditorNode, Frontend::FConstInputHandle InInputHandle);
@@ -259,25 +259,20 @@ namespace Metasound
 			static bool SynchronizeGraph(UObject& InMetaSound);
 
 			// Synchronizes editor nodes with frontend nodes, removing editor nodes that are not represented in the frontend, and adding editor nodes to represent missing frontend nodes.
-			//
-			// @return True if the UMetasoundEditorGraphNode was altered. False otherwise.
-			static bool SynchronizeNodes(UObject& InMetaSound);
+			static void SynchronizeNodes(UObject& InMetaSound);
 
-			// Synchronizes and reports to log whether or not an editor member node's associated FrontendNode ID has changed and therefore been updated through node versioning.
+			// Synchronizes and reports to log whether or not an output node's associated FrontendNode ID has changed and therefore been updated through node versioning.
 			//
 			// @return True if the UMetasoundEditorGraphNode was altered. False otherwise.
-			static bool SynchronizeNodeMembers(UObject& InMetaSound);
-
-			// Synchronizes node location data using handle's data
-			//
-			// @return True if the UMetasoundEditorGraphNode was altered. False otherwise.
-			static bool SynchronizeNodeLocation(const Frontend::FConstNodeHandle& InNode, UMetasoundEditorGraphNode& OutGraphNode);
+			static bool SynchronizeOutputNodes(UObject& InMetaSound);
 
 			// Adds and removes pins so that the UMetasoundEditorGraphNode matches the InNode.
 			//
 			// @return True if the UMetasoundEditorGraphNode was altered. False otherwise.
 			static bool SynchronizeNodePins(UMetasoundEditorGraphNode& InEditorNode, Frontend::FConstNodeHandle InNode, bool bRemoveUnusedPins = true, bool bLogChanges = true);
 
+			static bool SynchronizeComments(UObject& InMetaSound);
+
 			// Adds and removes connections so that the UEdGraph of the MetaSound has the same
 			// connections as the FMetasoundFrontendDocument graph.
 			//
diff --git a/Engine/Plugins/Runtime/Metasound/Source/MetasoundEditor/Private/MetasoundEditorGraphCommentNode.cpp b/Engine/Plugins/Runtime/Metasound/Source/MetasoundEditor/Private/MetasoundEditorGraphCommentNode.cpp
--- a/Engine/Plugins/Runtime/Metasound/Source/MetasoundEditor/Private/MetasoundEditorGraphCommentNode.cpp
+++ b/Engine/Plugins/Runtime/Metasound/Source/MetasoundEditor/Private/MetasoundEditorGraphCommentNode.cpp
@@ -6,6 +6,7 @@
 #include "Kismet2/Kismet2NameValidators.h"
 #include "Layout/SlateRect.h"
 #include "MetasoundBuilderSubsystem.h"
+#include "MetasoundEditorGraph.h"
 #include "MetasoundUObjectRegistry.h"
 #include "Styling/AppStyle.h"
 #include "Templates/Casts.h"
@@ -29,6 +30,11 @@ FMetasoundAssetBase& UMetasoundEditorGraphCommentNode::GetAssetChecked()
 	return *MetaSound;
 }
 
+bool UMetasoundEditorGraphCommentNode::CanUserDeleteNode() const
+{
+	return true;
+}
+
 void UMetasoundEditorGraphCommentNode::ConvertToFrontendComment(const UEdGraphNode_Comment& InEdNode, FMetaSoundFrontendGraphComment& OutComment)
 {
 	OutComment.bColorBubble = InEdNode.bColorCommentBubble != 0;
@@ -59,31 +65,9 @@ void UMetasoundEditorGraphCommentNode::PostEditChangeProperty(FPropertyChangedEv
 {
 	Super::PostEditChangeProperty(PropertyChangedEvent);
 
-	// TODO: Check property name and only change that which changed
 	FMetaSoundFrontendGraphComment Comment;
 	ConvertToFrontendComment(*this, Comment);
-
-	UMetaSoundBuilderBase& Builder = UMetaSoundBuilderSubsystem::GetChecked().AttachBuilderToAssetChecked(*GetOutermostObject());
-	Builder.FindOrAddGraphComment(CommentID) = MoveTemp(Comment);
-}
-
-void UMetasoundEditorGraphCommentNode::PostPlacedNewNode()
-{
-	Super::PostPlacedNewNode();
-
-	UMetaSoundBuilderBase& Builder = UMetaSoundBuilderSubsystem::GetChecked().AttachBuilderToAssetChecked(*GetOutermostObject());
-
-	// If the frontend version predates this new node, sync this node's comment to it...
-	if (const FMetaSoundFrontendGraphComment* Comment = Builder.FindGraphComment(CommentID))
-	{
-		NodeComment = Comment->Comment;
-	}
-
-	// otherwise, initialize the frontend version with the ed graph's default comment.
-	else
-	{
-		Builder.FindOrAddGraphComment(CommentID).Comment = NodeComment;
-	}
+	GetBuilderChecked().FindOrAddGraphComment(CommentID) = MoveTemp(Comment);
 }
 
 void UMetasoundEditorGraphCommentNode::ResizeNode(const FVector2D& NewSize)
@@ -91,22 +75,33 @@ void UMetasoundEditorGraphCommentNode::ResizeNode(const FVector2D& NewSize)
 	Super::ResizeNode(NewSize);
 	if (bCanResizeNode) 
 	{
-		UMetaSoundBuilderBase& Builder = UMetaSoundBuilderSubsystem::GetChecked().AttachBuilderToAssetChecked(*GetOutermostObject());
+		UMetasoundEditorGraph& Graph = *CastChecked<UMetasoundEditorGraph>(GetGraph());
+		UMetaSoundBuilderBase& Builder = UMetaSoundBuilderSubsystem::GetChecked().AttachBuilderToAssetChecked(Graph.GetMetasoundChecked());
 		Builder.FindOrAddGraphComment(CommentID).Size = FVector2D(NodeWidth, NodeHeight);
 	}
 }
 
+UMetaSoundBuilderBase& UMetasoundEditorGraphCommentNode::GetBuilderChecked() const
+{
+	UMetasoundEditorGraph* EdGraph = CastChecked<UMetasoundEditorGraph>(GetGraph());
+	return EdGraph->GetBuilderChecked();
+}
+
 FGuid UMetasoundEditorGraphCommentNode::GetCommentID() const
 {
 	return CommentID;
 }
 
+bool UMetasoundEditorGraphCommentNode::RemoveFromDocument() const
+{
+	return GetBuilderChecked().RemoveGraphComment(CommentID);
+}
+
 void UMetasoundEditorGraphCommentNode::SetBounds(const class FSlateRect& Rect)
 {
 	Super::SetBounds(Rect);
 
-	UMetaSoundBuilderBase& Builder = UMetaSoundBuilderSubsystem::GetChecked().AttachBuilderToAssetChecked(*GetOutermostObject());
-	FMetaSoundFrontendGraphComment& FrontendComment = Builder.FindOrAddGraphComment(CommentID);
+	FMetaSoundFrontendGraphComment& FrontendComment = GetBuilderChecked().FindOrAddGraphComment(CommentID);
 	FrontendComment.Position = FVector2D(NodePosX, NodePosY);
 	FrontendComment.Size = FVector2D(NodeWidth, NodeHeight);
 }
@@ -115,7 +110,7 @@ void UMetasoundEditorGraphCommentNode::OnRenameNode(const FString& NewName)
 {
 	Super::OnRenameNode(NewName);
 
-	UMetaSoundBuilderBase& Builder = UMetaSoundBuilderSubsystem::GetChecked().AttachBuilderToAssetChecked(*GetOutermostObject());
+	UMetaSoundBuilderBase& Builder = GetBuilderChecked();
 	Builder.FindOrAddGraphComment(CommentID).Comment = NewName;
 }
 #undef LOCTEXT_NAMESPACE
diff --git a/Engine/Plugins/Runtime/Metasound/Source/MetasoundEditor/Private/MetasoundEditorGraphConnectionDrawingPolicy.cpp b/Engine/Plugins/Runtime/Metasound/Source/MetasoundEditor/Private/MetasoundEditorGraphConnectionDrawingPolicy.cpp
--- a/Engine/Plugins/Runtime/Metasound/Source/MetasoundEditor/Private/MetasoundEditorGraphConnectionDrawingPolicy.cpp
+++ b/Engine/Plugins/Runtime/Metasound/Source/MetasoundEditor/Private/MetasoundEditorGraphConnectionDrawingPolicy.cpp
@@ -133,7 +133,10 @@ namespace Metasound
 					return;
 				}
 
+				PRAGMA_DISABLE_DEPRECATION_WARNINGS
 				Frontend::FConstOutputHandle OutputHandle = Frontend::FindReroutedOutput(InData.OutputHandle);
+				PRAGMA_ENABLE_DEPRECATION_WARNINGS
+
 				const FGuid NodeID = OutputHandle->GetOwningNodeID();
 				TNumericType Value = InDefaultValue;
 				FName OutputName = OutputHandle->GetName();
diff --git a/Engine/Plugins/Runtime/Metasound/Source/MetasoundEditor/Private/MetasoundEditorGraphInputNode.cpp b/Engine/Plugins/Runtime/Metasound/Source/MetasoundEditor/Private/MetasoundEditorGraphInputNode.cpp
--- a/Engine/Plugins/Runtime/Metasound/Source/MetasoundEditor/Private/MetasoundEditorGraphInputNode.cpp
+++ b/Engine/Plugins/Runtime/Metasound/Source/MetasoundEditor/Private/MetasoundEditorGraphInputNode.cpp
@@ -25,6 +25,53 @@
 
 #include UE_INLINE_GENERATED_CPP_BY_NAME(MetasoundEditorGraphInputNode)
 
+void UMetasoundEditorGraphInputNode::CacheTitle()
+{
+	using namespace Metasound::Frontend;
+
+	if (Input)
+	{
+		FConstNodeHandle NodeHandle = Input->GetNodeHandle();
+		CachedTitle = NodeHandle->GetDisplayTitle();
+	}
+	else
+	{
+		CachedTitle = FText::FromName(Breadcrumb.MemberName);
+	}
+}
+
+const FMetasoundEditorGraphVertexNodeBreadcrumb& UMetasoundEditorGraphInputNode::GetBreadcrumb() const
+{
+	return Breadcrumb;
+}
+
+void UMetasoundEditorGraphInputNode::CacheBreadcrumb()
+{
+	using namespace Metasound::Frontend;
+
+	Breadcrumb = { };
+
+	// Take data from associated input as pasted graph may not be same as local graph
+	// and associated input will not be copied with given node.  Need the following data
+	// to associate or create new associated input.
+	if (Input)
+	{
+		FConstNodeHandle NodeHandle = Input->GetConstNodeHandle();
+
+		Breadcrumb.MemberName = NodeHandle->GetNodeName();
+		Breadcrumb.ClassName = NodeHandle->GetClassMetadata().GetClassName();
+
+		FConstOutputHandle OutputHandle = NodeHandle->GetConstOutputs().Last();
+		Breadcrumb.AccessType = OutputHandle->GetVertexAccessType();
+		Breadcrumb.DataType = OutputHandle->GetDataType();
+
+		if (const UMetasoundEditorGraphMemberDefaultLiteral* Literal = Input->GetLiteral())
+		{
+			Breadcrumb.DefaultLiteral = Literal->GetDefault();
+		}
+	}
+}
+
 UMetasoundEditorGraphMember* UMetasoundEditorGraphInputNode::GetMember() const
 {
 	return Input;
@@ -32,12 +79,15 @@ UMetasoundEditorGraphMember* UMetasoundEditorGraphInputNode::GetMember() const
 
 FMetasoundFrontendClassName UMetasoundEditorGraphInputNode::GetClassName() const
 {
-	if (ensure(Input))
+	using namespace Metasound::Frontend;
+
+	if (Input)
 	{
-		return Input->ClassName;
+		FConstNodeHandle NodeHandle = Input->GetConstNodeHandle();
+		return NodeHandle->GetClassMetadata().GetClassName();
 	}
 
-	return Super::GetClassName();
+	return Breadcrumb.ClassName;
 }
 
 void UMetasoundEditorGraphInputNode::UpdatePreviewInstance(const Metasound::FVertexName& InParameterName, TScriptInterface<IAudioParameterControllerInterface>& InParameterInterface) const
@@ -53,12 +103,7 @@ void UMetasoundEditorGraphInputNode::UpdatePreviewInstance(const Metasound::FVer
 
 FGuid UMetasoundEditorGraphInputNode::GetNodeID() const
 {
-	if (Input)
-	{
-		return Input->NodeID;
-	}
-
-	return Super::GetNodeID();
+	return NodeID;
 }
 
 FLinearColor UMetasoundEditorGraphInputNode::GetNodeTitleColor() const
@@ -77,6 +122,40 @@ FSlateIcon UMetasoundEditorGraphInputNode::GetNodeTitleIcon() const
 	return FSlateIcon("MetaSoundStyle", NativeIconName);
 }
 
+void UMetasoundEditorGraphInputNode::GetPinHoverText(const UEdGraphPin& Pin, FString& OutHoverText) const
+{
+	using namespace Metasound::Editor;
+	using namespace Metasound::Frontend;
+
+	if (ensure(Pin.Direction == EGPD_Output)) // Should never display input pin for input node hover
+	{
+		if (ensure(Input))
+		{
+			FConstNodeHandle InputNode = Input->GetConstNodeHandle();
+			OutHoverText = InputNode->GetDescription().ToString();
+			if (ShowNodeDebugData())
+			{
+				FConstOutputHandle OutputHandle = FGraphBuilder::FindReroutedConstOutputHandleFromPin(&Pin);
+				OutHoverText = FString::Format(TEXT("{0}\nVertex Name: {1}\nDataType: {2}\nID: {3}"),
+				{
+					OutHoverText,
+					OutputHandle->GetName().ToString(),
+					OutputHandle->GetDataType().ToString(),
+					OutputHandle->GetID().ToString(),
+				});
+			}
+		}
+	}
+}
+
+void UMetasoundEditorGraphInputNode::ReconstructNode()
+{
+	using namespace Metasound::Editor;
+	using namespace Metasound::Frontend;
+
+	Super::ReconstructNode();
+}
+
 void UMetasoundEditorGraphInputNode::Validate(Metasound::Editor::FGraphNodeValidationResult& OutResult)
 {
 #if WITH_EDITOR
@@ -111,12 +190,3 @@ void UMetasoundEditorGraphInputNode::Validate(Metasound::Editor::FGraphNodeValid
 	}
 #endif // #if WITH_EDITOR
 }
-
-void UMetasoundEditorGraphInputNode::SetNodeID(FGuid InNodeID)
-{
-	if (Input)
-	{
-		Input->NodeID = InNodeID;
-	}
-}
-
