Version Metasound Document to include all ed data and make all Metasound EdGraph data transient

diff --git a/Engine/Plugins/Runtime/Metasound/Source/MetasoundEditor/Private/MetasoundEditorGraphNode.cpp b/Engine/Plugins/Runtime/Metasound/Source/MetasoundEditor/Private/MetasoundEditorGraphNode.cpp
--- a/Engine/Plugins/Runtime/Metasound/Source/MetasoundEditor/Private/MetasoundEditorGraphNode.cpp
+++ b/Engine/Plugins/Runtime/Metasound/Source/MetasoundEditor/Private/MetasoundEditorGraphNode.cpp
@@ -65,20 +65,84 @@ void UMetasoundEditorGraphNode::UpdateFrontendNodeLocation(const FVector2D& InLo
 {
 	using namespace Metasound::Frontend;
 	
+	ensureMsgf(NodeGuid.IsValid(), TEXT("Cannot update frontend node location prior to node guid being finalized."));
+
 	FNodeHandle NodeHandle = GetNodeHandle();
 	FMetasoundFrontendNodeStyle Style = NodeHandle->GetNodeStyle();
 	Style.Display.Locations.FindOrAdd(NodeGuid) = InLocation;
 	NodeHandle->SetNodeStyle(Style);
 }
 
+bool UMetasoundEditorGraphNode::ShowNodeDebugData()
+{
+	using namespace Metasound::Editor;
+
+	return GraphNodePrivate::ShowNodeDebugData != 0;
+}
+
+bool UMetasoundEditorGraphNode::RemoveFromDocument() const
+{
+	UMetaSoundBuilderBase& Builder = GetBuilderChecked();
+	EMetaSoundBuilderResult Result = EMetaSoundBuilderResult::Failed;
+	Builder.RemoveNode({ GetNodeID() }, Result);
+	return Result == EMetaSoundBuilderResult::Succeeded;
+}
+
 void UMetasoundEditorGraphNode::SetNodeLocation(const FVector2D& InLocation)
 {
-	NodePosX = FMath::TruncToInt(InLocation.X);
-	NodePosY = FMath::TruncToInt(InLocation.Y);
+	NodePosX = UE::LWC::FloatToIntCastChecked<int32>(InLocation.X);
+	NodePosY = UE::LWC::FloatToIntCastChecked<int32>(InLocation.Y);
 
 	UpdateFrontendNodeLocation(InLocation);
 }
 
+void UMetasoundEditorGraphNode::SyncCommentFromFrontendNode()
+{
+	using namespace Metasound::Frontend;
+
+	FConstNodeHandle NodeHandle = GetNodeHandle();
+	const FMetasoundFrontendNodeStyle& Style = NodeHandle->GetNodeStyle();
+	NodeComment = Style.Display.Comment;
+	bCommentBubbleMakeVisible = Style.Display.bCommentVisible;
+}
+
+bool UMetasoundEditorGraphNode::SyncLocationFromFrontendNode(bool bUpdateEditorNodeID)
+{
+	using namespace Metasound::Frontend;
+
+	FConstNodeHandle NodeHandle = GetConstNodeHandle();
+	const FMetasoundFrontendNodeStyle& Style = NodeHandle->GetNodeStyle();
+
+	if (!Style.Display.Locations.IsEmpty())
+	{
+		if (bUpdateEditorNodeID)
+		{
+			if (ensureMsgf(Style.Display.Locations.Num() == 1, TEXT("Style location map should only contain one location (legacy support for multiple location values should be migrated by this point)")))
+			{
+				for (const TPair<FGuid, FVector2D>& IDLocationPair : Style.Display.Locations)
+				{
+					if (ensureMsgf(IDLocationPair.Key.IsValid(), TEXT("Location cannot be updated for ed node with invalid guid")))
+					{
+						NodeGuid = IDLocationPair.Key;
+						NodePosX = UE::LWC::FloatToIntCastChecked<int32>(IDLocationPair.Value.X);
+						NodePosY = UE::LWC::FloatToIntCastChecked<int32>(IDLocationPair.Value.Y);
+						return true;
+					}
+				}
+			}
+		}
+		if (const FVector2D* Location = Style.Display.Locations.Find(NodeGuid))
+		{
+			NodePosX = UE::LWC::FloatToIntCastChecked<int32>(Location->X);
+			NodePosY = UE::LWC::FloatToIntCastChecked<int32>(Location->Y);
+
+			return true;
+		}
+	}
+
+	return false;
+}
+
 void UMetasoundEditorGraphNode::PostLoad()
 {
 	Super::PostLoad();
@@ -134,13 +198,10 @@ int32 UMetasoundEditorGraphNode::EstimateNodeWidth() const
 	}
 }
 
-const FMetaSoundFrontendDocumentBuilder& UMetasoundEditorGraphNode::GetFrontendBuilderChecked() const
+UMetaSoundBuilderBase& UMetasoundEditorGraphNode::GetBuilderChecked() const
 {
-	UObject* Outermost = GetOutermostObject();
-	check(Outermost);
-
-	const UMetaSoundBuilderBase& BuilderBase = UMetaSoundBuilderSubsystem::GetChecked().AttachBuilderToAssetChecked(*Outermost);
-	return BuilderBase.GetConstBuilder();
+	UMetasoundEditorGraph* EdGraph = CastChecked<UMetasoundEditorGraph>(GetGraph());
+	return UMetaSoundBuilderSubsystem::GetChecked().AttachBuilderToAssetChecked(EdGraph->GetMetasoundChecked());
 }
 
 const FMetasoundFrontendNode* UMetasoundEditorGraphNode::GetFrontendNode() const
@@ -168,6 +229,12 @@ const FMetasoundFrontendNode& UMetasoundEditorGraphNode::GetFrontendNodeChecked(
 	return *FrontendNode;
 }
 
+const FMetasoundEditorGraphNodeBreadcrumb& UMetasoundEditorGraphNode::GetBreadcrumb() const
+{
+	static const FMetasoundEditorGraphNodeBreadcrumb StubCrumb;
+	return StubCrumb;
+}
+
 UObject& UMetasoundEditorGraphNode::GetMetasoundChecked()
 {
 	UMetasoundEditorGraph* EdGraph = CastChecked<UMetasoundEditorGraph>(GetGraph());
@@ -183,10 +250,7 @@ const UObject& UMetasoundEditorGraphNode::GetMetasoundChecked() const
 Metasound::Frontend::FConstGraphHandle UMetasoundEditorGraphNode::GetConstRootGraphHandle() const
 {
 	const FMetasoundAssetBase* ConstMetasoundAsset = Metasound::IMetasoundUObjectRegistry::Get().GetObjectAsAssetBase(&GetMetasoundChecked());
-	FMetasoundAssetBase* MetasoundAsset = const_cast<FMetasoundAssetBase*>(ConstMetasoundAsset);
-	check(MetasoundAsset);
-
-	return MetasoundAsset->GetRootGraphHandle();
+	return ConstMetasoundAsset->GetRootGraphHandle();
 }
 
 Metasound::Frontend::FGraphHandle UMetasoundEditorGraphNode::GetRootGraphHandle() const
@@ -255,7 +319,7 @@ void UMetasoundEditorGraphNode::Validate(Metasound::Editor::FGraphNodeValidation
 
 #if WITH_EDITOR
 	// Validate that non-reroute inputs are connected to "real" outputs
-	if (GetClassName() != FRerouteNodeTemplate::ClassName)
+	if (GetBreadcrumb().ClassName != FRerouteNodeTemplate::ClassName)
 	{
 		for (UEdGraphPin* Pin : Pins)
 		{
@@ -269,7 +333,7 @@ void UMetasoundEditorGraphNode::Validate(Metasound::Editor::FGraphNodeValidation
 					{
 						if (UMetasoundEditorGraphExternalNode* ExternalNode = Cast<UMetasoundEditorGraphExternalNode>(ReroutedPin->GetOwningNode()))
 						{
-							if (ExternalNode->GetClassName() == FRerouteNodeTemplate::ClassName)
+							if (ExternalNode->GetBreadcrumb().ClassName == FRerouteNodeTemplate::ClassName)
 							{
 								FConstInputHandle InputHandle = FGraphBuilder::GetConstInputHandleFromPin(Pin);
 								const FString Msg = FString::Format(*GraphNodePrivate::MissingConcreteOutputConnectionFormat, { InputHandle->GetDisplayName().ToString() });
@@ -443,17 +507,28 @@ Metasound::Frontend::FDataTypeRegistryInfo UMetasoundEditorGraphNode::GetPinData
 	using namespace Metasound::Editor;
 	using namespace Metasound::Frontend;
 
- 	Metasound::Frontend::FDataTypeRegistryInfo DataTypeInfo;
+ 	FDataTypeRegistryInfo DataTypeInfo;
 
-	if (InPin.Direction == EGPD_Input)
-	{
-		FConstInputHandle Handle = FGraphBuilder::GetConstInputHandleFromPin(&InPin);
-		ensure(IDataTypeRegistry::Get().GetDataTypeInfo(Handle->GetDataType(), DataTypeInfo));
-	}
-	else // InPin.Direction == EGPD_Output
+	const FMetaSoundFrontendDocumentBuilder& Builder = GetBuilderChecked().GetConstBuilder();
+	const FMetasoundFrontendVertexHandle Handle = FGraphBuilder::GetPinVertexHandle(Builder, &InPin);
+	if (Handle.IsSet())
 	{
-		FConstOutputHandle Handle = FGraphBuilder::GetConstOutputHandleFromPin(&InPin);
-		ensure(IDataTypeRegistry::Get().GetDataTypeInfo(Handle->GetDataType(), DataTypeInfo));
+		if (InPin.Direction == EGPD_Input)
+		{
+			const FMetasoundFrontendVertex* Vertex = Builder.FindNodeInput(Handle.NodeID, Handle.VertexID);
+			if (Vertex)
+			{
+				IDataTypeRegistry::Get().GetDataTypeInfo(Vertex->TypeName, DataTypeInfo);
+			}
+		}
+		else // InPin.Direction == EGPD_Output
+		{
+			const FMetasoundFrontendVertex* Vertex = Builder.FindNodeOutput(Handle.NodeID, Handle.VertexID);
+			if (Vertex)
+			{
+				IDataTypeRegistry::Get().GetDataTypeInfo(Vertex->TypeName, DataTypeInfo);
+			}
+		}
 	}
 
 	return DataTypeInfo;
@@ -530,6 +605,22 @@ FString UMetasoundEditorGraphNode::GetPinMetaData(FName InPinName, FName InKey)
 	return Super::GetPinMetaData(InPinName, InKey);
 }
 
+void UMetasoundEditorGraphNode::OnUpdateCommentText(const FString& NewComment)
+{
+	using namespace Metasound::Frontend;
+
+	if (!NodeComment.Equals(NewComment))
+	{
+		const FScopedTransaction Transaction(LOCTEXT("CommentCommitted", "Comment Changed"));
+		Modify();
+		NodeComment = NewComment;
+		FNodeHandle NodeHandle = GetNodeHandle();
+		FMetasoundFrontendNodeStyle Style = NodeHandle->GetNodeStyle();
+		Style.Display.Comment = NewComment;
+		NodeHandle->SetNodeStyle(Style);
+	}
+}
+
 void UMetasoundEditorGraphNode::PreSave(FObjectPreSaveContext InSaveContext)
 {
 	using namespace Metasound::Editor;
@@ -678,7 +769,7 @@ void UMetasoundEditorGraphOutputNode::PinDefaultValueChanged(UEdGraphPin* InPin)
 		UObject& MetaSound = GetMetasoundChecked();
 		MetaSound.Modify();
 		
-		FInputHandle InputHandle = FGraphBuilder::GetInputHandleFromPin(InPin);
+		FConstInputHandle InputHandle = FGraphBuilder::GetConstInputHandleFromPin(InPin);
 		if (InputHandle->IsValid())
 		{
 			FMetasoundFrontendLiteral LiteralValue;
@@ -717,6 +808,46 @@ void UMetasoundEditorGraphOutputNode::PinDefaultValueChanged(UEdGraphPin* InPin)
 	}
 }
 
+void UMetasoundEditorGraphOutputNode::ReconstructNode()
+{
+	using namespace Metasound::Frontend;
+
+	if (!Output)
+	{
+		UMetasoundEditorGraph* Graph = CastChecked<UMetasoundEditorGraph>(GetGraph());
+		if (ensure(Graph))
+		{
+			Graph->IterateOutputs([&Graph, this](UMetasoundEditorGraphOutput& TestOutput)
+			{
+				FConstNodeHandle OutputHandle = TestOutput.GetConstNodeHandle();
+				FConstInputHandle TestInput = OutputHandle->GetConstInputs().Last();
+				const bool bTypeMatches = TestInput->GetDataType() == Breadcrumb.DataType;
+				const bool bAccessMatches = TestInput->GetVertexAccessType() == Breadcrumb.AccessType;
+				const bool bNameMatches = OutputHandle->GetNodeName() == Breadcrumb.MemberName;
+				if (bTypeMatches && bAccessMatches && bNameMatches)
+				{
+					Output = &TestOutput;
+				}
+			});
+		}
+	}
+
+	Super::ReconstructNode();
+}
+
+bool UMetasoundEditorGraphOutputNode::RemoveFromDocument() const
+{
+	if (Output)
+	{
+		UMetaSoundBuilderBase& Builder = GetBuilderChecked();
+		EMetaSoundBuilderResult Result = EMetaSoundBuilderResult::Failed;
+		Builder.RemoveGraphOutput(Output->GetMemberName(), Result);
+		return Result == EMetaSoundBuilderResult::Succeeded;
+	}
+
+	return false;
+}
+
 bool UMetasoundEditorGraphOutputNode::EnableInteractWidgets() const
 {
 	using namespace Metasound::Frontend;
@@ -734,7 +865,6 @@ bool UMetasoundEditorGraphOutputNode::EnableInteractWidgets() const
 
 void UMetasoundEditorGraphOutputNode::Validate(Metasound::Editor::FGraphNodeValidationResult& OutResult)
 {
-#if WITH_EDITOR
 	using namespace Metasound::Editor;
 	using namespace Metasound::Frontend;
 
@@ -763,16 +893,43 @@ void UMetasoundEditorGraphOutputNode::Validate(Metasound::Editor::FGraphNodeVali
 			}
 		}
 	}
-#endif // #if WITH_EDITOR
+}
+
+const FMetasoundEditorGraphNodeBreadcrumb& UMetasoundEditorGraphOutputNode::GetBreadcrumb() const
+{
+	return Breadcrumb;
+}
+
+void UMetasoundEditorGraphOutputNode::CacheBreadcrumb()
+{
+	using namespace Metasound::Frontend;
+	FConstNodeHandle NodeHandle = GetConstNodeHandle();
+
+	Breadcrumb.MemberName = NodeHandle->GetNodeName();
+	Breadcrumb.ClassName = NodeHandle->GetClassMetadata().GetClassName();
+
+	FConstInputHandle InputHandle = NodeHandle->GetConstInputs().Last();
+	Breadcrumb.AccessType = InputHandle->GetVertexAccessType();
+	Breadcrumb.DataType = InputHandle->GetDataType();
+
+
+	if (const UMetasoundEditorGraphMemberDefaultLiteral* Literal = Output->GetLiteral())
+	{
+		Breadcrumb.DefaultLiteral = Literal->GetDefault();
+	}
 }
 
 FMetasoundFrontendClassName UMetasoundEditorGraphOutputNode::GetClassName() const
 {
-	if (ensure(Output))
+	using namespace Metasound::Frontend;
+
+	if (Output)
 	{
-		return Output->ClassName;
+		FConstNodeHandle NodeHandle = Output->GetConstNodeHandle();
+		return NodeHandle->GetClassMetadata().GetClassName();
 	}
-	return FMetasoundFrontendClassName();
+
+	return Breadcrumb.ClassName;
 }
 
 FGuid UMetasoundEditorGraphOutputNode::GetNodeID() const
@@ -850,6 +1007,15 @@ bool UMetasoundEditorGraphExternalNode::CanAutoUpdate() const
 	return GetConstNodeHandle()->CanAutoUpdate(InterfaceUpdates);
 }
 
+void UMetasoundEditorGraphExternalNode::CacheBreadcrumb()
+{
+	using namespace Metasound::Frontend;
+	FConstNodeHandle NodeHandle = GetConstNodeHandle();
+	const FNodeRegistryKey RegistryKey(NodeHandle->GetClassMetadata());
+	Breadcrumb.bIsClassNative = FMetasoundFrontendRegistryContainer::Get()->IsNodeNative(RegistryKey);
+	Breadcrumb.ClassName = NodeHandle->GetClassMetadata().GetClassName();
+}
+
 void UMetasoundEditorGraphExternalNode::CacheTitle()
 {
 	using namespace Metasound::Editor;
@@ -939,7 +1105,7 @@ void UMetasoundEditorGraphExternalNode::Validate(Metasound::Editor::FGraphNodeVa
 			{
 #if WITH_EDITOR
 				FString Message;
-				if (!Template->HasRequiredConnections(GetFrontendBuilderChecked(), GetNodeID(), &Message))
+				if (!Template->HasRequiredConnections(GetBuilderChecked().GetConstBuilder(), GetNodeID(), &Message))
 				{
 					OutResult.SetMessage(EMessageSeverity::Warning, Message);
 				}
@@ -1025,8 +1191,8 @@ void UMetasoundEditorGraphExternalNode::Validate(Metasound::Editor::FGraphNodeVa
 		bool bWasRemoved = false;
 		if (Pin->Direction == EGPD_Input)
 		{
-			FInputHandle Input = FGraphBuilder::GetInputHandleFromPin(Pin);
-			bWasRemoved |= InterfaceUpdates.RemovedInputs.ContainsByPredicate([&](const FMetasoundFrontendClassInput* ClassInput)
+			FConstInputHandle Input = FGraphBuilder::GetConstInputHandleFromPin(Pin);
+			bWasRemoved |= InterfaceUpdates.RemovedInputs.ContainsByPredicate([&Input](const FMetasoundFrontendClassInput* ClassInput)
 			{
 				return Input->GetName() == ClassInput->Name && Input->GetDataType() == ClassInput->TypeName;
 			});
@@ -1034,8 +1200,8 @@ void UMetasoundEditorGraphExternalNode::Validate(Metasound::Editor::FGraphNodeVa
 
 		if (Pin->Direction == EGPD_Output)
 		{
-			FOutputHandle Output = FGraphBuilder::GetOutputHandleFromPin(Pin);
-			bWasRemoved |= InterfaceUpdates.RemovedOutputs.ContainsByPredicate([&](const FMetasoundFrontendClassOutput* ClassOutput)
+			FConstOutputHandle Output = FGraphBuilder::GetConstOutputHandleFromPin(Pin);
+			bWasRemoved |= InterfaceUpdates.RemovedOutputs.ContainsByPredicate([&Output](const FMetasoundFrontendClassOutput* ClassOutput)
 			{
 				return Output->GetName() == ClassOutput->Name && Output->GetDataType() == ClassOutput->TypeName;
 			});
@@ -1135,6 +1301,11 @@ void UMetasoundEditorGraphExternalNode::Validate(Metasound::Editor::FGraphNodeVa
 #endif // WITH_EDITOR
 }
 
+const FMetasoundEditorGraphNodeBreadcrumb& UMetasoundEditorGraphExternalNode::GetBreadcrumb() const
+{
+	return Breadcrumb;
+}
+
 FLinearColor UMetasoundEditorGraphExternalNode::GetNodeTitleColor() const
 {
 	if (const UMetasoundEditorSettings* EditorSettings = GetDefault<UMetasoundEditorSettings>())
@@ -1166,7 +1337,7 @@ bool UMetasoundEditorGraphExternalNode::ShouldDrawNodeAsControlPointOnly(int32&
 {
 	using namespace Metasound::Frontend;
 
-	if (GetClassName() == FRerouteNodeTemplate::ClassName)
+	if (GetBreadcrumb().ClassName == FRerouteNodeTemplate::ClassName)
 	{
 		OutInputPinIndex = 0;
 		OutOutputPinIndex = 1;
@@ -1176,6 +1347,25 @@ bool UMetasoundEditorGraphExternalNode::ShouldDrawNodeAsControlPointOnly(int32&
 	return Super::ShouldDrawNodeAsControlPointOnly(OutInputPinIndex, OutOutputPinIndex);
 }
 
+void UMetasoundEditorGraphVariableNode::CacheBreadcrumb()
+{
+	using namespace Metasound::Frontend;
+
+	Breadcrumb = { };
+	if (Variable)
+	{
+		FConstVariableHandle VariableHandle = Variable->GetConstVariableHandle();
+
+		Breadcrumb.MemberName = VariableHandle->GetName();
+		Breadcrumb.DataType = VariableHandle->GetDataType();
+	}
+}
+
+const FMetasoundEditorGraphNodeBreadcrumb& UMetasoundEditorGraphVariableNode::GetBreadcrumb() const
+{
+	return Breadcrumb;
+}
+
 UMetasoundEditorGraphMember* UMetasoundEditorGraphVariableNode::GetMember() const
 {
 	return Variable;
@@ -1246,7 +1436,7 @@ void UMetasoundEditorGraphVariableNode::PinDefaultValueChanged(UEdGraphPin* Pin)
 		UObject& MetaSound = GetMetasoundChecked();
 		MetaSound.Modify();
 
-		FInputHandle InputHandle = FGraphBuilder::GetInputHandleFromPin(Pin);
+		FConstInputHandle InputHandle = FGraphBuilder::GetConstInputHandleFromPin(Pin);
 		if (InputHandle->IsValid())
 		{
 			FMetasoundFrontendLiteral LiteralValue;
diff --git a/Engine/Plugins/Runtime/Metasound/Source/MetasoundEditor/Private/MetasoundEditorGraphNodeFactory.h b/Engine/Plugins/Runtime/Metasound/Source/MetasoundEditor/Private/MetasoundEditorGraphNodeFactory.h
--- a/Engine/Plugins/Runtime/Metasound/Source/MetasoundEditor/Private/MetasoundEditorGraphNodeFactory.h
+++ b/Engine/Plugins/Runtime/Metasound/Source/MetasoundEditor/Private/MetasoundEditorGraphNodeFactory.h
@@ -3,6 +3,7 @@
 
 #include "EdGraphUtilities.h"
 #include "EdGraph/EdGraphNode.h"
+#include "MetasoundEditorGraphCommentNode.h"
 #include "MetasoundEditorGraphNode.h"
 #include "MetasoundEditorGraphSchema.h"
 #include "NodeTemplates/MetasoundFrontendNodeTemplateReroute.h"
@@ -33,9 +34,10 @@ class FMetasoundGraphNodeFactory : public FGraphPanelNodeFactory
 			}
 			return SNew(SMetaSoundGraphNode, InNode);
 		}
-		else if (UEdGraphNode_Comment* CommentNode = Cast<UEdGraphNode_Comment>(InNode))
+		else if (UMetasoundEditorGraphCommentNode* CommentNode = Cast<UMetasoundEditorGraphCommentNode>(InNode))
 		{
-			if (CommentNode->GetSchema()->IsA(UMetasoundEditorGraphSchema::StaticClass()))
+			const UEdGraphSchema* EdGraphSchema = CommentNode->GetSchema();
+			if (EdGraphSchema && EdGraphSchema->IsA(UMetasoundEditorGraphSchema::StaticClass()))
 			{
 				return SNew(SMetasoundGraphNodeComment, CommentNode);
 			}
diff --git a/Engine/Plugins/Runtime/Metasound/Source/MetasoundEditor/Private/MetasoundEditorGraphSchema.cpp b/Engine/Plugins/Runtime/Metasound/Source/MetasoundEditor/Private/MetasoundEditorGraphSchema.cpp
--- a/Engine/Plugins/Runtime/Metasound/Source/MetasoundEditor/Private/MetasoundEditorGraphSchema.cpp
+++ b/Engine/Plugins/Runtime/Metasound/Source/MetasoundEditor/Private/MetasoundEditorGraphSchema.cpp
@@ -23,6 +23,7 @@
 #include "MetasoundEditorCommands.h"
 #include "MetasoundEditorGraph.h"
 #include "MetasoundEditorGraphBuilder.h"
+#include "MetasoundEditorGraphCommentNode.h"
 #include "MetasoundEditorGraphInputNode.h"
 #include "MetasoundEditorGraphNode.h"
 #include "MetasoundEditorGraphValidation.h"
@@ -37,11 +38,13 @@
 #include "MetasoundLiteral.h"
 #include "MetasoundNodeInterface.h"
 #include "MetasoundStandardNodesCategories.h"
+#include "MetasoundFrontendNodeTemplateRegistry.h"
 #include "MetasoundUObjectRegistry.h"
 #include "MetasoundVariableNodes.h"
 #include "MetasoundVertex.h"
 #include "NodeTemplates/MetasoundFrontendNodeTemplateReroute.h"
 #include "ScopedTransaction.h"
+#include "Settings/EditorStyleSettings.h"
 #include "Styling/SlateStyleRegistry.h"
 #include "Toolkits/ToolkitManager.h"
 #include "ToolMenus.h"
@@ -162,12 +165,12 @@ PRAGMA_ENABLE_DEPRECATION_WARNINGS
 
 				if (FromPin->Direction == EGPD_Input)
 				{
-					FInputHandle InputHandle = FGraphBuilder::GetInputHandleFromPin(FromPin);
+					FConstInputHandle InputHandle = FGraphBuilder::GetConstInputHandleFromPin(FromPin);
 					for (UEdGraphPin* Pin : NewGraphNode.Pins)
 					{
 						if (Pin->Direction == EGPD_Output)
 						{
-							FOutputHandle OutputHandle = FGraphBuilder::GetOutputHandleFromPin(Pin);
+							FConstOutputHandle OutputHandle = FGraphBuilder::GetConstOutputHandleFromPin(Pin);
 							if (OutputHandle->IsValid() && 
 								InputHandle->CanConnectTo(*OutputHandle).Connectable == FConnectability::EConnectable::Yes)
 							{
@@ -182,12 +185,12 @@ PRAGMA_ENABLE_DEPRECATION_WARNINGS
 
 				if (FromPin->Direction == EGPD_Output)
 				{
-					FOutputHandle OutputHandle = FGraphBuilder::GetOutputHandleFromPin(FromPin);
+					FConstOutputHandle OutputHandle = FGraphBuilder::GetConstOutputHandleFromPin(FromPin);
 					for (UEdGraphPin* Pin : NewGraphNode.Pins)
 					{
 						if (Pin->Direction == EGPD_Input)
 						{
-							FInputHandle InputHandle = FGraphBuilder::GetInputHandleFromPin(Pin);
+							FConstInputHandle InputHandle = FGraphBuilder::GetConstInputHandleFromPin(Pin);
 							if (InputHandle->IsValid() && 
 								InputHandle->CanConnectTo(*OutputHandle).Connectable == FConnectability::EConnectable::Yes)
 							{
@@ -285,10 +288,13 @@ PRAGMA_ENABLE_DEPRECATION_WARNINGS
 						constexpr bool bPostSubTransaction = false;
 						Variable->SetMemberName(NodeName, bPostSubTransaction);
 
-						FNodeHandle NodeHandle = FGraphBuilder::AddVariableNodeHandle(ParentMetaSound, Variable->GetVariableID(), ClassName.ToNodeClassName());
+						FConstNodeHandle NodeHandle = FGraphBuilder::AddVariableNodeHandle(ParentMetaSound, Variable->GetVariableID(), ClassName.ToNodeClassName());
 
-						if (UMetasoundEditorGraphVariableNode* NewGraphNode = FGraphBuilder::AddVariableNode(ParentMetaSound, NodeHandle, InLocation))
+						if (UMetasoundEditorGraphVariableNode* NewGraphNode = FGraphBuilder::AddVariableNode(ParentMetaSound, NodeHandle))
 						{
+							NewGraphNode->UpdateFrontendNodeLocation(InLocation);
+							NewGraphNode->SyncLocationFromFrontendNode();
+
 							// Set the literal using the new value if connecting to an input
 							if (FromPin.Direction == EGPD_Input)
 							{
@@ -451,9 +457,11 @@ UEdGraphNode* FMetasoundGraphSchemaAction_NewNode::PerformAction(UEdGraph* Paren
 	ParentMetasound.Modify();
 	ParentGraph->Modify();
 
-	if (UMetasoundEditorGraphExternalNode* NewGraphNode = FGraphBuilder::AddExternalNode(ParentMetasound, ClassMetadata, Location, bSelectNewNode))
+	if (UMetasoundEditorGraphExternalNode* NewGraphNode = FGraphBuilder::AddExternalNode(ParentMetasound, ClassMetadata, bSelectNewNode))
 	{
 		NewGraphNode->Modify();
+		NewGraphNode->UpdateFrontendNodeLocation(Location);
+		NewGraphNode->SyncLocationFromFrontendNode();
 		SchemaPrivate::TryConnectNewNodeToMatchingDataTypePin(*NewGraphNode, FromPin);
 		SchemaPrivate::SelecteNodeInEditor(*MetaSoundGraph, *NewGraphNode);
 		return NewGraphNode;
@@ -498,8 +506,8 @@ UEdGraphNode* FMetasoundGraphSchemaAction_NewInput::PerformAction(UEdGraph* Pare
 		return nullptr;
 	}
 
-	FNodeHandle NodeHandle = Input->GetNodeHandle();
-	if (!ensure(NodeHandle->IsValid()))
+	FNodeHandle InputNodeHandle = Input->GetNodeHandle();
+	if (!ensure(InputNodeHandle->IsValid()))
 	{
 		return nullptr;
 	}
@@ -509,9 +517,12 @@ UEdGraphNode* FMetasoundGraphSchemaAction_NewInput::PerformAction(UEdGraph* Pare
 	MetasoundGraph->Modify();
 	Input->Modify();
 
-	if (UMetasoundEditorGraphInputNode* NewGraphNode = FGraphBuilder::AddInputNode(ParentMetasound, NodeHandle, InLocation))
+	FConstNodeHandle TemplateNodeHandle = FGraphBuilder::AddInputTemplateNodeHandle(ParentMetasound, InputNodeHandle);
+	if (UMetasoundEditorGraphNode* NewGraphNode = FGraphBuilder::AddInputNode(ParentMetasound, TemplateNodeHandle))
 	{
 		NewGraphNode->Modify();
+		NewGraphNode->UpdateFrontendNodeLocation(InLocation);
+		NewGraphNode->SyncLocationFromFrontendNode();
 		SchemaPrivate::TryConnectNewNodeToMatchingDataTypePin(*NewGraphNode, FromPin);
 		FGraphBuilder::RegisterGraphWithFrontend(ParentMetasound);
 		return NewGraphNode;
@@ -553,14 +564,17 @@ UEdGraphNode* FMetasoundGraphSchemaAction_PromoteToInput::PerformAction(UEdGraph
 	// The promoted input must have the same vertex access type in order for it to be connectible
 	const FCreateNodeVertexParams VertexParams = { InputHandle->GetDataType(), InputHandle->GetVertexAccessType() };
 
-	FNodeHandle NodeHandle = FGraphBuilder::AddInputNodeHandle(ParentMetasound, VertexParams, &DefaultValue, &InputName);
-	if (ensure(NodeHandle->IsValid()))
+	FNodeHandle InputNodeHandle = FGraphBuilder::AddInputNodeHandle(ParentMetasound, VertexParams, &DefaultValue, &InputName);
+	if (ensure(InputNodeHandle->IsValid()))
 	{
-		UMetasoundEditorGraphInput* Input = MetasoundGraph->FindOrAddInput(NodeHandle);
+		UMetasoundEditorGraphInput* Input = MetasoundGraph->FindOrAddInput(InputNodeHandle);
 		if (ensure(Input))
 		{
-			if (UMetasoundEditorGraphInputNode* NewGraphNode = FGraphBuilder::AddInputNode(ParentMetasound, NodeHandle, InLocation))
+			FConstNodeHandle TemplateNodeHandle = FGraphBuilder::AddInputTemplateNodeHandle(ParentMetasound, InputNodeHandle);
+			if (UMetasoundEditorGraphNode* NewGraphNode = FGraphBuilder::AddInputNode(ParentMetasound, TemplateNodeHandle))
 			{
+				NewGraphNode->UpdateFrontendNodeLocation(InLocation);
+				NewGraphNode->SyncLocationFromFrontendNode();
 				UEdGraphNode* EdGraphNode = CastChecked<UEdGraphNode>(NewGraphNode);
 
 				if (ensure(SchemaPrivate::TryConnectNewNodeToMatchingDataTypePin(*EdGraphNode, FromPin)))
@@ -591,7 +605,7 @@ UEdGraphNode* FMetasoundGraphSchemaAction_PromoteToVariable_AccessorNode::Perfor
 	using namespace Metasound::Editor;
 	using namespace Metasound::Frontend;
 
-	FInputHandle InputHandle = FGraphBuilder::GetInputHandleFromPin(FromPin);
+	FConstInputHandle InputHandle = FGraphBuilder::GetConstInputHandleFromPin(FromPin);
 	if (!ensure(InputHandle->IsValid()))
 	{
 		return nullptr;
@@ -623,7 +637,7 @@ UEdGraphNode* FMetasoundGraphSchemaAction_PromoteToVariable_DeferredAccessorNode
 	using namespace Metasound::Editor;
 	using namespace Metasound::Frontend;
 
-	FInputHandle InputHandle = FGraphBuilder::GetInputHandleFromPin(FromPin);
+	FConstInputHandle InputHandle = FGraphBuilder::GetConstInputHandleFromPin(FromPin);
 	if (!ensure(InputHandle->IsValid()))
 	{
 		return nullptr;
@@ -655,7 +669,7 @@ UEdGraphNode* FMetasoundGraphSchemaAction_PromoteToVariable_MutatorNode::Perform
 	using namespace Metasound::Editor;
 	using namespace Metasound::Frontend;
 
-	FOutputHandle OutputHandle = FGraphBuilder::GetOutputHandleFromPin(FromPin);
+	FConstOutputHandle OutputHandle = FGraphBuilder::GetConstOutputHandleFromPin(FromPin);
 	if (!ensure(OutputHandle->IsValid()))
 	{
 		return nullptr;
@@ -701,13 +715,13 @@ UEdGraphNode* FMetasoundGraphSchemaAction_NewOutput::PerformAction(UEdGraph* Par
 	UMetasoundEditorGraph* MetasoundGraph = CastChecked<UMetasoundEditorGraph>(ParentGraph);
 	UObject& ParentMetasound = MetasoundGraph->GetMetasoundChecked();
 
-	UMetasoundEditorGraphOutput* Output = MetasoundGraph->FindOutput(NodeID);
+	const UMetasoundEditorGraphOutput* Output = MetasoundGraph->FindOutput(NodeID);
 	if (!ensure(Output))
 	{
 		return nullptr;
 	}
 
-	FNodeHandle NodeHandle = Output->GetNodeHandle();
+	FConstNodeHandle NodeHandle = Output->GetConstNodeHandle();
 	if (!ensure(NodeHandle->IsValid()))
 	{
 		return nullptr;
@@ -717,8 +731,10 @@ UEdGraphNode* FMetasoundGraphSchemaAction_NewOutput::PerformAction(UEdGraph* Par
 	ParentMetasound.Modify();
 	ParentGraph->Modify();
 
-	if (UMetasoundEditorGraphOutputNode* NewGraphNode = FGraphBuilder::AddOutputNode(ParentMetasound, NodeHandle, Location, bSelectNewNode))
+	if (UMetasoundEditorGraphOutputNode* NewGraphNode = FGraphBuilder::AddOutputNode(ParentMetasound, NodeHandle, bSelectNewNode))
 	{
+		NewGraphNode->UpdateFrontendNodeLocation(Location);
+		NewGraphNode->SyncLocationFromFrontendNode();
 		SchemaPrivate::TryConnectNewNodeToMatchingDataTypePin(*NewGraphNode, FromPin);
 		FGraphBuilder::RegisterGraphWithFrontend(ParentMetasound);
 		return NewGraphNode;
@@ -742,7 +758,7 @@ UEdGraphNode* FMetasoundGraphSchemaAction_PromoteToOutput::PerformAction(UEdGrap
 	using namespace Metasound::Editor;
 	using namespace Metasound::Frontend;
 
-	FOutputHandle OutputHandle = FGraphBuilder::GetOutputHandleFromPin(FromPin);
+	FConstOutputHandle OutputHandle = FGraphBuilder::GetConstOutputHandleFromPin(FromPin);
 	if (!ensure(OutputHandle->IsValid()))
 	{
 		return nullptr;
@@ -758,14 +774,17 @@ UEdGraphNode* FMetasoundGraphSchemaAction_PromoteToOutput::PerformAction(UEdGrap
 	const FVertexName NewNodeName = FGraphBuilder::GenerateUniqueNameByClassType(ParentMetasound, EMetasoundFrontendClassType::Output, OutputName);
 	const FCreateNodeVertexParams VertexParams = { OutputHandle->GetDataType(), OutputHandle->GetVertexAccessType() };
 
-	FNodeHandle NodeHandle = FGraphBuilder::AddOutputNodeHandle(ParentMetasound, VertexParams, &NewNodeName);
+	FConstNodeHandle NodeHandle = FGraphBuilder::AddOutputNodeHandle(ParentMetasound, VertexParams, &NewNodeName);
 	if (ensure(NodeHandle->IsValid()))
 	{
 		UMetasoundEditorGraphOutput* Output = MetasoundGraph->FindOrAddOutput(NodeHandle);
 		if (ensure(Output))
 		{
-			if (UMetasoundEditorGraphOutputNode* NewGraphNode = FGraphBuilder::AddOutputNode(ParentMetasound, NodeHandle, InLocation))
+			if (UMetasoundEditorGraphOutputNode* NewGraphNode = FGraphBuilder::AddOutputNode(ParentMetasound, NodeHandle))
 			{ 
+				NewGraphNode->UpdateFrontendNodeLocation(InLocation);
+				NewGraphNode->SyncLocationFromFrontendNode();
+
 				UEdGraphNode* EdGraphNode = CastChecked<UEdGraphNode>(NewGraphNode);
 
 				if (ensure(SchemaPrivate::TryConnectNewNodeToMatchingDataTypePin(*EdGraphNode, FromPin)))
@@ -803,12 +822,14 @@ UEdGraphNode* FMetasoundGraphSchemaAction_NewVariableNode::PerformAction(UEdGrap
 				MetasoundGraph->Modify();
 				Variable->Modify();
 				
-				FNodeHandle FrontendNode = CreateFrontendVariableNode(MetasoundGraph->GetGraphHandle(), VariableID);
+				FConstNodeHandle FrontendNode = CreateFrontendVariableNode(MetasoundGraph->GetGraphHandle(), VariableID);
 				if (ensure(FrontendNode->IsValid()))
 				{
-					if (UMetasoundEditorGraphVariableNode* NewGraphNode = FGraphBuilder::AddVariableNode(*ParentMetasound, FrontendNode, Location, bSelectNewNode))
+					if (UMetasoundEditorGraphVariableNode* NewGraphNode = FGraphBuilder::AddVariableNode(*ParentMetasound, FrontendNode, bSelectNewNode))
 					{
 						NewGraphNode->Modify();
+						NewGraphNode->UpdateFrontendNodeLocation(Location);
+						NewGraphNode->SyncLocationFromFrontendNode();
 						SchemaPrivate::TryConnectNewNodeToMatchingDataTypePin(*NewGraphNode, FromPin);
 						return NewGraphNode;
 					}
@@ -928,13 +949,16 @@ UEdGraphNode* FMetasoundGraphSchemaAction_NewReroute::PerformAction(UEdGraph* Pa
 		Params.InputsToConnect.Add(FromVertexDataType);
 	}
 
-	FNodeHandle NodeHandle = MetaSoundAsset->GetRootGraphHandle()->AddTemplateNode(*RerouteTemplate, MoveTemp(Params));
+	FConstNodeHandle NodeHandle = MetaSoundAsset->GetRootGraphHandle()->AddTemplateNode(*RerouteTemplate, MoveTemp(Params));
 
-	if (UMetasoundEditorGraphExternalNode* NewGraphNode = FGraphBuilder::AddExternalNode(ParentMetasound, NodeHandle, Location, bSelectNewNode))
+	if (UMetasoundEditorGraphExternalNode* NewGraphNode = FGraphBuilder::AddTemplateNode(ParentMetasound, NodeHandle->GetID(), NodeHandle->GetClassMetadata(), bSelectNewNode))
 	{
 		NewGraphNode->Modify();
+		NewGraphNode->UpdateFrontendNodeLocation(Location);
+		NewGraphNode->SyncLocationFromFrontendNode();
+
 		SchemaPrivate::TryConnectNewNodeToMatchingDataTypePin(*NewGraphNode, FromPin);
-		MetaSoundGraph->GetModifyContext().AddNodeIDsModified({ NewGraphNode->GetNodeID()});
+		MetaSoundAsset->GetModifyContext().AddNodeIDsModified({ NewGraphNode->GetNodeID()});
 
 		TSharedPtr<FEditor> ParentEditor = FGraphBuilder::GetEditorForMetasound(ParentMetasound);
 		if (ParentEditor.IsValid() && bSelectNewNode)
@@ -962,23 +986,39 @@ UEdGraphNode* FMetasoundGraphSchemaAction_NewComment::PerformAction(UEdGraph* Pa
 {
 	using namespace Metasound::Editor;
 
-	const FScopedTransaction Transaction(LOCTEXT("AddNewOutputNode1", "Add Comment to MetaSound Graph"));
-	ParentGraph->Modify();
+	UMetasoundEditorGraph* MetaSoundGraph = CastChecked<UMetasoundEditorGraph>(ParentGraph);
+	UObject& MetaSound = MetaSoundGraph->GetMetasoundChecked();
 
-	UEdGraphNode_Comment* CommentTemplate = NewObject<UEdGraphNode_Comment>();
+	const FScopedTransaction Transaction(LOCTEXT("AddNewOutputNode1", "Add Comment to MetaSound Graph"));
+	MetaSoundGraph->Modify();
+	MetaSound.Modify();
 
+	// Must cache bounds prior to comment creation as call selects new node and invalidates original selection
 	FSlateRect Bounds;
-	FVector2D SpawnLocation = Location;
 	TSharedPtr<FEditor> MetasoundEditor = FGraphBuilder::GetEditorForGraph(*ParentGraph);
+	const bool bUseBounds = MetasoundEditor.IsValid() && MetasoundEditor->GetBoundsForSelectedNodes(Bounds, 50.0f);
 
-	if (MetasoundEditor.IsValid() && MetasoundEditor->GetBoundsForSelectedNodes(Bounds, 50.0f))
+	if (UMetasoundEditorGraphCommentNode* NewComment = FGraphBuilder::CreateCommentNode(MetaSound, bSelectNewNode))
 	{
-		CommentTemplate->SetBounds(Bounds);
-		SpawnLocation.X = CommentTemplate->NodePosX;
-		SpawnLocation.Y = CommentTemplate->NodePosY;
+		if (bUseBounds)
+		{
+			NewComment->SetBounds(Bounds);
+		}
+		else
+		{
+			NewComment->NodePosX = Location.X;
+			NewComment->NodePosY = Location.Y;
+			NewComment->SnapToGrid(GetDefault<UEditorStyleSettings>()->GridSnapSize);
+
+			// Applies new snapped node data to frontend doc data
+			FMetaSoundFrontendGraphComment& FrontendComment = MetaSoundGraph->GetBuilderChecked().FindOrAddGraphComment(NewComment->GetCommentID());
+			UMetasoundEditorGraphCommentNode::ConvertToFrontendComment(*NewComment, FrontendComment);
+		}
+
+		return NewComment;
 	}
 
-	return FEdGraphSchemaAction_NewNode::SpawnNodeFromTemplate<UEdGraphNode_Comment>(ParentGraph, CommentTemplate, SpawnLocation);
+	return nullptr;
 }
 
 const FSlateBrush* FMetasoundGraphSchemaAction_NewComment::GetIconBrush() const
@@ -1266,8 +1306,8 @@ const FPinConnectionResponse UMetasoundEditorGraphSchema::CanCreateConnection(co
 		}
 	}
 
-	Frontend::FInputHandle InputHandle = Editor::FGraphBuilder::GetInputHandleFromPin(InputPin);
-	Frontend::FOutputHandle OutputHandle = Editor::FGraphBuilder::GetOutputHandleFromPin(OutputPin);
+	Frontend::FConstInputHandle InputHandle = Editor::FGraphBuilder::GetConstInputHandleFromPin(InputPin);
+	Frontend::FConstOutputHandle OutputHandle = Editor::FGraphBuilder::GetConstOutputHandleFromPin(OutputPin);
 
 	const bool bInputValid = InputHandle->IsValid();
 	const bool bOutputValid = OutputHandle->IsValid();
@@ -1387,6 +1427,7 @@ void UMetasoundEditorGraphSchema::OnPinConnectionDoubleCicked(UEdGraphPin* PinA,
 
 bool UMetasoundEditorGraphSchema::TryCreateConnection(UEdGraphPin* PinA, UEdGraphPin* PinB) const
 {
+	using namespace Metasound::Editor;
 	using namespace Metasound::Frontend;
 
 	if (!ensure(PinA && PinB))
@@ -1421,12 +1462,12 @@ bool UMetasoundEditorGraphSchema::TryCreateConnection(UEdGraphPin* PinA, UEdGrap
 		return false;
 	}
 
-	if (!Metasound::Editor::FGraphBuilder::ConnectNodes(*InputPin, *OutputPin, false /* bConnectEdPins */))
+	if (!FGraphBuilder::ConnectNodes(*InputPin, *OutputPin, false /* bConnectEdPins */))
 	{
 		return false;
 	}
 
-	Graph->GetModifyContext().SetDocumentModified();
+	FGraphBuilder::GetOutermostMetaSoundChecked(*Graph).GetModifyContext().SetDocumentModified();
 
 	return true;
 }
@@ -1475,15 +1516,14 @@ bool UMetasoundEditorGraphSchema::SafeDeleteNodeFromGraph(UEdGraph* Graph, UEdGr
 {
 	using namespace Metasound::Editor;
 
-	if (!InNodeToDelete || !Graph || InNodeToDelete->GetGraph() != Graph)
+	UMetasoundEditorGraph* MetaSoundGraph = Cast<UMetasoundEditorGraph>(Graph);
+	if (!InNodeToDelete || !MetaSoundGraph || InNodeToDelete->GetGraph() != Graph)
 	{
 		return false;
 	}
 
-	if (UObject* MetaSound = Graph->GetOutermostObject())
-	{
-		MetaSound->Modify();
-	}
+	UObject& MetaSound = MetaSoundGraph->GetMetasoundChecked();
+	MetaSound.Modify();
 	Graph->Modify();
 
 	return FGraphBuilder::DeleteNode(*InNodeToDelete);
@@ -1501,13 +1541,13 @@ bool UMetasoundEditorGraphSchema::ShouldHidePinDefaultValue(UEdGraphPin* Pin) co
 
 	if (Pin->Direction == EGPD_Input)
 	{
-		FInputHandle InputHandle = FGraphBuilder::GetInputHandleFromPin(Pin);
+		FConstInputHandle InputHandle = FGraphBuilder::GetConstInputHandleFromPin(Pin);
 		return !InputHandle->GetOwningNode()->GetClassStyle().Display.bShowLiterals;
 	}
 
 	if (Pin->Direction == EGPD_Output)
 	{
-		FOutputHandle OutputHandle = FGraphBuilder::GetOutputHandleFromPin(Pin);
+		FConstOutputHandle OutputHandle = FGraphBuilder::GetConstOutputHandleFromPin(Pin);
 		return !OutputHandle->GetOwningNode()->GetClassStyle().Display.bShowLiterals;
 	}
 
@@ -1524,7 +1564,8 @@ FText UMetasoundEditorGraphSchema::GetPinDisplayName(const UEdGraphPin* Pin) con
 
 	UMetasoundEditorGraphNode* Node = CastChecked<UMetasoundEditorGraphNode>(Pin->GetOwningNode());
 	FConstNodeHandle NodeHandle = Node->GetConstNodeHandle();
-	const EMetasoundFrontendClassType ClassType = NodeHandle->GetClassMetadata().GetType();
+	const FMetasoundFrontendClassMetadata& ClassMetadata = NodeHandle->GetClassMetadata();
+	const EMetasoundFrontendClassType ClassType = ClassMetadata.GetType();
 
 	switch (ClassType)
 	{
@@ -1549,7 +1590,6 @@ FText UMetasoundEditorGraphSchema::GetPinDisplayName(const UEdGraphPin* Pin) con
 
 		case EMetasoundFrontendClassType::Literal:
 		case EMetasoundFrontendClassType::External:
-		case EMetasoundFrontendClassType::Template:
 		{
 			if (Pin->Direction == EGPD_Input)
 			{
@@ -1571,6 +1611,33 @@ FText UMetasoundEditorGraphSchema::GetPinDisplayName(const UEdGraphPin* Pin) con
 			return Super::GetPinDisplayName(Pin);
 		}
 
+		case EMetasoundFrontendClassType::Template:
+		{
+			const INodeTemplate* Template = INodeTemplateRegistry::Get().FindTemplate(ClassMetadata.GetClassName());
+			if (ensure(Template))
+			{
+				if (Pin->Direction == EGPD_Input)
+				{
+					FConstInputHandle InputHandle = NodeHandle->GetConstInputWithVertexName(Pin->GetFName());
+					if (InputHandle->IsValid())
+					{
+						return Template->GetInputPinDisplayName(*InputHandle);
+					}
+				}
+				else
+				{
+					FConstOutputHandle OutputHandle = NodeHandle->GetConstOutputWithVertexName(Pin->GetFName());
+					if (OutputHandle->IsValid())
+					{
+						return Template->GetOutputPinDisplayName(*OutputHandle);
+					}
+				}
+			}
+
+			return Super::GetPinDisplayName(Pin);
+		}
+		break;
+
 		case EMetasoundFrontendClassType::Graph:
 		case EMetasoundFrontendClassType::Invalid:
 		default:
@@ -1807,7 +1874,7 @@ void UMetasoundEditorGraphSchema::DroppedAssetsOnGraph(const TArray<FAssetData>&
 			}
 
 			// This may not be necessary as dropping an asset on the graph may load it, thus triggering the registration from the MetaSoundAssetManager.
-			const FMetasoundFrontendDocument& DroppedDoc = DroppedMetaSoundAsset->GetDocumentChecked();
+			const FMetasoundFrontendDocument& DroppedDoc = DroppedMetaSoundAsset->GetConstDocumentChecked();
 			const FNodeRegistryKey RegistryKey = FNodeRegistryKey(DroppedDoc.RootGraph);
 			if (ensure(RegistryKey.IsValid()))
 			{
@@ -1819,7 +1886,10 @@ void UMetasoundEditorGraphSchema::DroppedAssetsOnGraph(const TArray<FAssetData>&
 			FMetasoundFrontendClass Class;
 			if (ensure(ISearchEngine::Get().FindClassWithHighestVersion(DroppedDoc.RootGraph.Metadata.GetClassName(), Class)))
 			{
-				Metasound::Editor::FGraphBuilder::AddExternalNode(MetaSound, Class.Metadata, GraphPosition);
+				UMetasoundEditorGraphNode* NewGraphNode = Metasound::Editor::FGraphBuilder::AddExternalNode(MetaSound, Class.Metadata);
+				NewGraphNode->UpdateFrontendNodeLocation(GraphPosition);
+				NewGraphNode->SyncLocationFromFrontendNode();
+
 				bTransactionSucceeded = true;
 			}
 		}
diff --git a/Engine/Plugins/Runtime/Metasound/Source/MetasoundEditor/Private/MetasoundEditorModule.cpp b/Engine/Plugins/Runtime/Metasound/Source/MetasoundEditor/Private/MetasoundEditorModule.cpp
--- a/Engine/Plugins/Runtime/Metasound/Source/MetasoundEditor/Private/MetasoundEditorModule.cpp
+++ b/Engine/Plugins/Runtime/Metasound/Source/MetasoundEditor/Private/MetasoundEditorModule.cpp
@@ -257,9 +257,7 @@ namespace Metasound
 				{
 					if (UObject* AssetObject = InAssetData.GetAsset())
 					{
-						FMetasoundAssetBase* MetaSoundAsset = IMetasoundUObjectRegistry::Get().GetObjectAsAssetBase(AssetObject);
-						check(MetaSoundAsset);
-						MetaSoundAsset->RegisterGraphWithFrontend(RegOptions);
+						FGraphBuilder::RegisterGraphWithFrontend(*AssetObject);
 					}
 				}
 				else
@@ -281,11 +279,12 @@ namespace Metasound
 					{
 						if (Result == EAsyncLoadingResult::Succeeded)
 						{
-							FMetasoundAssetBase* MetaSoundAsset = IMetasoundUObjectRegistry::Get().GetObjectAsAssetBase(ObjectPath.ResolveObject());
+							UObject* MetaSoundObj = ObjectPath.ResolveObject();
+							FMetasoundAssetBase* MetaSoundAsset = IMetasoundUObjectRegistry::Get().GetObjectAsAssetBase(MetaSoundObj);
 							check(MetaSoundAsset);
 							if (!MetaSoundAsset->IsRegistered())
 							{
-								MetaSoundAsset->RegisterGraphWithFrontend(RegOptions);
+								FGraphBuilder::RegisterGraphWithFrontend(*MetaSoundObj);
 							}
 						}
 
diff --git a/Engine/Plugins/Runtime/Metasound/Source/MetasoundEditor/Private/MetasoundEditorSubsystem.cpp b/Engine/Plugins/Runtime/Metasound/Source/MetasoundEditor/Private/MetasoundEditorSubsystem.cpp
--- a/Engine/Plugins/Runtime/Metasound/Source/MetasoundEditor/Private/MetasoundEditorSubsystem.cpp
+++ b/Engine/Plugins/Runtime/Metasound/Source/MetasoundEditor/Private/MetasoundEditorSubsystem.cpp
@@ -10,6 +10,7 @@
 #include "MetasoundEditorSettings.h"
 #include "MetasoundFactory.h"
 #include "MetasoundUObjectRegistry.h"
+#include "NodeTemplates/MetasoundFrontendNodeTemplateInput.h"
 #include "Sound/SoundSourceBusSend.h"
 #include "Sound/SoundSubmixSend.h"
 
@@ -64,10 +65,6 @@ TScriptInterface<IMetaSoundDocumentInterface> UMetaSoundEditorSubsystem::BuildTo
 				const bool bIsSource = &BuilderUClass == UMetaSoundSource::StaticClass();
 				if (InBuilder->IsPreset())
 				{
-					FMetasoundAssetBase* PresetAsset = IMetasoundUObjectRegistry::Get().GetObjectAsAssetBase(NewMetaSound);
-					check(PresetAsset);
-					PresetAsset->ConformObjectDataToInterfaces();
-
 					// Only use referenced UObject's SoundWave settings for sources if not overridden 
 					if (TemplateSoundWave == nullptr && bIsSource)
 					{
@@ -85,7 +82,11 @@ TScriptInterface<IMetaSoundDocumentInterface> UMetaSoundEditorSubsystem::BuildTo
 				}
 			}
 
-			InitEdGraph(*NewMetaSound);
+			UMetaSoundBuilderBase& NewDocBuilder = UMetaSoundBuilderSubsystem::GetChecked().AttachBuilderToAssetChecked(*NewMetaSound);
+
+			EMetaSoundBuilderResult InjectResult = EMetaSoundBuilderResult::Failed;
+			constexpr bool bForceNodeCreation = true;
+			NewDocBuilder.InjectInputTemplateNodes(bForceNodeCreation, InjectResult);
 
 			if (!bWasRooted)
 			{
@@ -105,6 +106,44 @@ TScriptInterface<IMetaSoundDocumentInterface> UMetaSoundEditorSubsystem::BuildTo
 	return nullptr;
 }
 
+bool UMetaSoundEditorSubsystem::BindMemberMetadata(
+	UMetasoundEditorGraphMember& InMember,
+	TSubclassOf<UMetasoundEditorGraphMemberDefaultLiteral> LiteralClass,
+	UMetasoundEditorGraphMemberDefaultLiteral* TemplateObject)
+{
+	const FGuid& MemberID = InMember.GetMemberID();
+
+	UMetasoundEditorGraph* Graph = InMember.GetOwningGraph();
+	check(Graph);
+	UObject& MetaSound = Graph->GetMetasoundChecked();
+	UMetaSoundBuilderBase& Builder = UMetaSoundBuilderSubsystem::GetChecked().AttachBuilderToAssetChecked(MetaSound);
+
+	if (TemplateObject)
+	{
+		Builder.ClearMemberMetadata(MemberID);
+	}
+	else
+	{
+		if (UMetaSoundFrontendMemberMetadata* Literal = Builder.FindMemberMetadata(MemberID))
+		{
+			InMember.Literal = CastChecked<UMetasoundEditorGraphMemberDefaultLiteral>(Literal);
+			return false;
+		}
+	}
+
+	if (UMetasoundEditorGraphMemberDefaultLiteral* NewLiteral = NewObject<UMetasoundEditorGraphMemberDefaultLiteral>(&MetaSound, LiteralClass, FName(), RF_Transactional, TemplateObject))
+	{
+		NewLiteral->MemberID = MemberID;
+
+		Builder.SetMemberMetadata(*NewLiteral);
+		InMember.Literal = NewLiteral;
+		return true;
+	}
+
+	checkNoEntry();
+	return false;
+}
+
 UMetaSoundEditorSubsystem& UMetaSoundEditorSubsystem::GetChecked()
 {
 	checkf(GEditor, TEXT("Cannot access UMetaSoundEditorSubsystem without editor loaded"));
@@ -147,10 +186,12 @@ void UMetaSoundEditorSubsystem::InitAsset(UObject& InNewMetaSound, UObject* InRe
 
 	TScriptInterface<IMetaSoundDocumentInterface> DocInterface = &InNewMetaSound;
 	FMetaSoundFrontendDocumentBuilder Builder(DocInterface);
+
 	Builder.InitDocument();
-#if WITH_EDITORONLY_DATA
 	Builder.InitNodeLocations();
-#endif // WITH_EDITORONLY_DATA
+
+	constexpr bool bForceNodeCreation = true;
+	FInputNodeTemplate::GetChecked().Inject(Builder, bForceNodeCreation);
 
 	const FString& Author = GetDefaultAuthor();
 	Builder.SetAuthor(Author);
@@ -165,27 +206,22 @@ void UMetaSoundEditorSubsystem::InitAsset(UObject& InNewMetaSound, UObject* InRe
 		TScriptInterface<IMetaSoundDocumentInterface> ReferencedDocInterface = InReferencedMetaSound;
 		Builder.ConvertToPreset(ReferencedDocInterface->GetConstDocument());
 
-		// Update asset object data from interfaces 
-		FMetasoundAssetBase* PresetAsset = IMetasoundUObjectRegistry::Get().GetObjectAsAssetBase(&InNewMetaSound);
-		check(PresetAsset);
-		PresetAsset->ConformObjectDataToInterfaces();
-
 		// Copy sound wave settings to preset for sources
 		if (&ReferencedDocInterface->GetBaseMetaSoundUClass() == UMetaSoundSource::StaticClass())
 		{
 			SetSoundWaveSettingsFromTemplate(*CastChecked<USoundWave>(&InNewMetaSound), *CastChecked<USoundWave>(InReferencedMetaSound));
 		}
 	}
-
-	// Initial graph generation is not something to be managed by the transaction
-	// stack, so don't track dirty state until after initial setup if necessary.
-	InitEdGraph(InNewMetaSound);
 }
 
 void UMetaSoundEditorSubsystem::InitEdGraph(UObject& InMetaSound)
 {
 	using namespace Metasound;
 	using namespace Metasound::Editor;
+	using namespace Metasound::Frontend;
+
+	TScriptInterface<IMetaSoundDocumentInterface> DocInterface = &InMetaSound;
+	const FMetasoundFrontendClassName& ClassName = DocInterface->GetConstDocument().RootGraph.Metadata.GetClassName();
 
 	FMetasoundAssetBase* MetaSoundAsset = IMetasoundUObjectRegistry::Get().GetObjectAsAssetBase(&InMetaSound);
 	checkf(MetaSoundAsset, TEXT("EdGraph can only be initialized on registered MetaSoundAsset type"));
diff --git a/Engine/Plugins/Runtime/Metasound/Source/MetasoundEditor/Private/MetasoundFactory.cpp b/Engine/Plugins/Runtime/Metasound/Source/MetasoundEditor/Private/MetasoundFactory.cpp
--- a/Engine/Plugins/Runtime/Metasound/Source/MetasoundEditor/Private/MetasoundFactory.cpp
+++ b/Engine/Plugins/Runtime/Metasound/Source/MetasoundEditor/Private/MetasoundFactory.cpp
@@ -37,43 +37,6 @@ UMetaSoundBaseFactory::UMetaSoundBaseFactory(const FObjectInitializer& ObjectIni
 	bEditAfterNew = true;
 }
 
-void UMetaSoundBaseFactory::InitAsset(UObject& InNewMetaSound, UObject* InReferencedMetaSound)
-{
-	using namespace Metasound;
-	using namespace Metasound::Editor;
-	using namespace Metasound::Frontend;
-
-	TScriptInterface<IMetaSoundDocumentInterface> DocInterface = &InNewMetaSound;
-	FMetaSoundFrontendDocumentBuilder Builder(DocInterface);
-	Builder.InitDocument();
-#if WITH_EDITORONLY_DATA
-	Builder.InitNodeLocations();
-#endif // WITH_EDITORONLY_DATA
-
-	const FString& Author = UMetaSoundEditorSubsystem::GetChecked().GetDefaultAuthor();
-	Builder.SetAuthor(Author);
-
-	// Initialize asset as a preset
-	if (InReferencedMetaSound)
-	{
-		// Ensure the referenced MetaSound is registered already
-		UMetaSoundEditorSubsystem::GetChecked().RegisterGraphWithFrontend(*InReferencedMetaSound);
-		
-		// Initialize preset with referenced Metasound 
-		TScriptInterface<IMetaSoundDocumentInterface> ReferencedDocInterface = InReferencedMetaSound;
-		Builder.ConvertToPreset(ReferencedDocInterface->GetConstDocument());
-		
-		// Update asset object data from interfaces 
-		FMetasoundAssetBase* PresetAsset = IMetasoundUObjectRegistry::Get().GetObjectAsAssetBase(&InNewMetaSound);
-		check(PresetAsset);
-		PresetAsset->ConformObjectDataToInterfaces();
-	}
-
-	// Initial graph generation is not something to be managed by the transaction
-	// stack, so don't track dirty state until after initial setup if necessary.
-	UMetaSoundEditorSubsystem::GetChecked().InitEdGraph(InNewMetaSound);
-}
-
 UMetaSoundFactory::UMetaSoundFactory(const FObjectInitializer& ObjectInitializer)
 	: Super(ObjectInitializer)
 {
diff --git a/Engine/Plugins/Runtime/Metasound/Source/MetasoundEditor/Private/MetasoundNodeDetailCustomization.cpp b/Engine/Plugins/Runtime/Metasound/Source/MetasoundEditor/Private/MetasoundNodeDetailCustomization.cpp
--- a/Engine/Plugins/Runtime/Metasound/Source/MetasoundEditor/Private/MetasoundNodeDetailCustomization.cpp
+++ b/Engine/Plugins/Runtime/Metasound/Source/MetasoundEditor/Private/MetasoundNodeDetailCustomization.cpp
@@ -66,10 +66,13 @@ namespace Metasound
 			static const FText ConstructorPinTooltip = LOCTEXT("ConstructorPinTooltip",
 				"Whether this input or output is a constructor pin. Constructor values are only read on construction (on play), and are not dynamically updated at runtime.");
 
-			void GetDataTypeFromElementPropertyHandle(TSharedPtr<IPropertyHandle> ElementPropertyHandle, Frontend::FDataTypeRegistryInfo& OutDataTypeInfo)
+
+			// Retrieves the data type info if the literal property's member is found. Returns if the associated member is found, false if not.
+			bool GetDataTypeFromElementPropertyHandle(TSharedPtr<IPropertyHandle> ElementPropertyHandle, Frontend::FDataTypeRegistryInfo& OutDataTypeInfo)
 			{
 				using namespace Frontend;
 
+				OutDataTypeInfo = { };
 				TArray<UObject*>OuterObjects;
 				ElementPropertyHandle->GetOuterObjects(OuterObjects);
 				if (OuterObjects.Num() == 1)
@@ -77,7 +80,7 @@ namespace Metasound
 					UObject* Outer = OuterObjects.Last();
 					if (const UMetasoundEditorGraphMemberDefaultLiteral* DefaultLiteral = Cast<UMetasoundEditorGraphMemberDefaultLiteral>(Outer))
 					{
-						if (const UMetasoundEditorGraphMember* Member = Cast<UMetasoundEditorGraphMember>(DefaultLiteral->GetOuter()))
+						if (const UMetasoundEditorGraphMember* Member = DefaultLiteral->FindMember())
 						{
 							FName DataTypeName = Member->GetDataType();
 							ensure(IDataTypeRegistry::Get().GetDataTypeInfo(DataTypeName, OutDataTypeInfo));
@@ -91,9 +94,13 @@ namespace Metasound
 									ensure(IDataTypeRegistry::Get().GetDataTypeInfo(DataTypeName, OutDataTypeInfo));
 								}
 							}
+
+							return true;
 						}
 					}
 				}
+
+				return false;
 			}
 
 			// If DataType is an array type, creates & returns the array's
@@ -120,9 +127,9 @@ namespace Metasound
 					}
 					
 					Frontend::FDataTypeRegistryInfo DataTypeInfo;
-					MemberCustomizationPrivate::GetDataTypeFromElementPropertyHandle(PropertyHandle, DataTypeInfo);
+					const bool bMemberFound = MemberCustomizationPrivate::GetDataTypeFromElementPropertyHandle(PropertyHandle, DataTypeInfo);
 					UClass* ProxyGenClass = DataTypeInfo.ProxyGeneratorClass;
-					if (!ProxyGenClass)
+					if (!bMemberFound || !ProxyGenClass)
 					{
 						return;
 					}
@@ -219,8 +226,10 @@ namespace Metasound
 				PropertyHandle->GetOuterObjects(OuterObjects);
 				if (!OuterObjects.IsEmpty())
 				{
-					UObject* Outer = OuterObjects[0]->GetOuter();
-					GraphMember = Cast<UMetasoundEditorGraphMember>(Outer);
+					if (UMetasoundEditorGraphMemberDefaultLiteral* Literal = Cast<UMetasoundEditorGraphMemberDefaultLiteral>(OuterObjects[0]))
+					{
+						GraphMember = Literal->FindMember();
+					}
 				}
 
 				FUIAction PasteAction;
@@ -310,12 +319,10 @@ namespace Metasound
 					if (FloatLiteral.IsValid())
 					{
 						FloatLiteral->OnClampChanged.Remove(OnClampChangedDelegateHandle);
-						if (FMetasoundAssetBase* MetasoundAsset = Metasound::IMetasoundUObjectRegistry::Get().GetObjectAsAssetBase(FloatLiteral->GetOutermostObject()))
+						if (const UMetasoundEditorGraphMember* Member = FloatLiteral->FindMember())
 						{
-							if (const UMetasoundEditorGraphMember* Member = Cast<UMetasoundEditorGraphMember>(FloatLiteral->GetOuter()))
-							{
-								MetasoundAsset->GetModifyContext().AddMemberIDsModified({ Member->GetMemberID() });
-							}
+							FMetasoundAssetBase& MetasoundAsset = Editor::FGraphBuilder::GetOutermostMetaSoundChecked(*FloatLiteral);
+							MetasoundAsset.GetModifyContext().AddMemberIDsModified({ Member->GetMemberID() });
 						}
 					}
 				});
@@ -332,7 +339,7 @@ namespace Metasound
 
 			// Enable widget options for editable inputs only 
 			bool bShowWidgetOptions = false;
-			if (const UMetasoundEditorGraphInput* ParentMember = Cast <UMetasoundEditorGraphInput>(InLiteral.GetParentMember()))
+			if (const UMetasoundEditorGraphInput* ParentMember = Cast <UMetasoundEditorGraphInput>(InLiteral.FindMember()))
 			{
 				if (const UMetasoundEditorGraph* OwningGraph = ParentMember->GetOwningGraph())
 				{
@@ -387,8 +394,8 @@ namespace Metasound
 
 			// Enable widget options for editable inputs only 
 			bool bShowWidgetOptions = false;
-			if (const UMetasoundEditorGraphInput* ParentMember = Cast <UMetasoundEditorGraphInput>(InLiteral.GetParentMember()))
-			{								
+			if (const UMetasoundEditorGraphInput* ParentMember = Cast <UMetasoundEditorGraphInput>(InLiteral.FindMember()))
+			{
 				if (const UMetasoundEditorGraph* OwningGraph = ParentMember->GetOwningGraph())
 				{
 					// Disable widget options for constructor inputs for now to prevent changing default value via widget while playing 
@@ -439,12 +446,9 @@ namespace Metasound
 				.OnAreAssetsAcceptableForDropWithReason_Lambda([this, DefaultValueHandle](TArrayView<FAssetData> InAssets, FText& OutReason)
 				{
 					Frontend::FDataTypeRegistryInfo DataTypeInfo;
-					MemberCustomizationPrivate::GetDataTypeFromElementPropertyHandle(DefaultValueHandle, DataTypeInfo);
-
-					const IMetasoundEditorModule& EditorModule = FModuleManager::GetModuleChecked<IMetasoundEditorModule>("MetaSoundEditor");
-					bool bCanDrop = false;
-
-					if (UClass* ProxyGenClass = DataTypeInfo.ProxyGeneratorClass)
+					const bool bMemberFound = MemberCustomizationPrivate::GetDataTypeFromElementPropertyHandle(DefaultValueHandle, DataTypeInfo);
+					bool bCanDrop = bMemberFound;
+					if (UClass* ProxyGenClass = DataTypeInfo.ProxyGeneratorClass; bCanDrop && bMemberFound)
 					{
 						bCanDrop = true;
 						for (const FAssetData& AssetData : InAssets)
@@ -452,6 +456,7 @@ namespace Metasound
 							if (UClass* Class = AssetData.GetClass())
 							{
 PRAGMA_DISABLE_DEPRECATION_WARNINGS
+								const IMetasoundEditorModule& EditorModule = FModuleManager::GetModuleChecked<IMetasoundEditorModule>("MetaSoundEditor");
 								if (EditorModule.IsExplicitProxyClass(*DataTypeInfo.ProxyGeneratorClass))
 								{
 									bCanDrop &= Class == DataTypeInfo.ProxyGeneratorClass;
@@ -534,7 +539,7 @@ PRAGMA_ENABLE_DEPRECATION_WARNINGS
 				{
 					if (UMetasoundEditorGraphMemberDefaultLiteral* Literal = Cast<UMetasoundEditorGraphMemberDefaultLiteral>(OuterObjects.Last()))
 					{
-						if (UMetasoundEditorGraphInput* Input = Cast<UMetasoundEditorGraphInput>(Literal->GetParentMember()))
+						if (UMetasoundEditorGraphInput* Input = Cast<UMetasoundEditorGraphInput>(Literal->FindMember()))
 						{
 							// Don't display trigger simulation widget if its a trigger
 							// provided by an interface that does not support transmission.
@@ -561,7 +566,10 @@ PRAGMA_ENABLE_DEPRECATION_WARNINGS
 			using namespace Frontend;
 			using namespace MemberCustomizationPrivate;
 
-			if (FMetasoundFrontendRegistryContainer* Registry = FMetasoundFrontendRegistryContainer::Get())
+			// DataType can be reset during deletion of a literal value.  Customization can repaint briefly before the literal is removed,
+			// so just ignores if DataType is invalid.
+			const bool bIsValidDataType = !DataTypeInfo.DataTypeName.IsNone();
+			if (bIsValidDataType)
 			{
 				TSharedPtr<IPropertyHandle> ValueProperty = StructPropertyHandle->GetChildHandle(GET_MEMBER_NAME_CHECKED(FMetasoundEditorGraphMemberDefaultIntRef, Value));
 				if (ValueProperty.IsValid())
@@ -783,7 +791,7 @@ PRAGMA_ENABLE_DEPRECATION_WARNINGS
 				}
 			}
 
-			MemberCustomizationPrivate::GetDataTypeFromElementPropertyHandle(ElementPropertyHandle, DataTypeInfo);
+			const bool bMemberFound = MemberCustomizationPrivate::GetDataTypeFromElementPropertyHandle(ElementPropertyHandle, DataTypeInfo);
 
 			TSharedRef<SWidget> ValueWidget = CreateValueWidget(ParentPropertyHandleArray, StructPropertyHandle);
 			FDetailWidgetRow& ValueRow = ChildBuilder.AddCustomRow(MemberCustomizationStyle::DefaultPropertyText);
@@ -1640,8 +1648,8 @@ PRAGMA_ENABLE_DEPRECATION_WARNINGS
 
 			if (bIsPreset && !bIsDefaultConstructed && !bIsTriggerDataType)
 			{
-				const UMetasoundEditorGraphInput* Input = Cast<UMetasoundEditorGraphInput>(MemberDefaultLiteral->GetParentMember());
-				if (ensure(Input))
+				const UMetasoundEditorGraphInput* Input = Cast<UMetasoundEditorGraphInput>(MemberDefaultLiteral->FindMember());
+				if (Input)
 				{
 					auto PropertyEnabled = TAttribute<bool>::CreateLambda([this] { return !GetInputInheritsDefault(); });
 					for (IDetailPropertyRow* DefaultPropertyRow : DefaultPropertyRows)
@@ -1653,12 +1661,12 @@ PRAGMA_ENABLE_DEPRECATION_WARNINGS
 						DefaultPropertyRow->OverrideResetToDefault(ResetOverride);
 					}
 				}
-			} 
+			}
 			else if (!bIsPreset)
 			{
 				// Make default value uneditable while playing for constructor inputs
-				const UMetasoundEditorGraphInput* Input = Cast<UMetasoundEditorGraphInput>(MemberDefaultLiteral->GetParentMember());
-				if (ensure(Input))
+				const UMetasoundEditorGraphInput* Input = Cast<UMetasoundEditorGraphInput>(MemberDefaultLiteral->FindMember());
+				if (Input)
 				{
 					auto PropertyEnabled = TAttribute<bool>::CreateLambda([this, Input]
 					{
diff --git a/Engine/Plugins/Runtime/Metasound/Source/MetasoundEditor/Private/SMetasoundActionMenu.cpp b/Engine/Plugins/Runtime/Metasound/Source/MetasoundEditor/Private/SMetasoundActionMenu.cpp
--- a/Engine/Plugins/Runtime/Metasound/Source/MetasoundEditor/Private/SMetasoundActionMenu.cpp
+++ b/Engine/Plugins/Runtime/Metasound/Source/MetasoundEditor/Private/SMetasoundActionMenu.cpp
@@ -92,12 +92,12 @@ namespace Metasound
 					const IMetasoundEditorModule& EditorModule = FModuleManager::GetModuleChecked<IMetasoundEditorModule>("MetaSoundEditor");
 					if (Pin->Direction == EGPD_Input)
 					{
-						FInputHandle InputHandle = FGraphBuilder::GetInputHandleFromPin(Pin);
+						FConstInputHandle InputHandle = FGraphBuilder::GetConstInputHandleFromPin(Pin);
 						IDataTypeRegistry::Get().GetDataTypeInfo(InputHandle->GetDataType(), RegistryInfo);
 					}
 					else
 					{
-						FOutputHandle OutputHandle = FGraphBuilder::GetOutputHandleFromPin(Pin);
+						FConstOutputHandle OutputHandle = FGraphBuilder::GetConstOutputHandleFromPin(Pin);
 						IDataTypeRegistry::Get().GetDataTypeInfo(OutputHandle->GetDataType(), RegistryInfo);
 					}
 
@@ -165,12 +165,12 @@ namespace Metasound
 
 									if (DraggedFromPins[0]->Direction == EGPD_Input)
 									{
-										FInputHandle InputHandle = FGraphBuilder::GetInputHandleFromPin(Pin);
+										FConstInputHandle InputHandle = FGraphBuilder::GetConstInputHandleFromPin(Pin);
 										return FText::Format(LOCTEXT("ContextTypeFilteredText_Output", "Classes with output of type '{0}'"), FText::FromName(InputHandle->GetDataType()));
 									}
 									else
 									{
-										FOutputHandle OutputHandle = FGraphBuilder::GetOutputHandleFromPin(Pin);
+										FConstOutputHandle OutputHandle = FGraphBuilder::GetConstOutputHandleFromPin(Pin);
 										return FText::Format(LOCTEXT("ContextTypeFilteredText_Input", "Classes with input of type '{0}'"), FText::FromName(OutputHandle->GetDataType()));
 									}
 								})
diff --git a/Engine/Plugins/Runtime/Metasound/Source/MetasoundEditor/Private/SMetasoundGraphNode.cpp b/Engine/Plugins/Runtime/Metasound/Source/MetasoundEditor/Private/SMetasoundGraphNode.cpp
--- a/Engine/Plugins/Runtime/Metasound/Source/MetasoundEditor/Private/SMetasoundGraphNode.cpp
+++ b/Engine/Plugins/Runtime/Metasound/Source/MetasoundEditor/Private/SMetasoundGraphNode.cpp
@@ -17,6 +17,7 @@
 #include "KismetPins/SGraphPinNum.h"
 #include "KismetPins/SGraphPinObject.h"
 #include "KismetPins/SGraphPinString.h"
+#include "MetasoundBuilderSubsystem.h"
 #include "MetasoundEditorGraph.h"
 #include "MetasoundEditorGraphBuilder.h"
 #include "MetasoundEditorGraphInputNode.h"
@@ -34,7 +35,6 @@
 #include "PropertyCustomizationHelpers.h"
 #include "SAudioRadialSlider.h"
 #include "SAudioSlider.h"
-#include "SCommentBubble.h"
 #include "ScopedTransaction.h"
 #include "SGraphNode.h"
 #include "SGraphPinComboBox.h"
@@ -94,7 +94,7 @@ namespace Metasound
 							{
 								constexpr bool bPostTransaction = false;
 								GraphMember->UpdateFrontendDefaultLiteral(bPostTransaction);
-								Graph->GetModifyContext().AddMemberIDsModified({ GraphMember->GetMemberID() });
+								FGraphBuilder::GetOutermostMetaSoundChecked(*Graph).GetModifyContext().AddMemberIDsModified({ GraphMember->GetMemberID() });
 							}
 						}
 						DefaultFloat->OnDefaultValueChanged.Remove(InputSliderOnValueChangedDelegateHandle);
@@ -137,7 +137,7 @@ namespace Metasound
 
 		void SMetaSoundGraphNode::ExecuteTrigger(UMetasoundEditorGraphMemberDefaultLiteral& Literal)
 		{
-			UMetasoundEditorGraphMember* Member = Cast<UMetasoundEditorGraphMember>(Literal.GetOuter());
+			UMetasoundEditorGraphMember* Member = Literal.FindMember();
 			if (!ensure(Member))
 			{
 				return;
@@ -462,6 +462,43 @@ namespace Metasound
 			}
 		}
 
+		void SMetaSoundGraphNode::OnCommentBubbleToggled(bool bInCommentBubbleVisible)
+		{
+			UMetasoundEditorGraphNode& EdNode = GetMetaSoundNode();
+			UObject& MetaSound = EdNode.GetMetasoundChecked();
+			UMetaSoundBuilderBase& Builder = UMetaSoundBuilderSubsystem::GetChecked().AttachBuilderToAssetChecked(MetaSound);
+			if (const FMetasoundFrontendNode* Node = Builder.GetConstBuilder().FindNode(EdNode.GetNodeID()))
+			{
+				if (bInCommentBubbleVisible != Node->Style.Display.bCommentVisible)
+				{
+					const FScopedTransaction Transaction(LOCTEXT("GraphNodeCommentCommitted", "Graph Node Comment Bubble Toggled"));
+					MetaSound.Modify();
+					EMetaSoundBuilderResult Result;
+					EdNode.bCommentBubbleVisible = bInCommentBubbleVisible;
+					Builder.SetNodeCommentVisible(Node->GetID(), bInCommentBubbleVisible, Result);
+				}
+			}
+		}
+
+		void SMetaSoundGraphNode::OnCommentTextCommitted(const FText& NewComment, ETextCommit::Type CommitInfo)
+		{
+			FString NewCommentString = NewComment.ToString();
+			UMetasoundEditorGraphNode& EdNode = GetMetaSoundNode();
+			UObject& MetaSound = EdNode.GetMetasoundChecked();
+			UMetaSoundBuilderBase& Builder = UMetaSoundBuilderSubsystem::GetChecked().AttachBuilderToAssetChecked(MetaSound);
+			if (const FMetasoundFrontendNode* Node = Builder.GetConstBuilder().FindNode(EdNode.GetNodeID()))
+			{
+				if (!Node->Style.Display.Comment.Equals(NewCommentString))
+				{
+					const FScopedTransaction Transaction(LOCTEXT("GraphNodeCommentCommitted", "Graph Node Comment Changed"));
+					MetaSound.Modify();
+					EMetaSoundBuilderResult Result;
+					EdNode.NodeComment = MoveTemp(NewCommentString);
+					Builder.SetNodeComment(Node->GetID(), EdNode.NodeComment, Result);
+				}
+			}
+		}
+
 		FLinearColor SMetaSoundGraphNode::GetNodeTitleColorOverride() const
 		{
 			FLinearColor ReturnTitleColor = GraphNode->IsDeprecated() ? FLinearColor::Red : GetNodeObj()->GetNodeTitleColor();
@@ -553,7 +590,8 @@ namespace Metasound
 
 			UMetasoundEditorGraphNode& Node = GetMetaSoundNode();
 			Node.GetMetasoundChecked().Modify();
-			Node.SetNodeLocation(NewPosition);
+			Node.UpdateFrontendNodeLocation(NewPosition);
+			Node.SyncLocationFromFrontendNode();
 		}
 
 		const FSlateBrush* SMetaSoundGraphNode::GetNodeBodyBrush() const
@@ -727,7 +765,7 @@ namespace Metasound
 
 								if (UMetasoundEditorGraph* Graph = GraphMember->GetOwningGraph())
 								{
-									Graph->GetModifyContext().AddMemberIDsModified({ GraphMember->GetMemberID() });
+									FGraphBuilder::GetOutermostMetaSoundChecked(*Graph).GetModifyContext().AddMemberIDsModified({ GraphMember->GetMemberID() });
 								}
 							}
 						};
@@ -1004,7 +1042,7 @@ namespace Metasound
 								.AutoHeight()
 								[
 									SAssignNew(MaterialButtonWidget, SAudioMaterialButton)
-										.OnBooleanValueChanged_Lambda(OnboolValueChangedLambda)									
+										.OnBooleanValueChanged_Lambda(OnboolValueChangedLambda)
 								];
 
 							MaterialButtonWidget->SetPressedState(DefaultBool->GetDefault());
@@ -1080,7 +1118,8 @@ namespace Metasound
 
 			UMetasoundEditorGraphNode& Node = GetMetaSoundNode();
 			Node.GetMetasoundChecked().Modify();
-			Node.SetNodeLocation(NewPosition);
+			Node.UpdateFrontendNodeLocation(NewPosition);
+			Node.SyncLocationFromFrontendNode();
 		}
 
 		UMetasoundEditorGraphNode& SMetaSoundGraphNodeKnot::GetMetaSoundNode()
diff --git a/Engine/Plugins/Runtime/Metasound/Source/MetasoundEditor/Private/SMetasoundGraphNode.h b/Engine/Plugins/Runtime/Metasound/Source/MetasoundEditor/Private/SMetasoundGraphNode.h
--- a/Engine/Plugins/Runtime/Metasound/Source/MetasoundEditor/Private/SMetasoundGraphNode.h
+++ b/Engine/Plugins/Runtime/Metasound/Source/MetasoundEditor/Private/SMetasoundGraphNode.h
@@ -23,93 +23,95 @@ class UMetasoundEditorGraphMemberNode;
 class UMetasoundEditorGraphNode;
 
 
-namespace Metasound
+namespace Metasound::Editor
 {
-	namespace Editor
+	class SMetaSoundGraphNode : public SGraphNode
 	{
-		class SMetaSoundGraphNode : public SGraphNode
-		{
-			public:
-			SLATE_BEGIN_ARGS(SMetaSoundGraphNode)
-			{
-			}
-
-			SLATE_END_ARGS()
-			virtual ~SMetaSoundGraphNode();
-
-			void Construct(const FArguments& InArgs, class UEdGraphNode* InNode);
-
-		protected:
-			bool IsVariableAccessor() const;
-			bool IsVariableMutator() const;
-
-			// SGraphNode Interface
-			virtual const FSlateBrush* GetShadowBrush(bool bSelected) const override;
-			virtual void CreateInputSideAddButton(TSharedPtr<SVerticalBox> InputBox) override;
-			virtual void CreateOutputSideAddButton(TSharedPtr<SVerticalBox> OutputBox) override;
-			virtual TSharedPtr<SGraphPin> CreatePinWidget(UEdGraphPin* InPin) const override;
-			virtual void CreateStandardPinWidget(UEdGraphPin* InPin) override;
-			virtual TSharedRef<SWidget> CreateNodeContentArea() override;
-			virtual TSharedRef<SWidget> CreateTitleWidget(TSharedPtr<SNodeTitle> NodeTitle) override;
-			virtual const FSlateBrush* GetNodeBodyBrush() const override;
-			virtual TSharedRef<SWidget> CreateTitleRightWidget() override;
-			virtual EVisibility IsAddPinButtonVisible() const override;
-			virtual FReply OnAddPin() override;
-			virtual void MoveTo(const FVector2D& NewPosition, FNodeSet& NodeFilter, bool bMarkDirty = true) override;
-			virtual void SetDefaultTitleAreaWidget(TSharedRef<SOverlay> DefaultTitleAreaWidget) override;
-			virtual void GetOverlayBrushes(bool bSelected, const FVector2D WidgetSize, TArray<FOverlayBrushInfo>& Brushes) const override;
-
-			FLinearColor GetNodeTitleColorOverride() const;
-
-			FName GetLiteralDataType() const;
-			UMetasoundEditorGraphNode& GetMetaSoundNode();
-			const UMetasoundEditorGraphNode& GetMetaSoundNode() const;
-
 		public:
-			static void ExecuteTrigger(UMetasoundEditorGraphMemberDefaultLiteral& Literal);
-			static TSharedRef<SWidget> CreateTriggerSimulationWidget(UMetasoundEditorGraphMemberDefaultLiteral& Literal, TAttribute<EVisibility>&& InVisibility, TAttribute<bool>&& InEnablement, const FText* InToolTip = nullptr);
-
-		private:
-			// If this node represents a graph member node, returns corresponding member.
-			UMetasoundEditorGraphMember* GetMetaSoundMember();
-
-			// If this node represents a graph member node, returns cast node.
-			UMetasoundEditorGraphMemberNode* GetMetaSoundMemberNode();
-
-			TAttribute<EVisibility> GetSimulationVisibilityAttribute() const;
-
-			// Slider widget for float input
-			TSharedPtr<SAudioInputWidget> InputWidget;
-			//Button Widget for bool input.
-			TSharedPtr <SAudioMaterialButton> MaterialButtonWidget;
-			// Handle for on value changed delegate for input slider 
-			FDelegateHandle InputSliderOnValueChangedDelegateHandle;
-			// Handle for on input slider range changed  
-			FDelegateHandle InputSliderOnRangeChangedDelegateHandle;
-
-			// Whether the input widget is currently transacting 
-			// for keeping track of transaction state across delegates to only commit transaction on value commit
-			bool bIsInputWidgetTransacting = false;
-
-			EMetasoundFrontendClassType ClassType;
-		};
-
-		class SMetaSoundGraphNodeKnot : public SGraphNodeKnot
+		SLATE_BEGIN_ARGS(SMetaSoundGraphNode)
 		{
-		public:
-			SLATE_BEGIN_ARGS(SMetaSoundGraphNode)
-			{
-			}
+		}
+
+		SLATE_END_ARGS()
+		virtual ~SMetaSoundGraphNode();
+
+		void Construct(const FArguments& InArgs, class UEdGraphNode* InNode);
+
+	protected:
+		bool IsVariableAccessor() const;
+		bool IsVariableMutator() const;
+
+		// SGraphNode Interface
+		virtual const FSlateBrush* GetShadowBrush(bool bSelected) const override;
+		virtual void CreateInputSideAddButton(TSharedPtr<SVerticalBox> InputBox) override;
+		virtual void CreateOutputSideAddButton(TSharedPtr<SVerticalBox> OutputBox) override;
+		virtual TSharedPtr<SGraphPin> CreatePinWidget(UEdGraphPin* InPin) const override;
+		virtual void CreateStandardPinWidget(UEdGraphPin* InPin) override;
+		virtual TSharedRef<SWidget> CreateNodeContentArea() override;
+		virtual TSharedRef<SWidget> CreateTitleWidget(TSharedPtr<SNodeTitle> NodeTitle) override;
+		virtual const FSlateBrush* GetNodeBodyBrush() const override;
+		virtual TSharedRef<SWidget> CreateTitleRightWidget() override;
+		virtual EVisibility IsAddPinButtonVisible() const override;
+		virtual FReply OnAddPin() override;
+		virtual void MoveTo(const FVector2D& NewPosition, FNodeSet& NodeFilter, bool bMarkDirty = true) override;
+		virtual void SetDefaultTitleAreaWidget(TSharedRef<SOverlay> DefaultTitleAreaWidget) override;
+		virtual void GetOverlayBrushes(bool bSelected, const FVector2D WidgetSize, TArray<FOverlayBrushInfo>& Brushes) const override;
+		virtual void OnCommentBubbleToggled(bool bInCommentBubbleVisible) override;
+		virtual void OnCommentTextCommitted(const FText& NewComment, ETextCommit::Type CommitInfo) override;
+
+		FLinearColor GetNodeTitleColorOverride() const;
+
+		FName GetLiteralDataType() const;
+		UMetasoundEditorGraphNode& GetMetaSoundNode();
+		const UMetasoundEditorGraphNode& GetMetaSoundNode() const;
+
+	public:
+		static void ExecuteTrigger(UMetasoundEditorGraphMemberDefaultLiteral& Literal);
+		static TSharedRef<SWidget> CreateTriggerSimulationWidget(UMetasoundEditorGraphMemberDefaultLiteral& Literal, TAttribute<EVisibility>&& InVisibility, TAttribute<bool>&& InEnablement, const FText* InToolTip = nullptr);
+
+	private:
+		// If this node represents a graph member node, returns corresponding member.
+		UMetasoundEditorGraphMember* GetMetaSoundMember();
+
+		// If this node represents a graph member node, returns cast node.
+		UMetasoundEditorGraphMemberNode* GetMetaSoundMemberNode();
+
+		TAttribute<EVisibility> GetSimulationVisibilityAttribute() const;
+
+		// Slider widget for float input
+		TSharedPtr<SAudioInputWidget> InputWidget;
+
+		//Button Widget for bool input.
+		TSharedPtr <SAudioMaterialButton> MaterialButtonWidget;
+
+		// Handle for on value changed delegate for input slider 
+		FDelegateHandle InputSliderOnValueChangedDelegateHandle;
+		
+		// Handle for on input slider range changed  
+		FDelegateHandle InputSliderOnRangeChangedDelegateHandle;
+
+		// Whether the input widget is currently transacting 
+		// for keeping track of transaction state across delegates to only commit transaction on value commit
+		bool bIsInputWidgetTransacting = false;
+
+		EMetasoundFrontendClassType ClassType;
+	};
+
+	class SMetaSoundGraphNodeKnot : public SGraphNodeKnot
+	{
+	public:
+		SLATE_BEGIN_ARGS(SMetaSoundGraphNode)
+		{
+		}
 
-			SLATE_END_ARGS()
+		SLATE_END_ARGS()
 
-			void Construct(const FArguments& InArgs, class UEdGraphNode* InNode);
+		void Construct(const FArguments& InArgs, class UEdGraphNode* InNode);
 
-			virtual TSharedPtr<SGraphPin> CreatePinWidget(UEdGraphPin* Pin) const override;
-			virtual void MoveTo(const FVector2D& NewPosition, FNodeSet& NodeFilter, bool bMarkDirty) override;
+		virtual TSharedPtr<SGraphPin> CreatePinWidget(UEdGraphPin* Pin) const override;
+		virtual void MoveTo(const FVector2D& NewPosition, FNodeSet& NodeFilter, bool bMarkDirty) override;
 
-			UMetasoundEditorGraphNode& GetMetaSoundNode();
-			const UMetasoundEditorGraphNode& GetMetaSoundNode() const;
-		};
-	} // namespace Editor
-} // namespace Metasound
+		UMetasoundEditorGraphNode& GetMetaSoundNode();
+		const UMetasoundEditorGraphNode& GetMetaSoundNode() const;
+	};
+} // namespace Metasound::Editor
diff --git a/Engine/Plugins/Runtime/Metasound/Source/MetasoundEditor/Private/SMetasoundGraphNodeComment.cpp b/Engine/Plugins/Runtime/Metasound/Source/MetasoundEditor/Private/SMetasoundGraphNodeComment.cpp
--- a/Engine/Plugins/Runtime/Metasound/Source/MetasoundEditor/Private/SMetasoundGraphNodeComment.cpp
+++ b/Engine/Plugins/Runtime/Metasound/Source/MetasoundEditor/Private/SMetasoundGraphNodeComment.cpp
@@ -3,6 +3,7 @@
 
 #include "EdGraphNode_Comment.h"
 #include "Framework/Application/SlateApplication.h"
+#include "MetasoundEditorGraphCommentNode.h"
 #include "MetasoundEditorGraphNode.h"
 
 namespace Metasound
@@ -19,7 +20,7 @@ namespace Metasound
 			FModifierKeysState KeysState = FSlateApplication::Get().GetModifierKeys();
 			if (!KeysState.IsShiftDown())
 			{
-				UEdGraphNode_Comment* CommentNode = Cast<UEdGraphNode_Comment>(GraphNode);
+				UMetasoundEditorGraphCommentNode* CommentNode = Cast<UMetasoundEditorGraphCommentNode>(GraphNode);
 				if (CommentNode && CommentNode->MoveMode == ECommentBoxMode::GroupMovement)
 				{
 					FVector2D PositionDelta = NewPosition - GetPosition();
diff --git a/Engine/Plugins/Runtime/Metasound/Source/MetasoundEditor/Private/SMetasoundGraphPin.cpp b/Engine/Plugins/Runtime/Metasound/Source/MetasoundEditor/Private/SMetasoundGraphPin.cpp
--- a/Engine/Plugins/Runtime/Metasound/Source/MetasoundEditor/Private/SMetasoundGraphPin.cpp
+++ b/Engine/Plugins/Runtime/Metasound/Source/MetasoundEditor/Private/SMetasoundGraphPin.cpp
@@ -12,6 +12,11 @@ namespace Metasound
 {
 	namespace Editor
 	{
+		bool SMetaSoundGraphPinKnot::CanInspectPin(const UEdGraphPin* InPin)
+		{
+			return false;
+		}
+
 		void SMetaSoundGraphPinKnot::Construct(const FArguments& InArgs, UEdGraphPin* InPin)
 		{
 			SGraphPinKnot::Construct(SGraphPinKnot::FArguments(), InPin);
@@ -38,8 +43,7 @@ namespace Metasound
 				{
 					if (const INodeTemplate* Template = INodeTemplateRegistry::Get().FindTemplate(FRerouteNodeTemplate::GetRegistryKey()))
 					{
-						const FMetaSoundFrontendDocumentBuilder& Builder = OwningNode->GetFrontendBuilderChecked();
-						bHasRequiredConnections = Template->HasRequiredConnections(Builder, OwningNode->GetNodeID());
+						bHasRequiredConnections = Template->HasRequiredConnections(GetBuilderChecked().GetConstBuilder(), OwningNode->GetNodeID());
 					}
 				}
 			}
diff --git a/Engine/Plugins/Runtime/Metasound/Source/MetasoundEditor/Private/SMetasoundGraphPin.h b/Engine/Plugins/Runtime/Metasound/Source/MetasoundEditor/Private/SMetasoundGraphPin.h
--- a/Engine/Plugins/Runtime/Metasound/Source/MetasoundEditor/Private/SMetasoundGraphPin.h
+++ b/Engine/Plugins/Runtime/Metasound/Source/MetasoundEditor/Private/SMetasoundGraphPin.h
@@ -8,6 +8,7 @@
 #include "KismetPins/SGraphPinInteger.h"
 #include "KismetPins/SGraphPinObject.h"
 #include "KismetPins/SGraphPinString.h"
+#include "MetasoundBuilderSubsystem.h"
 #include "MetasoundEditorGraph.h"
 #include "MetasoundEditorGraphBuilder.h"
 #include "MetasoundEditorGraphMemberDefaults.h"
@@ -42,6 +43,11 @@ namespace Metasound
 			TSharedPtr<FMetasoundPinAudioInspector> PinAudioInspector;
 
 		protected:
+			virtual bool CanInspectPin(const UEdGraphPin* InPin)
+			{
+				return FGraphBuilder::CanInspectPin(InPin);
+			}
+
 			static TWeakPtr<FPinValueInspectorTooltip> OpenPinInspector(UEdGraphPin& InPin, TSharedPtr<SMetasoundPinValueInspector>& OutPinInspector)
 			{
 				TSharedPtr<SMetasoundPinValueInspector> NewPinInspector = SNew(SMetasoundPinValueInspector);
@@ -55,37 +61,43 @@ namespace Metasound
 				return nullptr;
 			}
 
-			static void UpdatePinInspector(
+			void UpdatePinInspector(
 				UEdGraphPin& InPin,
 				const bool bIsHoveringPin,
 				TSharedPtr<SMetasoundPinValueInspector>& OutPinInspector,
 				TWeakPtr<FPinValueInspectorTooltip>& OutInspectorTooltip,
 				TFunctionRef<void(FVector2D&)> InGetTooltipLocation)
 			{
-				const bool bCanInspectPin = FGraphBuilder::CanInspectPin(&InPin);
-				if (bIsHoveringPin && bCanInspectPin)
+				if (bIsHoveringPin)
 				{
-					if (OutPinInspector.IsValid())
+					const bool bCanInspectPin = CanInspectPin(&InPin);
+					if (bCanInspectPin)
 					{
-						const UEdGraphPin* InspectedPin = OutPinInspector->GetPinRef().Get();
-						if (InspectedPin == &InPin)
+						if (OutPinInspector.IsValid())
 						{
-							OutPinInspector->UpdateMessage();
+							const UEdGraphPin* InspectedPin = OutPinInspector->GetPinRef().Get();
+							if (InspectedPin == &InPin)
+							{
+								OutPinInspector->UpdateMessage();
+							}
 						}
-					}
-					else
-					{
-						OutInspectorTooltip = OpenPinInspector(InPin, OutPinInspector);
-						TSharedPtr<FPinValueInspectorTooltip> NewTooltip = OutInspectorTooltip.Pin();
-						if (NewTooltip.IsValid())
+						else
 						{
-							FVector2D TooltipLocation;
-							InGetTooltipLocation(TooltipLocation);
-							NewTooltip->MoveTooltip(TooltipLocation);
+							OutInspectorTooltip = OpenPinInspector(InPin, OutPinInspector);
+							TSharedPtr<FPinValueInspectorTooltip> NewTooltip = OutInspectorTooltip.Pin();
+							if (NewTooltip.IsValid())
+							{
+								FVector2D TooltipLocation;
+								InGetTooltipLocation(TooltipLocation);
+								NewTooltip->MoveTooltip(TooltipLocation);
+							}
 						}
+
+						return;
 					}
 				}
-				else if (OutPinInspector.IsValid())
+
+				if (OutPinInspector.IsValid())
 				{
 					TSharedPtr<FPinValueInspectorTooltip> InspectorTooltip = OutInspectorTooltip.Pin();
 					if (InspectorTooltip.IsValid())
@@ -119,29 +131,34 @@ namespace Metasound
 				return nullptr;
 			}
 
-			static void UpdatePinAudioInspector(UEdGraphPin& InPin,
+			void UpdatePinAudioInspector(UEdGraphPin& InPin,
 				const bool bIsHoveringPin,
 				TSharedPtr<FMetasoundPinAudioInspector>& OutPinAudioInspector,
 				TWeakPtr<FPinValueInspectorTooltip>& OutInspectorTooltip,
 				TFunctionRef<void(FVector2D&)> InGetTooltipLocation)
 			{
-				const bool bCanInspectPin = FGraphBuilder::CanInspectPin(&InPin);
-
-				if (bIsHoveringPin && bCanInspectPin)
+				if (bIsHoveringPin)
 				{
-					if (!OutPinAudioInspector.IsValid())
+					const bool bCanInspectPin = CanInspectPin(&InPin);
+					if (bCanInspectPin)
 					{
-						OutInspectorTooltip = OpenPinAudioInspector(InPin, OutPinAudioInspector);
-						TSharedPtr<FPinValueInspectorTooltip> NewTooltip = OutInspectorTooltip.Pin();
-						if (NewTooltip.IsValid())
+						if (!OutPinAudioInspector.IsValid())
 						{
-							FVector2D TooltipLocation;
-							InGetTooltipLocation(TooltipLocation);
-							NewTooltip->MoveTooltip(TooltipLocation);
+							OutInspectorTooltip = OpenPinAudioInspector(InPin, OutPinAudioInspector);
+							TSharedPtr<FPinValueInspectorTooltip> NewTooltip = OutInspectorTooltip.Pin();
+							if (NewTooltip.IsValid())
+							{
+								FVector2D TooltipLocation;
+								InGetTooltipLocation(TooltipLocation);
+								NewTooltip->MoveTooltip(TooltipLocation);
+							}
 						}
+
+						return;
 					}
 				}
-				else if (OutPinAudioInspector.IsValid())
+
+				if (OutPinAudioInspector.IsValid())
 				{
 					TSharedPtr<FPinValueInspectorTooltip> InspectorTooltip = OutInspectorTooltip.Pin();
 					if (InspectorTooltip.IsValid())
@@ -180,12 +197,12 @@ namespace Metasound
 						{
 							if (Pin->Direction == EGPD_Input)
 							{
-								Frontend::FInputHandle InputHandle = FGraphBuilder::GetInputHandleFromPin(Pin);
+								Frontend::FConstInputHandle InputHandle = FGraphBuilder::GetConstInputHandleFromPin(Pin);
 								AccessType = InputHandle->GetVertexAccessType();
 							}
 							else if (Pin->Direction == EGPD_Output)
 							{
-								Frontend::FOutputHandle OutputHandle = FGraphBuilder::GetOutputHandleFromPin(Pin);
+								Frontend::FConstOutputHandle OutputHandle = FGraphBuilder::GetConstOutputHandleFromPin(Pin);
 								AccessType = OutputHandle->GetVertexAccessType();
 							}
 						}
@@ -250,6 +267,56 @@ namespace Metasound
 				return IInputController::GetInvalidHandle();
 			}
 
+			const UMetasoundEditorGraphNode* GetOwningMetaSoundNode() const
+			{
+				if (const UEdGraphPin* Pin = ParentPinType::GetPinObj())
+				{
+					UObject* Node = Pin->GetOwningNode();
+					return Cast<UMetasoundEditorGraphNode>(Node);
+				}
+
+				return nullptr;
+			}
+
+			UMetaSoundBuilderBase& GetBuilderChecked() const
+			{
+				const UMetasoundEditorGraphNode* Node = GetOwningMetaSoundNode();
+				check(Node);
+				UObject* Outermost = Node->GetOutermostObject();
+				check(Outermost);
+				return UMetaSoundBuilderSubsystem::GetChecked().AttachBuilderToAssetChecked(*Outermost);
+			}
+
+			const FMetasoundFrontendNode* GetFrontendNode() const
+			{
+				if (const UMetasoundEditorGraphNode* Node = GetOwningMetaSoundNode())
+				{
+					if (UObject* Outermost = Node->GetOutermostObject())
+					{
+						const FGuid NodeID = Node->GetNodeID();
+						const UMetaSoundBuilderBase& Builder = UMetaSoundBuilderSubsystem::GetChecked().AttachBuilderToAssetChecked(*Outermost);
+						return Builder.GetConstBuilder().FindNode(NodeID);
+					}
+				}
+
+				return nullptr;
+			}
+
+			const FMetasoundFrontendNode& GetFrontendNodeChecked() const
+			{
+				const UMetasoundEditorGraphNode* Node = GetOwningMetaSoundNode();
+				check(Node);
+				UObject* Outermost = Node->GetOutermostObject();
+				check(Outermost);
+
+				const FGuid NodeID = Node->GetNodeID();
+				const UMetaSoundBuilderBase& Builder = UMetaSoundBuilderSubsystem::GetChecked().AttachBuilderToAssetChecked(*Outermost);
+
+				const FMetasoundFrontendNode* FrontendNode = Builder.GetConstBuilder().FindNode(NodeID);
+				check(FrontendNode);
+				return *FrontendNode;
+			}
+
 			bool ShowDefaultValueWidget() const
 			{
 				UEdGraphPin* Pin = ParentPinType::GetPinObj();
@@ -327,7 +394,7 @@ namespace Metasound
 									{
 										if (UMetasoundEditorGraphExternalNode* Node = Cast<UMetasoundEditorGraphExternalNode>(Pin->GetOwningNode()))
 										{
-											if (Node->GetClassName() == FRerouteNodeTemplate::ClassName)
+											if (Node->GetBreadcrumb().ClassName == FRerouteNodeTemplate::ClassName)
 											{
 												bIsRerouteNode = true;
 											}
@@ -345,6 +412,7 @@ namespace Metasound
 						}))
 						.OnClicked(FOnClicked::CreateLambda([this]()
 						{
+							using namespace Editor;
 							using namespace Frontend;
 
 							if (UEdGraphPin* Pin = ParentPinType::GetPinObj())
@@ -354,7 +422,7 @@ namespace Metasound
 									if (UMetasoundEditorGraph* MetaSoundGraph = CastChecked<UMetasoundEditorGraph>(Node->GetGraph()))
 									{
 										UObject& MetaSound = MetaSoundGraph->GetMetasoundChecked();
-
+										FMetasoundFrontendDocumentModifyContext& ModifyContext = FGraphBuilder::GetOutermostMetaSoundChecked(MetaSound).GetModifyContext();
 										{
 											const FScopedTransaction Transaction(LOCTEXT("MetaSoundEditorResetToClassDefault", "Reset to Class Default"));
 											MetaSound.Modify();
@@ -366,18 +434,18 @@ namespace Metasound
 												if (ensure(Member))
 												{
 													Member->ResetToClassDefault();
-													MetaSoundGraph->GetModifyContext().AddMemberIDsModified({ Member->GetMemberID() });
+													ModifyContext.AddMemberIDsModified({ Member->GetMemberID() });
 												}
 												else
 												{
-													MetaSoundGraph->GetModifyContext().SetDocumentModified();
+													ModifyContext.SetDocumentModified();
 												}
 											}
 											else
 											{
 												FInputHandle InputHandle = GetInputHandle();
 												InputHandle->ClearLiteral();
-												MetaSoundGraph->GetModifyContext().SetDocumentModified();
+												ModifyContext.SetDocumentModified();
 											}
 										}
 									}
@@ -544,6 +612,8 @@ namespace Metasound
 			virtual const FSlateBrush* GetPinIcon() const override;
 
 		protected:
+			virtual bool CanInspectPin(const UEdGraphPin* InPin) override;
+
 			void CacheHasRequiredConnections();
 
 			bool bHasRequiredConnections = false;
diff --git a/Engine/Plugins/Runtime/Metasound/Source/MetasoundEditor/Private/SMetasoundPinValueInspector.cpp b/Engine/Plugins/Runtime/Metasound/Source/MetasoundEditor/Private/SMetasoundPinValueInspector.cpp
--- a/Engine/Plugins/Runtime/Metasound/Source/MetasoundEditor/Private/SMetasoundPinValueInspector.cpp
+++ b/Engine/Plugins/Runtime/Metasound/Source/MetasoundEditor/Private/SMetasoundPinValueInspector.cpp
@@ -171,7 +171,7 @@ namespace Metasound
 			return nullptr;
 		}
 
-		Frontend::FConstOutputHandle FMetasoundNumericDebugLineItem::GetReroutedOutputHandle() const
+		Frontend::FConstOutputHandle FMetasoundNumericDebugLineItem::GetConstReroutedOutputHandle() const
 		{
 			return FGraphBuilder::FindReroutedConstOutputHandleFromPin(GraphPinObj);
 		}
@@ -336,9 +336,9 @@ namespace Metasound
 			return OutputHandle->GetOwningNode()->GetOwningGraph();
 		}
 
-		Frontend::FConstGraphHandle FMetasoundNumericDebugLineItem::GetGraphHandle() const
+		Frontend::FConstGraphHandle FMetasoundNumericDebugLineItem::GetConstGraphHandle() const
 		{
-			Frontend::FConstOutputHandle OutputHandle = FGraphBuilder::GetOutputHandleFromPin(GraphPinObj);
+			Frontend::FConstOutputHandle OutputHandle = FGraphBuilder::GetConstOutputHandleFromPin(GraphPinObj);
 			return OutputHandle->GetOwningNode()->GetOwningGraph();
 		}
 
@@ -360,21 +360,21 @@ namespace Metasound
 		{
 			using namespace Frontend;
 
-			FConstGraphHandle OwningGraph = GetGraphHandle();
+			FConstGraphHandle OwningGraph = GetConstGraphHandle();
 			if (!OwningGraph->IsValid())
 			{
 				return nullptr;
 			}
 
-			const FConstOutputHandle OutputHandle = GetReroutedOutputHandle();
+			const FConstOutputHandle OutputHandle = GetConstReroutedOutputHandle();
 			const FGuid NodeID = OutputHandle->GetOwningNodeID();
 			const FName OutputName = OutputHandle->GetName();
 			const FMetasoundFrontendGraphStyle& Style = OwningGraph->GetGraphStyle();
 
 			return Style.EdgeStyles.FindByPredicate([&NodeID, &OutputName](FMetasoundFrontendEdgeStyle& EdgeStyle)
-				{
-					return EdgeStyle.NodeID == NodeID && EdgeStyle.OutputName == OutputName;
-				});
+			{
+				return EdgeStyle.NodeID == NodeID && EdgeStyle.OutputName == OutputName;
+			});
 		}
 
 		void FMetasoundNumericDebugLineItem::Update()
@@ -383,7 +383,7 @@ namespace Metasound
 
 			bIsValueColorizationEnabled = false;
 
-			const FConstOutputHandle OutputHandle = GetReroutedOutputHandle();
+			const FConstOutputHandle OutputHandle = GetConstReroutedOutputHandle();
 			const FGuid NodeID = OutputHandle->GetOwningNodeID();
 			const FName OutputName = OutputHandle->GetName();
 			if (const FMetasoundFrontendEdgeStyle* EdgeStyle = GetEdgeStyle())
@@ -413,12 +413,6 @@ namespace Metasound
 				return FReply::Unhandled();
 			}
 
-			FGraphHandle OwningGraph = GetGraphHandle();
-			if (!OwningGraph->IsValid())
-			{
-				return FReply::Unhandled();
-			}
-
 			const FMetasoundFrontendEdgeStyle* EdgeStyle = GetEdgeStyle();
 			if (!EdgeStyle)
 			{
diff --git a/Engine/Plugins/Runtime/Metasound/Source/MetasoundEditor/Private/SMetasoundPinValueInspector.h b/Engine/Plugins/Runtime/Metasound/Source/MetasoundEditor/Private/SMetasoundPinValueInspector.h
--- a/Engine/Plugins/Runtime/Metasound/Source/MetasoundEditor/Private/SMetasoundPinValueInspector.h
+++ b/Engine/Plugins/Runtime/Metasound/Source/MetasoundEditor/Private/SMetasoundPinValueInspector.h
@@ -61,13 +61,13 @@ namespace Metasound
 
 			UObject* GetOutermostObject();
 
-			Frontend::FConstOutputHandle GetReroutedOutputHandle() const;
+			Frontend::FConstOutputHandle GetConstReroutedOutputHandle() const;
 
 			Frontend::FOutputHandle GetReroutedOutputHandle();
 
 			Frontend::FGraphHandle GetGraphHandle();
 
-			Frontend::FConstGraphHandle GetGraphHandle() const;
+			Frontend::FConstGraphHandle GetConstGraphHandle() const;
 
 			FGraphConnectionManager* GetConnectionManager();
 
