Version Metasound Document to include all ed data and make all Metasound EdGraph data transient

diff --git a/Engine/Plugins/Runtime/Metasound/Source/MetasoundEditor/Public/MetasoundEditorGraph.h b/Engine/Plugins/Runtime/Metasound/Source/MetasoundEditor/Public/MetasoundEditorGraph.h
--- a/Engine/Plugins/Runtime/Metasound/Source/MetasoundEditor/Public/MetasoundEditorGraph.h
+++ b/Engine/Plugins/Runtime/Metasound/Source/MetasoundEditor/Public/MetasoundEditorGraph.h
@@ -2,6 +2,7 @@
 #pragma once
 
 #include "AudioParameterControllerInterface.h"
+#include "MetasoundBuilderSubsystem.h"
 #include "MetasoundEditor.h"
 #include "MetasoundEditorGraphNode.h"
 #include "MetasoundFrontendController.h"
@@ -46,7 +47,7 @@ DECLARE_MULTICAST_DELEGATE_OneParam(FOnMetasoundMemberNameChanged, FGuid /* Node
 DECLARE_MULTICAST_DELEGATE(FOnMetasoundMemberRenameRequested);
 
 UCLASS()
-class METASOUNDEDITOR_API UMetasoundEditorGraphMemberDefaultLiteral : public UObject
+class METASOUNDEDITOR_API UMetasoundEditorGraphMemberDefaultLiteral : public UMetaSoundFrontendMemberMetadata
 {
 	GENERATED_BODY()
 
@@ -55,35 +56,22 @@ public:
 	{
 	}
 
-	virtual FMetasoundFrontendLiteral GetDefault() const
-	{
-		return FMetasoundFrontendLiteral();
-	}
-
-	virtual EMetasoundFrontendLiteralType GetLiteralType() const
-	{
-		return EMetasoundFrontendLiteralType::None;
-	}
-
-	virtual void SetFromLiteral(const FMetasoundFrontendLiteral& InLiteral)
-	{
-	}
-
-	virtual void ForceRefresh()
-	{
-	}
-
 #if WITH_EDITOR
 	virtual void PostEditChangeProperty(FPropertyChangedEvent& InPropertyChangedEvent) override;
 	virtual void PostEditChangeChainProperty(FPropertyChangedChainEvent& InPropertyChangedEvent) override;
 	virtual void PostEditUndo() override;
 #endif // WITH_EDITOR
 
+	const UMetasoundEditorGraphMember* FindMember() const;
+	UMetasoundEditorGraphMember* FindMember();
+
+	UE_DEPRECATED(5.5, "Due to serialization optimization, literals now inherit from LiteralMetadata and are no longer parented under members. Use FindMember instead")
 	const UMetasoundEditorGraphMember* GetParentMember() const
 	{
 		return Cast<UMetasoundEditorGraphMember>(GetOuter());
 	}
 
+	UE_DEPRECATED(5.5, "Due to serialization optimization, literals now inherit from LiteralMetadata and are no longer parented under members. Use FindMember instead")
 	UMetasoundEditorGraphMember* GetParentMember()
 	{
 		return Cast<UMetasoundEditorGraphMember>(GetOuter());
@@ -176,6 +164,8 @@ public:
 	void InitializeLiteral();
 
 protected:
+	virtual UMetasoundEditorGraphNode* AddNode(Metasound::Frontend::FNodeHandle InNodeHandle, bool bInSelectNewNode) PURE_VIRTUAL(UMetasoundEditorGraphMember::AddNode, return nullptr;);
+
 	/** Default literal value of member */
 	UPROPERTY()
 	TObjectPtr<UMetasoundEditorGraphMemberDefaultLiteral> Literal;
@@ -184,7 +174,8 @@ protected:
 	UPROPERTY()
 	FName TypeName;
 
-friend class UMetasoundEditorGraph;
+	friend class UMetasoundEditorGraph;
+	friend class UMetaSoundEditorSubsystem;
 };
 
 /** Base class for an input or output of the graph. */
@@ -194,9 +185,6 @@ class METASOUNDEDITOR_API UMetasoundEditorGraphVertex : public UMetasoundEditorG
 	GENERATED_BODY()
 
 protected:
-	UE_DEPRECATED(5.1, "Use AddNodeHandle with FCreateNodeVertexParams instead.")
-	virtual Metasound::Frontend::FNodeHandle AddNodeHandle(const FName& InNodeName, FName InDataType) PURE_VIRTUAL(UMetasoundEditorGraphVertex::AddNodeHandle, return Metasound::Frontend::INodeController::GetInvalidHandle(); )
-
 	/** Adds the node handle for a newly created vertex. */
 	virtual Metasound::Frontend::FNodeHandle AddNodeHandle(const FName& InNodeName, const Metasound::Editor::FCreateNodeVertexParams& InParams) PURE_VIRTUAL(UMetasoundEditorGraphVertex::AddNodeHandle, return Metasound::Frontend::INodeController::GetInvalidHandle(); )
 
@@ -223,7 +211,6 @@ public:
 	virtual void SetDescription(const FText& InDescription, bool bPostTransaction) override;
 	virtual void SetMemberName(const FName& InNewName, bool bPostTransaction) override;
 	virtual void SetDisplayName(const FText& InNewName, bool bPostTransaction) override;
-	virtual void SetDataType(FName InNewType, bool bPostTransaction = true) override;
 	/* ~End UMetasoundEditorGraphMember interface */
 
 	/** Version of interface membership, or invalid version if not an interface member. */
@@ -233,15 +220,19 @@ public:
 	virtual bool IsInterfaceMember() const;
 
 	/** Returns the Metasound class type of the associated node */
-	virtual EMetasoundFrontendClassType GetClassType() const PURE_VIRTUAL(UMetasoundEditorGraphMember::GetClassType, return EMetasoundFrontendClassType::Invalid; )
+	virtual EMetasoundFrontendClassType GetClassType() const PURE_VIRTUAL(UMetasoundEditorGraphVertex::GetClassType, return EMetasoundFrontendClassType::Invalid; )
 
 	/** Returns the SortOrderIndex assigned to this member. */
-	virtual int32 GetSortOrderIndex() const PURE_VIRTUAL(UMetasoundEditorGraphMember::GetSortOrderIndex, return 0; )
+	virtual int32 GetSortOrderIndex() const PURE_VIRTUAL(UMetasoundEditorGraphVertex::GetSortOrderIndex, return 0; )
+
+	/** Sets the DataType assigned to this member. */
+	virtual void SetDataType(FName DataType, bool bPostTransaction = true) PURE_VIRTUAL(UMetasoundEditorGraphVertex::SetDataType, )
 
 	/** Sets the SortOrderIndex assigned to this member. */
-	virtual void SetSortOrderIndex(int32 InSortOrderIndex) PURE_VIRTUAL(UMetasoundEditorGraphMember::SetSortOrderIndex, )
+	virtual void SetSortOrderIndex(int32 InSortOrderIndex) PURE_VIRTUAL(UMetasoundEditorGraphVertex::SetSortOrderIndex, )
 	
-	void SetVertexAccessType(EMetasoundFrontendVertexAccessType InNewAccessType, bool bPostTransaction = true);
+	/** Sets the VertexAccessType assigned to this member. */
+	virtual void SetVertexAccessType(EMetasoundFrontendVertexAccessType InNewAccessType, bool bPostTransaction = true) PURE_VIRTUAL(UMetasoundEditorGraphVertex::SetVertexAccessType, )
 
 	/** Returns the node handle associated with the vertex. */
 	Metasound::Frontend::FNodeHandle GetNodeHandle();
@@ -252,6 +243,8 @@ public:
 	virtual EMetasoundFrontendVertexAccessType GetVertexAccessType() const PURE_VIRTUAL(UMetasoundEditorGraphVertex::GetVertexAccessType, return EMetasoundFrontendVertexAccessType::Reference; );
 
 	virtual bool CanRename() const override;
+
+	friend class Metasound::Editor::FEditor;
 };
 
 UCLASS()
@@ -262,16 +255,20 @@ class METASOUNDEDITOR_API UMetasoundEditorGraphInput : public UMetasoundEditorGr
 public:
 	virtual int32 GetSortOrderIndex() const override;
 	virtual void SetSortOrderIndex(int32 InSortOrderIndex) override;
+	virtual TArray<UMetasoundEditorGraphMemberNode*> GetNodes() const override;
 
 	virtual const FText& GetGraphMemberLabel() const override;
 	virtual void ResetToClassDefault() override;
+	virtual void SetDataType(FName InNewType, bool bPostTransaction) override;
+	virtual void SetVertexAccessType(EMetasoundFrontendVertexAccessType InNewAccessType, bool bPostTransaction) override;
 	virtual void SetMemberName(const FName& InNewName, bool bPostTransaction) override;
 	virtual void UpdateFrontendDefaultLiteral(bool bPostTransaction) override;
 	virtual EMetasoundFrontendVertexAccessType GetVertexAccessType() const override;
 
 protected:
-	virtual Metasound::Frontend::FNodeHandle AddNodeHandle(const FName& InNodeName, FName InDataType) override;
+	virtual UMetasoundEditorGraphNode* AddNode(Metasound::Frontend::FNodeHandle InNodeHandle, bool bInSelectNewNode) override;
 	virtual Metasound::Frontend::FNodeHandle AddNodeHandle(const FName& InNodeName, const Metasound::Editor::FCreateNodeVertexParams& InParams) override;
+
 	virtual EMetasoundFrontendClassType GetClassType() const override { return EMetasoundFrontendClassType::Input; }
 	virtual Metasound::Editor::ENodeSection GetSectionID() const override;
 };
@@ -287,11 +284,13 @@ public:
 	virtual void SetSortOrderIndex(int32 InSortOrderIndex) override;
 	virtual const FText& GetGraphMemberLabel() const override;
 	virtual void ResetToClassDefault() override;
+	virtual void SetDataType(FName InNewType, bool bPostTransaction) override;
+	virtual void SetVertexAccessType(EMetasoundFrontendVertexAccessType InNewAccessType, bool bPostTransaction) override;
 	virtual void UpdateFrontendDefaultLiteral(bool bPostTransaction) override;
 	virtual EMetasoundFrontendVertexAccessType GetVertexAccessType() const override;
 
 protected:
-	virtual Metasound::Frontend::FNodeHandle AddNodeHandle(const FName& InNodeName, FName InDataType) override;
+	virtual UMetasoundEditorGraphNode* AddNode(Metasound::Frontend::FNodeHandle InNodeHandle, bool bInSelectNewNode) override;
 	virtual Metasound::Frontend::FNodeHandle AddNodeHandle(const FName& InNodeName, const Metasound::Editor::FCreateNodeVertexParams& InParams) override;
 	virtual EMetasoundFrontendClassType GetClassType() const override { return EMetasoundFrontendClassType::Output; }
 	virtual Metasound::Editor::ENodeSection GetSectionID() const override;
@@ -338,6 +337,9 @@ public:
 
 	virtual bool CanRename() const override;
 
+protected:
+	virtual UMetasoundEditorGraphNode* AddNode(Metasound::Frontend::FNodeHandle InNodeHandle, bool bInSelectNewNode) override;
+
 private:
 	struct FVariableEditorNodes
 	{
@@ -365,6 +367,7 @@ class METASOUNDEDITOR_API UMetasoundEditorGraph : public UMetasoundEditorGraphBa
 	GENERATED_BODY()
 
 public:
+	UE_DEPRECATED(5.5, "Input node creation is no longer supported publically via the MetasoundEditorGraph.")
 	UMetasoundEditorGraphInputNode* CreateInputNode(Metasound::Frontend::FNodeHandle InNodeHandle, bool bInSelectNewNode);
 
 	Metasound::Frontend::FDocumentHandle GetDocumentHandle();
@@ -374,6 +377,7 @@ public:
 
 	virtual void PreSave(FObjectPreSaveContext InSaveContext) override;
 
+	UMetaSoundBuilderBase& GetBuilderChecked();
 	UObject* GetMetasound();
 	const UObject* GetMetasound() const;
 	UObject& GetMetasoundChecked();
@@ -394,25 +398,15 @@ public:
 
 	// UMetasoundEditorGraphBase Implementation
 	virtual void RegisterGraphWithFrontend() override;
-	virtual FMetasoundFrontendDocumentModifyContext& GetModifyContext() override;
-	virtual const FMetasoundFrontendDocumentModifyContext& GetModifyContext() const override;
-
-	virtual void ClearVersionedOnLoad() override;
-	virtual bool GetVersionedOnLoad() const override;
-	virtual void SetVersionedOnLoad() override;
+	virtual void MigrateEditorDocumentData(FMetaSoundFrontendDocumentBuilder& OutBuilder) override;
 
 private:
-	bool RemoveFrontendInput(UMetasoundEditorGraphInput& Input);
-	bool RemoveFrontendOutput(UMetasoundEditorGraphOutput& Output);
-	bool RemoveFrontendVariable(UMetasoundEditorGraphVariable& Variable);
 	void ValidateInternal(Metasound::Editor::FGraphValidationResults& OutResults);
 
 	// Preview ID is the Unique ID provided by the UObject that implements
 	// a sound's ParameterInterface when a sound begins playing.
 	uint32 PreviewID = INDEX_NONE;
 
-	bool bVersionedOnLoad = false;
-
 	// Used as a means of forcing the graph to rebuild nodes on next tick.
 	// TODO: Will no longer require this once all editor metadata is migrated
 	// to the frontend & the system can adequately rely on the changeIDs as a
@@ -431,21 +425,26 @@ private:
 public:
 	UMetasoundEditorGraphInput* FindInput(FGuid InNodeID) const;
 	UMetasoundEditorGraphInput* FindInput(FName InName) const;
-	UMetasoundEditorGraphInput* FindOrAddInput(Metasound::Frontend::FNodeHandle InNodeHandle);
+	UMetasoundEditorGraphInput* FindOrAddInput(Metasound::Frontend::FConstNodeHandle InNodeHandle);
 
 	UMetasoundEditorGraphOutput* FindOutput(FGuid InNodeID) const;
 	UMetasoundEditorGraphOutput* FindOutput(FName InName) const;
-	UMetasoundEditorGraphOutput* FindOrAddOutput(Metasound::Frontend::FNodeHandle InNodeHandle);
+	UMetasoundEditorGraphOutput* FindOrAddOutput(Metasound::Frontend::FConstNodeHandle InNodeHandle);
 
 	UMetasoundEditorGraphVariable* FindVariable(const FGuid& InVariableID) const;
 	UMetasoundEditorGraphVariable* FindOrAddVariable(const Metasound::Frontend::FConstVariableHandle& InVariableHandle);
 
-	UMetasoundEditorGraphMember* FindMember(FGuid InNodeID) const;
+	UMetasoundEditorGraphMember* FindMember(FGuid InMemberID) const;
 	UMetasoundEditorGraphMember* FindAdjacentMember(const UMetasoundEditorGraphMember& InMember);
 
-	bool RemoveMember(UMetasoundEditorGraphMember& InGraphMember);
-	bool RemoveMemberNodes(UMetasoundEditorGraphMember& InGraphMember);
-	bool RemoveFrontendMember(UMetasoundEditorGraphMember& InMember);
+	UE_DEPRECATED(5.5, "Use the associated Frontend builder and synchronize the graph remove a given member.")
+	bool RemoveMember(UMetasoundEditorGraphMember& InGraphMember) { return false; }
+
+	UE_DEPRECATED(5.5, "Use the associated Frontend builder and synchronize the graph remove all member nodes.")
+	bool RemoveMemberNodes(UMetasoundEditorGraphMember& InGraphMember) { return false; }
+
+	UE_DEPRECATED(5.5, "Use RemoveMember instead")
+	bool RemoveFrontendMember(UMetasoundEditorGraphMember& InGraphMember) { return false; }
 
 	friend class UMetaSoundFactory;
 	friend class UMetaSoundSourceFactory;
diff --git a/Engine/Plugins/Runtime/Metasound/Source/MetasoundEditor/Public/MetasoundEditorGraphCommentNode.h b/Engine/Plugins/Runtime/Metasound/Source/MetasoundEditor/Public/MetasoundEditorGraphCommentNode.h
--- a/Engine/Plugins/Runtime/Metasound/Source/MetasoundEditor/Public/MetasoundEditorGraphCommentNode.h
+++ b/Engine/Plugins/Runtime/Metasound/Source/MetasoundEditor/Public/MetasoundEditorGraphCommentNode.h
@@ -14,6 +14,7 @@
 
 struct FPropertyChangedEvent;
 class FMetasoundAssetBase;
+class UMetaSoundBuilderBase;
 
 
 // Forward Declarations
@@ -40,17 +41,21 @@ public:
 	//~ End UObject Interface
 
 	//~ Begin UEdGraphNode Interface
-	METASOUNDEDITOR_API virtual void ResizeNode(const FVector2D& NewSize) override;
-	METASOUNDEDITOR_API virtual void PostPlacedNewNode() override;
+	METASOUNDEDITOR_API virtual bool CanUserDeleteNode() const override;
 	METASOUNDEDITOR_API virtual void OnRenameNode(const FString& NewName) override;
+	METASOUNDEDITOR_API virtual void ResizeNode(const FVector2D& NewSize) override;
 	//~ End UEdGraphNode Interface
 
 	/** Set the Bounds for the comment node */
 	METASOUNDEDITOR_API void SetBounds(const class FSlateRect& Rect);
 
+	UMetaSoundBuilderBase& GetBuilderChecked() const;
+
 	// Returns frontend comment ID
 	FGuid GetCommentID() const;
 
+	bool RemoveFromDocument() const;
+
 private:
 	FMetasoundAssetBase& GetAssetChecked();
 
diff --git a/Engine/Plugins/Runtime/Metasound/Source/MetasoundEditor/Public/MetasoundEditorGraphInputNode.h b/Engine/Plugins/Runtime/Metasound/Source/MetasoundEditor/Public/MetasoundEditorGraphInputNode.h
--- a/Engine/Plugins/Runtime/Metasound/Source/MetasoundEditor/Public/MetasoundEditorGraphInputNode.h
+++ b/Engine/Plugins/Runtime/Metasound/Source/MetasoundEditor/Public/MetasoundEditorGraphInputNode.h
@@ -25,19 +25,31 @@ public:
 	UPROPERTY()
 	TObjectPtr<UMetasoundEditorGraphInput> Input;
 
+	UPROPERTY()
+	FGuid NodeID;
+
 public:
+	const FMetasoundEditorGraphVertexNodeBreadcrumb& GetBreadcrumb() const;
+
+	virtual void CacheTitle() override;
+
+	virtual void CacheBreadcrumb() override;
 	virtual UMetasoundEditorGraphMember* GetMember() const override;
-	virtual void UpdatePreviewInstance(const Metasound::FVertexName& InParameterName, TScriptInterface<IAudioParameterControllerInterface>& InParameterInterface) const;
+	virtual void ReconstructNode() override;
 
+	virtual void UpdatePreviewInstance(const Metasound::FVertexName& InParameterName, TScriptInterface<IAudioParameterControllerInterface>& InParameterInterface) const;
 	virtual FMetasoundFrontendClassName GetClassName() const;
 
 	virtual FGuid GetNodeID() const override;
 	virtual FLinearColor GetNodeTitleColor() const override;
 	virtual FSlateIcon GetNodeTitleIcon() const override;
+	virtual void GetPinHoverText(const UEdGraphPin& Pin, FString& OutHoverText) const override;
 	virtual void Validate(Metasound::Editor::FGraphNodeValidationResult& OutResult) override;
 
 protected:
-	virtual void SetNodeID(FGuid InNodeID) override;
+	// Breadcrumb used if associated FrontendNode cannot be found or has been unlinked
+	UPROPERTY()
+	FMetasoundEditorGraphVertexNodeBreadcrumb Breadcrumb;
 
 	friend class Metasound::Editor::FGraphBuilder;
 };
diff --git a/Engine/Plugins/Runtime/Metasound/Source/MetasoundEditor/Public/MetasoundEditorGraphNode.h b/Engine/Plugins/Runtime/Metasound/Source/MetasoundEditor/Public/MetasoundEditorGraphNode.h
--- a/Engine/Plugins/Runtime/Metasound/Source/MetasoundEditor/Public/MetasoundEditorGraphNode.h
+++ b/Engine/Plugins/Runtime/Metasound/Source/MetasoundEditor/Public/MetasoundEditorGraphNode.h
@@ -17,24 +17,62 @@
 
 // Forward Declarations
 class UEdGraphPin;
+class UMetaSoundBuilderBase;
 class UMetaSoundPatch;
 class UMetasoundEditorGraphOutput;
 class UMetasoundEditorGraphMember;
 class UMetasoundEditorGraphVariable;
 class UMetasoundEditorGraphMemberDefaultFloat;
+
 namespace EMessageSeverity { enum Type : int; }
 
-namespace Metasound
+namespace Metasound::Editor
 {
-	namespace Editor
-	{
-		struct FGraphNodeValidationResult;
-		class FGraphBuilder;
+	struct FDocumentClipboardUtils;
+	struct FGraphNodeValidationResult;
+
+	class FGraphBuilder;
+
+	// Map of class names to sorted array of registered version numbers
+	using FSortedClassVersionMap = TMap<FName, TArray<FMetasoundFrontendVersionNumber>>;
+} // namespace Metasound::Editor
+
+
+USTRUCT()
+struct FMetasoundEditorGraphNodeBreadcrumb
+{
+	GENERATED_BODY()
+
+	UPROPERTY()
+	FMetasoundFrontendClassName ClassName;
+
+	UPROPERTY()
+	bool bIsClassNative = true;
+};
 
-		// Map of class names to sorted array of registered version numbers
-		using FSortedClassVersionMap = TMap<FName, TArray<FMetasoundFrontendVersionNumber>>;
-	} // namespace Editor
-} // namespace Metasound
+USTRUCT()
+struct FMetasoundEditorGraphMemberNodeBreadcrumb : public FMetasoundEditorGraphNodeBreadcrumb
+{
+	GENERATED_BODY()
+
+	UPROPERTY()
+	FName MemberName;
+
+	UPROPERTY()
+	FName DataType;
+
+	UPROPERTY()
+	FMetasoundFrontendLiteral DefaultLiteral;
+};
+
+USTRUCT()
+struct FMetasoundEditorGraphVertexNodeBreadcrumb : public FMetasoundEditorGraphMemberNodeBreadcrumb
+{
+	GENERATED_BODY()
+
+	UPROPERTY()
+	EMetasoundFrontendVertexAccessType AccessType = EMetasoundFrontendVertexAccessType::Unset;
+};
 
 UCLASS(MinimalAPI)
 class UMetasoundEditorGraphNode : public UEdGraphNode
@@ -64,6 +102,7 @@ public:
 	virtual void PinDefaultValueChanged(UEdGraphPin* Pin) override;
 	virtual void ReconstructNode() override;
 	virtual FString GetPinMetaData(FName InPinName, FName InKey) override;
+	virtual void OnUpdateCommentText(const FString& NewComment) override;
 	// End of UEdGraphNode interface
 
 	// UObject interface
@@ -78,21 +117,36 @@ public:
 
 	virtual FSlateIcon GetNodeTitleIcon() const { return FSlateIcon(); }
 	virtual FName GetCornerIcon() const { return FName(); }
-	virtual bool CanAddInputPin() const
-	{
-		return false;
-	}
+	virtual bool CanAddInputPin() const { return false; }
 
-	const FMetaSoundFrontendDocumentBuilder& GetFrontendBuilderChecked() const;
+	UMetaSoundBuilderBase& GetBuilderChecked() const;
 	const FMetasoundFrontendNode* GetFrontendNode() const;
 	const FMetasoundFrontendNode& GetFrontendNodeChecked() const;
 
+	virtual const FMetasoundEditorGraphNodeBreadcrumb& GetBreadcrumb() const;
+
+	// Caches any "breadcrumb" data associated with a particular MetaSound Editor node. Called
+	// when copying edgraph data to the clipboard or validating for fast access. Also generally
+	// provides a mechanism for MetaSound nodes to cache Frontend data for use to look-up Frontend
+	// data if re-associated should the associated document data/node becomes unlinked.
+	virtual void CacheBreadcrumb() { }
+
 	UObject& GetMetasoundChecked();
 	const UObject& GetMetasoundChecked() const;
 
-	// Sets the node's location, both on this graph member node and on the frontend handle
+	virtual bool RemoveFromDocument() const;
+
+	UE_DEPRECATED(5.4, "Use UpdateFrontendNodeLocation and/or SyncLocationFromFrontendNode")
 	void SetNodeLocation(const FVector2D& InLocation);
-	// Helper function to update node location on frontend handle
+
+	// Finds the associated node with the given ID and sets this EdGraphNode's comment and comment visibility boolean.
+	void SyncCommentFromFrontendNode();
+
+	// Finds the associated node with the given ID and sets this EdGraphNode's location.
+	// Returns whether or not node ID entry exists and if location was set.
+	bool SyncLocationFromFrontendNode(bool bUpdateEditorNodeID = false);
+
+	// Helper function that sets the associated frontend node's location. Does NOT set this EdGraphNode's location.
 	void UpdateFrontendNodeLocation(const FVector2D& InLocation);
 
 	Metasound::Frontend::FGraphHandle GetRootGraphHandle() const;
@@ -103,7 +157,9 @@ public:
 
 	TSet<FString> GetDisallowedPinClassNames(const UEdGraphPin& InPin) const;
 
+	UE_DEPRECATED(5.4, "Use the Frontend node or Breadcrumb directly to get the class name.")
 	virtual FMetasoundFrontendClassName GetClassName() const { return FMetasoundFrontendClassName(); }
+
 	virtual FGuid GetNodeID() const { return FGuid(); }
 	virtual FText GetDisplayName() const;
 	virtual void CacheTitle();
@@ -112,7 +168,6 @@ public:
 	// Mark node for refresh
 	void SyncChangeIDs();
 
-
 	FText GetCachedTitle() const { return CachedTitle; }
 
 	// Returns whether or not the class interface, metadata, or style has been changed since the last node refresh
@@ -127,9 +182,10 @@ protected:
 	// a new name if the external definition changes between application sessions.
 	FText CachedTitle;
 
-	virtual void SetNodeID(FGuid InNodeID) { }
+	static bool ShowNodeDebugData();
 
-	friend class Metasound::Editor::FGraphBuilder;
+	UE_DEPRECATED(5.4, "Now set directly on implementing nodes")
+	virtual void SetNodeID(FGuid InNodeID) { }
 };
 
 /** Node that represents a graph member */
@@ -147,7 +203,7 @@ public:
 	{
 		return true;
 	}
-protected:
+
 	// Clamp float literal value based on the given default float literal. 
 	// Returns whether the literal was clamped. 
 	static bool ClampFloatLiteral(const UMetasoundEditorGraphMemberDefaultFloat* DefaultFloatLiteral, FMetasoundFrontendLiteral& LiteralValue);
@@ -163,7 +219,12 @@ public:
 	UPROPERTY()
 	TObjectPtr<UMetasoundEditorGraphOutput> Output;
 
+	virtual const FMetasoundEditorGraphNodeBreadcrumb& GetBreadcrumb() const override;
+	virtual void CacheBreadcrumb() override;
+
+	UE_DEPRECATED(5.4, "Use the Frontend node or Breadcrumb directly to get the class name.")
 	virtual FMetasoundFrontendClassName GetClassName() const override;
+
 	virtual FGuid GetNodeID() const override;
 	virtual UMetasoundEditorGraphMember* GetMember() const override;
 
@@ -177,6 +238,8 @@ public:
 	virtual bool CanUserDeleteNode() const override;
 
 	virtual void PinDefaultValueChanged(UEdGraphPin* InPin) override;
+	virtual void ReconstructNode() override;
+	virtual bool RemoveFromDocument() const override;
 
 	// Disables interact widgets (ex. sliders, knobs) when input is connected
 	virtual bool EnableInteractWidgets() const override;
@@ -184,10 +247,16 @@ public:
 	virtual void Validate(Metasound::Editor::FGraphNodeValidationResult& OutResult) override;
 
 protected:
+	// Breadcrumb used if associated FrontendNode cannot be found or has been unlinked
+	UPROPERTY()
+	FMetasoundEditorGraphVertexNodeBreadcrumb Breadcrumb;
+
 	virtual FLinearColor GetNodeTitleColor() const override;
 	virtual FSlateIcon GetNodeTitleIcon() const override;
 	virtual void SetNodeID(FGuid InNodeID) override;
 
+	// Friended to enable mutation of Frontend NodeID & direct breadcrumb access
+	friend struct Metasound::Editor::FDocumentClipboardUtils;
 	friend class Metasound::Editor::FGraphBuilder;
 };
 
@@ -198,25 +267,29 @@ class UMetasoundEditorGraphExternalNode : public UMetasoundEditorGraphNode
 
 protected:
 	UPROPERTY()
+	FMetasoundEditorGraphNodeBreadcrumb Breadcrumb;
+
+	UPROPERTY(meta = (Deprecated = "5.4", DeprecationMessage = "Use Breadcrumb value when manipulating clipboard data or validating. Otherwise, look-up frontend node's associated class directly"))
 	FMetasoundFrontendClassName ClassName;
 
 	UPROPERTY()
 	FGuid NodeID;
 
-	// Whether or not the referenced class is natively defined
-	// (false if defined in another asset). Cached from node
-	// implementation for fast access when validated.
-	UPROPERTY()
+	UPROPERTY(meta = (Deprecated = "5.4", DeprecationMessage = "Use Breadcrumb value when manipulating clipboard data or validating, Otherwise, look-up frontend node's associated class directly"))
 	bool bIsClassNative = true;
 
 public:
-	virtual FMetasoundFrontendClassName GetClassName() const override { return ClassName; }
+	UE_DEPRECATED(5.4, "Use the Frontend node or Breadcrumb directly to get the class name.")
+	virtual FMetasoundFrontendClassName GetClassName() const override { return Breadcrumb.ClassName; }
+
+	virtual const FMetasoundEditorGraphNodeBreadcrumb& GetBreadcrumb() const;
 	virtual FGuid GetNodeID() const override { return NodeID; }
 	virtual FLinearColor GetNodeTitleColor() const override;
 	virtual FSlateIcon GetNodeTitleIcon() const override;
 	virtual bool ShouldDrawNodeAsControlPointOnly(int32& OutInputPinIndex, int32& OutOutputPinIndex) const override;
 
 	virtual void ReconstructNode() override;
+	virtual void CacheBreadcrumb() override;
 	virtual void CacheTitle() override;
 	virtual void GetPinHoverText(const UEdGraphPin& Pin, FString& OutHoverText) const override;
 
@@ -228,11 +301,8 @@ public:
 
 
 protected:
-	virtual void SetNodeID(FGuid InNodeID) override
-	{
-		NodeID = InNodeID;
-	}
-
+	// Friended to enable mutation of Frontend NodeID & direct breadcrumb access
+	friend struct Metasound::Editor::FDocumentClipboardUtils;
 	friend class Metasound::Editor::FGraphBuilder;
 };
 
@@ -263,8 +333,13 @@ public:
 	// Variables do not have titles to distinguish more visually from vertex types
 	virtual void CacheTitle() override { }
 
+	virtual const FMetasoundEditorGraphNodeBreadcrumb& GetBreadcrumb() const override;
+
+	virtual void CacheBreadcrumb() override;
 	virtual UMetasoundEditorGraphMember* GetMember() const override;
 	virtual bool EnableInteractWidgets() const override;
+
+	
 	virtual FMetasoundFrontendClassName GetClassName() const override;
 	virtual FGuid GetNodeID() const override;
 	virtual FName GetCornerIcon() const override;
@@ -273,9 +348,15 @@ public:
 	virtual EMetasoundFrontendClassType GetClassType() const;
 
 protected:
+	// Breadcrumb used if associated FrontendNode cannot be found or has been unlinked
+	UPROPERTY()
+	FMetasoundEditorGraphMemberNodeBreadcrumb Breadcrumb;
+
 	virtual FLinearColor GetNodeTitleColor() const override;
 	virtual FSlateIcon GetNodeTitleIcon() const override;
 	virtual void SetNodeID(FGuid InNodeID) override;
 
+	// Friended to enable mutation of Frontend NodeID & direct breadcrumb access
+	friend struct Metasound::Editor::FDocumentClipboardUtils;
 	friend class Metasound::Editor::FGraphBuilder;
 };
diff --git a/Engine/Plugins/Runtime/Metasound/Source/MetasoundEditor/Public/MetasoundEditorSubsystem.h b/Engine/Plugins/Runtime/Metasound/Source/MetasoundEditor/Public/MetasoundEditorSubsystem.h
--- a/Engine/Plugins/Runtime/Metasound/Source/MetasoundEditor/Public/MetasoundEditorSubsystem.h
+++ b/Engine/Plugins/Runtime/Metasound/Source/MetasoundEditor/Public/MetasoundEditorSubsystem.h
@@ -12,6 +12,8 @@
 
 // Forward Declarations
 class UMetaSoundBuilderBase;
+class UMetasoundEditorGraphMember;
+class UMetasoundEditorGraphMemberDefaultLiteral;
 
 /** The subsystem in charge of editor MetaSound functionality */
 UCLASS()
@@ -20,6 +22,15 @@ class METASOUNDEDITOR_API UMetaSoundEditorSubsystem : public UEditorSubsystem
 	GENERATED_BODY()
 
 public:
+	// Binds literal editor Metadata to the given member.  If the literal already exists, adds literal
+	// reference to given member (asserts that existing literal is of similar subclass provided).  If 
+	// it does not exist, or an optional template object is provided, metadata is generated then bound.
+	// Returns true if new literal metadata was generated, false if not. Asserts if bind failed.
+	bool BindMemberMetadata(
+		UMetasoundEditorGraphMember& InMember,
+		TSubclassOf<UMetasoundEditorGraphMemberDefaultLiteral> LiteralClass,
+		UMetasoundEditorGraphMemberDefaultLiteral* TemplateObject = nullptr);
+
 	// Build the given builder to a MetaSound asset
 	// @param Author - Sets the author on the given builder's document.
 	// @param AssetName - Name of the asset to build.
@@ -42,7 +53,7 @@ public:
 		UPARAM(DisplayName = "Node Handle") const FMetaSoundNodeHandle& InNode,
 		UPARAM(DisplayName = "Location") const FVector2D& InLocation,
 		EMetaSoundBuilderResult& OutResult);
-	
+
 	// Initialize the UObject asset, with an optional MetaSound to be referenced if the asset is a preset
 	void InitAsset(UObject& InNewMetaSound, UObject* InReferencedMetaSound = nullptr);
 
diff --git a/Engine/Plugins/Runtime/Metasound/Source/MetasoundEditor/Public/MetasoundFactory.h b/Engine/Plugins/Runtime/Metasound/Source/MetasoundEditor/Public/MetasoundFactory.h
--- a/Engine/Plugins/Runtime/Metasound/Source/MetasoundEditor/Public/MetasoundFactory.h
+++ b/Engine/Plugins/Runtime/Metasound/Source/MetasoundEditor/Public/MetasoundFactory.h
@@ -20,9 +20,6 @@ class UMetaSoundBaseFactory : public UFactory
 	// using the provided MetaSound as a Reference
 	UPROPERTY(EditAnywhere, Transient, Category = Factory)
 	TObjectPtr<UObject> ReferencedMetaSoundObject;
-	
-	UE_DEPRECATED(5.3, "Use UMetaSoundEditorSubsystem::InitAsset instead.")
-	static void InitAsset(UObject& InNewMetaSound, UObject* InReferencedMetaSound = nullptr);
 };
 
 UCLASS(hidecategories=Object, MinimalAPI)
diff --git a/Engine/Plugins/Runtime/Metasound/Source/MetasoundEngine/Private/Metasound.cpp b/Engine/Plugins/Runtime/Metasound/Source/MetasoundEngine/Private/Metasound.cpp
--- a/Engine/Plugins/Runtime/Metasound/Source/MetasoundEngine/Private/Metasound.cpp
+++ b/Engine/Plugins/Runtime/Metasound/Source/MetasoundEngine/Private/Metasound.cpp
@@ -11,6 +11,7 @@
 #include "MetasoundEngineEnvironment.h"
 #include "MetasoundEnvironment.h"
 #include "MetasoundFrontendController.h"
+#include "MetasoundFrontendDocument.h"
 #include "MetasoundFrontendQuery.h"
 #include "MetasoundFrontendQuerySteps.h"
 #include "MetasoundFrontendRegistries.h"
@@ -52,7 +53,6 @@ int32 UMetasoundEditorGraphBase::GetHighestMessageSeverity() const
 	return HighestMessageSeverity;
 }
 
-
 UMetaSoundPatch::UMetaSoundPatch(const FObjectInitializer& ObjectInitializer)
 	: Super(ObjectInitializer)
 	, FMetasoundAssetBase()
@@ -135,26 +135,37 @@ void UMetaSoundPatch::Serialize(FArchive& InArchive)
 }
 
 #if WITH_EDITORONLY_DATA
+void UMetaSoundPatch::MigrateEditorGraph(FMetaSoundFrontendDocumentBuilder& OutBuilder)
+{
+	PRAGMA_DISABLE_DEPRECATION_WARNINGS
+	if (Graph)
+	{
+		Graph->MigrateEditorDocumentData(OutBuilder);
+		Graph = nullptr;
+	}
+	PRAGMA_ENABLE_DEPRECATION_WARNINGS
+}
+
 UEdGraph* UMetaSoundPatch::GetGraph()
 {
-	return Graph;
+	return EditorGraph;
 }
 
 const UEdGraph* UMetaSoundPatch::GetGraph() const
 {
-	return Graph;
+	return EditorGraph;
 }
 
 UEdGraph& UMetaSoundPatch::GetGraphChecked()
 {
-	check(Graph);
-	return *Graph;
+	check(EditorGraph);
+	return *EditorGraph;
 }
 
 const UEdGraph& UMetaSoundPatch::GetGraphChecked() const
 {
-	check(Graph);
-	return *Graph;
+	check(EditorGraph);
+	return *EditorGraph;
 }
 FText UMetaSoundPatch::GetDisplayName() const
 {
diff --git a/Engine/Plugins/Runtime/Metasound/Source/MetasoundEngine/Private/MetasoundAssetSubsystem.cpp b/Engine/Plugins/Runtime/Metasound/Source/MetasoundEngine/Private/MetasoundAssetSubsystem.cpp
--- a/Engine/Plugins/Runtime/Metasound/Source/MetasoundEngine/Private/MetasoundAssetSubsystem.cpp
+++ b/Engine/Plugins/Runtime/Metasound/Source/MetasoundEngine/Private/MetasoundAssetSubsystem.cpp
@@ -144,7 +144,7 @@ void UMetaSoundAssetSubsystem::AddAssetReferences(FMetasoundAssetBase& InAssetBa
 	using namespace Metasound;
 	using namespace Metasound::Frontend;
 
-	const FMetasoundFrontendDocument& Document = InAssetBase.GetDocumentChecked();
+	const FMetasoundFrontendDocument& Document = InAssetBase.GetConstDocumentChecked();
 	const FNodeRegistryKey AssetClassKey = FNodeRegistryKey(Document.RootGraph);
 
 	if (!ContainsKey(AssetClassKey))
@@ -179,7 +179,7 @@ void UMetaSoundAssetSubsystem::AddAssetReferences(FMetasoundAssetBase& InAssetBa
 	{
 		if (Asset)
 		{
-			const FMetasoundFrontendDocument& RefDocument = Asset->GetDocumentChecked();
+			const FMetasoundFrontendDocument& RefDocument = Asset->GetConstDocumentChecked();
 			const FNodeRegistryKey ClassKey = FNodeRegistryKey(RefDocument.RootGraph);
 			if (!ContainsKey(ClassKey))
 			{
@@ -214,7 +214,7 @@ Metasound::Frontend::FNodeRegistryKey UMetaSoundAssetSubsystem::AddOrUpdateAsset
 	const FMetasoundAssetBase* MetaSoundAsset = Metasound::IMetasoundUObjectRegistry::Get().GetObjectAsAssetBase(&InObject);
 	check(MetaSoundAsset);
 
-	const FMetasoundFrontendDocument& Document = MetaSoundAsset->GetDocumentChecked();
+	const FMetasoundFrontendDocument& Document = MetaSoundAsset->GetConstDocumentChecked();
 	const FNodeRegistryKey RegistryKey = FNodeRegistryKey(Document.RootGraph);
 
 	if (RegistryKey.IsValid())
@@ -275,6 +275,11 @@ Metasound::Frontend::FNodeRegistryKey UMetaSoundAssetSubsystem::AddOrUpdateAsset
 	return FNodeRegistryKey::GetInvalid();
 }
 
+FMetaSoundFrontendDocumentBuilder& UMetaSoundAssetSubsystem::AttachDocumentBuilderChecked(UObject& InObject) const
+{
+	return UMetaSoundBuilderSubsystem::GetChecked().AttachBuilderToAssetChecked(InObject).Builder;
+}
+
 bool UMetaSoundAssetSubsystem::CanAutoUpdate(const FMetasoundFrontendClassName& InClassName) const
 {
 	const UMetaSoundSettings* Settings = GetDefault<UMetaSoundSettings>();
@@ -344,7 +349,7 @@ TSet<UMetaSoundAssetSubsystem::FAssetInfo> UMetaSoundAssetSubsystem::GetReferenc
 	using namespace Metasound::Frontend;
 
 	TSet<FAssetInfo> OutAssetInfos;
-	const FMetasoundFrontendDocument& Document = InAssetBase.GetDocumentChecked();
+	const FMetasoundFrontendDocument& Document = InAssetBase.GetConstDocumentChecked();
 	for (const FMetasoundFrontendClass& Class : Document.Dependencies)
 	{
 		const FNodeRegistryKey Key = FNodeRegistryKey(Class.Metadata);
@@ -622,12 +627,10 @@ void UMetaSoundAssetSubsystem::RemoveAsset(const UObject& InObject)
 	if (const FMetasoundAssetBase* MetaSoundAsset = Metasound::IMetasoundUObjectRegistry::Get().GetObjectAsAssetBase(&InObject))
 	{
 		const FSoftObjectPath ObjectPath(&InObject);
-		const FMetasoundFrontendDocument& Document = MetaSoundAsset->GetDocumentChecked();
+		const FMetasoundFrontendDocument& Document = MetaSoundAsset->GetConstDocumentChecked();
 		FNodeRegistryKey RegistryKey = FNodeRegistryKey(Document.RootGraph);
-		if (AssetSubsystemPrivate::RemoveIfExactMatch(PathMap, RegistryKey, ObjectPath))
-		{
-			UMetaSoundBuilderSubsystem::GetChecked().DetachBuilderFromAsset(Document.RootGraph.Metadata.GetClassName());
-		}
+		AssetSubsystemPrivate::RemoveIfExactMatch(PathMap, RegistryKey, ObjectPath);
+		UMetaSoundBuilderSubsystem::GetChecked().DetachBuilderFromAsset(Document.RootGraph.Metadata.GetClassName());
 	}
 }
 
@@ -641,10 +644,8 @@ void UMetaSoundAssetSubsystem::RemoveAsset(const FAssetData& InAssetData)
 	{
 		const FNodeRegistryKey RegistryKey(ClassInfo);
 		const FSoftObjectPath ObjectPath = InAssetData.GetSoftObjectPath();
-		if (AssetSubsystemPrivate::RemoveIfExactMatch(PathMap, RegistryKey, ObjectPath))
-		{
-			UMetaSoundBuilderSubsystem::GetChecked().DetachBuilderFromAsset(ClassInfo.ClassName);
-		}
+		AssetSubsystemPrivate::RemoveIfExactMatch(PathMap, RegistryKey, ObjectPath);
+		UMetaSoundBuilderSubsystem::GetChecked().DetachBuilderFromAsset(ClassInfo.ClassName);
 	}
 }
 
diff --git a/Engine/Plugins/Runtime/Metasound/Source/MetasoundEngine/Private/MetasoundBuilderSubsystem.cpp b/Engine/Plugins/Runtime/Metasound/Source/MetasoundEngine/Private/MetasoundBuilderSubsystem.cpp
--- a/Engine/Plugins/Runtime/Metasound/Source/MetasoundEngine/Private/MetasoundBuilderSubsystem.cpp
+++ b/Engine/Plugins/Runtime/Metasound/Source/MetasoundEngine/Private/MetasoundBuilderSubsystem.cpp
@@ -26,6 +26,7 @@
 #include "MetasoundTrace.h"
 #include "MetasoundUObjectRegistry.h"
 #include "MetasoundVertex.h"
+#include "NodeTemplates/MetasoundFrontendNodeTemplateInput.h"
 #include "PerPlatformProperties.h"
 
 #include UE_INLINE_GENERATED_CPP_BY_NAME(MetasoundBuilderSubsystem)
@@ -71,8 +72,6 @@ void UMetaSoundBuilderBase::BeginDestroy()
 {
 	if (UMetaSoundBuilderSubsystem* Subsystem = UMetaSoundBuilderSubsystem::Get())
 	{
-		const FMetasoundFrontendDocument& Document = static_cast<const FMetaSoundFrontendDocumentBuilder&>(Builder).GetDocument();
-		const FMetasoundFrontendClassName& ClassName = Document.RootGraph.Metadata.GetClassName();
 		Subsystem->TransientBuilders.Remove(ClassName);
 	}
 
@@ -80,7 +79,6 @@ void UMetaSoundBuilderBase::BeginDestroy()
 	// often holds a TScriptInterface<IMetaSoundDocumentInterface> of
 	// a UPROPERTY that lives on this or derived objects. 
 	Builder.FinishBuilding();
-
 	Super::BeginDestroy();
 }
 
@@ -230,18 +228,18 @@ FMetaSoundNodeHandle UMetaSoundBuilderBase::AddNode(const TScriptInterface<IMeta
 	return NewHandle;
 }
 
-FMetaSoundNodeHandle UMetaSoundBuilderBase::AddNodeByClassName(const FMetasoundFrontendClassName& ClassName, int32 MajorVersion, EMetaSoundBuilderResult& OutResult)
+FMetaSoundNodeHandle UMetaSoundBuilderBase::AddNodeByClassName(const FMetasoundFrontendClassName& InClassName, int32 MajorVersion, EMetaSoundBuilderResult& OutResult)
 {
-	return AddNodeByClassName(ClassName, OutResult, MajorVersion);
+	return AddNodeByClassName(InClassName, OutResult, MajorVersion);
 }
 
-FMetaSoundNodeHandle UMetaSoundBuilderBase::AddNodeByClassName(const FMetasoundFrontendClassName& ClassName, EMetaSoundBuilderResult& OutResult, int32 MajorVersion)
+FMetaSoundNodeHandle UMetaSoundBuilderBase::AddNodeByClassName(const FMetasoundFrontendClassName& InClassName, EMetaSoundBuilderResult& OutResult, int32 MajorVersion)
 {
 	using namespace Metasound;
 	using namespace Metasound::Frontend;
 
 	FMetaSoundNodeHandle NewHandle;
-	if (const FMetasoundFrontendNode* NewNode = Builder.AddNodeByClassName(ClassName, MajorVersion))
+	if (const FMetasoundFrontendNode* NewNode = Builder.AddNodeByClassName(InClassName, MajorVersion))
 	{
 		NewHandle.NodeID = NewNode->GetID();
 	}
@@ -250,6 +248,31 @@ FMetaSoundNodeHandle UMetaSoundBuilderBase::AddNodeByClassName(const FMetasoundF
 	return NewHandle;
 }
 
+void UMetaSoundBuilderBase::BuildInternal(TScriptInterface<IMetaSoundDocumentInterface> NewMetaSound, const FMetasoundFrontendClassName* InDocClassName) const
+{
+	using namespace Metasound::Frontend;
+
+	// If using existing class name, ensure that a builder does not exist for it to
+	// avoid build active flag conflation with the locally generated frontend builder below.
+	if (InDocClassName)
+	{
+		UMetaSoundBuilderSubsystem::GetChecked().DetachBuilderFromAsset(*InDocClassName);
+	}
+
+	FMetaSoundFrontendDocumentBuilder NewDocBuilder(NewMetaSound);
+
+	constexpr bool bResetVersion = false;
+	NewDocBuilder.InitDocument(&GetConstBuilder().GetDocument(), InDocClassName, bResetVersion);
+	NewMetaSound->ConformObjectToDocument();
+}
+
+#if WITH_EDITOR
+bool UMetaSoundBuilderBase::ClearMemberMetadata(const FGuid& InMemberID)
+{
+	return Builder.ClearMemberMetadata(InMemberID);
+}
+#endif // WITH_EDITOR
+
 bool UMetaSoundBuilderBase::ContainsNode(const FMetaSoundNodeHandle& NodeHandle) const
 {
 	return Builder.ContainsNode(NodeHandle.NodeID);
@@ -420,6 +443,8 @@ void UMetaSoundBuilderBase::ConvertFromPreset(EMetaSoundBuilderResult& OutResult
 
 void UMetaSoundBuilderBase::ConvertToPreset(const TScriptInterface<IMetaSoundDocumentInterface>& ReferencedNodeClass, EMetaSoundBuilderResult& OutResult)
 {
+	using namespace Metasound::Frontend;
+
 	const IMetaSoundDocumentInterface* ReferencedInterface = ReferencedNodeClass.GetInterface();
 	if (!ReferencedInterface)
 	{
@@ -427,7 +452,7 @@ void UMetaSoundBuilderBase::ConvertToPreset(const TScriptInterface<IMetaSoundDoc
 		return;
 	}
 	
-	// Ensure the referenced node class isn't transient 
+	// Ensure the referenced node class isn't transient
 	if (Cast<UMetaSoundBuilderDocument>(ReferencedInterface))
 	{
 		UE_LOG(LogMetaSound, Warning, TEXT("Transient document builders cannot be referenced when converting builder '%s' to a preset. Build the referenced node class an asset first or use an existing asset instead"), *GetName());
@@ -452,14 +477,12 @@ void UMetaSoundBuilderBase::ConvertToPreset(const TScriptInterface<IMetaSoundDoc
 	}
 
 	const FMetasoundFrontendDocument& ReferencedDocument = ReferencedInterface->GetConstDocument();
-	if (Builder.ConvertToPreset(ReferencedDocument))
-	{
-		OutResult = EMetaSoundBuilderResult::Succeeded;
-	}
-	else
-	{
-		OutResult = EMetaSoundBuilderResult::Failed;
-	}
+
+	TSharedRef<FDocumentModifyDelegates> DocumentDelegates = MakeShared<FDocumentModifyDelegates>();
+	InitDelegates(*DocumentDelegates);
+	const bool bConvertToPreset = Builder.ConvertToPreset(ReferencedDocument, DocumentDelegates);
+	ConformObjectToDocument();
+	OutResult = bConvertToPreset ? EMetaSoundBuilderResult::Succeeded : EMetaSoundBuilderResult::Failed;
 }
 
 void UMetaSoundBuilderBase::DisconnectNodes(const FMetaSoundBuilderNodeOutputHandle& NodeOutputHandle, const FMetaSoundBuilderNodeInputHandle& NodeInputHandle, EMetaSoundBuilderResult& OutResult)
@@ -693,6 +716,13 @@ FMetaSoundNodeHandle UMetaSoundBuilderBase::FindGraphOutputNode(FName OutputName
 	return { };
 }
 
+#if WITH_EDITOR
+UMetaSoundFrontendMemberMetadata* UMetaSoundBuilderBase::FindMemberMetadata(const FGuid& InMemberID)
+{
+	return Builder.FindMemberMetadata(InMemberID);
+}
+#endif // WITH_EDITOR
+
 UMetaSoundBuilderDocument* UMetaSoundBuilderBase::CreateTransientDocumentObject() const
 {
 	return &UMetaSoundBuilderDocument::Create(GetBuilderUClass());
@@ -716,12 +746,14 @@ UObject* UMetaSoundBuilderBase::GetReferencedPresetAsset() const
 	auto FindExternalNode = [this](const FMetasoundFrontendNode& Node) 
 	{
 		const FMetasoundFrontendClass* Class = Builder.FindDependency(Node.ClassID);
+		check(Class);
 		return Class->Metadata.GetType() == EMetasoundFrontendClassType::External;
 	};
 	const FMetasoundFrontendNode* Node = Builder.GetDocument().RootGraph.Graph.Nodes.FindByPredicate(FindExternalNode);
 	if (Node != nullptr)
 	{
 		const FMetasoundFrontendClass* NodeClass = Builder.FindDependency(Node->ClassID);
+		check(NodeClass);
 		const FNodeRegistryKey NodeClassRegistryKey = FNodeRegistryKey(NodeClass->Metadata);
 		if (FMetasoundAssetBase* Asset = IMetaSoundAssetManager::GetChecked().TryLoadAssetFromKey(NodeClassRegistryKey))
 		{
@@ -744,14 +776,28 @@ void UMetaSoundBuilderBase::InitNodeLocations()
 	Builder.InitNodeLocations();
 }
 
+#if WITH_EDITOR
+void UMetaSoundBuilderBase::InjectInputTemplateNodes(bool bForceNodeCreation, EMetaSoundBuilderResult& OutResult)
+{
+	using namespace Metasound::Frontend;
+	const bool bTemplateNodesInjected = FInputNodeTemplate::GetChecked().Inject(Builder, bForceNodeCreation);
+	OutResult = bTemplateNodesInjected ? EMetaSoundBuilderResult::Succeeded : EMetaSoundBuilderResult::Failed;
+}
+#endif // WITH_EDITOR
+
 bool UMetaSoundBuilderBase::InterfaceIsDeclared(FName InterfaceName) const
 {
 	return Builder.IsInterfaceDeclared(InterfaceName);
 }
 
-void UMetaSoundBuilderBase::InvalidateCache()
+void UMetaSoundBuilderBase::InvalidateCache(bool bPrimeCache)
 {
-	Builder.InvalidateCache();
+	using namespace Metasound::Frontend;
+
+	TSharedRef<FDocumentModifyDelegates> DocumentDelegates = MakeShared<FDocumentModifyDelegates>();
+	InitDelegates(*DocumentDelegates);
+
+	Builder.Reload(DocumentDelegates, bPrimeCache);
 }
 
 bool UMetaSoundBuilderBase::IsPreset() const
@@ -814,9 +860,18 @@ void UMetaSoundBuilderBase::RegisterGraphIfOutstandingTransactions(UObject& InMe
 	MetaSoundAsset->RegisterGraphWithFrontend(Options);
 }
 
+void UMetaSoundBuilderBase::Reload(bool bPrimeCache)
+{
+	using namespace Metasound::Frontend;
+
+	TSharedRef<FDocumentModifyDelegates> DocumentDelegates = MakeShared<FDocumentModifyDelegates>();
+	InitDelegates(*DocumentDelegates);
+	Builder.Reload(DocumentDelegates, bPrimeCache);
+}
+
 void UMetaSoundBuilderBase::ReloadCache(bool bPrimeCache)
 {
-	Builder.ReloadCache();
+	Reload(bPrimeCache);
 }
 
 #if WITH_EDITOR
@@ -844,9 +899,13 @@ void UMetaSoundBuilderBase::RemoveInterface(FName InterfaceName, EMetaSoundBuild
 	OutResult = bInterfaceRemoved ? EMetaSoundBuilderResult::Succeeded : EMetaSoundBuilderResult::Failed;
 }
 
-void UMetaSoundBuilderBase::RemoveNode(const FMetaSoundNodeHandle& NodeHandle, EMetaSoundBuilderResult& OutResult)
+void UMetaSoundBuilderBase::RemoveNode(const FMetaSoundNodeHandle& NodeHandle, EMetaSoundBuilderResult& OutResult, bool bRemoveUnusedDependencies)
 {
 	const bool bNodeRemoved = Builder.RemoveNode(NodeHandle.NodeID);
+	if (bNodeRemoved && bRemoveUnusedDependencies)
+	{
+		Builder.RemoveUnusedDependencies();
+	}
 	OutResult = bNodeRemoved ? EMetaSoundBuilderResult::Succeeded : EMetaSoundBuilderResult::Failed;
 }
 
@@ -856,6 +915,11 @@ void UMetaSoundBuilderBase::RemoveNodeInputDefault(const FMetaSoundBuilderNodeIn
 	OutResult = bInputDefaultRemoved ? EMetaSoundBuilderResult::Succeeded : EMetaSoundBuilderResult::Failed;
 }
 
+void UMetaSoundBuilderBase::RemoveUnusedDependencies()
+{
+	Builder.RemoveUnusedDependencies();
+}
+
 void UMetaSoundBuilderBase::RenameRootGraphClass(const FMetasoundFrontendClassName& InName)
 {
 	Builder.RenameRootGraphClass(InName);
@@ -868,11 +932,42 @@ void UMetaSoundBuilderBase::SetAuthor(const FString& InAuthor)
 }
 #endif // WITH_EDITOR
 
+void UMetaSoundBuilderBase::SetGraphInputAccessType(FName InputName, EMetasoundFrontendVertexAccessType AccessType, EMetaSoundBuilderResult& OutResult)
+{
+	const bool bSet = Builder.SetGraphInputAccessType(InputName, AccessType);
+	OutResult = bSet ? EMetaSoundBuilderResult::Succeeded : EMetaSoundBuilderResult::Failed;
+}
+
+void UMetaSoundBuilderBase::SetGraphInputDataType(FName InputName, FName DataType, EMetaSoundBuilderResult& OutResult)
+{
+	const bool bSet = Builder.SetGraphInputDataType(InputName, DataType);
+	OutResult = bSet ? EMetaSoundBuilderResult::Succeeded : EMetaSoundBuilderResult::Failed;
+}
+
 void UMetaSoundBuilderBase::SetGraphInputDefault(FName InputName, const FMetasoundFrontendLiteral& Literal, EMetaSoundBuilderResult& OutResult)
 {
-	const bool bInputDefaultSet = Builder.SetGraphInputDefault(InputName, Literal);
-	OutResult = bInputDefaultSet ? EMetaSoundBuilderResult::Succeeded : EMetaSoundBuilderResult::Failed;
+	const bool bSet = Builder.SetGraphInputDefault(InputName, Literal);
+	OutResult = bSet ? EMetaSoundBuilderResult::Succeeded : EMetaSoundBuilderResult::Failed;
+}
+
+void UMetaSoundBuilderBase::SetGraphOutputAccessType(FName OutputName, EMetasoundFrontendVertexAccessType AccessType, EMetaSoundBuilderResult& OutResult)
+{
+	const bool bSet = Builder.SetGraphOutputAccessType(OutputName, AccessType);
+	OutResult = bSet ? EMetaSoundBuilderResult::Succeeded : EMetaSoundBuilderResult::Failed;
+}
+
+void UMetaSoundBuilderBase::SetGraphOutputDataType(FName OutputName, FName DataType, EMetaSoundBuilderResult& OutResult)
+{
+	const bool bSet = Builder.SetGraphOutputDataType(OutputName, DataType);
+	OutResult = bSet ? EMetaSoundBuilderResult::Succeeded : EMetaSoundBuilderResult::Failed;
+}
+
+#if WITH_EDITOR
+void UMetaSoundBuilderBase::SetMemberMetadata(UMetaSoundFrontendMemberMetadata& NewMetadata)
+{
+	Builder.SetMemberMetadata(NewMetadata);
 }
+#endif // WITH_EDITOR
 
 void UMetaSoundBuilderBase::SetNodeInputDefault(const FMetaSoundBuilderNodeInputHandle& InputHandle, const FMetasoundFrontendLiteral& Literal, EMetaSoundBuilderResult& OutResult)
 {
@@ -1024,6 +1119,11 @@ void UMetaSoundBuilderBase::UpdateDependencyClassNames(const TMap<FMetasoundFron
 	Builder.UpdateDependencyClassNames(OldToNewReferencedClassNames);
 }
 
+bool UMetaSoundPatchBuilder::ConformObjectToDocument()
+{
+	return false;
+}
+
 void UMetaSoundPatchBuilder::CreateTransientBuilder()
 {
 	using namespace Metasound::Frontend;
@@ -1106,7 +1206,7 @@ void UMetaSoundSourceBuilder::Audition(UObject* Parent, UAudioComponent* AudioCo
 
 	// Must be called post register as register ensures cached runtime data passed to transactor is up-to-date
 	MetaSoundSource.SetDynamicGeneratorEnabled(MetaSoundSource.GetAssetPathChecked(), bLiveUpdatesEnabled);
-	MetaSoundSource.ConformObjectDataToInterfaces();
+	MetaSoundSource.ConformObjectToDocument();
 
 	AudioComponent->SetSound(&MetaSoundSource);
 
@@ -1156,6 +1256,12 @@ TScriptInterface<IMetaSoundDocumentInterface> UMetaSoundSourceBuilder::Build(UOb
 	return &BuildInternal<UMetaSoundSource>(Parent, InBuilderOptions);
 }
 
+bool UMetaSoundSourceBuilder::ConformObjectToDocument()
+{
+	UMetaSoundSource& Source = GetMetaSoundSource();
+	return Source.ConformObjectToDocument();
+}
+
 void UMetaSoundSourceBuilder::CreateTransientBuilder()
 {
 	using namespace Metasound::Frontend;
@@ -1960,7 +2066,7 @@ void UMetaSoundBuilderSubsystem::InvalidateDocumentCache(const FMetasoundFronten
 	TWeakObjectPtr<UMetaSoundBuilderBase> BuilderPtr = AssetBuilders.FindRef(InClassName);
 	if (BuilderPtr.IsValid())
 	{
-		BuilderPtr->InvalidateCache();
+		BuilderPtr->Reload();
 	}
 }
 
@@ -1975,16 +2081,16 @@ UMetaSoundBuilderBase* UMetaSoundBuilderSubsystem::FindBuilderOfDocument(TScript
 	if (const UObject* MetaSoundObject = InMetaSound.GetObject())
 	{
 		const FMetasoundFrontendDocument& Document = InMetaSound->GetConstDocument();
-		const FMetasoundFrontendClassName& ClassName = Document.RootGraph.Metadata.GetClassName();
+		const FMetasoundFrontendClassName& InClassName = Document.RootGraph.Metadata.GetClassName();
 
-		Builder = AssetBuilders.FindRef(ClassName);
+		Builder = AssetBuilders.FindRef(InClassName);
 		if (Builder.IsValid())
 		{
 			ensureAlwaysMsgf(MetaSoundObject->IsAsset(), TEXT("MetaSound is asset but Builder was registered with Subsystem as transient"));
 		}
 		else
 		{
-			Builder = TransientBuilders.FindRef(ClassName);
+			Builder = TransientBuilders.FindRef(InClassName);
 			if (Builder.IsValid())
 			{
 				ensureAlwaysMsgf(!MetaSoundObject->IsAsset(), TEXT("MetaSound is transient object but Builder was registered with Subsystem as asset"));
@@ -2043,7 +2149,7 @@ void UMetaSoundBuilderSubsystem::PostBuilderAssetTransaction(const FMetasoundFro
 	TWeakObjectPtr<UMetaSoundBuilderBase> BuilderPtr = AssetBuilders.FindRef(InClassName);
 	if (UMetaSoundBuilderBase* Builder = BuilderPtr.Get())
 	{
-		Builder->ReloadCache();
+		Builder->Reload();
 	}
 }
 #endif // WITH_EDITOR
diff --git a/Engine/Plugins/Runtime/Metasound/Source/MetasoundEngine/Private/MetasoundEngineAsset.h b/Engine/Plugins/Runtime/Metasound/Source/MetasoundEngine/Private/MetasoundEngineAsset.h
--- a/Engine/Plugins/Runtime/Metasound/Source/MetasoundEngine/Private/MetasoundEngineAsset.h
+++ b/Engine/Plugins/Runtime/Metasound/Source/MetasoundEngine/Private/MetasoundEngineAsset.h
@@ -31,7 +31,7 @@ namespace Metasound
 		{
 			InMetaSound.GetModifyContext().SetForceRefreshViews();
 
-			const FMetasoundFrontendDocument& Document = InMetaSound.GetDocumentChecked();
+			const FMetasoundFrontendDocument& Document = InMetaSound.GetConstDocumentChecked();
 			const FMetasoundFrontendClassName& ClassName = Document.RootGraph.Metadata.GetClassName();
 			UMetaSoundBuilderSubsystem::GetChecked().PostBuilderAssetTransaction(ClassName);
 
@@ -95,44 +95,38 @@ namespace Metasound
 			return ReferencedAssets;
 		}
 
-		template <typename TMetaSoundObject>
-		static void PreSaveAsset(TMetaSoundObject& InMetaSound, FObjectPreSaveContext InSaveContext)
+		static void PreSaveAsset(FMetasoundAssetBase& InMetaSound, FObjectPreSaveContext InSaveContext)
 		{
 #if WITH_EDITORONLY_DATA
 			using namespace Frontend;
 
-			// Do not call asset manager on CDO objects which may be loaded before asset 
+			// Do not call asset manager on CDO objects which may be loaded before asset
 			// manager is set.
 			if (IMetaSoundAssetManager* AssetManager = IMetaSoundAssetManager::Get())
 			{
 				AssetManager->WaitUntilAsyncLoadReferencedAssetsComplete(InMetaSound);
 			}
 
-			if (UMetasoundEditorGraphBase* MetaSoundGraph = Cast<UMetasoundEditorGraphBase>(InMetaSound.GetGraph()))
+			if (InSaveContext.IsCooking() || IsRunningCommandlet())
 			{
-				if (InSaveContext.IsCooking() || IsRunningCommandlet())
-				{
-					// Use deterministic ID generation so more can be done at cook rather than runtime
-					if (MetaSoundEnableCookDeterministicIDGeneration != 0)
-					{
-						{
-							constexpr bool bIsDeterministic = true;
-							FDocumentIDGenerator::FScopeDeterminism DeterminismScope = FDocumentIDGenerator::FScopeDeterminism(bIsDeterministic);
-							InMetaSound.CookMetaSound();
-						}
-					}
-				}
- 				else if (FApp::CanEverRenderAudio())
+				constexpr bool bIsDeterministic = true;
+				FDocumentIDGenerator::FScopeDeterminism DeterminismScope = FDocumentIDGenerator::FScopeDeterminism(bIsDeterministic);
+				InMetaSound.CookMetaSound();
+			}
+ 			else if (FApp::CanEverRenderAudio())
+			{
+				if (UMetasoundEditorGraphBase* MetaSoundGraph = Cast<UMetasoundEditorGraphBase>(InMetaSound.GetGraph()))
 				{
+					// Uses graph flavor of register with frontend to update editor systems/asset editors in case editor is enabled.
 					MetaSoundGraph->RegisterGraphWithFrontend();
-					MetaSoundGraph->GetModifyContext().SetForceRefreshViews();
-				}
-				else
-				{
-					UE_LOG(LogMetaSound, Warning, TEXT("PreSaveAsset for MetaSound: (%s) is doing nothing because InSaveContext.IsCooking, IsRunningCommandlet, and FApp::CanEverRenderAudio were all false")
-						, *InMetaSound.GetPathName());
+					InMetaSound.GetModifyContext().SetForceRefreshViews();
 				}
 			}
+			else
+			{
+				UE_LOG(LogMetaSound, Warning, TEXT("PreSaveAsset for MetaSound: (%s) is doing nothing because InSaveContext.IsCooking, IsRunningCommandlet, and FApp::CanEverRenderAudio were all false")
+					, *InMetaSound.GetOwningAssetName());
+			}
 #endif // WITH_EDITORONLY_DATA
 		}
 
@@ -144,10 +138,7 @@ namespace Metasound
 				if (InMetaSound.VersionAsset())
 				{
 #if WITH_EDITORONLY_DATA
-					if (UMetasoundEditorGraphBase* MetaSoundGraph = Cast<UMetasoundEditorGraphBase>(InMetaSound.GetGraph()))
-					{
-						MetaSoundGraph->SetVersionedOnLoad();
-					}
+					InMetaSound.SetVersionedOnLoad();
 #endif // WITH_EDITORONLY_DATA
 				}
 			}
diff --git a/Engine/Plugins/Runtime/Metasound/Source/MetasoundEngine/Private/MetasoundSource.cpp b/Engine/Plugins/Runtime/Metasound/Source/MetasoundEngine/Private/MetasoundSource.cpp
--- a/Engine/Plugins/Runtime/Metasound/Source/MetasoundEngine/Private/MetasoundSource.cpp
+++ b/Engine/Plugins/Runtime/Metasound/Source/MetasoundEngine/Private/MetasoundSource.cpp
@@ -93,12 +93,9 @@ namespace Metasound
 			Frontend::FMetaSoundAssetRegistrationOptions RegOptions;
 			RegOptions.bForceReregister = false;
 #if !WITH_EDITOR 
-			if (Frontend::MetaSoundEnableCookDeterministicIDGeneration != 0)
-			{
-				// When without editor, don't AutoUpdate or ResolveDocument at runtime. This only happens at cook or save.
-				// When with editor, those are needed because sounds are not necessarily saved before previewing.
-				RegOptions.bAutoUpdate = false;
-			}
+			// When without editor, don't AutoUpdate or ResolveDocument at runtime. This only happens at cook or save.
+			// When with editor, those are needed because sounds are not necessarily saved before previewing.
+			RegOptions.bAutoUpdate = false;
 #endif // !WITH_EDITOR
 			if (const UMetaSoundSettings* Settings = GetDefault<UMetaSoundSettings>())
 			{
@@ -433,17 +430,12 @@ void UMetaSoundSource::PostEditChangeOutputFormat()
 		UMetaSoundSourceBuilder* SourceBuilder = BuilderSubsystem.AttachSourceBuilderToAsset(this);
 		check(SourceBuilder);
 		SourceBuilder->SetFormat(OutputFormat, Result);
-
-		// TODO: Once builders are notified of controller changes and can be safely persistent, this
-		// can be removed so builders can be shared and not have to be created for each change output
-		// format mutation transaction.
-		BuilderSubsystem.DetachBuilderFromAsset(GetConstDocument().RootGraph.Metadata.GetClassName());
 	}
 
 	if (Result == EMetaSoundBuilderResult::Succeeded)
 	{
 		// Update the data in this UMetaSoundSource to reflect what is in the metasound document.
-		ConformObjectDataToInterfaces();
+		ConformObjectToDocument();
 
 		// Use the editor form of register to ensure other editors'
 		// MetaSounds are auto-updated if they are referencing this graph.
@@ -472,7 +464,7 @@ void UMetaSoundSource::PostEditChangeQualitySettings()
 }
 #endif // WITH_EDITOR
 
-bool UMetaSoundSource::ConformObjectDataToInterfaces()
+bool UMetaSoundSource::ConformObjectToDocument()
 {
 	using namespace Metasound::Engine;
 	using namespace Metasound::Frontend;
@@ -643,27 +635,37 @@ void UMetaSoundSource::OnAsyncReferencedAssetsLoaded(const TArray<FMetasoundAsse
 }
 
 #if WITH_EDITORONLY_DATA
+void UMetaSoundSource::MigrateEditorGraph(FMetaSoundFrontendDocumentBuilder& OutBuilder)
+{
+	PRAGMA_DISABLE_DEPRECATION_WARNINGS
+		if (Graph)
+		{
+			Graph->MigrateEditorDocumentData(OutBuilder);
+			Graph = nullptr;
+		}
+	PRAGMA_ENABLE_DEPRECATION_WARNINGS
+}
 
 UEdGraph* UMetaSoundSource::GetGraph()
 {
-	return Graph;
+	return EditorGraph;
 }
 
 const UEdGraph* UMetaSoundSource::GetGraph() const
 {
-	return Graph;
+	return EditorGraph;
 }
 
 UEdGraph& UMetaSoundSource::GetGraphChecked()
 {
-	check(Graph);
-	return *Graph;
+	check(EditorGraph);
+	return *EditorGraph;
 }
 
 const UEdGraph& UMetaSoundSource::GetGraphChecked() const
 {
-	check(Graph);
-	return *Graph;
+	check(EditorGraph);
+	return *EditorGraph;
 }
 
 FText UMetaSoundSource::GetDisplayName() const
@@ -920,7 +922,7 @@ Metasound::Frontend::FConstDocumentAccessPtr UMetaSoundSource::GetDocumentConstA
 bool UMetaSoundSource::ImplementsParameterInterface(Audio::FParameterInterfacePtr InInterface) const
 {
 	const FMetasoundFrontendVersion Version { InInterface->GetName(), { InInterface->GetVersion().Major, InInterface->GetVersion().Minor } };
-	return GetDocumentChecked().Interfaces.Contains(Version);
+	return GetConstDocumentChecked().Interfaces.Contains(Version);
 }
 
 ISoundGeneratorPtr UMetaSoundSource::CreateSoundGenerator(const FSoundGeneratorInitParams& InParams, TArray<FAudioParameter>&& InDefaultParameters)
@@ -1217,7 +1219,7 @@ void UMetaSoundSource::InitParametersInternal(const Metasound::TSortedVertexName
 
 bool UMetaSoundSource::IsParameterValid(const FAudioParameter& InParameter) const
 {
-	const TArray<FMetasoundFrontendClassInput>& Inputs = GetDocumentChecked().RootGraph.Interface.Inputs;
+	const TArray<FMetasoundFrontendClassInput>& Inputs = GetConstDocumentChecked().RootGraph.Interface.Inputs;
 	const FMetasoundFrontendVertex* Vertex = Algo::FindByPredicate(Inputs, [&InParameter] (const FMetasoundFrontendClassInput& Input)
 	{
 		return Input.Name == InParameter.ParamName;
@@ -1738,7 +1740,10 @@ void UMetaSoundSource::CacheRuntimeInputData()
 
 	if (bIsBuilderActive)
 	{
+#if !WITH_EDITOR
 		UE_LOG(LogMetaSound, Warning, TEXT("Skipping caching of runtime inputs for UMetaSoundSource %s because there is an active builder"), *GetOwningAssetName());
+#endif // !WITH_EDITOR
+		return;
 	}
 
 	constexpr bool bCreateUObjectProxies = true; 
diff --git a/Engine/Plugins/Runtime/Metasound/Source/MetasoundEngine/Public/Metasound.h b/Engine/Plugins/Runtime/Metasound/Source/MetasoundEngine/Public/Metasound.h
--- a/Engine/Plugins/Runtime/Metasound/Source/MetasoundEngine/Public/Metasound.h
+++ b/Engine/Plugins/Runtime/Metasound/Source/MetasoundEngine/Public/Metasound.h
@@ -31,11 +31,22 @@ public:
 	virtual void RegisterGraphWithFrontend() PURE_VIRTUAL(UMetasoundEditorGraphBase::RegisterGraphWithFrontend(), )
 
 #if WITH_EDITORONLY_DATA
-	virtual FMetasoundFrontendDocumentModifyContext& GetModifyContext() PURE_VIRTUAL(UMetasoundEditorGraphBase::GetModifyContext, static FMetasoundFrontendDocumentModifyContext InvalidModifyData; return InvalidModifyData; )
-	virtual const FMetasoundFrontendDocumentModifyContext& GetModifyContext() const PURE_VIRTUAL(UMetasoundEditorGraphBase::GetModifyContext, static const FMetasoundFrontendDocumentModifyContext InvalidModifyData; return InvalidModifyData; )
-	virtual void ClearVersionedOnLoad() PURE_VIRTUAL(UMetasoundEditorGraphBase::ClearVersionedOnLoad, )
-	virtual bool GetVersionedOnLoad() const PURE_VIRTUAL(UMetasoundEditorGraphBase::GetVersionedOnLoad, return false; )
-	virtual void SetVersionedOnLoad() PURE_VIRTUAL(UMetasoundEditorGraphBase::SetVersionedOnLoad, )
+	UE_DEPRECATED(5.5, "ModifyContext is to be replaced by builder API delegates providing context when items changed and it will be up to the caller to track modification deltas.")
+	virtual FMetasoundFrontendDocumentModifyContext& GetModifyContext() { static FMetasoundFrontendDocumentModifyContext InvalidModifyData; return InvalidModifyData; }
+
+	UE_DEPRECATED(5.5, "ModifyContext is to be replaced by builder API delegates providing context when items changed and it will be up to the caller to track modification deltas.")
+	virtual const FMetasoundFrontendDocumentModifyContext& GetModifyContext() const { static const FMetasoundFrontendDocumentModifyContext InvalidModifyData; return InvalidModifyData; }
+
+	UE_DEPRECATED(5.5, "Editor Graph is now transient, so versioning flag moved to AssetBase.")
+	virtual void ClearVersionedOnLoad() { }
+
+	UE_DEPRECATED(5.5, "Editor Graph is now transient, so versioning flag moved to AssetBase.")
+	virtual bool GetVersionedOnLoad() const { return false; }
+
+	UE_DEPRECATED(5.5, "Editor Graph is now transient, so versioning flag moved to AssetBase.")
+	virtual void SetVersionedOnLoad() {  }
+
+	virtual void MigrateEditorDocumentData(FMetaSoundFrontendDocumentBuilder & OutBuilder) PURE_VIRTUAL(UMetasoundEditorGraphBase::MigrateEditorDocumentData(), )
 #endif // WITH_EDITORONLY_DATA
 
 	int32 GetHighestMessageSeverity() const;
@@ -73,8 +84,11 @@ protected:
 	TSet<FSoftObjectPath> ReferenceAssetClassCache;
 
 #if WITH_EDITORONLY_DATA
-	UPROPERTY()
+	UPROPERTY(meta=(DeprecatedProperty, DeprecationMessage = "Use EditorGraph instead as it is now transient and generated via the FrontendDocument dynamically."))
 	TObjectPtr<UMetasoundEditorGraphBase> Graph;
+
+	UPROPERTY(Transient)
+	TObjectPtr<UMetasoundEditorGraphBase> EditorGraph;
 #endif // WITH_EDITORONLY_DATA
 
 public:
@@ -119,13 +133,14 @@ public:
 	virtual const UEdGraph* GetGraph() const override;
 	virtual UEdGraph& GetGraphChecked() override;
 	virtual const UEdGraph& GetGraphChecked() const override;
+	virtual void MigrateEditorGraph(FMetaSoundFrontendDocumentBuilder& OutBuilder) override;
 
 	// Sets the graph associated with this Metasound. Graph is required to be referenced on
 	// Metasound UObject for editor serialization purposes.
 	// @param Editor graph associated with UMetaSoundSource.
 	virtual void SetGraph(UEdGraph* InGraph) override
 	{
-		Graph = CastChecked<UMetasoundEditorGraphBase>(InGraph);
+		EditorGraph = CastChecked<UMetasoundEditorGraphBase>(InGraph);
 	}
 #endif // #if WITH_EDITORONLY_DATA
 
@@ -143,7 +158,7 @@ public:
 	virtual void Serialize(FArchive& InArchive) override;
 	virtual void PostLoad() override;
 
-	virtual bool ConformObjectDataToInterfaces() override { return false; }
+	virtual bool ConformObjectToDocument() override { return false; }
 
 	virtual const TSet<FString>& GetReferencedAssetClassKeys() const override
 	{
@@ -163,6 +178,8 @@ public:
 		return this;
 	}
 
+	virtual bool IsBuilderActive() const override;
+
 protected:
 #if WITH_EDITOR
 	virtual void SetReferencedAssetClasses(TSet<Metasound::Frontend::IMetaSoundAssetManager::FAssetInfo>&& InAssetClasses) override;
@@ -177,7 +194,6 @@ private:
 		return RootMetaSoundDocument;
 	}
 
-	virtual bool IsBuilderActive() const override;
 	virtual void OnBeginActiveBuilder() override;
 	virtual void OnFinishActiveBuilder() override;
 
diff --git a/Engine/Plugins/Runtime/Metasound/Source/MetasoundEngine/Public/MetasoundAssetSubsystem.h b/Engine/Plugins/Runtime/Metasound/Source/MetasoundEngine/Public/MetasoundAssetSubsystem.h
--- a/Engine/Plugins/Runtime/Metasound/Source/MetasoundEngine/Public/MetasoundAssetSubsystem.h
+++ b/Engine/Plugins/Runtime/Metasound/Source/MetasoundEngine/Public/MetasoundAssetSubsystem.h
@@ -66,7 +66,9 @@ public:
 #if WITH_EDITORONLY_DATA
 	virtual void AddAssetReferences(FMetasoundAssetBase& InAssetBase) override;
 #endif
-	virtual Metasound::Frontend::FNodeRegistryKey  AddOrUpdateAsset(const UObject& InObject) override;
+	virtual Metasound::Frontend::FNodeRegistryKey AddOrUpdateAsset(const UObject& InObject) override;
+	virtual FMetaSoundFrontendDocumentBuilder& AttachDocumentBuilderChecked(UObject& InObject) const override;
+
 	virtual bool CanAutoUpdate(const FMetasoundFrontendClassName& InClassName) const override;
 	virtual bool ContainsKey(const Metasound::Frontend::FNodeRegistryKey& InRegistryKey) const override;
 	virtual const FSoftObjectPath* FindObjectPathFromKey(const Metasound::Frontend::FNodeRegistryKey& RegistryKey) const override;
diff --git a/Engine/Plugins/Runtime/Metasound/Source/MetasoundEngine/Public/MetasoundBuilderSubsystem.h b/Engine/Plugins/Runtime/Metasound/Source/MetasoundEngine/Public/MetasoundBuilderSubsystem.h
--- a/Engine/Plugins/Runtime/Metasound/Source/MetasoundEngine/Public/MetasoundBuilderSubsystem.h
+++ b/Engine/Plugins/Runtime/Metasound/Source/MetasoundEngine/Public/MetasoundBuilderSubsystem.h
@@ -165,6 +165,10 @@ public:
 	UE_DEPRECATED(5.4, "This version of AddNodeByClassName is deprecated. Use the one with a default MajorVersion of 1.")
 	UPARAM(DisplayName = "Node Handle") FMetaSoundNodeHandle AddNodeByClassName(const FMetasoundFrontendClassName& ClassName, int32 MajorVersion, EMetaSoundBuilderResult& OutResult);
 
+#if WITH_EDITOR
+	bool ClearMemberMetadata(const FGuid& InMemberID);
+#endif // WITH_EDITOR
+
 	// Connects node output to a node input. Does *NOT* provide loop detection for performance reasons.  Loop detection is checked on class registration when built or played.
 	// Returns succeeded if connection made, failed if connection already exists with input, the data types do not match, or the connection is not supported due to access type
 	// incompatibility (ex. constructor input to non-constructor input).
@@ -219,7 +223,7 @@ public:
 	// all connections were found and removed, failed if any connections were not.
 	UFUNCTION(BlueprintCallable, Category = "Audio|MetaSound|Builder", meta = (ExpandEnumAsExecs = "OutResult"))
 	void DisconnectNodesByInterfaceBindings(const FMetaSoundNodeHandle& FromNodeHandle, const FMetaSoundNodeHandle& ToNodeHandle, EMetaSoundBuilderResult& OutResult);
-	
+
 	// Returns graph input node by the given name if it exists, or an invalid handle if not found.
 	UFUNCTION(BlueprintCallable, Category = "Audio|MetaSound|Builder", meta = (ExpandEnumAsExecs = "OutResult"))
 	UPARAM(DisplayName = "Node Handle") FMetaSoundNodeHandle FindGraphInputNode(FName InputName, EMetaSoundBuilderResult& OutResult);
@@ -228,6 +232,10 @@ public:
 	UFUNCTION(BlueprintCallable, Category = "Audio|MetaSound|Builder", meta = (ExpandEnumAsExecs = "OutResult"))
 	UPARAM(DisplayName = "Node Handle") FMetaSoundNodeHandle FindGraphOutputNode(FName OutputName, EMetaSoundBuilderResult& OutResult);
 
+#if WITH_EDITOR
+	UMetaSoundFrontendMemberMetadata* FindMemberMetadata(const FGuid& InMemberID);
+#endif // WITH_EDITOR
+
 	// Returns node input by the given name if it exists, or an invalid handle if not found.
 	UFUNCTION(BlueprintCallable, Category = "Audio|MetaSound|Builder", meta = (ExpandEnumAsExecs = "OutResult"))
 	UPARAM(DisplayName = "Input Handle") FMetaSoundBuilderNodeInputHandle FindNodeInputByName(const FMetaSoundNodeHandle& NodeHandle, FName InputName, EMetaSoundBuilderResult& OutResult);
@@ -346,20 +354,28 @@ public:
 
 	// Removes node and any associated connections from the builder's MetaSound.
 	UFUNCTION(BlueprintCallable, Category = "Audio|MetaSound|Builder", meta = (ExpandEnumAsExecs = "OutResult"))
-	void RemoveNode(const FMetaSoundNodeHandle& NodeHandle, EMetaSoundBuilderResult& OutResult);
+	void RemoveNode(const FMetaSoundNodeHandle& NodeHandle, EMetaSoundBuilderResult& OutResult, bool bRemoveUnusedDependencies = true);
 
 	// Removes node input literal default if set, reverting the value to be whatever the node class defaults the value to.
 	// Returns success if value was removed, false if not removed (i.e. wasn't set to begin with).
 	UFUNCTION(BlueprintCallable, Category = "Audio|MetaSound|Builder", meta = (ExpandEnumAsExecs = "OutResult"))
 	void RemoveNodeInputDefault(const FMetaSoundBuilderNodeInputHandle& InputHandle, EMetaSoundBuilderResult& OutResult);
-	
+
+	// Removes dependencies in document that are no longer referenced by nodes
+	UFUNCTION(BlueprintCallable, Category = "Audio|MetaSound|Builder")
+	void RemoveUnusedDependencies();
+
 	// Rename the document's root graph class with a guid and optional namespace and variant
 	void RenameRootGraphClass(const FMetasoundFrontendClassName& InName);
 
-	// Primarily used by editor transaction stack to avoid corruption when object is
-	// changed outside of the builder API. Generally discouraged for direct use otherwise.
+	UE_DEPRECATED(5.5, "Moved to 'Reload' to discourage priming by defaulting to false. "
+		"In addition, internal delegates and respective bindings are now regenerated when reloading, "
+		"effectively resetting the builder's entire state and not just strictly the underlying cache.")
 	void ReloadCache(bool bPrimeCache = true);
 
+	// Reloads the builder, freeing the internal cache and rebuilding delegate bindings
+	void Reload(bool bPrimeCache = false);
+
 #if WITH_EDITOR
 	// Sets the author of the MetaSound.
 	void SetAuthor(const FString& InAuthor);
@@ -369,10 +385,32 @@ public:
 	UFUNCTION(BlueprintCallable, Category = "Audio|MetaSound|Builder", meta = (ExpandEnumAsExecs = "OutResult"))
 	void SetNodeInputDefault(const FMetaSoundBuilderNodeInputHandle& NodeInputHandle, const FMetasoundFrontendLiteral& Literal, EMetaSoundBuilderResult& OutResult);
 
+	// Disconnects the given graph input's respective template nodes and sets the graph input's AccessType should it not match the current AccessType.
+	// Result succeeds if the AccessType was successfully changed or if the provided AccessType is already the input's current AccessType.
+	UFUNCTION(BlueprintCallable, Category = "Audio|MetaSound|Builder", meta = (ExpandEnumAsExecs = "OutResult"))
+	void SetGraphInputAccessType(FName InputName, EMetasoundFrontendVertexAccessType AccessType, EMetaSoundBuilderResult& OutResult);
+
+	// Disconnects the given graph input's respective template nodes and sets the graph input's DataType should it not match the current DataType.
+	// Result succeeds if the DataType was successfully changed or if the provided DataType is already the input's current DataType.
+	UFUNCTION(BlueprintCallable, Category = "Audio|MetaSound|Builder", meta = (ExpandEnumAsExecs = "OutResult"))
+	void SetGraphInputDataType(FName InputName, FName DataType, EMetaSoundBuilderResult& OutResult);
+
 	// Sets the input node's default value, overriding the default provided by the referenced graph if the graph is a preset.
 	UFUNCTION(BlueprintCallable, Category = "Audio|MetaSound|Builder", meta = (ExpandEnumAsExecs = "OutResult"))
 	void SetGraphInputDefault(FName InputName, const FMetasoundFrontendLiteral& Literal, EMetaSoundBuilderResult& OutResult);
 
+	// Disconnects the given graph output's respective template nodes and sets the graph output's AccessType should it not match the current AccessType.
+	// Result succeeds if the AccessType was successfully changed or if the provided AccessType is already the output's current AccessType.
+	UFUNCTION(BlueprintCallable, Category = "Audio|MetaSound|Builder", meta = (ExpandEnumAsExecs = "OutResult"))
+	void SetGraphOutputAccessType(FName OutputName, EMetasoundFrontendVertexAccessType AccessType, EMetaSoundBuilderResult& OutResult);
+
+	// Disconnects the given graph output's respective template nodes and sets the graph output's DataType should it not match the current DataType.
+	// Result succeeds if the DataType was successfully changed or if the provided DataType is already the output's current DataType.
+	UFUNCTION(BlueprintCallable, Category = "Audio|MetaSound|Builder", meta = (ExpandEnumAsExecs = "OutResult"))
+	void SetGraphOutputDataType(FName OutputName, FName DataType, EMetaSoundBuilderResult& OutResult);
+
+	void SetMemberMetadata(UMetaSoundFrontendMemberMetadata& NewMetadata);
+
 	// Update dependency class names given a map of old to new referenced class names. 
 	void UpdateDependencyClassNames(const TMap<FMetasoundFrontendClassName, FMetasoundFrontendClassName>& OldToNewReferencedClassNames);
 
@@ -384,10 +422,19 @@ public:
 	UE_DEPRECATED(5.4, "Moved to protected pure virtual 'CreateTransientBuilder'")
 	virtual void InitFrontendBuilder();
 
+	virtual bool ConformObjectToDocument() PURE_VIRTUAL(UMetaSoundBuilderBase::ConformObjectToDocument, return false; );
+
 	// Initializes and ensures all nodes have a position (required prior to exporting to an asset if expected to be viewed in the editor).
 	virtual void InitNodeLocations();
 
 #if WITH_EDITOR
+	// Injects template nodes between builder's document inputs not connected
+	// to existing template inputs, copying locational data from the represented
+	// input metadata. If bForceNodeCreation is false, only generates a template
+	// input node if a connection between the input and other nodes exists. If true,
+	// will inject template node irrespective of whether or not the input has connections.
+	void InjectInputTemplateNodes(bool bForceNodeCreation, EMetaSoundBuilderResult& OutResult);
+
 	const FMetaSoundFrontendGraphComment* FindGraphComment(const FGuid& InCommentID) const;
 	FMetaSoundFrontendGraphComment* FindGraphComment(const FGuid & InCommentID);
 	FMetaSoundFrontendGraphComment& FindOrAddGraphComment(const FGuid& InCommentID);
@@ -402,11 +449,14 @@ public:
 	const FMetaSoundFrontendDocumentBuilder& GetConstBuilder() const;
 
 protected:
-	// Creates a FrontendBuilder wrapping the transient MetaSound object of the class supported by the given subsystem builder (See GetBuilderUClass).
+	virtual void InitDelegates(Metasound::Frontend::FDocumentModifyDelegates & OutDocumentDelegates) { };
+
+	// Creates a FrontendBuilder wrapping a transient MetaSound object of the class supported by the given subsystem builder (See GetBuilderUClass).
 	// Should only be used when builders is not being attached to an existing, serialized MetaSound asset.
 	virtual void CreateTransientBuilder() PURE_VIRTUAL(UMetaSoundBuilderBase::CreateTransientBuilder, );
 
-	void InvalidateCache();
+	UE_DEPRECATED(5.5, "Moved to 'Reload', to enforce generation of new delegates")
+	void InvalidateCache(bool bPrimeCache = false);
 
 	// Runs build, conforming the document and corresponding object data on a MetaSound UObject to that managed by this builder.
 	template <typename UClassType>
@@ -428,7 +478,7 @@ protected:
 			{
 				// If MetaSound already exists, preserve the class name to avoid
 				// nametable bloat & preserve potentially existing references.
-				const FMetasoundFrontendDocument& ExistingDoc = CastChecked<const UClassType>(MetaSound)->GetDocumentChecked();
+				const FMetasoundFrontendDocument& ExistingDoc = CastChecked<const UClassType>(MetaSound)->GetConstDocumentChecked();
 				if (!BuilderOptions.bForceUniqueClassName)
 				{
 					DocClassName = &ExistingDoc.RootGraph.Metadata.GetClassName();
@@ -454,17 +504,9 @@ protected:
 		}
 
 		checkf(MetaSound, TEXT("Failed to build MetaSound from builder '%s'"), *GetPathName());
-		FMetasoundFrontendDocument NewDocument = GetConstBuilder().GetDocument();
-		{
-			// This is required to ensure the newly build document has a unique class
-			// identifier than that of the given builder's document copy to avoid collisions
-			// if added to the Frontend class registry (either below or at a later point in time).
-			constexpr bool bResetVersion = false;
-			FMetaSoundFrontendDocumentBuilder::InitGraphClassMetadata(NewDocument.RootGraph.Metadata, bResetVersion, DocClassName);
-		}
 
-		MetaSound->SetDocument(MoveTemp(NewDocument));
-		MetaSound->ConformObjectDataToInterfaces();
+		TScriptInterface<IMetaSoundDocumentInterface> MetaSoundInterface = MetaSound;
+		BuildInternal(MetaSoundInterface, DocClassName);
 
 		if (BuilderOptions.bAddToRegistry)
 		{
@@ -486,15 +528,25 @@ protected:
 	UPROPERTY()
 	FMetaSoundFrontendDocumentBuilder Builder;
 
+	// ClassName of associated document object.  Cached in case DocumentObject is deleted but continues to provide an identifier after destruction.
+	UPROPERTY()
+	FMetasoundFrontendClassName ClassName;
+
 #if WITH_EDITORONLY_DATA
 	UPROPERTY(meta = (DeprecatedProperty, DeprecationMessage = "5.4 - All source builders now operate on an underlying document source document that is also used to audition."))
 	bool bIsAttached = false;
 #endif // WITH_EDITORONLY_DATA
 
+
+private:
+	void BuildInternal(TScriptInterface<IMetaSoundDocumentInterface> NewMetaSound, const FMetasoundFrontendClassName* InDocClassName) const;
+
 private:
 	int32 LastTransactionRegistered = 0;
 
 	// Friending allows for swapping the builder in certain circumstances where desired (eg. attaching a builder to an existing asset)
+	// or performing editor-only manipulations to the builder (ex. adding template nodes) and versioning the document data where necessary.
+	friend class UMetaSoundAssetSubsystem;
 	friend class UMetaSoundBuilderSubsystem;
 };
 
@@ -508,6 +560,7 @@ public:
 	UFUNCTION(BlueprintCallable, Category = "Audio|MetaSound|Builder", meta = (WorldContext = "Parent"))
 	virtual UPARAM(DisplayName = "MetaSound") TScriptInterface<IMetaSoundDocumentInterface> Build(UObject* Parent, const FMetaSoundBuilderOptions& Options) const override;
 
+	virtual bool ConformObjectToDocument() override;
 	virtual const UClass& GetBuilderUClass() const override;
 
 protected:
@@ -547,6 +600,7 @@ public:
 
 	const Metasound::Engine::FOutputAudioFormatInfoPair* FindOutputAudioFormatInfo() const;
 
+	virtual bool ConformObjectToDocument() override;
 	virtual const UClass& GetBuilderUClass() const override;
 
 #if WITH_EDITORONLY_DATA
@@ -563,6 +617,7 @@ public:
 
 protected:
 	virtual void CreateTransientBuilder() override;
+	virtual void InitDelegates(Metasound::Frontend::FDocumentModifyDelegates& OutDocumentDelegates) override;
 
 private:
 	static TOptional<Metasound::FAnyDataReference> CreateDataReference(const Metasound::FOperatorSettings& InOperatorSettings, FName DataType, const Metasound::FLiteral& InLiteral, Metasound::EDataReferenceAccessType AccessType);
@@ -570,7 +625,6 @@ private:
 	const UMetaSoundSource& GetMetaSoundSource() const;
 	UMetaSoundSource& GetMetaSoundSource();
 
-	void InitDelegates(Metasound::Frontend::FDocumentModifyDelegates& OutDocumentDelegates);
 
 	void OnEdgeAdded(int32 EdgeIndex) const;
 	void OnInputAdded(int32 InputIndex);
@@ -618,7 +672,7 @@ public:
 	static const UMetaSoundBuilderSubsystem* GetConst();
 	static const UMetaSoundBuilderSubsystem& GetConstChecked();
 
-	UMetaSoundBuilderBase& AttachBuilderToAssetChecked(UObject& InMetaSound) const;
+	UMetaSoundBuilderBase& AttachBuilderToAssetChecked(UObject& InObject) const;
 
 	// Moving toward becoming UFUNCTIONS in the future, but right now unguarded case
 	// where a user could improperly reference transient assets. Also, unsafe where
@@ -740,6 +794,7 @@ private:
 
 		const FMetasoundFrontendDocument& Document = NewBuilder->GetConstBuilder().GetDocument();
 		const FMetasoundFrontendClassName& ClassName = Document.RootGraph.Metadata.GetClassName();
+		NewBuilder->ClassName = ClassName;
 		TransientBuilders.Add(ClassName, NewBuilder);
 		return *NewBuilder.Get();
 	}
@@ -752,18 +807,29 @@ private:
 
 		TScriptInterface<IMetaSoundDocumentInterface> DocInterface = InMetaSoundObject;
 		const FMetasoundFrontendDocument& Document = DocInterface->GetConstDocument();
-		const FMetasoundFrontendClassName& ClassName = Document.RootGraph.Metadata.GetClassName();
-		TWeakObjectPtr<UMetaSoundBuilderBase> Builder = AssetBuilders.FindRef(ClassName);
-		if (Builder.IsValid())
+		const FMetasoundFrontendClassName& FullClassName = Document.RootGraph.Metadata.GetClassName();
+
+		if (FullClassName.IsValid())
 		{
-			return *CastChecked<BuilderClass>(Builder.Get());
+			TWeakObjectPtr<UMetaSoundBuilderBase> Builder = AssetBuilders.FindRef(FullClassName);
+			if (Builder.IsValid())
+			{
+				return *CastChecked<BuilderClass>(Builder.Get());
+			}
 		}
 
 		TObjectPtr<BuilderClass> NewBuilder = NewObject<BuilderClass>(InMetaSoundObject);
 		check(NewBuilder);
 		NewBuilder->Builder = FMetaSoundFrontendDocumentBuilder(DocInterface);
+
+		if (!FullClassName.IsValid())
+		{
+			NewBuilder->Builder.InitDocument();
+		}
+
+		checkf(FullClassName.IsValid(), TEXT("Document initialization must result in a valid class name being generated"));
 		TObjectPtr<UMetaSoundBuilderBase> NewBuilderBase = CastChecked<UMetaSoundBuilderBase>(NewBuilder);
-		AssetBuilders.Add(ClassName, NewBuilderBase);
+		AssetBuilders.Add(FullClassName, NewBuilderBase);
 		return *NewBuilder;
 	}
 
diff --git a/Engine/Plugins/Runtime/Metasound/Source/MetasoundEngine/Public/MetasoundSource.h b/Engine/Plugins/Runtime/Metasound/Source/MetasoundEngine/Public/MetasoundSource.h
--- a/Engine/Plugins/Runtime/Metasound/Source/MetasoundEngine/Public/MetasoundSource.h
+++ b/Engine/Plugins/Runtime/Metasound/Source/MetasoundEngine/Public/MetasoundSource.h
@@ -26,6 +26,8 @@
 class UMetaSoundSettings;
 struct FMetaSoundQualitySettings;
 
+struct FMetaSoundFrontendDocumentBuilder;
+
 namespace Audio
 {
 	using DeviceID = uint32;
@@ -105,8 +107,11 @@ protected:
 	TSet<FSoftObjectPath> ReferenceAssetClassCache;
 
 #if WITH_EDITORONLY_DATA
-	UPROPERTY()
+	UPROPERTY(meta=(DeprecatedProperty, DeprecationMessage = "Use EditorGraph instead as it is now transient and generated via the FrontendDocument dynamically."))
 	TObjectPtr<UMetasoundEditorGraphBase> Graph;
+
+	UPROPERTY(Transient)
+	TObjectPtr<UMetasoundEditorGraphBase> EditorGraph;
 #endif // WITH_EDITORONLY_DATA
 
 public:
@@ -175,13 +180,14 @@ public:
 	virtual const UEdGraph* GetGraph() const override;
 	virtual UEdGraph& GetGraphChecked() override;
 	virtual const UEdGraph& GetGraphChecked() const override;
+	virtual void MigrateEditorGraph(FMetaSoundFrontendDocumentBuilder& OutBuilder) override;
 
 	// Sets the graph associated with this Metasound. Graph is required to be referenced on
 	// Metasound UObject for editor serialization purposes.
 	// @param Editor graph associated with UMetaSoundSource.
 	virtual void SetGraph(UEdGraph* InGraph) override
 	{
-		Graph = CastChecked<UMetasoundEditorGraphBase>(InGraph);
+		EditorGraph = CastChecked<UMetasoundEditorGraphBase>(InGraph);
 	}
 #endif // #if WITH_EDITORONLY_DATA
 
@@ -212,10 +218,9 @@ public:
 private:
 	void PostEditChangeOutputFormat();
 	void PostEditChangeQualitySettings();
-public:
-
 #endif // WITH_EDITOR
 
+public:
 	virtual const TSet<FString>& GetReferencedAssetClassKeys() const override
 	{
 		return ReferencedAssetClassKeys;
@@ -231,7 +236,7 @@ public:
 
 	void PostLoadQualitySettings();
 
-	virtual bool ConformObjectDataToInterfaces() override;
+	virtual bool ConformObjectToDocument() override;
 
 	virtual FTopLevelAssetPath GetAssetPathChecked() const override;
 
@@ -268,6 +273,7 @@ public:
 	Metasound::FOperatorSettings GetOperatorSettings(Metasound::FSampleRate InDeviceSampleRate) const;
 
 	virtual const FMetasoundFrontendDocument& GetConstDocument() const override;
+	virtual bool IsBuilderActive() const override;
 
 protected:
 	Metasound::Frontend::FDocumentAccessPtr GetDocumentAccessPtr() override;
@@ -288,7 +294,6 @@ private:
 		return RootMetasoundDocument;
 	}
 
-	virtual bool IsBuilderActive() const override;
 	virtual void OnBeginActiveBuilder() override;
 	virtual void OnFinishActiveBuilder() override;
 
@@ -351,7 +356,7 @@ private:
 
 	// Quality settings. 
 	bool GetQualitySettings(const FName InPlatformName, Metasound::SourcePrivate::FCookedQualitySettings& OutQualitySettings) const;
-	void ResolveQualitySettings(const UMetaSoundSettings* Settings);	
+	void ResolveQualitySettings(const UMetaSoundSettings* Settings);
 	void SerializeCookedQualitySettings(const FName PlatformName, FArchive& Ar);
 	TPimplPtr<Metasound::SourcePrivate::FCookedQualitySettings> CookedQualitySettings;
 };
diff --git a/Engine/Plugins/Runtime/Metasound/Source/MetasoundEngineTest/Private/EngineTestMetaSoundBuilder.h b/Engine/Plugins/Runtime/Metasound/Source/MetasoundEngineTest/Private/EngineTestMetaSoundBuilder.h
--- a/Engine/Plugins/Runtime/Metasound/Source/MetasoundEngineTest/Private/EngineTestMetaSoundBuilder.h
+++ b/Engine/Plugins/Runtime/Metasound/Source/MetasoundEngineTest/Private/EngineTestMetaSoundBuilder.h
@@ -13,12 +13,14 @@
 #include "MetasoundDataReference.h"
 #include "MetasoundFrontendController.h"
 #include "MetasoundFrontendSearchEngine.h"
+#include "MetasoundTrigger.h"
 #include "Misc/AutomationTest.h"
 #include "Misc/Paths.h"
+#include "NodeTemplates/MetasoundFrontendNodeTemplateInput.h"
 #include "Tests/AutomationCommon.h"
 
-
 #if WITH_DEV_AUTOMATION_TESTS
+
 namespace EngineTestMetaSoundPatchBuilderPrivate
 {
 	UMetaSoundPatchBuilder& CreatePatchBuilderChecked(FAutomationTestBase& Test, FName BuilderName, const TArray<FName>& InterfaceNamesToAdd)
@@ -193,8 +195,8 @@ bool FAudioMetaSoundBuilderInterfaceBindingConnectAndDisconnect::RunTest(const F
 }
 
 // This test creates a MetaSound patch, then adds and connects sin oscillator, attempts to retrieve a default input set, then clears it, and finally retrieves the class default
-IMPLEMENT_SIMPLE_AUTOMATION_TEST(FAudioMetasoundNodeClassQueryFunctions, "Audio.Metasound.Builder.NodeClassQueryFunctions", EAutomationTestFlags::EditorContext | EAutomationTestFlags::EngineFilter)
-bool FAudioMetasoundNodeClassQueryFunctions::RunTest(const FString& Parameters)
+IMPLEMENT_SIMPLE_AUTOMATION_TEST(FAudioMetaSoundNodeClassQueryFunctions, "Audio.Metasound.Builder.NodeClassQueryFunctions", EAutomationTestFlags::EditorContext | EAutomationTestFlags::EngineFilter)
+bool FAudioMetaSoundNodeClassQueryFunctions::RunTest(const FString& Parameters)
 {
 	using namespace EngineTestMetaSoundPatchBuilderPrivate;
 	using namespace Metasound;
@@ -203,7 +205,7 @@ bool FAudioMetasoundNodeClassQueryFunctions::RunTest(const FString& Parameters)
 	constexpr float NodeDefaultFreq = 100.0f;
 
 	EMetaSoundBuilderResult Result = EMetaSoundBuilderResult::Failed;
-	UMetaSoundPatchBuilder& Builder = CreatePatchBuilderChecked(*this, "DefaultLiteralAssignment", { });
+	UMetaSoundPatchBuilder& Builder = CreatePatchBuilderChecked(*this, "DefaultLiteralAssignment_Test", { });
 
 	// Sine Oscillator Node
 	const FMetaSoundNodeHandle OscNode = Builder.AddNodeByClassName({ "UE", "Sine", "Audio" }, Result, 1);
@@ -271,4 +273,62 @@ bool FAudioMetasoundNodeClassQueryFunctions::RunTest(const FString& Parameters)
 
 	return true;
 }
+
+// This test creates a MetaSound patch, adds a frequency input, connects a sin oscillator, and if the editor is loaded, and then attempts to inject template input nodes.
+IMPLEMENT_SIMPLE_AUTOMATION_TEST(FAudioMetaSoundInjectTemplateNodes, "Audio.Metasound.Builder.InjectInputTemplateNodes", EAutomationTestFlags::EditorContext | EAutomationTestFlags::EngineFilter)
+bool FAudioMetaSoundInjectTemplateNodes::RunTest(const FString& Parameters)
+{
+	using namespace EngineTestMetaSoundPatchBuilderPrivate;
+	using namespace Metasound;
+	using namespace Metasound::Frontend;
+
+	constexpr float NodeDefaultFreq = 100.0f;
+
+	EMetaSoundBuilderResult Result = EMetaSoundBuilderResult::Failed;
+	UMetaSoundPatchBuilder& Builder = CreatePatchBuilderChecked(*this, "TemplateInputInjection_Test", { });
+
+	// Sine Oscillator Node
+	const FMetaSoundNodeHandle OscNode = Builder.AddNodeByClassName({ "UE", "Sine", "Audio" }, Result, 1);
+	AddErrorIfFalse(Result == EMetaSoundBuilderResult::Succeeded && OscNode.IsSet(), TEXT("Failed to create new MetaSound node by class name"));
+
+	auto MakeAndConnectSinInputToGraphInput = [&](FName InputName, FName TypeName, const FVector2D* Location, const FMetasoundFrontendLiteral& InLiteral, bool bIsCtorPin)
+	{
+		const FMetaSoundBuilderNodeOutputHandle InputNodeOutputHandle = Builder.AddGraphInputNode(InputName, TypeName, InLiteral, Result, bIsCtorPin);
+		AddErrorIfFalse(Result == EMetaSoundBuilderResult::Succeeded && InputNodeOutputHandle.IsSet(), TEXT("Failed to create new MetaSound graph input"));
+
+		const FMetaSoundBuilderNodeInputHandle OscNodeInput = Builder.FindNodeInputByName(OscNode, InputName, Result);
+		AddErrorIfFalse(Result == EMetaSoundBuilderResult::Succeeded && OscNodeInput.IsSet(), TEXT("Failed to find Sine Oscillator node input"));
+
+		Builder.ConnectNodes(InputNodeOutputHandle, OscNodeInput, Result);
+		AddErrorIfFalse(Result == EMetaSoundBuilderResult::Succeeded, TEXT("Failed to connect input node to node input"));
+
+#if WITH_EDITOR
+		if (Location)
+		{
+			Builder.SetNodeLocation(FMetaSoundNodeHandle{ InputNodeOutputHandle.NodeID }, *Location, Result);
+			AddErrorIfFalse(Result == EMetaSoundBuilderResult::Succeeded, TEXT("Failed to set input node location"));
+		}
+#endif // WITH_EDITOR
+	};
+	FMetasoundFrontendLiteral Literal;
+
+	Literal.Set(true);
+	MakeAndConnectSinInputToGraphInput("Enabled", GetMetasoundDataTypeName<bool>(), nullptr, Literal, true);
+
+	Literal.Set(100.0f);
+	FVector2D FreqLoc(100, 100);
+	MakeAndConnectSinInputToGraphInput("Frequency", GetMetasoundDataTypeName<float>(), nullptr, Literal, false);
+
+	Literal.Set(FMetasoundFrontendLiteral::FDefault { });
+	FVector2D SyncLoc(300, 300);
+	MakeAndConnectSinInputToGraphInput("Sync", GetMetasoundDataTypeName<FTrigger>(), nullptr, Literal, false);
+
+#if WITH_EDITOR
+	constexpr bool bForceNodeCreation = false;
+	Builder.InjectInputTemplateNodes(bForceNodeCreation, Result);
+	AddErrorIfFalse(Result == EMetaSoundBuilderResult::Succeeded, TEXT("Failed to inject input template node"));
+#endif // WITH_EDITOR
+
+	return true;
+}
 #endif // WITH_DEV_AUTOMATION_TESTS
diff --git a/Engine/Plugins/Runtime/Metasound/Source/MetasoundEngineTest/Private/EngineTestMetasoundSource.cpp b/Engine/Plugins/Runtime/Metasound/Source/MetasoundEngineTest/Private/EngineTestMetasoundSource.cpp
--- a/Engine/Plugins/Runtime/Metasound/Source/MetasoundEngineTest/Private/EngineTestMetasoundSource.cpp
+++ b/Engine/Plugins/Runtime/Metasound/Source/MetasoundEngineTest/Private/EngineTestMetasoundSource.cpp
@@ -310,7 +310,7 @@ bool FMetaSoundSourceBuilderAuditionLatentCommand::Update()
 			if (const USoundBase* InitSound = AudioComponent->GetSound())
 			{
 				const UMetaSoundSource* InitMetaSound = CastChecked<UMetaSoundSource>(InitSound);
-				return InitMetaSound->GetDocumentChecked().RootGraph.Metadata.GetClassName().GetFullName();
+				return InitMetaSound->GetConstDocumentChecked().RootGraph.Metadata.GetClassName().GetFullName();
 			}
 
 			return { };
@@ -384,7 +384,10 @@ bool FAudioMetasoundSourceTest::RunTest(const FString& Parameters)
 	UMetaSoundSource* MetaSoundSource = NewObject<UMetaSoundSource>(GetTransientPackage(), FName(*LexToString(FGuid::NewGuid())));;
 	if (ensure(nullptr != MetaSoundSource))
 	{
+		// TODO: Move to using builder to swap in doc
+		PRAGMA_DISABLE_DEPRECATION_WARNINGS
 		MetaSoundSource->SetDocument(CreateMonoSourceDocument());
+		PRAGMA_ENABLE_DEPRECATION_WARNINGS
 
 		if (UAudioComponent* AudioComponent = CreateTestComponent(*this, MetaSoundSource))
 		{
