Version Metasound Document to include all ed data and make all Metasound EdGraph data transient

diff --git a/Engine/Plugins/Runtime/Metasound/Source/MetasoundFrontend/Private/MetasoundAssetBase.cpp b/Engine/Plugins/Runtime/Metasound/Source/MetasoundFrontend/Private/MetasoundAssetBase.cpp
--- a/Engine/Plugins/Runtime/Metasound/Source/MetasoundFrontend/Private/MetasoundAssetBase.cpp
+++ b/Engine/Plugins/Runtime/Metasound/Source/MetasoundFrontend/Private/MetasoundAssetBase.cpp
@@ -125,6 +125,68 @@ namespace Metasound
 				FMetasoundFrontendClass FrontendClass;
 				FNodeClassInfo ClassInfo;
 			};
+
+			void GetUpdatePathForDocument(const FMetasoundFrontendVersion& InCurrentVersion, const FMetasoundFrontendVersion& InTargetVersion, TArray<const IInterfaceRegistryEntry*>& OutUpgradePath)
+			{
+				if (InCurrentVersion.Name == InTargetVersion.Name)
+				{
+					// Get all associated registered interfaces
+					TArray<FMetasoundFrontendVersion> RegisteredVersions = ISearchEngine::Get().FindAllRegisteredInterfacesWithName(InTargetVersion.Name);
+
+					// Filter registry entries that exist between current version and target version
+					auto FilterRegistryEntries = [&InCurrentVersion, &InTargetVersion](const FMetasoundFrontendVersion& InVersion)
+					{
+						const bool bIsGreaterThanCurrent = InVersion.Number > InCurrentVersion.Number;
+						const bool bIsLessThanOrEqualToTarget = InVersion.Number <= InTargetVersion.Number;
+
+						return bIsGreaterThanCurrent && bIsLessThanOrEqualToTarget;
+					};
+					RegisteredVersions = RegisteredVersions.FilterByPredicate(FilterRegistryEntries);
+
+					// sort registry entries to create an ordered upgrade path.
+					RegisteredVersions.Sort();
+
+					// Get registry entries from registry keys.
+					auto GetRegistryEntry = [](const FMetasoundFrontendVersion& InVersion)
+					{
+						FInterfaceRegistryKey Key = GetInterfaceRegistryKey(InVersion);
+						return IInterfaceRegistry::Get().FindInterfaceRegistryEntry(Key);
+					};
+					Algo::Transform(RegisteredVersions, OutUpgradePath, GetRegistryEntry);
+				}
+			}
+
+			bool UpdateDocumentInterface(const TArray<const IInterfaceRegistryEntry*>& InUpgradePath, const FMetasoundFrontendVersion& InterfaceVersion, FDocumentHandle InDocument)
+			{
+				const FMetasoundFrontendVersionNumber* LastVersionUpdated = nullptr;
+				for (const IInterfaceRegistryEntry* Entry : InUpgradePath)
+				{
+					if (ensure(nullptr != Entry))
+					{
+						if (Entry->UpdateRootGraphInterface(InDocument))
+						{
+							LastVersionUpdated = &Entry->GetInterface().Version.Number;
+						}
+					}
+				}
+
+				if (LastVersionUpdated)
+				{
+#if WITH_EDITOR
+					const FString AssetName = *InDocument->GetRootGraphClass().Metadata.GetDisplayName().ToString();
+#else
+					const FString AssetName = *InDocument->GetRootGraphClass().Metadata.GetClassName().ToString();
+#endif // !WITH_EDITOR
+					UE_LOG(LogMetaSound, Display, TEXT("Asset '%s' interface '%s' updated: '%s' --> '%s'"),
+						*AssetName,
+						*InterfaceVersion.Name.ToString(),
+						*InterfaceVersion.Number.ToString(),
+						*LastVersionUpdated->ToString());
+					return true;
+				}
+
+				return false;
+			}
 		} // namespace AssetBasePrivate
 
 		namespace ConsoleVariables
@@ -205,6 +267,11 @@ namespace Metasound
 
 const FString FMetasoundAssetBase::FileExtension(TEXT(".metasound"));
 
+bool FMetasoundAssetBase::ConformObjectDataToInterfaces()
+{
+	return false;
+}
+
 void FMetasoundAssetBase::RegisterGraphWithFrontend(Metasound::Frontend::FMetaSoundAssetRegistrationOptions InRegistrationOptions)
 {
 	using namespace Metasound;
@@ -271,15 +338,8 @@ void FMetasoundAssetBase::RegisterGraphWithFrontend(Metasound::Frontend::FMetaSo
 
 	// Register graphs async by default;
 	const bool bAsync = !ConsoleVariables::bDisableAsyncGraphRegistration;
-	// Force a copy if async registration is enabled and we need to protect against
-	// race conditions from external modifications.
-	bool bForceCopy = (IsBuilderActive() && bAsync);
-#if WITH_EDITOR
-	bForceCopy |= InRegistrationOptions.bRegisterCopyIfAsync;
-#endif // WITH_EDITOR
-
-	GraphRegistryKey = FRegistryContainerImpl::Get().RegisterGraph(Owner, bAsync, bForceCopy);
 
+	GraphRegistryKey = FRegistryContainerImpl::Get().RegisterGraph(Owner, bAsync);
 	if (GraphRegistryKey.IsValid())
 	{
 #if WITH_EDITORONLY_DATA
@@ -329,7 +389,7 @@ void FMetasoundAssetBase::CookMetaSound()
 
 	{
 		// Performs document transforms on local copy, which reduces document footprint & renders transforming unnecessary unless altered at runtime when registering
-		FMetaSoundFrontendDocumentBuilder DocBuilder(Owner);
+		FMetaSoundFrontendDocumentBuilder& DocBuilder = IMetaSoundAssetManager::GetChecked().AttachDocumentBuilderChecked(*Owner);
 		const bool bContainsTemplateDependency = DocBuilder.ContainsDependencyOfType(EMetasoundFrontendClassType::Template);
 		if (bContainsTemplateDependency)
 		{
@@ -419,26 +479,9 @@ void FMetasoundAssetBase::UnregisterGraphWithFrontend()
 	}
 }
 
-void FMetasoundAssetBase::SetMetadata(FMetasoundFrontendClassMetadata& InMetadata)
-{
-	FMetasoundFrontendDocument& Doc = GetDocumentChecked();
-	Doc.RootGraph.Metadata = InMetadata;
-
-	if (Doc.RootGraph.Metadata.GetType() != EMetasoundFrontendClassType::Graph)
-	{
-		UE_LOG(LogMetaSound, Display, TEXT("Forcing class type to EMetasoundFrontendClassType::Graph on root graph metadata"));
-		Doc.RootGraph.Metadata.SetType(EMetasoundFrontendClassType::Graph);
-	}
-}
-
-bool FMetasoundAssetBase::GetDeclaredInterfaces(TArray<const Metasound::Frontend::IInterfaceRegistryEntry*>& OutInterfaces) const
-{
-	return FMetaSoundFrontendDocumentBuilder::FindDeclaredInterfaces(GetDocumentChecked(), OutInterfaces);
-}
-
 bool FMetasoundAssetBase::IsInterfaceDeclared(const FMetasoundFrontendVersion& InVersion) const
 {
-	return GetDocumentChecked().Interfaces.Contains(InVersion);
+	return GetConstDocumentChecked().Interfaces.Contains(InVersion);
 }
 
 Metasound::Frontend::FNodeClassInfo FMetasoundAssetBase::GetAssetClassInfo() const
@@ -448,38 +491,13 @@ Metasound::Frontend::FNodeClassInfo FMetasoundAssetBase::GetAssetClassInfo() con
 	const UObject* Owner = GetOwningAsset();
 	check(Owner);
 	TScriptInterface<const IMetaSoundDocumentInterface> DocInterface((UObject*)Owner);
-	return FNodeClassInfo { GetDocumentChecked().RootGraph, DocInterface->GetAssetPathChecked() };
+	return FNodeClassInfo { GetConstDocumentChecked().RootGraph, DocInterface->GetAssetPathChecked() };
 }
 
-void FMetasoundAssetBase::SetDocument(const FMetasoundFrontendDocument& InDocument, bool bMarkDirty)
+void FMetasoundAssetBase::SetDocument(FMetasoundFrontendDocument InDocument, bool bMarkDirty)
 {
-	FMetasoundFrontendDocument& Document = GetDocumentChecked();
-	Document = InDocument;
-	if (bMarkDirty)
-	{
-		MarkMetasoundDocumentDirty();
-	}
-}
-
-void FMetasoundAssetBase::AddDefaultInterfaces()
-{
-	using namespace Metasound::Frontend;
-
-	UObject* OwningAsset = GetOwningAsset();
-	check(OwningAsset);
-
-	UClass* AssetClass = OwningAsset->GetClass();
-	check(AssetClass);
-	const FTopLevelAssetPath& ClassPath = AssetClass->GetClassPathName();
-
-	TArray<FMetasoundFrontendVersion> InitVersions = ISearchEngine::Get().FindUClassDefaultInterfaceVersions(ClassPath);
-	FModifyRootGraphInterfaces({ }, InitVersions).Transform(GetDocumentChecked());
-}
-
-void FMetasoundAssetBase::SetDocument(FMetasoundFrontendDocument&& InDocument, bool bMarkDirty)
-{
-	FMetasoundFrontendDocument& Document = GetDocumentChecked();
-	Document = MoveTemp(InDocument);
+	FMetasoundFrontendDocument* Document = GetDocumentAccessPtr().Get();
+	*Document = MoveTemp(InDocument);
 	if (bMarkDirty)
 	{
 		MarkMetasoundDocumentDirty();
@@ -495,41 +513,47 @@ bool FMetasoundAssetBase::VersionAsset()
 	constexpr bool bIsDeterministic = true;
 	FDocumentIDGenerator::FScopeDeterminism DeterminismScope(bIsDeterministic);
 
-	FName AssetName;
-	FString AssetPath;
-	if (const UObject* OwningAsset = GetOwningAsset())
-	{
-		AssetName = FName(OwningAsset->GetName());
-		AssetPath = OwningAsset->GetPathName();
-	}
+	UObject* OwningAsset = GetOwningAsset();
+	check(OwningAsset);
 
-	FMetasoundFrontendDocument* Doc = GetDocumentAccessPtr().Get();
-	if (!ensure(Doc))
-	{
-		return false;
-	}
+	FName AssetName = FName(OwningAsset->GetName());
+	FString AssetPath = OwningAsset->GetPathName();
 
-	bool bDidEdit = Doc->VersionInterfaces();
+	bool bDidEdit = false;
 
+#if WITH_EDITORONLY_DATA
 	// Version Document Model
-	FDocumentHandle DocHandle = GetDocumentHandle();
+	if (GetConstDocumentChecked().RequiresInterfaceVersioning())
 	{
-		bDidEdit |= FVersionDocument(AssetName, AssetPath).Transform(DocHandle);
+		GetDocumentChecked().VersionInterfaces();
+		bDidEdit = true;
 	}
+#endif // WITH_EDITORONLY_DATA
+
+	TScriptInterface<IMetaSoundDocumentInterface> Interface(OwningAsset);
+	bDidEdit |= FVersionDocument(AssetName, AssetPath).Transform(Interface);
 
 	// Version Interfaces. Has to be re-run until no pass reports an update in case
 	// versions fork (ex. an interface splits into two newly named interfaces).
+	// TODO: Need to add function to determine if interfaces require updating before
+	// running update to avoid invalidating builder cache when not necessary. This
+	// will require either re-writing old transforms with the builder API or just
+	// leaving handle/controller logic and moving to deprecated implementation that
+	// is bypassed on versioned documents in most cases.
 	{
 		bool bInterfaceUpdated = false;
 		bool bPassUpdated = true;
+		const IMetaSoundDocumentInterface* ConstInterface = Interface.GetInterface();
 		while (bPassUpdated)
 		{
 			bPassUpdated = false;
 
-			const TArray<FMetasoundFrontendVersion> Versions = Doc->Interfaces.Array();
+			const TArray<FMetasoundFrontendVersion> Versions = ConstInterface->GetConstDocument().Interfaces.Array();
 			for (const FMetasoundFrontendVersion& Version : Versions)
 			{
-				bPassUpdated |= FUpdateRootGraphInterface(Version, GetOwningAssetName()).Transform(DocHandle);
+				FUpdateRootGraphInterface UpdateTransform(Version, GetOwningAssetName());
+				FConstDocumentHandle ConstDoc = GetDocumentHandle();
+				bPassUpdated = TryUpdateInterfaceFromVersion(Version);
 			}
 
 			bInterfaceUpdated |= bPassUpdated;
@@ -537,7 +561,7 @@ bool FMetasoundAssetBase::VersionAsset()
 
 		if (bInterfaceUpdated)
 		{
-			ConformObjectDataToInterfaces();
+			Interface->ConformObjectToDocument();
 		}
 		bDidEdit |= bInterfaceUpdated;
 	}
@@ -666,14 +690,17 @@ FMetasoundFrontendDocumentModifyContext& FMetasoundAssetBase::GetModifyContext()
 	return GetDocumentChecked().Metadata.ModifyContext;
 }
 
+const FMetasoundFrontendDocumentModifyContext& FMetasoundAssetBase::GetConstModifyContext() const
+{
+	return GetConstDocumentChecked().Metadata.ModifyContext;
+}
+
 const FMetasoundFrontendDocumentModifyContext& FMetasoundAssetBase::GetModifyContext() const
 {
-	return GetDocumentChecked().Metadata.ModifyContext;
+	return GetConstDocumentChecked().Metadata.ModifyContext;
 }
 #endif // WITH_EDITOR
 
-
-
 bool FMetasoundAssetBase::IsRegistered() const
 {
 	using namespace Metasound::Frontend;
@@ -735,22 +762,6 @@ bool FMetasoundAssetBase::AddingReferenceCausesLoop(const FSoftObjectPath& InRef
 	return bCausesLoop;
 }
 
-void FMetasoundAssetBase::ConvertFromPreset()
-{
-	using namespace Metasound::Frontend;
-	FGraphHandle GraphHandle = GetRootGraphHandle();
-
-#if WITH_EDITOR
-	FMetasoundFrontendGraphStyle Style = GraphHandle->GetGraphStyle();
-	Style.bIsGraphEditable = true;
-	GraphHandle->SetGraphStyle(Style);
-#endif // WITH_EDITOR
-
-	FMetasoundFrontendGraphClassPresetOptions PresetOptions = GraphHandle->GetGraphPresetOptions();
-	PresetOptions.bIsPreset = false;
-	GraphHandle->SetGraphPresetOptions(PresetOptions);
-}
-
 TArray<FMetasoundAssetBase::FSendInfoAndVertexName> FMetasoundAssetBase::GetSendInfos(uint64 InInstanceID) const
 {
 	return TArray<FSendInfoAndVertexName>();
@@ -777,7 +788,7 @@ bool FMetasoundAssetBase::MarkMetasoundDocumentDirty() const
 {
 	if (const UObject* OwningAsset = GetOwningAsset())
 	{
-		return ensure(OwningAsset->MarkPackageDirty());
+		return OwningAsset->MarkPackageDirty();
 	}
 	return false;
 }
@@ -840,6 +851,17 @@ bool FMetasoundAssetBase::ImportFromJSONAsset(const FString& InAbsolutePath)
 	return false;
 }
 
+const FMetasoundFrontendDocument& FMetasoundAssetBase::GetConstDocumentChecked() const
+{
+	const UObject* Owner = GetOwningAsset();
+	check(Owner);
+
+	// Annoying const cast requirement because ScriptInterface isn't ctor isn't const correct :(
+	TScriptInterface<IMetaSoundDocumentInterface> DocInterface = const_cast<UObject*>(Owner);
+
+	return DocInterface->GetConstDocument();
+}
+
 FMetasoundFrontendDocument& FMetasoundAssetBase::GetDocumentChecked()
 {
 	FMetasoundFrontendDocument* Document = GetDocumentAccessPtr().Get();
@@ -849,10 +871,7 @@ FMetasoundFrontendDocument& FMetasoundAssetBase::GetDocumentChecked()
 
 const FMetasoundFrontendDocument& FMetasoundAssetBase::GetDocumentChecked() const
 {
-	const FMetasoundFrontendDocument* Document = GetDocumentConstAccessPtr().Get();
-
-	check(nullptr != Document);
-	return *Document;
+	return GetConstDocumentChecked();
 }
 
 const Metasound::Frontend::FGraphRegistryKey& FMetasoundAssetBase::GetGraphRegistryKey() const
@@ -960,11 +979,73 @@ void FMetasoundAssetBase::UpdateAssetRegistry()
 }
 #endif
 
+bool FMetasoundAssetBase::TryUpdateInterfaceFromVersion(const FMetasoundFrontendVersion& Version)
+{
+	using namespace Metasound::Frontend;
+	using namespace AssetBasePrivate;
+
+	FConstDocumentHandle ConstDoc = IDocumentController::CreateDocumentHandle(GetDocumentConstAccessPtr());
+	FMetasoundFrontendInterface TargetInterface = GetInterfaceToVersion(Version, ConstDoc);
+
+	if (TargetInterface.Version.IsValid())
+	{
+		TArray<const IInterfaceRegistryEntry*> UpgradePath;
+		GetUpdatePathForDocument(Version, TargetInterface.Version, UpgradePath);
+		const bool bUpdated = UpdateDocumentInterface(UpgradePath, Version, GetDocumentHandle());
+		ensureMsgf(bUpdated, TEXT("Target interface '%s' was out-of-date but interface failed to be updated"), *TargetInterface.Version.ToString());
+		return bUpdated;
+	}
+
+	return false;
+}
+
+FMetasoundFrontendInterface FMetasoundAssetBase::GetInterfaceToVersion(const FMetasoundFrontendVersion& InterfaceVersion, Metasound::Frontend::FConstDocumentHandle InDocument) const
+{
+	using namespace Metasound::Frontend;
+
+	// Find registered target interface.
+	FMetasoundFrontendInterface TargetInterface;
+	bool bFoundTargetInterface = ISearchEngine::Get().FindInterfaceWithHighestVersion(InterfaceVersion.Name, TargetInterface);
+	if (!bFoundTargetInterface)
+	{
+		UE_LOG(LogMetaSound, Warning,
+			TEXT("Could not check for interface updates. Target interface is not registered [InterfaceVersion:%s] when attempting to update root graph of asset (%s). "
+				"Ensure that the module which registers the interface has been loaded before the asset is loaded."),
+			*InterfaceVersion.ToString(),
+			*GetOwningAssetName());
+		return { };
+	}
+
+	if (TargetInterface.Version == InterfaceVersion)
+	{
+		return { };
+	}
+
+	return TargetInterface;
+}
+
 TSharedPtr<FMetasoundFrontendDocument> FMetasoundAssetBase::PreprocessDocument()
 {
 	return nullptr;
 }
 
+#if WITH_EDITORONLY_DATA
+bool FMetasoundAssetBase::GetVersionedOnLoad() const
+{
+	return bVersionedOnLoad;
+}
+
+void FMetasoundAssetBase::ClearVersionedOnLoad()
+{
+	bVersionedOnLoad = false;
+}
+
+void FMetasoundAssetBase::SetVersionedOnLoad()
+{
+	bVersionedOnLoad = true;
+}
+#endif // WITH_EDITORONLY_DATA
+
 PRAGMA_DISABLE_DEPRECATION_WARNINGS
 const FMetasoundAssetBase::FRuntimeData& FMetasoundAssetBase::GetRuntimeData() const
 {
diff --git a/Engine/Plugins/Runtime/Metasound/Source/MetasoundFrontend/Private/MetasoundFrontend.cpp b/Engine/Plugins/Runtime/Metasound/Source/MetasoundFrontend/Private/MetasoundFrontend.cpp
--- a/Engine/Plugins/Runtime/Metasound/Source/MetasoundFrontend/Private/MetasoundFrontend.cpp
+++ b/Engine/Plugins/Runtime/Metasound/Source/MetasoundFrontend/Private/MetasoundFrontend.cpp
@@ -18,6 +18,7 @@
 #include "MetasoundRouter.h"
 #include "MetasoundVertex.h"
 #include "Modules/ModuleManager.h"
+#include "NodeTemplates/MetasoundFrontendNodeTemplateInput.h"
 #include "NodeTemplates/MetasoundFrontendNodeTemplateReroute.h"
 #include "StructDeserializer.h"
 #include "StructSerializer.h"
@@ -94,8 +95,8 @@ class FMetasoundFrontendModule : public IModuleInterface
 	{
 		using namespace Metasound::Frontend;
 
-		TUniquePtr<INodeTemplate> RerouteTemplate = MakeUnique<FRerouteNodeTemplate>();
-		RegisterNodeTemplate(MoveTemp(RerouteTemplate));
+		RegisterNodeTemplate(MakeUnique<FInputNodeTemplate>());
+		RegisterNodeTemplate(MakeUnique<FRerouteNodeTemplate>());
 
 		FMetasoundFrontendRegistryContainer* Registry = FMetasoundFrontendRegistryContainer::Get();
 		if (ensure(nullptr != Registry))
@@ -108,6 +109,7 @@ class FMetasoundFrontendModule : public IModuleInterface
 	{
 		using namespace Metasound::Frontend;
 
+		UnregisterNodeTemplate(FInputNodeTemplate::ClassName, FInputNodeTemplate::VersionNumber);
 		UnregisterNodeTemplate(FRerouteNodeTemplate::ClassName, FRerouteNodeTemplate::VersionNumber);
 	}
 };
diff --git a/Engine/Plugins/Runtime/Metasound/Source/MetasoundFrontend/Private/MetasoundFrontendController.cpp b/Engine/Plugins/Runtime/Metasound/Source/MetasoundFrontend/Private/MetasoundFrontendController.cpp
--- a/Engine/Plugins/Runtime/Metasound/Source/MetasoundFrontend/Private/MetasoundFrontendController.cpp
+++ b/Engine/Plugins/Runtime/Metasound/Source/MetasoundFrontend/Private/MetasoundFrontendController.cpp
@@ -4,6 +4,7 @@
 
 #include "MetasoundFrontendInvalidController.h"
 #include "MetasoundFrontendDocumentController.h"
+#include "NodeTemplates/MetasoundFrontendNodeTemplateInput.h"
 #include "NodeTemplates/MetasoundFrontendNodeTemplateReroute.h"
 
 
@@ -11,6 +12,17 @@ namespace Metasound
 {
 	namespace Frontend
 	{
+		namespace ControllerPrivate
+		{
+			// This check for class name is a hack due to the fact that handles do not provide a builder
+			// in order to do a proper look-up of a rerouted output's access/data type. Once calling systems of the
+			// recurse functions utilizing this function are refactored to use the builder API, this can be removed.
+			bool IsRerouteClass(const FMetasoundFrontendClassName& ClassName)
+			{
+				return ClassName == FRerouteNodeTemplate::ClassName || ClassName == FInputNodeTemplate::ClassName;
+			}
+		}
+
 		FOutputHandle IOutputController::GetInvalidHandle()
 		{
 			static FOutputHandle Invalid = MakeShared<FInvalidOutputController>();
@@ -55,6 +67,15 @@ namespace Metasound
 
 		FDocumentHandle IDocumentController::CreateDocumentHandle(FMetasoundFrontendDocument& InDocument)
 		{
+			// Mutation of a document via the soft deprecated access ptr/controller system is not tracked by
+			// the builder registry, so the document cache is invalidated here. It is discouraged to mutate
+			// documents using both systems at the same time as it can corrupt a builder document's cache.
+			if (IDocumentBuilderRegistry* DocRegistry = IDocumentBuilderRegistry::Get())
+			{
+				const FMetasoundFrontendClassName& Name = InDocument.RootGraph.Metadata.GetClassName();
+				DocRegistry->InvalidateDocumentCache(Name);
+			}
+
 			return CreateDocumentHandle(MakeAccessPtr<FDocumentAccessPtr>(InDocument.AccessPoint, InDocument));
 		}
 
@@ -81,14 +102,13 @@ namespace Metasound
 
 		FConstOutputHandle FindReroutedOutput(FConstOutputHandle InOutputHandle)
 		{
-			using namespace Frontend;
-
 			if (InOutputHandle->IsValid())
 			{
 				FConstNodeHandle NodeHandle = InOutputHandle->GetOwningNode();
 				if (NodeHandle->IsValid())
 				{
-					if (NodeHandle->GetClassMetadata().GetClassName() == FRerouteNodeTemplate::ClassName)
+					const FMetasoundFrontendClassName& ClassName = NodeHandle->GetClassMetadata().GetClassName();
+					if (ControllerPrivate::IsRerouteClass(ClassName))
 					{
 						TArray<FConstInputHandle> Inputs = NodeHandle->GetConstInputs();
 						if (!Inputs.IsEmpty())
@@ -109,14 +129,14 @@ namespace Metasound
 
 		void FindReroutedInputs(FConstInputHandle InHandleToCheck, TArray<FConstInputHandle>& InOutInputHandles)
 		{
-			using namespace Frontend;
+			using namespace ControllerPrivate;
 
 			if (InHandleToCheck->IsValid())
 			{
 				FConstNodeHandle NodeHandle = InHandleToCheck->GetOwningNode();
 				if (NodeHandle->IsValid())
 				{
-					if (NodeHandle->GetClassMetadata().GetClassName() == FRerouteNodeTemplate::ClassName)
+					if (ControllerPrivate::IsRerouteClass(NodeHandle->GetClassMetadata().GetClassName()))
 					{
 						TArray<FConstOutputHandle> Outputs = NodeHandle->GetConstOutputs();
 						for (FConstOutputHandle& OutputHandle : Outputs)
@@ -139,14 +159,12 @@ namespace Metasound
 
 		void IterateReroutedInputs(FConstInputHandle InHandleToCheck, TFunctionRef<void(FConstInputHandle)> Func)
 		{
-			using namespace Frontend;
-
 			if (InHandleToCheck->IsValid())
 			{
 				FConstNodeHandle NodeHandle = InHandleToCheck->GetOwningNode();
 				if (NodeHandle->IsValid())
 				{
-					if (NodeHandle->GetClassMetadata().GetClassName() == FRerouteNodeTemplate::ClassName)
+					if (ControllerPrivate::IsRerouteClass(NodeHandle->GetClassMetadata().GetClassName()))
 					{
 						TArray<FConstOutputHandle> Outputs = NodeHandle->GetConstOutputs();
 						for (FConstOutputHandle& OutputHandle : Outputs)
diff --git a/Engine/Plugins/Runtime/Metasound/Source/MetasoundFrontend/Private/MetasoundFrontendDocumentBuilder.cpp b/Engine/Plugins/Runtime/Metasound/Source/MetasoundFrontend/Private/MetasoundFrontendDocumentBuilder.cpp
--- a/Engine/Plugins/Runtime/Metasound/Source/MetasoundFrontend/Private/MetasoundFrontendDocumentBuilder.cpp
+++ b/Engine/Plugins/Runtime/Metasound/Source/MetasoundFrontend/Private/MetasoundFrontendDocumentBuilder.cpp
@@ -13,6 +13,7 @@
 #include "MetasoundAssetManager.h"
 #include "MetasoundDocumentInterface.h"
 #include "MetasoundFrontendController.h"
+#include "MetasoundFrontendDataTypeRegistry.h"
 #include "MetasoundFrontendDocumentCache.h"
 #include "MetasoundFrontendDocument.h"
 #include "MetasoundFrontendDocumentIdGenerator.h"
@@ -23,6 +24,8 @@
 #include "MetasoundFrontendSearchEngine.h"
 #include "MetasoundFrontendTransform.h"
 #include "MetasoundTrace.h"
+#include "MetasoundVariableNodes.h"
+#include "NodeTemplates/MetasoundFrontendNodeTemplateInput.h"
 
 #include UE_INLINE_GENERATED_CPP_BY_NAME(MetasoundFrontendDocumentBuilder)
 
@@ -31,6 +34,60 @@ namespace Metasound::Frontend
 {
 	namespace DocumentBuilderPrivate
 	{
+		bool FindInputRegistryClass(FName TypeName, EMetasoundFrontendVertexAccessType AccessType, FMetasoundFrontendClass& OutClass)
+		{
+			switch (AccessType)
+			{
+				case EMetasoundFrontendVertexAccessType::Value:
+				{
+					return IDataTypeRegistry::Get().GetFrontendConstructorInputClass(TypeName, OutClass);
+				}
+				break;
+
+				case EMetasoundFrontendVertexAccessType::Reference:
+				{
+					return IDataTypeRegistry::Get().GetFrontendInputClass(TypeName, OutClass);
+				}
+				break;
+
+				case EMetasoundFrontendVertexAccessType::Unset:
+				default:
+				{
+					checkNoEntry();
+				}
+				break;
+			}
+
+			return false;
+		}
+
+		bool FindOutputRegistryClass(FName TypeName, EMetasoundFrontendVertexAccessType AccessType, FMetasoundFrontendClass& OutClass)
+		{
+			switch (AccessType)
+			{
+				case EMetasoundFrontendVertexAccessType::Value:
+				{
+					return IDataTypeRegistry::Get().GetFrontendConstructorOutputClass(TypeName, OutClass);
+				}
+				break;
+
+				case EMetasoundFrontendVertexAccessType::Reference:
+				{
+					return IDataTypeRegistry::Get().GetFrontendOutputClass(TypeName, OutClass);
+				}
+				break;
+
+				case EMetasoundFrontendVertexAccessType::Unset:
+				default:
+				{
+					checkNoEntry();
+				}
+				break;
+			}
+
+			return false;
+		}
+
 		bool TryGetInterfaceBoundEdges(
 			const FGuid& InFromNodeID,
 			const TSet<FMetasoundFrontendVersion>& InFromNodeInterfaces,
@@ -398,7 +455,7 @@ namespace Metasound::Frontend
 					const FName NodeName = Node->Name;
 					if (AddedNames.Contains(NodeName))
 					{
-						Node->Style.Display.Locations.Add(FGuid(), NextLocation);
+						Node->Style.Display.Locations.Add(FGuid::NewGuid(), NextLocation);
 						NextLocation += DisplayStyle::NodeLayout::DefaultOffsetY;
 					}
 					else
@@ -455,7 +512,7 @@ namespace Metasound::Frontend
 				}
 
 #if WITH_EDITORONLY_DATA
-				if (bAddedVertices && Options.bSetDefaultNodeLocations)
+				if (bAddedVertices && Options.bSetDefaultNodeLocations && !IsRunningCookCommandlet())
 				{
 					TArray<FMetasoundFrontendNode>& Nodes = Document.RootGraph.Graph.Nodes;
 					// Sort/Place Inputs
@@ -588,6 +645,11 @@ UMetaSoundBuilderDocument& UMetaSoundBuilderDocument::Create(const IMetaSoundDoc
 	return *DocObject;
 }
 
+bool UMetaSoundBuilderDocument::ConformObjectToDocument()
+{
+	return false;
+}
+
 FTopLevelAssetPath UMetaSoundBuilderDocument::GetAssetPathChecked() const
 {
 	FTopLevelAssetPath Path;
@@ -622,32 +684,12 @@ void UMetaSoundBuilderDocument::OnFinishActiveBuilder()
 	// Nothing to do here. UMetaSoundBuilderDocuments are always being used by builders
 }
 
-FMetaSoundFrontendDocumentBuilder::FMetaSoundFrontendDocumentBuilder()
-	: DocumentDelegates(MakeShared<Metasound::Frontend::FDocumentModifyDelegates>())
-{
-}
-
-FMetaSoundFrontendDocumentBuilder::FMetaSoundFrontendDocumentBuilder(TScriptInterface<IMetaSoundDocumentInterface> InDocumentInterface)
+FMetaSoundFrontendDocumentBuilder::FMetaSoundFrontendDocumentBuilder(TScriptInterface<IMetaSoundDocumentInterface> InDocumentInterface, TSharedPtr<Metasound::Frontend::FDocumentModifyDelegates> InDocumentDelegates, bool bPrimeCache)
 	: DocumentInterface(InDocumentInterface)
-	, DocumentDelegates(MakeShared<Metasound::Frontend::FDocumentModifyDelegates>())
 {
-	if (DocumentInterface)
-	{
-		BeginBuilding();
-		InitCacheInternal();
-	}
+	BeginBuilding(InDocumentDelegates, bPrimeCache);
 }
 
-FMetaSoundFrontendDocumentBuilder::FMetaSoundFrontendDocumentBuilder(TScriptInterface<IMetaSoundDocumentInterface> InDocumentInterface, TSharedRef<Metasound::Frontend::FDocumentModifyDelegates> InDocumentDelegates)
-	: DocumentInterface(InDocumentInterface)
-	, DocumentDelegates(InDocumentDelegates)
-{
-	if (DocumentInterface)
-	{
-		BeginBuilding();
-		InitCacheInternal();
-	}
-}
 FMetaSoundFrontendDocumentBuilder::~FMetaSoundFrontendDocumentBuilder()
 {
 	FinishBuilding();
@@ -682,7 +724,15 @@ void FMetaSoundFrontendDocumentBuilder::AddEdge(FMetasoundFrontendEdge&& InNewEd
 {
 	using namespace Metasound::Frontend;
 
-	checkf(CanAddEdge(InNewEdge), TEXT("Attempted call to AddEdge in MetaSound Builder where not valid."));
+#if DO_CHECK
+	{
+		const IDocumentGraphEdgeCache& EdgeCache = DocumentCache->GetEdgeCache();
+		checkf(!EdgeCache.IsNodeInputConnected(InNewEdge.ToNodeID, InNewEdge.ToVertexID), TEXT("Failed to add edge in MetaSound Builder: Destination input already connected"));
+
+		const EInvalidEdgeReason Reason = IsValidEdge(InNewEdge);
+		checkf(Reason == Metasound::Frontend::EInvalidEdgeReason::None, TEXT("Attempted call to AddEdge in MetaSound Builder where edge is invalid: %s."), *LexToString(Reason));
+	}
+#endif // DO_CHECK
 
 	FMetasoundFrontendDocument& Document = GetDocument();
 	FMetasoundFrontendGraph& Graph = Document.RootGraph.Graph;
@@ -945,42 +995,15 @@ const FMetasoundFrontendNode* FMetaSoundFrontendDocumentBuilder::AddGraphInput(c
 		return nullptr;
 	}
 
-	auto FindRegistryClass = [&InClassInput](FMetasoundFrontendClass& OutClass) -> bool
-	{
-		switch (InClassInput.AccessType)
-		{
-		case EMetasoundFrontendVertexAccessType::Value:
-		{
-			return IDataTypeRegistry::Get().GetFrontendConstructorInputClass(InClassInput.TypeName, OutClass);
-		}
-		break;
-
-		case EMetasoundFrontendVertexAccessType::Reference:
-		{
-			return IDataTypeRegistry::Get().GetFrontendInputClass(InClassInput.TypeName, OutClass);
-		}
-		break;
-
-		case EMetasoundFrontendVertexAccessType::Unset:
-		default:
-		{
-			checkNoEntry();
-		}
-		break;
-		}
-
-		return false;
-	};
-
 	FMetasoundFrontendClass Class;
-	if (FindRegistryClass(Class))
+	if (DocumentBuilderPrivate::FindInputRegistryClass(InClassInput.TypeName, InClassInput.AccessType, Class))
 	{
 		if(!FindDependency(Class.Metadata))
 		{
 			AddDependency(Class);
 		}
 
-		auto FinalizeNode = [&InClassInput](FMetasoundFrontendNode& InOutNode, const Metasound::Frontend::FNodeRegistryKey&)
+		auto FinalizeNode = [this, &InClassInput](FMetasoundFrontendNode& InOutNode, const Metasound::Frontend::FNodeRegistryKey&)
 		{
 			// Sets the name of the node an vertices on the node to match the class vertex name
 			DocumentBuilderPrivate::SetNodeAndVertexNames(InOutNode, InClassInput);
@@ -1040,35 +1063,8 @@ const FMetasoundFrontendNode* FMetaSoundFrontendDocumentBuilder::AddGraphOutput(
 		return nullptr;
 	}
 
-	auto FindRegistryClass = [&InClassOutput](FMetasoundFrontendClass& OutClass) -> bool
-	{
-		switch (InClassOutput.AccessType)
-		{
-		case EMetasoundFrontendVertexAccessType::Value:
-		{
-			return IDataTypeRegistry::Get().GetFrontendConstructorOutputClass(InClassOutput.TypeName, OutClass);
-		}
-		break;
-
-		case EMetasoundFrontendVertexAccessType::Reference:
-		{
-			return IDataTypeRegistry::Get().GetFrontendOutputClass(InClassOutput.TypeName, OutClass);
-		}
-		break;
-
-		case EMetasoundFrontendVertexAccessType::Unset:
-		default:
-		{
-			checkNoEntry();
-		}
-		break;
-		}
-
-		return false;
-	};
-
 	FMetasoundFrontendClass Class;
-	if (FindRegistryClass(Class))
+	if (DocumentBuilderPrivate::FindOutputRegistryClass(InClassOutput.TypeName, InClassOutput.AccessType, Class))
 	{
 		if (!FindDependency(Class.Metadata))
 		{
@@ -1085,7 +1081,7 @@ const FMetasoundFrontendNode* FMetaSoundFrontendDocumentBuilder::AddGraphOutput(
 			FMetasoundFrontendClassOutput& NewOutput = RootGraph.Interface.Outputs.Add_GetRef(InClassOutput);
 			if (!NewOutput.VertexID.IsValid())
 			{
-				NewOutput.VertexID = FDocumentIDGenerator::Get().CreateVertexID(Document);;
+				NewOutput.VertexID = FDocumentIDGenerator::Get().CreateVertexID(Document);
 			}
 
 			DocumentDelegates->InterfaceDelegates.OnOutputAdded.Broadcast(NewIndex);
@@ -1203,7 +1199,33 @@ const FMetasoundFrontendNode* FMetaSoundFrontendDocumentBuilder::AddNodeByClassN
 	return nullptr;
 }
 
-FMetasoundFrontendNode* FMetaSoundFrontendDocumentBuilder::AddNodeInternal(const FMetasoundFrontendClassMetadata& InClassMetadata, FFinalizeNodeFunctionRef FinalizeNode, FGuid InNodeID)
+const FMetasoundFrontendNode* FMetaSoundFrontendDocumentBuilder::AddNodeByTemplate(const Metasound::Frontend::INodeTemplate& InTemplate, Metasound::Frontend::FNodeTemplateGenerateInterfaceParams Params, FGuid InNodeID)
+{
+	using namespace Metasound;
+	using namespace Metasound::Frontend;
+
+	const FMetasoundFrontendClass& TemplateClass = InTemplate.GetFrontendClass();
+	checkf(TemplateClass.Metadata.GetType() == EMetasoundFrontendClassType::Template, TEXT("INodeTemplate ClassType must always be 'Template'"));
+	const FMetasoundFrontendClass* Dependency = FindDependency(TemplateClass.Metadata);
+	if (!Dependency)
+	{
+		Dependency = AddDependency(TemplateClass);
+	}
+	check(Dependency);
+
+	auto FinalizeNodeFunction = [](const FMetasoundFrontendNode& Node, const Metasound::Frontend::FNodeRegistryKey& ClassKey)
+	{
+		return Node.Name;
+	};
+
+	FMetasoundFrontendNode* NewNode = AddNodeInternal(Dependency->Metadata, FinalizeNodeFunction, InNodeID);
+	check(NewNode);
+	NewNode->Interface = InTemplate.GenerateNodeInterface(MoveTemp(Params));
+
+	return NewNode;
+}
+
+FMetasoundFrontendNode* FMetaSoundFrontendDocumentBuilder::AddNodeInternal(const FMetasoundFrontendClassMetadata& InClassMetadata, FFinalizeNodeFunctionRef FinalizeNode, FGuid InNodeID, int32* NewNodeIndex)
 {
 	METASOUND_TRACE_CPUPROFILER_EVENT_SCOPE(FMetaSoundFrontendDocumentBuilder::AddNodeInternal);
 
@@ -1223,6 +1245,11 @@ FMetasoundFrontendNode* FMetaSoundFrontendDocumentBuilder::AddNodeInternal(const
 		const IDocumentGraphNodeCache& NodeCache = DocumentCache->GetNodeCache();
 		DocumentDelegates->NodeDelegates.OnNodeAdded.Broadcast(NewIndex);
 
+		if (NewNodeIndex)
+		{
+			*NewNodeIndex = NewIndex;
+		}
+
 #if WITH_EDITORONLY_DATA
 		Document.Metadata.ModifyContext.AddNodeIDModified(InNodeID);
 #endif // WITH_EDITORONLY_DATA
@@ -1257,25 +1284,137 @@ bool FMetaSoundFrontendDocumentBuilder::RemoveUnusedDependencies()
 		const FGuid& ClassID = Dependencies[i].ID;
 		if (!DocumentCache->GetNodeCache().ContainsNodesOfClassID(ClassID))
 		{
-			checkf(RemoveDependency(ClassID), TEXT("Failed to remove dependency that was found on document and was not referenced by nodes"));
-			bDidEdit = true;
+			const bool bRemoved = RemoveDependency(ClassID);
+			bDidEdit |= ensureAlwaysMsgf(bRemoved, TEXT("Failed to remove dependency that was found on document and was not referenced by nodes"));
 		}
 	}
 
 	return bDidEdit;
 }
 
-void FMetaSoundFrontendDocumentBuilder::ClearGraph()
+void FMetaSoundFrontendDocumentBuilder::ClearGraph(TSharedRef<Metasound::Frontend::FDocumentModifyDelegates> ModifyDelegates)
 {
-	FMetasoundFrontendGraphClass& GraphClass = GetDocument().RootGraph;
+	FMetasoundFrontendDocument& Doc = GetDocument();
+	FMetasoundFrontendGraphClass& GraphClass = Doc.RootGraph;
 	GraphClass.Graph.Nodes.Reset();
 	GraphClass.Graph.Edges.Reset();
 	GraphClass.Interface.Inputs.Reset();
 	GraphClass.Interface.Outputs.Reset();
 	GraphClass.PresetOptions.InputsInheritingDefault.Reset();
-	GetDocument().Interfaces.Reset();
-	RemoveUnusedDependencies();
-	InitCacheInternal();
+	Doc.Interfaces.Reset();
+	Doc.Dependencies.Reset();
+
+#if WITH_EDITORONLY_DATA
+	Doc.Metadata.MemberMetadata.Reset();
+#endif // WITH_EDITORONLY_DATA
+
+	Reload(ModifyDelegates);
+}
+
+#if WITH_EDITORONLY_DATA
+bool FMetaSoundFrontendDocumentBuilder::ClearMemberMetadata(const FGuid& InMemberID)
+{
+	return GetDocument().Metadata.MemberMetadata.Remove(InMemberID) > 0;
+}
+#endif // WITH_EDITORONLY_DATA
+
+bool FMetaSoundFrontendDocumentBuilder::ConformGraphInputNodeToClass(const FMetasoundFrontendClassInput& GraphInput)
+{
+	using namespace Metasound::Frontend;
+
+	FMetasoundFrontendDocument& Document = GetDocument();
+	const IDocumentGraphNodeCache& NodeCache = DocumentCache->GetNodeCache();
+	if (const int32* NodeIndexPtr = NodeCache.FindNodeIndex(GraphInput.NodeID))
+	{
+		TArray<FMetasoundFrontendNode>& Nodes = Document.RootGraph.Graph.Nodes;
+		FMetasoundFrontendNode& Node = Nodes[*NodeIndexPtr];
+		FMetasoundFrontendClass Class;
+		const bool bClassFound = DocumentBuilderPrivate::FindInputRegistryClass(GraphInput.TypeName, GraphInput.AccessType, Class);
+		if (ensureAlways(bClassFound))
+		{
+			const FMetasoundFrontendClass* Dependency = FindDependency(Class.Metadata);
+			if (!Dependency)
+			{
+				Dependency = AddDependency(Class);
+			}
+
+			if (ensureAlways(Dependency))
+			{
+				const int32 RemovalIndex = *NodeIndexPtr; // Have to cache as next delegate broadcast invalidates index pointer
+				DocumentDelegates->NodeDelegates.OnRemoveSwappingNode.Broadcast(RemovalIndex, Nodes.Num() - 1);
+				FMetasoundFrontendNode NewNode = MoveTemp(Node);
+				Nodes.RemoveAtSwap(RemovalIndex, EAllowShrinking::No);
+				NewNode.ClassID = Dependency->ID;
+				NewNode.Interface.Inputs.Last().TypeName = GraphInput.TypeName;
+				NewNode.Interface.Outputs.Last().TypeName = GraphInput.TypeName;
+
+#if WITH_EDITORONLY_DATA
+				Document.Metadata.ModifyContext.AddNodeIDModified(NewNode.GetID());
+#endif // WITH_EDITORONLY_DATA
+
+				// Set the default literal on the nodes inputs so that it gets passed to the instantiated TInputNode on a live
+				// auditioned MetaSound.
+				DocumentBuilderPrivate::SetDefaultLiteralOnInputNode(NewNode, GraphInput);
+
+				FMetasoundFrontendNode& NewNodeRef = Nodes.Add_GetRef(MoveTemp(NewNode));
+				DocumentDelegates->NodeDelegates.OnNodeAdded.Broadcast(Nodes.Num() - 1);
+
+				// Remove the default literal on the node added during the "FinalizeNode" call. This matches how 
+				// nodes are serialized in editor. The default literals are only stored on the FMetasoundFrontendClassInputs.
+				NewNodeRef.InputLiterals.Reset();
+
+				RemoveUnusedDependencies();
+				return true;
+			}
+		}
+	}
+
+	return false;
+}
+
+bool FMetaSoundFrontendDocumentBuilder::ConformGraphOutputNodeToClass(const FMetasoundFrontendClassOutput& GraphOutput)
+{
+	using namespace Metasound::Frontend;
+
+	FMetasoundFrontendDocument& Document = GetDocument();
+	const IDocumentGraphNodeCache& NodeCache = DocumentCache->GetNodeCache();
+	if (const int32* NodeIndexPtr = NodeCache.FindNodeIndex(GraphOutput.NodeID))
+	{
+		TArray<FMetasoundFrontendNode>& Nodes = Document.RootGraph.Graph.Nodes;
+		FMetasoundFrontendNode& Node = Nodes[*NodeIndexPtr];
+		FMetasoundFrontendClass Class;
+		const bool bClassFound = DocumentBuilderPrivate::FindOutputRegistryClass(GraphOutput.TypeName, GraphOutput.AccessType, Class);
+		if (ensureAlways(bClassFound))
+		{
+			const FMetasoundFrontendClass* Dependency = FindDependency(Class.Metadata);
+			if (!Dependency)
+			{
+				Dependency = AddDependency(Class);
+			}
+
+			if (ensureAlways(Dependency))
+			{
+				const int32 RemovalIndex = *NodeIndexPtr; // Have to cache as next delegate broadcast invalidates index pointer
+				DocumentDelegates->NodeDelegates.OnRemoveSwappingNode.Broadcast(RemovalIndex, Nodes.Num() - 1);
+				FMetasoundFrontendNode NewNode = MoveTemp(Node);
+				Nodes.RemoveAtSwap(RemovalIndex, EAllowShrinking::No);
+				NewNode.ClassID = Dependency->ID;
+				NewNode.Interface.Inputs.Last().TypeName = GraphOutput.TypeName;
+				NewNode.Interface.Outputs.Last().TypeName = GraphOutput.TypeName;
+
+#if WITH_EDITORONLY_DATA
+				Document.Metadata.ModifyContext.AddNodeIDModified(NewNode.GetID());
+#endif // WITH_EDITORONLY_DATA
+				Nodes.Add(MoveTemp(NewNode));
+				DocumentDelegates->NodeDelegates.OnNodeAdded.Broadcast(Nodes.Num() - 1);
+
+				RemoveUnusedDependencies();
+				return true;
+			}
+		}
+	}
+
+	return false;
 }
 
 bool FMetaSoundFrontendDocumentBuilder::ContainsDependencyOfType(EMetasoundFrontendClassType ClassType) const
@@ -1319,12 +1458,12 @@ bool FMetaSoundFrontendDocumentBuilder::ConvertFromPreset()
 	return false;
 }
 
-bool FMetaSoundFrontendDocumentBuilder::ConvertToPreset(const FMetasoundFrontendDocument& InReferencedDocument)
+bool FMetaSoundFrontendDocumentBuilder::ConvertToPreset(const FMetasoundFrontendDocument& InReferencedDocument, TSharedRef<Metasound::Frontend::FDocumentModifyDelegates> ModifyDelegates)
 {
 	using namespace Metasound;
 	using namespace Metasound::Frontend;
 
-	ClearGraph();
+	ClearGraph(ModifyDelegates);
 
 	FMetasoundFrontendGraphClass& PresetAssetRootGraph = GetDocument().RootGraph;
 	FMetasoundFrontendGraph& PresetAssetGraph = PresetAssetRootGraph.Graph;
@@ -1334,23 +1473,32 @@ bool FMetaSoundFrontendDocumentBuilder::ConvertToPreset(const FMetasoundFrontend
 #endif // WITH_EDITORONLY_DATA
 
 	// Mark all inputs as inherited by default
-	TArray<FName> InputsInheritingDefault;
-	Algo::Transform(PresetAssetRootGraph.Interface.Inputs, InputsInheritingDefault, [](const FMetasoundFrontendClassInput& Input)
 	{
-		return Input.Name;
-	});
+		TArray<FName> InputsInheritingDefault;
+		auto GetInputName = [](const FMetasoundFrontendClassInput& Input) { return Input.Name; };
+		Algo::Transform(PresetAssetRootGraph.Interface.Inputs, InputsInheritingDefault, GetInputName);
 
-	PresetAssetRootGraph.PresetOptions.bIsPreset = true;
-	PresetAssetRootGraph.PresetOptions.InputsInheritingDefault = TSet<FName>(InputsInheritingDefault);
+		PresetAssetRootGraph.PresetOptions.bIsPreset = true;
+		PresetAssetRootGraph.PresetOptions.InputsInheritingDefault = TSet<FName>(InputsInheritingDefault);
+	}
 
 	// Apply root graph transform 
 	FRebuildPresetRootGraph RebuildPresetRootGraph(InReferencedDocument);
 	if (RebuildPresetRootGraph.Transform(GetDocument()))
 	{
-		InitCacheInternal();
+		DocumentInterface->ConformObjectToDocument();
+
+		// TL/DR: Have to reload and assign delegates here due to the rebuild preset transform still being implemented via controllers.
+		// Onces its reimplemented with the builder API, this can be removed.
+		// 
+		// The invalidate cache call when accessing the mutable document handle from within the transform unfortunately doesn't reach this
+		// builder's cache indirectly as converting to preset can be called by transient builders that are not registered with the MetaSound
+		// builder subsystem.
+		Reload(ModifyDelegates);
 		return true;
 	}
 
+	// Hack until prior processes are all moved over to builder API
 	return false;
 }
 
@@ -1537,6 +1685,16 @@ const FMetasoundFrontendNode* FMetaSoundFrontendDocumentBuilder::FindGraphOutput
 	return nullptr;
 }
 
+#if WITH_EDITOR
+UMetaSoundFrontendMemberMetadata* FMetaSoundFrontendDocumentBuilder::FindMemberMetadata(const FGuid& InMemberID)
+{
+	FMetasoundFrontendDocument& Document = GetDocument();
+	TMap<FGuid, TObjectPtr<UMetaSoundFrontendMemberMetadata>>& LiteralMetadata = Document.Metadata.MemberMetadata;
+	TObjectPtr<UMetaSoundFrontendMemberMetadata> ToReturn = LiteralMetadata.FindRef(InMemberID);
+	return ToReturn;
+}
+#endif // WITH_EDITOR
+
 const FMetasoundFrontendNode* FMetaSoundFrontendDocumentBuilder::FindNode(const FGuid& InNodeID) const
 {
 	using namespace Metasound::Frontend;
@@ -1701,7 +1859,7 @@ FMetasoundFrontendDocument& FMetaSoundFrontendDocumentBuilder::GetDocument()
 
 const FMetasoundFrontendDocument& FMetaSoundFrontendDocumentBuilder::GetDocument() const
 {
-	return DocumentInterface->GetDocument();
+	return DocumentInterface->GetConstDocument();
 }
 
 const Metasound::Frontend::FDocumentModifyDelegates& FMetaSoundFrontendDocumentBuilder::GetDocumentDelegates() const
@@ -1716,6 +1874,28 @@ const IMetaSoundDocumentInterface& FMetaSoundFrontendDocumentBuilder::GetDocumen
 	return *Interface;
 }
 
+FMetasoundAssetBase& FMetaSoundFrontendDocumentBuilder::GetMetasoundAsset()
+{
+	using namespace Metasound::Frontend;
+
+	UObject* Object = DocumentInterface.GetObject();
+	check(Object);
+	FMetasoundAssetBase* Asset = IMetaSoundAssetManager::GetChecked().GetAsAsset(*Object);
+	check(Asset);
+	return *Asset;
+}
+
+const FMetasoundAssetBase& FMetaSoundFrontendDocumentBuilder::GetMetasoundAsset() const
+{
+	using namespace Metasound::Frontend;
+
+	UObject* Object = DocumentInterface.GetObject();
+	check(Object);
+	const FMetasoundAssetBase* Asset = IMetaSoundAssetManager::GetChecked().GetAsAsset(*Object);
+	check(Asset);
+	return *Asset;
+}
+
 EMetasoundFrontendVertexAccessType FMetaSoundFrontendDocumentBuilder::GetNodeInputAccessType(const FGuid& InNodeID, const FGuid& InVertexID) const
 {
 	using namespace Metasound::Frontend;
@@ -1897,6 +2077,48 @@ EMetasoundFrontendVertexAccessType FMetaSoundFrontendDocumentBuilder::GetNodeOut
 	return EMetasoundFrontendVertexAccessType::Unset;
 }
 
+void FMetaSoundFrontendDocumentBuilder::InitDocument(const FMetasoundFrontendDocument* InDocumentTemplate, const FMetasoundFrontendClassName* InNewClassName, bool bResetVersion)
+{
+	using namespace Metasound;
+	using namespace Metasound::Frontend;
+
+	METASOUND_TRACE_CPUPROFILER_EVENT_SCOPE(FMetaSoundFrontendDocumentBuilder::InitDocument);
+
+	FMetasoundFrontendDocument& Document = GetDocument();
+
+	// 1. Set default class Metadata.
+	if (InDocumentTemplate)
+	{
+		// 1a. If template provided, copy that.
+		Document = *InDocumentTemplate;
+		InitGraphClassMetadata(bResetVersion, InNewClassName);
+	}
+	else
+	{
+		// 1a. Initialize class using default data
+		FMetasoundFrontendClassMetadata& ClassMetadata = Document.RootGraph.Metadata;
+		InitGraphClassMetadata(Document.RootGraph.Metadata, bResetVersion, InNewClassName);
+
+		// 1b. Set default doc version Metadata
+		{
+			FMetasoundFrontendDocumentMetadata& DocMetadata = Document.Metadata;
+			DocMetadata.Version.Number = FMetasoundFrontendDocument::GetMaxVersion();
+		}
+
+		// 1c. Add default interfaces for given UClass
+		{
+			TArray<FMetasoundFrontendVersion> InitVersions = ISearchEngine::Get().FindUClassDefaultInterfaceVersions(GetBuilderClassPath());
+			FModifyInterfaceOptions Options({ }, InitVersions);
+			ModifyInterfaces(MoveTemp(Options));
+		}
+	}
+}
+
+bool FMetaSoundFrontendDocumentBuilder::IsValid() const
+{
+	return DocumentInterface.GetObject() != nullptr;
+}
+
 int32 FMetaSoundFrontendDocumentBuilder::GetTransactionCount() const
 {
 	using namespace Metasound::Frontend;
@@ -1928,34 +2150,9 @@ void FMetaSoundFrontendDocumentBuilder::InitGraphClassMetadata(FMetasoundFronten
 	InOutMetadata.SetType(EMetasoundFrontendClassType::Graph);
 }
 
-void FMetaSoundFrontendDocumentBuilder::InitDocument()
+void FMetaSoundFrontendDocumentBuilder::InitGraphClassMetadata(bool bResetVersion, const FMetasoundFrontendClassName* NewClassName)
 {
-	using namespace Metasound;
-	using namespace Metasound::Frontend;
-
-	METASOUND_TRACE_CPUPROFILER_EVENT_SCOPE(FMetaSoundFrontendDocumentBuilder::InitDocument);
-
-	FMetasoundFrontendDocument& Document = GetDocument();
-
-	// 1. Set default class Metadata
-	{
-		constexpr bool bResetVersion = true;
-		FMetasoundFrontendClassMetadata& ClassMetadata = Document.RootGraph.Metadata;
-		InitGraphClassMetadata(ClassMetadata, bResetVersion);
-	}
-
-	// 2. Set default doc version Metadata
-	{
-		FMetasoundFrontendDocumentMetadata& DocMetadata = Document.Metadata;
-		DocMetadata.Version.Number = FMetasoundFrontendDocument::GetMaxVersion();
-	}
-
-	// 3. Add default interfaces for given UClass
-	{
-		TArray<FMetasoundFrontendVersion> InitVersions = ISearchEngine::Get().FindUClassDefaultInterfaceVersions(GetBuilderClassPath());
-		FModifyInterfaceOptions Options({ }, InitVersions);
-		ModifyInterfaces(MoveTemp(Options));
-	}
+	InitGraphClassMetadata(GetDocument().RootGraph.Metadata, bResetVersion, NewClassName);
 }
 
 void FMetaSoundFrontendDocumentBuilder::InitNodeLocations()
@@ -2015,12 +2212,6 @@ void FMetaSoundFrontendDocumentBuilder::InitNodeLocations()
 #endif // WITH_EDITORONLY_DATA
 }
 
-void FMetaSoundFrontendDocumentBuilder::InvalidateCache()
-{
-	constexpr bool bPrimeCache = false;
-	InitCacheInternal(bPrimeCache);
-}
-
 bool FMetaSoundFrontendDocumentBuilder::IsDependencyReferenced(const FGuid& InClassID) const
 {
 	return DocumentCache->GetNodeCache().ContainsNodesOfClassID(InClassID);
@@ -2059,6 +2250,15 @@ bool FMetaSoundFrontendDocumentBuilder::IsPreset() const
 	return GetDocument().RootGraph.PresetOptions.bIsPreset;
 }
 
+bool FMetaSoundFrontendDocumentBuilder::IsVariableClass(EMetasoundFrontendClassType ClassType) const
+{
+	const bool bIsVariableNode = (ClassType == EMetasoundFrontendClassType::Variable)
+		|| (ClassType == EMetasoundFrontendClassType::VariableAccessor)
+		|| (ClassType == EMetasoundFrontendClassType::VariableDeferredAccessor)
+		|| (ClassType == EMetasoundFrontendClassType::VariableMutator);
+	return bIsVariableNode;
+}
+
 Metasound::Frontend::EInvalidEdgeReason FMetaSoundFrontendDocumentBuilder::IsValidEdge(const FMetasoundFrontendEdge& InEdge) const
 {
 	using namespace Metasound::Frontend;
@@ -2094,6 +2294,27 @@ Metasound::Frontend::EInvalidEdgeReason FMetaSoundFrontendDocumentBuilder::IsVal
 	return EInvalidEdgeReason::None;
 }
 
+void FMetaSoundFrontendDocumentBuilder::IterateNodesConnectedWithVertex(const FMetasoundFrontendVertexHandle& Vertex, TFunctionRef<void(const FMetasoundFrontendEdge&, FMetasoundFrontendNode&)> NodeIndexIterFunc)
+{
+	using namespace Metasound::Frontend;
+
+	FMetasoundFrontendDocument& Document = GetDocument();
+	TArray<FMetasoundFrontendEdge> EdgesToConnectedNodes; // Have to cache to avoid pointers becoming garbage in subsequent removal loop
+	const IDocumentGraphEdgeCache& EdgeCache = DocumentCache->GetEdgeCache();
+	const IDocumentGraphNodeCache& NodeCache = DocumentCache->GetNodeCache();
+	const TArray<const FMetasoundFrontendEdge*> Edges = EdgeCache.FindEdges(Vertex.NodeID, Vertex.VertexID);
+	Algo::Transform(Edges, EdgesToConnectedNodes, [](const FMetasoundFrontendEdge* Edge) { check(Edge); return *Edge; });
+	for (const FMetasoundFrontendEdge& Edge : EdgesToConnectedNodes)
+	{
+		const FGuid& ConnectedNodeID = Edge.ToNodeID == Vertex.NodeID ? Edge.FromNodeID : Edge.ToNodeID;
+		if (const int32* ConnectedNodeIndex = NodeCache.FindNodeIndex(ConnectedNodeID))
+		{
+			FMetasoundFrontendNode& Node = Document.RootGraph.Graph.Nodes[*ConnectedNodeIndex];
+			NodeIndexIterFunc(Edge, Node);
+		}
+	}
+}
+
 bool FMetaSoundFrontendDocumentBuilder::ModifyInterfaces(Metasound::Frontend::FModifyInterfaceOptions&& InOptions)
 {
 	using namespace Metasound::Frontend;
@@ -2156,7 +2377,7 @@ bool FMetaSoundFrontendDocumentBuilder::TransformTemplateNodes()
 		const FTemplateTransformParams& Params = Pair.Value;
 		if (Params.Template)
 		{
-			TUniquePtr<INodeTransform> NodeTransform = Params.Template->GenerateNodeTransform();
+			TUniquePtr<INodeTemplateTransform> NodeTransform = Params.Template->GenerateNodeTransform();
 			check(NodeTransform.IsValid());
 
 			for (const FGuid& NodeID : Params.NodeIDs)
@@ -2182,27 +2403,31 @@ bool FMetaSoundFrontendDocumentBuilder::TransformTemplateNodes()
 	return bModified;
 }
 
-void FMetaSoundFrontendDocumentBuilder::InitCacheInternal(bool bPrimeCache)
+void FMetaSoundFrontendDocumentBuilder::BeginBuilding(TSharedPtr<Metasound::Frontend::FDocumentModifyDelegates> Delegates, bool bPrimeCache)
 {
 	using namespace Metasound::Frontend;
-	DocumentCache = FDocumentCache::Create(GetDocument(), DocumentDelegates, bPrimeCache);
-}
 
-void FMetaSoundFrontendDocumentBuilder::BeginBuilding()
-{
+	DocumentDelegates = Delegates.IsValid() ? Delegates : MakeShared<FDocumentModifyDelegates>();
+
 	if (DocumentInterface)
 	{
 		DocumentInterface->OnBeginActiveBuilder();
+		DocumentCache = FDocumentCache::Create(GetDocument(), DocumentDelegates.ToSharedRef(), bPrimeCache);
 	}
 }
 
 void FMetaSoundFrontendDocumentBuilder::FinishBuilding()
 {
+	using namespace Metasound::Frontend;
+
 	if (DocumentInterface)
 	{
 		DocumentInterface->OnFinishActiveBuilder();
+		DocumentInterface = { };
 	}
-	DocumentInterface = TScriptInterface<IMetaSoundDocumentInterface>();
+
+	DocumentDelegates.Reset();
+	DocumentCache.Reset();
 }
 
 bool FMetaSoundFrontendDocumentBuilder::RemoveDependency(const FGuid& InClassID)
@@ -2339,17 +2564,35 @@ bool FMetaSoundFrontendDocumentBuilder::RemoveNamedEdges(const TSet<Metasound::F
 	return bSuccess;
 }
 
-bool FMetaSoundFrontendDocumentBuilder::RemoveNodeInputDefault(const FGuid& InNodeID, const FGuid& InVertexID)
+void FMetaSoundFrontendDocumentBuilder::Reload(TSharedPtr<Metasound::Frontend::FDocumentModifyDelegates> Delegates, bool bPrimeCache)
 {
 	using namespace Metasound::Frontend;
 
-	const IDocumentGraphNodeCache& NodeCache = DocumentCache->GetNodeCache();
-	if (const int32* NodeIndex = NodeCache.FindNodeIndex(InNodeID))
+	if (DocumentInterface)
 	{
-		FMetasoundFrontendGraph& Graph = GetDocument().RootGraph.Graph;
-		FMetasoundFrontendNode& Node = Graph.Nodes[*NodeIndex];
+		DocumentInterface->OnFinishActiveBuilder();
+	}
 
-		auto IsVertex = [&InVertexID](const FMetasoundFrontendVertex& Vertex) { return Vertex.VertexID == InVertexID; };
+	DocumentDelegates = Delegates.IsValid() ? Delegates : MakeShared<FDocumentModifyDelegates>();
+
+	if (DocumentInterface)
+	{
+		DocumentCache = FDocumentCache::Create(GetDocument(), DocumentDelegates.ToSharedRef(), bPrimeCache);
+		DocumentInterface->OnBeginActiveBuilder();
+	}
+}
+
+bool FMetaSoundFrontendDocumentBuilder::RemoveNodeInputDefault(const FGuid& InNodeID, const FGuid& InVertexID)
+{
+	using namespace Metasound::Frontend;
+
+	const IDocumentGraphNodeCache& NodeCache = DocumentCache->GetNodeCache();
+	if (const int32* NodeIndex = NodeCache.FindNodeIndex(InNodeID))
+	{
+		FMetasoundFrontendGraph& Graph = GetDocument().RootGraph.Graph;
+		FMetasoundFrontendNode& Node = Graph.Nodes[*NodeIndex];
+
+		auto IsVertex = [&InVertexID](const FMetasoundFrontendVertex& Vertex) { return Vertex.VertexID == InVertexID; };
 		const int32 VertexIndex = Node.Interface.Inputs.IndexOfByPredicate(IsVertex);
 		if (VertexIndex != INDEX_NONE)
 		{
@@ -2440,6 +2683,10 @@ bool FMetaSoundFrontendDocumentBuilder::RemoveEdgeToNodeInput(const FGuid& InNod
 		DocumentDelegates->EdgeDelegates.OnRemoveSwappingEdge.Broadcast(Index, LastIndex);
 		Graph.Edges.RemoveAtSwap(Index, EAllowShrinking::No);
 
+#if WITH_EDITORONLY_DATA
+		Document.Metadata.ModifyContext.AddNodeIDModified(InNodeID);
+#endif // WITH_EDITORONLY_DATA
+
 		return true;
 	}
 
@@ -2449,51 +2696,128 @@ bool FMetaSoundFrontendDocumentBuilder::RemoveEdgeToNodeInput(const FGuid& InNod
 #if WITH_EDITOR
 bool FMetaSoundFrontendDocumentBuilder::RemoveGraphComment(const FGuid& InCommentID)
 {
-	TMap<FGuid, FMetaSoundFrontendGraphComment>& Comments = GetDocument().RootGraph.Graph.Style.Comments;
-	return Comments.Remove(InCommentID) > 0;
+	FMetasoundFrontendDocument& Document = GetDocument();
+	TMap<FGuid, FMetaSoundFrontendGraphComment>& Comments = Document.RootGraph.Graph.Style.Comments;
+
+	if (Comments.Remove(InCommentID) > 0)
+	{
+#if WITH_EDITORONLY_DATA
+		Document.Metadata.ModifyContext.SetDocumentModified();
+#endif // WITH_EDITORONLY_DATA
+
+		return true;
+	}
+
+	return false;
 }
 #endif // WITH_EDITOR
 
-bool FMetaSoundFrontendDocumentBuilder::RemoveGraphInput(FName InInputName)
+TArray<const FMetasoundFrontendNode*> FMetaSoundFrontendDocumentBuilder::GetGraphInputTemplateNodes(FName InInputName)
 {
-	if (const FMetasoundFrontendNode* Node = FindGraphInputNode(InInputName))
+	using namespace Metasound::Frontend;
+
+	TArray<const FMetasoundFrontendNode*> TemplateNodes;
+
+	FMetasoundFrontendDocument& Document = GetDocument();
+	if (const int32* Index = DocumentCache->GetInterfaceCache().FindInputIndex(InInputName))
 	{
-		const FGuid ClassID = Node->ClassID;
-		if (RemoveNode(Node->GetID()))
+		const FMetasoundFrontendClassInput& InputClass = Document.RootGraph.Interface.Inputs[*Index];
+		const IDocumentGraphNodeCache& NodeCache = DocumentCache->GetNodeCache();
+		const IDocumentGraphEdgeCache& EdgeCache = DocumentCache->GetEdgeCache();
+
+		if (const FMetasoundFrontendNode* InputNode = NodeCache.FindNode(InputClass.NodeID))
 		{
-			TArray<FMetasoundFrontendClassInput>& Inputs = GetDocument().RootGraph.Interface.Inputs;
-			auto InputNameMatches = [InInputName](const FMetasoundFrontendClassInput& Input) { return Input.Name == InInputName; };
-			const int32 Index = Inputs.IndexOfByPredicate(InputNameMatches);
-			if (Index != INDEX_NONE)
+			const FGuid OutputVertexID = InputNode->Interface.Outputs.Last().VertexID;
+			const TArray<const FMetasoundFrontendEdge*> ConnectedEdges = EdgeCache.FindEdges(InputClass.NodeID, OutputVertexID);
+			for (const FMetasoundFrontendEdge* Edge : ConnectedEdges)
 			{
-				DocumentDelegates->InterfaceDelegates.OnRemovingInput.Broadcast(Index);
-
-				const int32 LastIndex = Inputs.Num() - 1;
-				if (Index != LastIndex)
-				{
-					DocumentDelegates->InterfaceDelegates.OnRemovingInput.Broadcast(LastIndex);
-				}
-				Inputs.RemoveAtSwap(Index, EAllowShrinking::No);
-				if (Index != LastIndex)
+				check(Edge);
+				if (const int32* ConnectedNodeIndex = NodeCache.FindNodeIndex(Edge->ToNodeID))
 				{
-					DocumentDelegates->InterfaceDelegates.OnInputAdded.Broadcast(Index);
-				}
-
-				if (IsDependencyReferenced(ClassID))
-				{
-					return true;
-				}
-				else
-				{
-					if (RemoveDependency(ClassID))
+					const FMetasoundFrontendNode& ConnectedNode = Document.RootGraph.Graph.Nodes[*ConnectedNodeIndex];
+					if (const FMetasoundFrontendClass* ConnectedNodeClass = FindDependency(ConnectedNode.ClassID))
 					{
-						return true;
+						if (ConnectedNodeClass->Metadata.GetClassName() == FInputNodeTemplate::ClassName)
+						{
+							TemplateNodes.Add(&ConnectedNode);
+						}
 					}
 				}
 			}
 		}
 	}
 
+	return TemplateNodes;
+}
+
+bool FMetaSoundFrontendDocumentBuilder::RemoveGraphInput(FName InInputName)
+{
+	using namespace Metasound::Frontend;
+
+	FMetasoundFrontendDocument& Document = GetDocument();
+	if (const int32* IndexPtr = DocumentCache->GetInterfaceCache().FindInputIndex(InInputName))
+	{
+		TArray<FMetasoundFrontendClassInput>& Inputs = GetDocument().RootGraph.Interface.Inputs;
+		const FGuid NodeID = Inputs[*IndexPtr].NodeID;
+		FGuid ClassID;
+		TArray<FGuid> NodeIDsToRemove { NodeID };
+
+		if (const FMetasoundFrontendNode* Node = FindNode(NodeID))
+		{
+			ClassID = Node->ClassID;
+		}
+		else
+		{
+			return false;
+		}
+
+		const TArray<const FMetasoundFrontendNode*> TemplateNodes = GetGraphInputTemplateNodes(InInputName);
+		Algo::Transform(TemplateNodes, NodeIDsToRemove, [](const FMetasoundFrontendNode* Node) { return Node->GetID(); });
+
+		bool bNodesRemoved = true;
+		for (const FGuid& ToRemove : NodeIDsToRemove)
+		{
+			if (RemoveNode(ToRemove))
+			{
+#if WITH_EDITORONLY_DATA
+				Document.Metadata.ModifyContext.AddNodeIDModified(ToRemove);
+#endif // WITH_EDITORONLY_DATA
+			}
+			else
+			{
+				bNodesRemoved = false;
+			}
+		}
+
+		if (bNodesRemoved)
+		{
+			const int32 Index = *IndexPtr;
+			DocumentDelegates->InterfaceDelegates.OnRemovingInput.Broadcast(Index);
+
+			const int32 LastIndex = Inputs.Num() - 1;
+			if (Index != LastIndex)
+			{
+				DocumentDelegates->InterfaceDelegates.OnRemovingInput.Broadcast(LastIndex);
+			}
+			Inputs.RemoveAtSwap(Index, EAllowShrinking::No);
+			if (Index != LastIndex)
+			{
+				DocumentDelegates->InterfaceDelegates.OnInputAdded.Broadcast(Index);
+			}
+
+#if WITH_EDITORONLY_DATA
+			ClearMemberMetadata(NodeID);
+			Document.Metadata.ModifyContext.AddMemberIDModified(NodeID);
+#endif // WITH_EDITORONLY_DATA
+
+			const bool bDependencyReferenced = IsDependencyReferenced(ClassID);
+			if (bDependencyReferenced || RemoveDependency(ClassID))
+			{
+				return true;
+			}
+		}
+	}
+
 	return false;
 }
 
@@ -2502,9 +2826,16 @@ bool FMetaSoundFrontendDocumentBuilder::RemoveGraphOutput(FName InOutputName)
 	if (const FMetasoundFrontendNode* Node = FindGraphOutputNode(InOutputName))
 	{
 		const FGuid ClassID = Node->ClassID;
-		if (RemoveNode(Node->GetID()))
+		const FGuid NodeID = Node->GetID();
+		if (RemoveNode(NodeID))
 		{
-			TArray<FMetasoundFrontendClassOutput>& Outputs = GetDocument().RootGraph.Interface.Outputs;
+			FMetasoundFrontendDocument& Document = GetDocument();
+
+#if WITH_EDITORONLY_DATA
+			Document.Metadata.ModifyContext.AddNodeIDModified(NodeID);
+#endif // WITH_EDITORONLY_DATA
+
+			TArray<FMetasoundFrontendClassOutput>& Outputs = Document.RootGraph.Interface.Outputs;
 			auto OutputNameMatches = [InOutputName](const FMetasoundFrontendClassOutput& Output) { return Output.Name == InOutputName; };
 			const int32 Index = Outputs.IndexOfByPredicate(OutputNameMatches);
 			if (Index != INDEX_NONE)
@@ -2522,17 +2853,16 @@ bool FMetaSoundFrontendDocumentBuilder::RemoveGraphOutput(FName InOutputName)
 					DocumentDelegates->InterfaceDelegates.OnOutputAdded.Broadcast(Index);
 				}
 
-				if (IsDependencyReferenced(ClassID))
+#if WITH_EDITORONLY_DATA
+				ClearMemberMetadata(NodeID);
+				Document.Metadata.ModifyContext.AddMemberIDModified(NodeID);
+#endif // WITH_EDITORONLY_DATA
+
+				const bool bDependencyReferenced = IsDependencyReferenced(ClassID);
+				if (bDependencyReferenced || RemoveDependency(ClassID))
 				{
 					return true;
 				}
-				else
-				{
-					if (RemoveDependency(ClassID))
-					{
-						return true;
-					}
-				}
 			}
 		}
 	}
@@ -2581,6 +2911,7 @@ bool FMetaSoundFrontendDocumentBuilder::RemoveNode(const FGuid& InNodeID)
 	using namespace Metasound::Frontend;
 
 	const IDocumentGraphNodeCache& NodeCache = DocumentCache->GetNodeCache();
+	const IDocumentGraphEdgeCache& EdgeCache = DocumentCache->GetEdgeCache();
 
 	if (const int32* IndexPtr = NodeCache.FindNodeIndex(InNodeID))
 	{
@@ -2590,15 +2921,55 @@ bool FMetaSoundFrontendDocumentBuilder::RemoveNode(const FGuid& InNodeID)
 		FMetasoundFrontendGraph& Graph = Document.RootGraph.Graph;
 		TArray<FMetasoundFrontendNode>& Nodes = Graph.Nodes;
 		const FMetasoundFrontendNode& Node = Nodes[Index];
+		const FGuid& NodeID = Node.GetID();
+
+		const FMetasoundFrontendClass* NodeClass = DocumentCache->FindDependency(Node.ClassID);
+		check(NodeClass);
+
+		const bool bIsVariableNode = IsVariableClass(NodeClass->Metadata.GetType());
+		if (bIsVariableNode)
+		{
+			const bool bUnlinkedVariable = UnlinkVariableNode(Node.GetID());
+			ensureMsgf(bUnlinkedVariable, TEXT("Removed variable node that was not linked to variable"));
+		}
 
+		FMetasoundFrontendVertexHandle FromVertexHandle;
 		for (const FMetasoundFrontendVertex& Vertex : Node.Interface.Inputs)
 		{
+			if (bIsVariableNode)
+			{
+				if (const int32* InputEdgeIndex = EdgeCache.FindEdgeIndexToNodeInput(InNodeID, Vertex.VertexID))
+				{
+					FromVertexHandle = Graph.Edges[*InputEdgeIndex].GetFromVertexHandle();
+				}
+			}
+
 			RemoveEdgeToNodeInput(InNodeID, Vertex.VertexID);
 		}
 
+		TArray<FMetasoundFrontendVertexHandle> ToVertexHandles;
 		for (const FMetasoundFrontendVertex& Vertex : Node.Interface.Outputs)
 		{
+			if (bIsVariableNode)
+			{
+				ToVertexHandles.Reset();
+				const TArrayView<const int32> OutputEdgeIndices = EdgeCache.FindEdgeIndicesFromNodeOutput(InNodeID, Vertex.VertexID);
+				Algo::Transform(OutputEdgeIndices, ToVertexHandles, [&Graph](const int32& VertIndex) { return Graph.Edges[VertIndex].GetToVertexHandle(); });
+			}
 			RemoveEdgesFromNodeOutput(InNodeID, Vertex.VertexID);
+			if (bIsVariableNode)
+			{
+				for (const FMetasoundFrontendVertexHandle& ToHandle : ToVertexHandles)
+				{
+					AddEdge(FMetasoundFrontendEdge
+					{
+						FromVertexHandle.NodeID,
+						FromVertexHandle.VertexID,
+						ToHandle.NodeID,
+						ToHandle.VertexID
+					});
+				}
+			}
 		}
 
 		const int32 LastIndex = Nodes.Num() - 1;
@@ -2658,8 +3029,9 @@ bool FMetaSoundFrontendDocumentBuilder::RenameRootGraphClass(const FMetasoundFro
 
 void FMetaSoundFrontendDocumentBuilder::ReloadCache()
 {
-	constexpr bool bPrimeCache = true;
-	InitCacheInternal(bPrimeCache);
+	using namespace Metasound::Frontend;
+
+	Reload(DocumentDelegates, true);
 }
 
 #if WITH_EDITOR
@@ -2670,6 +3042,142 @@ void FMetaSoundFrontendDocumentBuilder::SetAuthor(const FString& InAuthor)
 }
 #endif // WITH_EDITOR
 
+bool FMetaSoundFrontendDocumentBuilder::SetGraphInputAccessType(FName InputName, EMetasoundFrontendVertexAccessType AccessType)
+{
+	using namespace Metasound::Frontend;
+
+	if (!ensureMsgf(AccessType != EMetasoundFrontendVertexAccessType::Unset, TEXT("Cannot set graph input access type to '%s'"), LexToString(AccessType)))
+	{
+		return false;
+	}
+
+	if (const int32* Index = DocumentCache->GetInterfaceCache().FindInputIndex(InputName))
+	{
+		FMetasoundFrontendDocument& Document = GetDocument();
+		FMetasoundFrontendClassInput& GraphInput = Document.RootGraph.Interface.Inputs[*Index];
+
+		if (GraphInput.AccessType != AccessType)
+		{
+			GraphInput.AccessType = AccessType;
+			if (AccessType == EMetasoundFrontendVertexAccessType::Reference)
+			{
+				const IDocumentGraphNodeCache& NodeCache = DocumentCache->GetNodeCache();
+				if (const int32* NodeIndex = NodeCache.FindNodeIndex(GraphInput.NodeID))
+				{
+					FMetasoundFrontendNode& Node = Document.RootGraph.Graph.Nodes[*NodeIndex];
+					const FMetasoundFrontendVertex& NodeOutput = Node.Interface.Outputs.Last();
+					IterateNodesConnectedWithVertex({ GraphInput.NodeID, NodeOutput.VertexID }, [this, &AccessType, &Document](const FMetasoundFrontendEdge& Edge, FMetasoundFrontendNode& ConnectedNode)
+					{
+						if (const FMetasoundFrontendClass* ConnectedNodeClass = FindDependency(ConnectedNode.ClassID))
+						{
+							// If connected to an input template node, disconnect the template node from other nodes as the data type is
+							// about to be mismatched.  Otherwise, direct connection to other nodes (i.e. at runtime when template
+							// nodes aren't injected) forcefully remove to avoid data type mismatch.
+							if (ConnectedNodeClass->Metadata.GetClassName() == FInputNodeTemplate::ClassName)
+							{
+								const FMetasoundFrontendVertex& ConnectedNodeOutput = ConnectedNode.Interface.Outputs.Last();
+								IterateNodesConnectedWithVertex({ Edge.ToNodeID, ConnectedNodeOutput.VertexID }, [this, &AccessType](const FMetasoundFrontendEdge& TempEdge, FMetasoundFrontendNode&)
+								{
+									const EMetasoundFrontendVertexAccessType ConnectedAccessType = GetNodeInputAccessType(TempEdge.ToNodeID, TempEdge.ToVertexID);
+									if (!FMetasoundFrontendClassVertex::CanConnectVertexAccessTypes(AccessType, ConnectedAccessType))
+									{
+										RemoveEdgeToNodeInput(TempEdge.ToNodeID, TempEdge.ToVertexID);
+									}
+								});
+							}
+							else
+							{
+								const EMetasoundFrontendVertexAccessType ConnectedAccessType = GetNodeInputAccessType(Edge.ToNodeID, Edge.ToVertexID);
+								if (!FMetasoundFrontendClassVertex::CanConnectVertexAccessTypes(AccessType, ConnectedAccessType))
+								{
+									RemoveEdgeToNodeInput(Edge.ToNodeID, Edge.ToVertexID);
+								}
+							}
+						}
+					});
+
+				}
+			}
+
+			const bool bNodeConformed = ConformGraphInputNodeToClass(GraphInput);
+			if (bNodeConformed)
+			{
+#if WITH_EDITORONLY_DATA
+				Document.Metadata.ModifyContext.AddMemberIDModified(GraphInput.NodeID);
+#endif // WITH_EDITORONLY_DATA
+
+				return true;
+			}
+		}
+	}
+
+	return false;
+}
+
+bool FMetaSoundFrontendDocumentBuilder::SetGraphInputDataType(FName InputName, FName DataType)
+{
+	using namespace Metasound::Frontend;
+
+	if (IDataTypeRegistry::Get().IsRegistered(DataType))
+	{
+		if (const int32* Index = DocumentCache->GetInterfaceCache().FindInputIndex(InputName))
+		{
+			FMetasoundFrontendDocument& Document = GetDocument();
+			FMetasoundFrontendClassInput& GraphInput = Document.RootGraph.Interface.Inputs[*Index];
+			if (GraphInput.TypeName != DataType)
+			{
+				FMetasoundFrontendLiteral DefaultLiteral;
+				DefaultLiteral.SetFromLiteral(IDataTypeRegistry::Get().CreateDefaultLiteral(DataType));
+				GraphInput.DefaultLiteral = DefaultLiteral;
+				GraphInput.TypeName = DataType;
+
+				const IDocumentGraphNodeCache& NodeCache = DocumentCache->GetNodeCache();
+				if (const int32* NodeIndex = NodeCache.FindNodeIndex(GraphInput.NodeID))
+				{
+					FMetasoundFrontendNode& Node = Document.RootGraph.Graph.Nodes[*NodeIndex];
+					FMetasoundFrontendVertex& NodeOutput = Node.Interface.Outputs.Last();
+					IterateNodesConnectedWithVertex({ GraphInput.NodeID, NodeOutput.VertexID }, [this, &DataType](const FMetasoundFrontendEdge& Edge, FMetasoundFrontendNode& ConnectedNode)
+					{
+						const FMetasoundFrontendClass* ConnectedNodeClass = FindDependency(ConnectedNode.ClassID);
+						if (ensure(ConnectedNodeClass))
+						{
+							// If connected to an input template node, disconnect the template node from other nodes as the data type is
+							// about to be mismatched.  Otherwise, direct connection to other nodes (i.e. at runtime when template
+							// nodes aren't injected) forcefully remove to avoid data type mismatch.
+							if (ConnectedNodeClass->Metadata.GetClassName() == FInputNodeTemplate::ClassName)
+							{
+								RemoveEdgesFromNodeOutput(Edge.ToNodeID, ConnectedNode.Interface.Outputs.Last().VertexID);
+								ConnectedNode.Interface.Inputs.Last().TypeName = DataType;
+								ConnectedNode.Interface.Outputs.Last().TypeName = DataType;
+							}
+							else
+							{
+								RemoveEdgeToNodeInput(Edge.ToNodeID, Edge.ToVertexID);
+							}
+						}
+					});
+				}
+
+				const bool bNodeConformed = ConformGraphInputNodeToClass(GraphInput);
+				if (bNodeConformed)
+				{
+					RemoveUnusedDependencies();
+
+#if WITH_EDITORONLY_DATA
+					ClearMemberMetadata(GraphInput.NodeID);
+					Document.Metadata.ModifyContext.AddMemberIDModified(GraphInput.NodeID);
+					Document.Metadata.ModifyContext.AddNodeIDModified(GraphInput.NodeID);
+#endif // WITH_EDITORONLY_DATA
+
+					return true;
+				}
+			}
+		}
+	}
+
+	return false;
+}
+
 bool FMetaSoundFrontendDocumentBuilder::SetGraphInputDefault(FName InputName, const FMetasoundFrontendLiteral& InDefaultLiteral)
 {
 	using namespace Metasound::Frontend;
@@ -2722,39 +3230,100 @@ bool FMetaSoundFrontendDocumentBuilder::SetGraphInputInheritsDefault(FName InNam
 	return false;
 }
 
-bool FMetaSoundFrontendDocumentBuilder::SetNodeInputDefault(const FGuid& InNodeID, const FGuid& InVertexID, const FMetasoundFrontendLiteral& InLiteral)
+bool FMetaSoundFrontendDocumentBuilder::SetGraphOutputAccessType(FName OutputName, EMetasoundFrontendVertexAccessType AccessType)
 {
 	using namespace Metasound::Frontend;
 
-	const IDocumentGraphNodeCache& NodeCache = DocumentCache->GetNodeCache();
-	if (const int32* NodeIndex = NodeCache.FindNodeIndex(InNodeID))
+	if (!ensureMsgf(AccessType != EMetasoundFrontendVertexAccessType::Unset, TEXT("Cannot set graph output access type to '%s'"), LexToString(AccessType)))
 	{
-		FMetasoundFrontendGraph& Graph = GetDocument().RootGraph.Graph;
-		FMetasoundFrontendNode& Node = Graph.Nodes[*NodeIndex];
+		return false;
+	}
 
-		auto IsVertex = [&InVertexID](const FMetasoundFrontendVertex& Vertex) { return Vertex.VertexID == InVertexID; };
-		int32 VertexIndex = Node.Interface.Inputs.IndexOfByPredicate(IsVertex);
-		if (VertexIndex != INDEX_NONE)
+	if (const int32* Index = DocumentCache->GetInterfaceCache().FindOutputIndex(OutputName))
+	{
+		FMetasoundFrontendDocument& Document = GetDocument();
+		FMetasoundFrontendClassOutput& GraphOutput = Document.RootGraph.Interface.Outputs[*Index];
+		if (GraphOutput.AccessType != AccessType)
 		{
-			FMetasoundFrontendVertexLiteral NewVertexLiteral;
-			NewVertexLiteral.VertexID = InVertexID;
-			NewVertexLiteral.Value = InLiteral;
-
-			auto IsLiteral = [&InVertexID](const FMetasoundFrontendVertexLiteral& Literal) { return Literal.VertexID == InVertexID; };
-			int32 LiteralIndex = Node.InputLiterals.IndexOfByPredicate(IsLiteral);
-			if (LiteralIndex == INDEX_NONE)
+			GraphOutput.AccessType = AccessType;
+			if (AccessType == EMetasoundFrontendVertexAccessType::Value)
 			{
-				LiteralIndex = Node.InputLiterals.Num();
-				Node.InputLiterals.Add(MoveTemp(NewVertexLiteral));
+				const IDocumentGraphNodeCache& NodeCache = DocumentCache->GetNodeCache();
+				if (const int32* NodeIndex = NodeCache.FindNodeIndex(GraphOutput.NodeID))
+				{
+					FMetasoundFrontendNode& Node = Document.RootGraph.Graph.Nodes[*NodeIndex];
+					const FMetasoundFrontendVertex& NodeInput = Node.Interface.Inputs.Last();
+					IterateNodesConnectedWithVertex({ GraphOutput.NodeID, NodeInput.VertexID }, [this, &AccessType](const FMetasoundFrontendEdge& Edge, FMetasoundFrontendNode& ConnectedNode)
+					{
+						if (const FMetasoundFrontendClass* ConnectedNodeClass = FindDependency(ConnectedNode.ClassID))
+						{
+							const FMetasoundFrontendVertex& ConnectedNodeOutput = ConnectedNode.Interface.Outputs.Last();
+							const EMetasoundFrontendVertexAccessType ConnectedAccessType = GetNodeOutputAccessType(ConnectedNode.GetID(), ConnectedNodeOutput.VertexID);
+							if (!FMetasoundFrontendClassVertex::CanConnectVertexAccessTypes(ConnectedAccessType, AccessType))
+							{
+								RemoveEdgeToNodeInput(Edge.ToNodeID, Edge.ToVertexID);
+							}
+						}
+					});
+				}
 			}
-			else
+
+			const bool bNodeConformed = ConformGraphOutputNodeToClass(GraphOutput);
+			if (bNodeConformed)
 			{
-				Node.InputLiterals[LiteralIndex] = MoveTemp(NewVertexLiteral);
+	#if WITH_EDITORONLY_DATA
+				Document.Metadata.ModifyContext.AddMemberIDModified(GraphOutput.NodeID);
+	#endif // WITH_EDITORONLY_DATA
+
+				return true;
 			}
+		}
+	}
 
-			const FOnMetaSoundFrontendDocumentMutateNodeInputLiteralArray& OnNodeInputLiteralSet = DocumentDelegates->NodeDelegates.OnNodeInputLiteralSet;
-			OnNodeInputLiteralSet.Broadcast(*NodeIndex, VertexIndex, LiteralIndex);
-			return true;
+	return false;
+}
+
+bool FMetaSoundFrontendDocumentBuilder::SetGraphOutputDataType(FName OutputName, FName DataType)
+{
+	using namespace Metasound::Frontend;
+
+	if (IDataTypeRegistry::Get().IsRegistered(DataType))
+	{
+		if (const int32* Index = DocumentCache->GetInterfaceCache().FindOutputIndex(OutputName))
+		{
+			FMetasoundFrontendDocument& Document = GetDocument();
+			FMetasoundFrontendClassOutput& GraphOutput = Document.RootGraph.Interface.Outputs[*Index];
+			if (GraphOutput.TypeName != DataType)
+			{
+				GraphOutput.TypeName = DataType;
+
+				const IDocumentGraphNodeCache& NodeCache = DocumentCache->GetNodeCache();
+				if (const int32* NodeIndex = NodeCache.FindNodeIndex(GraphOutput.NodeID))
+				{
+					FMetasoundFrontendNode& Node = Document.RootGraph.Graph.Nodes[*NodeIndex];
+
+					FMetasoundFrontendLiteral DefaultLiteral;
+					DefaultLiteral.SetFromLiteral(IDataTypeRegistry::Get().CreateDefaultLiteral(DataType));
+					FMetasoundFrontendVertex& NodeInput = Node.Interface.Inputs.Last();
+					Node.InputLiterals = { FMetasoundFrontendVertexLiteral { NodeInput.VertexID, MoveTemp(DefaultLiteral) } };
+
+					RemoveEdgeToNodeInput(GraphOutput.NodeID, NodeInput.VertexID);
+					GraphOutput.TypeName = DataType;
+
+				}
+
+				const bool bNodeConformed = ConformGraphOutputNodeToClass(GraphOutput);
+				if (bNodeConformed)
+				{
+
+#if WITH_EDITORONLY_DATA
+					ClearMemberMetadata(GraphOutput.NodeID);
+					Document.Metadata.ModifyContext.AddMemberIDModified(GraphOutput.NodeID);
+#endif // WITH_EDITORONLY_DATA
+
+					return true;
+				}
+			}
 		}
 	}
 
@@ -2762,6 +3331,15 @@ bool FMetaSoundFrontendDocumentBuilder::SetNodeInputDefault(const FGuid& InNodeI
 }
 
 #if WITH_EDITOR
+void FMetaSoundFrontendDocumentBuilder::SetMemberMetadata(UMetaSoundFrontendMemberMetadata& NewMetadata)
+{
+	check(NewMetadata.MemberID.IsValid());
+
+	TMap<FGuid, TObjectPtr<UMetaSoundFrontendMemberMetadata>>& LiteralMetadata = GetDocument().Metadata.MemberMetadata;
+	LiteralMetadata.Remove(NewMetadata.MemberID);
+	LiteralMetadata.Add(NewMetadata.MemberID, NewMetadata);
+}
+
 bool FMetaSoundFrontendDocumentBuilder::SetNodeComment(const FGuid& InNodeID, FString&& InNewComment)
 {
 	if (FMetasoundFrontendNode* Node = FindNodeInternal(InNodeID))
@@ -2783,7 +3361,48 @@ bool FMetaSoundFrontendDocumentBuilder::SetNodeCommentVisible(const FGuid& InNod
 
 	return false;
 }
+#endif // WITH_EDITOR
+
+bool FMetaSoundFrontendDocumentBuilder::SetNodeInputDefault(const FGuid& InNodeID, const FGuid& InVertexID, const FMetasoundFrontendLiteral& InLiteral)
+{
+	using namespace Metasound::Frontend;
 
+	const IDocumentGraphNodeCache& NodeCache = DocumentCache->GetNodeCache();
+	if (const int32* NodeIndex = NodeCache.FindNodeIndex(InNodeID))
+	{
+		FMetasoundFrontendGraph& Graph = GetDocument().RootGraph.Graph;
+		FMetasoundFrontendNode& Node = Graph.Nodes[*NodeIndex];
+
+		auto IsVertex = [&InVertexID](const FMetasoundFrontendVertex& Vertex) { return Vertex.VertexID == InVertexID; };
+		int32 VertexIndex = Node.Interface.Inputs.IndexOfByPredicate(IsVertex);
+		if (VertexIndex != INDEX_NONE)
+		{
+			FMetasoundFrontendVertexLiteral NewVertexLiteral;
+			NewVertexLiteral.VertexID = InVertexID;
+			NewVertexLiteral.Value = InLiteral;
+
+			auto IsLiteral = [&InVertexID](const FMetasoundFrontendVertexLiteral& Literal) { return Literal.VertexID == InVertexID; };
+			int32 LiteralIndex = Node.InputLiterals.IndexOfByPredicate(IsLiteral);
+			if (LiteralIndex == INDEX_NONE)
+			{
+				LiteralIndex = Node.InputLiterals.Num();
+				Node.InputLiterals.Add(MoveTemp(NewVertexLiteral));
+			}
+			else
+			{
+				Node.InputLiterals[LiteralIndex] = MoveTemp(NewVertexLiteral);
+			}
+
+			const FOnMetaSoundFrontendDocumentMutateNodeInputLiteralArray& OnNodeInputLiteralSet = DocumentDelegates->NodeDelegates.OnNodeInputLiteralSet;
+			OnNodeInputLiteralSet.Broadcast(*NodeIndex, VertexIndex, LiteralIndex);
+			return true;
+		}
+	}
+
+	return false;
+}
+
+#if WITH_EDITOR
 bool FMetaSoundFrontendDocumentBuilder::SetNodeLocation(const FGuid& InNodeID, const FVector2D& InLocation, const FGuid* InLocationGuid)
 {
 	if (FMetasoundFrontendNode* Node = FindNodeInternal(InNodeID))
@@ -2791,6 +3410,7 @@ bool FMetaSoundFrontendDocumentBuilder::SetNodeLocation(const FGuid& InNodeID, c
 		FMetasoundFrontendNodeStyle& Style = Node->Style;
 		if (InLocationGuid)
 		{
+			ensureMsgf(InLocationGuid->IsValid(), TEXT("Invalid Location Guid no longer supported."));
 			Style.Display.Locations.FindOrAdd(*InLocationGuid) = InLocation;
 		}
 		else if (Style.Display.Locations.IsEmpty())
@@ -2972,6 +3592,39 @@ bool FMetaSoundFrontendDocumentBuilder::SwapGraphOutput(const FMetasoundFrontend
 	return true;
 }
 
+bool FMetaSoundFrontendDocumentBuilder::UnlinkVariableNode(const FGuid& InNodeID)
+{
+	using namespace Metasound::VariableNames;
+
+	FMetasoundFrontendGraphClass& GraphClass = GetDocument().RootGraph;
+	for (FMetasoundFrontendVariable& Variable : GraphClass.Graph.Variables)
+	{
+		if (InNodeID == Variable.VariableNodeID)
+		{
+			Variable.VariableNodeID = FGuid();
+			return true;
+		}
+
+		if (InNodeID == Variable.MutatorNodeID)
+		{
+			Variable.MutatorNodeID = FGuid();
+			return true;
+		}
+
+		if (Variable.AccessorNodeIDs.Remove(InNodeID) > 0)
+		{
+			return true;
+		}
+
+		if (Variable.DeferredAccessorNodeIDs.Remove(InNodeID) > 0)
+		{
+			return true;
+		}
+	}
+
+	return false;
+}
+
 bool FMetaSoundFrontendDocumentBuilder::UpdateDependencyClassNames(const TMap<FMetasoundFrontendClassName, FMetasoundFrontendClassName>& OldToNewReferencedClassNames)
 {
 	for (FMetasoundFrontendClass& Dependency : GetDocument().Dependencies)
diff --git a/Engine/Plugins/Runtime/Metasound/Source/MetasoundFrontend/Private/MetasoundFrontendDocumentCache.cpp b/Engine/Plugins/Runtime/Metasound/Source/MetasoundFrontend/Private/MetasoundFrontendDocumentCache.cpp
--- a/Engine/Plugins/Runtime/Metasound/Source/MetasoundFrontend/Private/MetasoundFrontendDocumentCache.cpp
+++ b/Engine/Plugins/Runtime/Metasound/Source/MetasoundFrontend/Private/MetasoundFrontendDocumentCache.cpp
@@ -338,6 +338,11 @@ namespace Metasound::Frontend
 		return nullptr;
 	}
 
+	const int32* FDocumentGraphInterfaceCache::FindInputIndex(FName InputName) const
+	{
+		return InputNameToIndex.Find(InputName);
+	}
+
 	const FMetasoundFrontendClassOutput* FDocumentGraphInterfaceCache::FindOutput(FName OutputName) const
 	{
 		if (const int32* Index = OutputNameToIndex.Find(OutputName))
@@ -349,6 +354,11 @@ namespace Metasound::Frontend
 		return nullptr;
 	}
 
+	const int32* FDocumentGraphInterfaceCache::FindOutputIndex(FName OutputName) const
+	{
+		return OutputNameToIndex.Find(OutputName);
+	}
+
 	int32 FDocumentGraphInterfaceCache::GetTransactionCount() const
 	{
 		return TransactionCount;
diff --git a/Engine/Plugins/Runtime/Metasound/Source/MetasoundFrontend/Private/MetasoundFrontendDocumentCache.h b/Engine/Plugins/Runtime/Metasound/Source/MetasoundFrontend/Private/MetasoundFrontendDocumentCache.h
--- a/Engine/Plugins/Runtime/Metasound/Source/MetasoundFrontend/Private/MetasoundFrontendDocumentCache.h
+++ b/Engine/Plugins/Runtime/Metasound/Source/MetasoundFrontend/Private/MetasoundFrontendDocumentCache.h
@@ -126,7 +126,9 @@ namespace Metasound::Frontend
 
 		// IDocumentGraphInterfaceCache implementation
 		virtual const FMetasoundFrontendClassInput* FindInput(FName InputName) const override;
+		virtual const int32* FindInputIndex(FName InputName) const override;
 		virtual const FMetasoundFrontendClassOutput* FindOutput(FName OutputName) const override;
+		virtual const int32* FindOutputIndex(FName OutputName) const override;
 
 		int32 GetTransactionCount() const;
 
diff --git a/Engine/Plugins/Runtime/Metasound/Source/MetasoundFrontend/Private/MetasoundFrontendDocumentIdGenerator.cpp b/Engine/Plugins/Runtime/Metasound/Source/MetasoundFrontend/Private/MetasoundFrontendDocumentIdGenerator.cpp
--- a/Engine/Plugins/Runtime/Metasound/Source/MetasoundFrontend/Private/MetasoundFrontendDocumentIdGenerator.cpp
+++ b/Engine/Plugins/Runtime/Metasound/Source/MetasoundFrontend/Private/MetasoundFrontendDocumentIdGenerator.cpp
@@ -11,14 +11,6 @@ namespace Metasound
 			std::atomic<uint64> GlobalAtomicMetasoundIdCounter = 1; // First ID will be 1. This is because we will make this into an FGuid which cannot be zero.
 		}
 
-		int32 MetaSoundEnableCookDeterministicIDGeneration = 1;
-		FAutoConsoleVariableRef CVarMetaSoundEnableCookDeterministicIDGeneration(
-			TEXT("au.MetaSound.EnableCookDeterministicIDGeneration"),
-			MetaSoundEnableCookDeterministicIDGeneration,
-			TEXT("Enable moving MetaSound registration operations like AutoUpdate and some template node transformations from runtime to cook using deterministic ID generation \n")
-			TEXT("0: Disabled, !0: Enabled (default)"),
-			ECVF_Default);
-
 		FDocumentIDGenerator::FScopeDeterminism::FScopeDeterminism(bool bInIsDeterministic)
 		{
 			FDocumentIDGenerator& IDGen = Get();
diff --git a/Engine/Plugins/Runtime/Metasound/Source/MetasoundFrontend/Private/MetasoundFrontendDocumentVersioning.h b/Engine/Plugins/Runtime/Metasound/Source/MetasoundFrontend/Private/MetasoundFrontendDocumentVersioning.h
--- a/Engine/Plugins/Runtime/Metasound/Source/MetasoundFrontend/Private/MetasoundFrontendDocumentVersioning.h
+++ b/Engine/Plugins/Runtime/Metasound/Source/MetasoundFrontend/Private/MetasoundFrontendDocumentVersioning.h
@@ -2,31 +2,35 @@
 #pragma once
 
 #include "Containers/UnrealString.h"
+#include "MetasoundDocumentInterface.h"
 #include "MetasoundFrontendController.h"
 #include "MetasoundFrontendTransform.h"
 #include "MetasoundFrontendDocument.h"
 #include "UObject/NameTypes.h"
+#include "UObject/ScriptInterface.h"
 
+class IMetaSoundDocumentInterface;
 
-namespace Metasound
+
+namespace Metasound::Frontend
 {
-	namespace Frontend
+	/** Base class for versioning a document. */
+	class METASOUNDFRONTEND_API FVersionDocument
 	{
-		/** Base class for versioning a document. */
-		class METASOUNDFRONTEND_API FVersionDocument : public IDocumentTransform
+		const FName Name;
+		const FString& Path;
+
+	public:
+		static FMetasoundFrontendVersionNumber GetMaxVersion()
 		{
-			const FName Name;
-			const FString& Path;
+			return FMetasoundFrontendVersionNumber { 1, 12 };
+		}
 
-		public:
-			static FMetasoundFrontendVersionNumber GetMaxVersion()
-			{
-				return FMetasoundFrontendVersionNumber { 1, 11 };
-			}
+		FVersionDocument(FName InName, const FString& InPath);
 
-			FVersionDocument(FName InName, const FString& InPath);
+		UE_DEPRECATED(5.4, "Use DocumentInterface overload instead")
+		bool Transform(FDocumentHandle InDocument) const;
 
-			bool Transform(FDocumentHandle InDocument) const override;
-		};
-	} // namespace Frontend
-} // namespace Metasound
+		bool Transform(TScriptInterface<IMetaSoundDocumentInterface> DocumentInterface) const;
+	};
+} // namespace Metasound::Frontend
diff --git a/Engine/Plugins/Runtime/Metasound/Source/MetasoundFrontend/Private/MetasoundFrontendGraphController.cpp b/Engine/Plugins/Runtime/Metasound/Source/MetasoundFrontend/Private/MetasoundFrontendGraphController.cpp
--- a/Engine/Plugins/Runtime/Metasound/Source/MetasoundFrontend/Private/MetasoundFrontendGraphController.cpp
+++ b/Engine/Plugins/Runtime/Metasound/Source/MetasoundFrontend/Private/MetasoundFrontendGraphController.cpp
@@ -750,6 +750,13 @@ namespace Metasound
 				GraphClass->PresetOptions.InputsInheritingDefault.Reset();
 				OwningDocument->ClearInterfaceVersions();
 				OwningDocument->RemoveUnreferencedDependencies();
+
+#if WITH_EDITORONLY_DATA
+				FDocumentAccessPtr DocumentPtr = OwningDocument->GetDocumentPtr();
+				FMetasoundFrontendDocument* Document = DocumentPtr.Get();
+				check(Document);
+				Document->Metadata.MemberMetadata.Reset();
+#endif // WITH_EDITORONLY_DATA
 			}
 		}
 
diff --git a/Engine/Plugins/Runtime/Metasound/Source/MetasoundFrontend/Private/MetasoundFrontendNodeTemplateRegistry.cpp b/Engine/Plugins/Runtime/Metasound/Source/MetasoundFrontend/Private/MetasoundFrontendNodeTemplateRegistry.cpp
--- a/Engine/Plugins/Runtime/Metasound/Source/MetasoundFrontend/Private/MetasoundFrontendNodeTemplateRegistry.cpp
+++ b/Engine/Plugins/Runtime/Metasound/Source/MetasoundFrontend/Private/MetasoundFrontendNodeTemplateRegistry.cpp
@@ -12,11 +12,6 @@
 
 namespace Metasound::Frontend
 {
-	TUniquePtr<INodeTransform> INodeTemplate::GenerateNodeTransform(FMetasoundFrontendDocument& InDocument) const
-	{
-		return nullptr;
-	}
-
 	class FNodeTemplateRegistry : public INodeTemplateRegistry
 	{
 	public:
@@ -58,9 +53,9 @@ namespace Metasound::Frontend
 
 	const INodeTemplate* FNodeTemplateRegistry::FindTemplate(const FNodeRegistryKey& InKey) const
 	{
-		if (const TUniquePtr<const INodeTemplate>* TemplatePtr = Templates.Find(InKey))
+		if (const TUniquePtr<const INodeTemplate>* Template = Templates.Find(InKey))
 		{
-			return TemplatePtr->Get();
+			return Template->Get();
 		}
 
 		return nullptr;
@@ -87,6 +82,11 @@ namespace Metasound::Frontend
 		return Registry;
 	}
 
+	TUniquePtr<INodeTransform> INodeTemplate::GenerateNodeTransform(FMetasoundFrontendDocument& InDocument) const
+	{
+		return nullptr;
+	}
+
 #if WITH_EDITOR
 	FText FNodeTemplateBase::GetNodeDisplayName(const IMetaSoundDocumentInterface& Interface, const FGuid& InNodeID) const
 	{
diff --git a/Engine/Plugins/Runtime/Metasound/Source/MetasoundFrontend/Private/MetasoundFrontendOutputController.cpp b/Engine/Plugins/Runtime/Metasound/Source/MetasoundFrontend/Private/MetasoundFrontendOutputController.cpp
--- a/Engine/Plugins/Runtime/Metasound/Source/MetasoundFrontend/Private/MetasoundFrontendOutputController.cpp
+++ b/Engine/Plugins/Runtime/Metasound/Source/MetasoundFrontend/Private/MetasoundFrontendOutputController.cpp
@@ -2,6 +2,7 @@
 
 #include "MetasoundFrontendOutputController.h"
 #include "MetasoundFrontendInvalidController.h"
+#include "MetasoundFrontendNodeTemplateRegistry.h"
 
 #define LOCTEXT_NAMESPACE "MetasoundFrontendOutputController"
 
@@ -51,7 +52,10 @@ namespace Metasound
 
 		EMetasoundFrontendVertexAccessType FBaseOutputController::GetVertexAccessType() const
 		{
+			PRAGMA_DISABLE_DEPRECATION_WARNINGS
 			FConstOutputHandle ReroutedOutput = Frontend::FindReroutedOutput(AsShared());
+			PRAGMA_ENABLE_DEPRECATION_WARNINGS
+
 			if (ReroutedOutput->IsValid())
 			{
 				if (ReroutedOutput->GetOwningNodeID() != GetOwningNodeID() || ReroutedOutput->GetID() != GetID())
@@ -145,6 +149,17 @@ namespace Metasound
 
 		bool FBaseOutputController::IsConnectionUserModifiable() const 
 		{
+			FConstNodeHandle Owner = GetOwningNode();
+			if (Owner->GetClassMetadata().GetType() == EMetasoundFrontendClassType::Template)
+			{
+				const FNodeRegistryKey Key(GetOwningNode()->GetClassMetadata());
+				const INodeTemplate* Template = INodeTemplateRegistry::Get().FindTemplate(Key);
+				if (ensure(Template))
+				{
+					return Template->IsOutputConnectionUserModifiable();
+				}
+			}
+
 			return true;
 		}
 
diff --git a/Engine/Plugins/Runtime/Metasound/Source/MetasoundFrontend/Private/MetasoundFrontendRegistryContainerImpl.cpp b/Engine/Plugins/Runtime/Metasound/Source/MetasoundFrontend/Private/MetasoundFrontendRegistryContainerImpl.cpp
--- a/Engine/Plugins/Runtime/Metasound/Source/MetasoundFrontend/Private/MetasoundFrontendRegistryContainerImpl.cpp
+++ b/Engine/Plugins/Runtime/Metasound/Source/MetasoundFrontend/Private/MetasoundFrontendRegistryContainerImpl.cpp
@@ -17,6 +17,7 @@
 #include "UObject/NameTypes.h"
 #include "UObject/ScriptInterface.h"
 
+#include "MetasoundAssetBase.h"
 #include "MetasoundAssetManager.h"
 #include "MetasoundDocumentInterface.h"
 #include "MetasoundFrontendDataTypeRegistry.h"
@@ -42,51 +43,41 @@ namespace Metasound::Frontend
 {
 	namespace RegistryPrivate
 	{
-		TScriptInterface<IMetaSoundDocumentInterface> BuildRegistryDocument(TScriptInterface<IMetaSoundDocumentInterface> DocumentInterface, bool bForceCopy)
+		TScriptInterface<IMetaSoundDocumentInterface> BuildRegistryDocument(TScriptInterface<IMetaSoundDocumentInterface> DocumentInterface, bool bAsync)
 		{
 			using namespace Metasound::Frontend;
 
 			METASOUND_TRACE_CPUPROFILER_EVENT_SCOPE(Metasound::Frontend::BuildRegistryDocument);
 
+			UObject* DocObject = DocumentInterface.GetObject();
+			check(DocObject);
+
 #if WITH_EDITOR
-			// Node template transform is performed on copy of local document to avoid overwriting editable data
-			constexpr bool bTransformDocumentBeforeRegistering = true;
-#else // !WITH_EDITOR
-			// Node template transform is performed on local document only if cook determinism ID generation
-			// is enabled to avoid transforms potentially creating new edges with non-deterministic IDs.
-			const bool bTransformDocumentBeforeRegistering = MetaSoundEnableCookDeterministicIDGeneration == 0;
-#endif // !WITH_EDITOR
-
-			if (bTransformDocumentBeforeRegistering)
+			// Only assets require template node processing
+			if (DocObject->IsAsset())
 			{
-				// 1. Find template dependencies to build prior to making new document/builder as an optimization
-				// (no sense in creating new document/builder if no templates need processing)
-				FMetaSoundFrontendDocumentBuilder OriginalDocBuilder(DocumentInterface);
-				const bool bContainsTemplateDependency = OriginalDocBuilder.ContainsDependencyOfType(EMetasoundFrontendClassType::Template);
-				if (bContainsTemplateDependency)
-				{
-					UMetaSoundBuilderDocument& RegistryDocObject = UMetaSoundBuilderDocument::Create(*DocumentInterface.GetInterface());
-					FMetaSoundFrontendDocumentBuilder RegistryDocBuilder(&RegistryDocObject);
-
-					RegistryDocBuilder.TransformTemplateNodes();
-
-					return &RegistryDocObject;
-				}
+				DocumentInterface = &UMetaSoundBuilderDocument::Create(*DocumentInterface.GetInterface());
+				FMetaSoundFrontendDocumentBuilder Builder(DocumentInterface);
+				Builder.TransformTemplateNodes();
 			}
-#if !NO_LOGGING
-			else
+
+			return DocumentInterface;
+#else
+	#if !NO_LOGGING
+			const FMetaSoundFrontendDocumentBuilder& OriginalDocBuilder = IMetaSoundAssetManager::GetChecked().AttachDocumentBuilderChecked(*DocObject);
+			const bool bContainsTemplateDependency = OriginalDocBuilder.ContainsDependencyOfType(EMetasoundFrontendClassType::Template);
+			if (bContainsTemplateDependency)
 			{
-				FMetaSoundFrontendDocumentBuilder OriginalDocBuilder(DocumentInterface);
-				const bool bContainsTemplateDependency = OriginalDocBuilder.ContainsDependencyOfType(EMetasoundFrontendClassType::Template);
-				if (bContainsTemplateDependency)
-				{
-					UE_LOG(LogMetaSound, Error,
-						TEXT("Template node processing disabled but provided asset class at '%s' to register contains template nodes. Runtime graph will fail to build."),
-						*OriginalDocBuilder.GetDebugName());
-				}
+				UE_LOG(LogMetaSound, Error,
+					TEXT("Template node processing disabled but provided asset class at '%s' to register contains template nodes. Runtime graph will fail to build."),
+					*OriginalDocBuilder.GetDebugName());
 			}
-#endif // !NO_LOGGING
+	#endif // !NO_LOGGING
 
+			// Force a copy if async registration is enabled and we need to protect against race conditions from external modifications.
+			const FMetasoundAssetBase* AssetBase = IMetaSoundAssetManager::GetChecked().GetAsAsset(*DocObject);
+			check(AssetBase);
+			const bool bForceCopy = AssetBase->IsBuilderActive() && bAsync;
 			if (bForceCopy)
 			{
 				return &UMetaSoundBuilderDocument::Create(*DocumentInterface.GetInterface());
@@ -95,6 +86,7 @@ namespace Metasound::Frontend
 			{
 				return DocumentInterface;
 			}
+#endif // WITH_EDITOR
 		}
 
 		// FGraphNode is used to create unique INodes based off of a IGraph.
@@ -437,7 +429,7 @@ namespace Metasound::Frontend
 		return MakeUnique<FNodeRegistryTransactionStream>(TransactionBuffer);
 	}
 
-	FGraphRegistryKey FRegistryContainerImpl::RegisterGraph(const TScriptInterface<IMetaSoundDocumentInterface>& InDocumentInterface, bool bAsync, bool bForceCopy)
+	FGraphRegistryKey FRegistryContainerImpl::RegisterGraph(const TScriptInterface<IMetaSoundDocumentInterface>& InDocumentInterface, bool bAsync)
 	{
 		using namespace UE;
 
@@ -446,15 +438,8 @@ namespace Metasound::Frontend
 		check(InDocumentInterface);
 		check(IsInGameThread());
 
-		// Use the asset path of the provided document interface object for identification, *NOT* the
-		// built version as the build process may in fact create a new object with a transient path.
+		const FMetasoundFrontendDocument& Document = InDocumentInterface->GetConstDocument();
 		const FTopLevelAssetPath AssetPath = InDocumentInterface->GetAssetPathChecked();
-		const TScriptInterface<IMetaSoundDocumentInterface> RegistryDocInterface = RegistryPrivate::BuildRegistryDocument(InDocumentInterface, bForceCopy);
-
-		UObject* OwningObject = RegistryDocInterface.GetObject();
-		check(OwningObject);
-
-		const FMetasoundFrontendDocument& Document = RegistryDocInterface->GetConstDocument();
 		const FGraphRegistryKey RegistryKey { FNodeRegistryKey(Document.RootGraph), AssetPath };
 
 		if (!RegistryKey.IsValid())
@@ -466,6 +451,19 @@ namespace Metasound::Frontend
 
 		METASOUND_TRACE_CPUPROFILER_EVENT_SCOPE_TEXT(*FString::Printf(TEXT("FRegistryContainerImpl::RegisterGraph key:%s, asset %s"), *RegistryKey.ToString(), *AssetPath.ToString()));
 
+		// Wait for any async tasks that are in flight which correspond to the same graph prior to building as building, if force copy is false, may mutate the serialized document
+		if (bAsync)
+		{
+			WaitForAsyncGraphRegistration(RegistryKey);
+		}
+
+		// Use the asset path of the provided document interface object for identification, *NOT* the
+		// built version as the build process may in fact create a new object with a transient path.
+		const TScriptInterface<IMetaSoundDocumentInterface> RegistryDocInterface = RegistryPrivate::BuildRegistryDocument(InDocumentInterface, bAsync);
+
+		UObject* OwningObject = RegistryDocInterface.GetObject();
+		check(OwningObject);
+
 		FNodeClassInfo NodeClassInfo(Document.RootGraph, AssetPath);
 
 		// Proxies are created synchronously to avoid creating proxies in async tasks. Proxies
@@ -482,20 +480,13 @@ namespace Metasound::Frontend
 
 		if (bAsync)
 		{
-			// Wait for any async tasks that are in flight which correspond to the same graph
-			WaitForAsyncGraphRegistration(RegistryKey);
-
 			Tasks::FTask BuildAndRegisterTask = AsyncRegistrationPipe.Launch(
 				UE_SOURCE_LOCATION,
 				[RegistryKey, ClassInfo = MoveTemp(NodeClassInfo), RegistryDocInterface, ProxyDataCache = MoveTemp(ProxyDataCache)]() mutable
 				{
 					FRegistryContainerImpl& Registry = FRegistryContainerImpl::Get();
 					// Unregister the graph before re-registering
-					if (Registry.IsGraphRegistered(RegistryKey))
-					{
-						Registry.UnregisterGraphInternal(RegistryKey);
-					}
-
+					Registry.UnregisterGraphInternal(RegistryKey);
 					Registry.BuildAndRegisterGraphFromDocument(RegistryDocInterface, ProxyDataCache, MoveTemp(ClassInfo));
 					Registry.RemoveRegistrationTask(RegistryKey, FNodeRegistryTransaction::ETransactionType::NodeRegistration);
 					Registry.RemoveDocumentReference(RegistryDocInterface);
@@ -512,10 +503,7 @@ namespace Metasound::Frontend
 		}
 		else
 		{
-			if (IsGraphRegistered(RegistryKey))
-			{
-				UnregisterGraphInternal(RegistryKey);
-			}
+			UnregisterGraphInternal(RegistryKey);
 
 			// Build and register graph synchronously
 			BuildAndRegisterGraphFromDocument(RegistryDocInterface, ProxyDataCache, MoveTemp(NodeClassInfo));
@@ -620,6 +608,11 @@ namespace Metasound::Frontend
 
 		FScopeLock Lock(&RegistryMapsCriticalSection);
 		{
+			if (!RegisteredGraphs.Contains(InKey))
+			{
+				return false;
+			}
+
 			const int32 GraphUnregistered = RegisteredGraphs.Remove(InKey) > 0;
 			const bool bNodeUnregistered = UnregisterNodeInternal(InKey.NodeKey);
 
diff --git a/Engine/Plugins/Runtime/Metasound/Source/MetasoundFrontend/Private/MetasoundFrontendRegistryContainerImpl.h b/Engine/Plugins/Runtime/Metasound/Source/MetasoundFrontend/Private/MetasoundFrontendRegistryContainerImpl.h
--- a/Engine/Plugins/Runtime/Metasound/Source/MetasoundFrontend/Private/MetasoundFrontendRegistryContainerImpl.h
+++ b/Engine/Plugins/Runtime/Metasound/Source/MetasoundFrontend/Private/MetasoundFrontendRegistryContainerImpl.h
@@ -128,7 +128,7 @@ namespace Metasound
 			virtual TArray<FConverterNodeInfo> GetPossibleConverterNodes(const FName& FromDataType, const FName& ToDataType) override;
 
 			// Register a graph from an IMetaSoundDocumentInterface
-			FGraphRegistryKey RegisterGraph(const TScriptInterface<IMetaSoundDocumentInterface>& InDocumentInterface, bool bAsync = true, bool bForceCopy = false);
+			FGraphRegistryKey RegisterGraph(const TScriptInterface<IMetaSoundDocumentInterface>& InDocumentInterface, bool bAsync = true);
 
 			// Unregister a graph
 			bool UnregisterGraph(const FGraphRegistryKey& InRegistryKey, const TScriptInterface<IMetaSoundDocumentInterface>& InDocumentInterface, bool bAsync = true);
