Version Metasound Document to include all ed data and make all Metasound EdGraph data transient

diff --git a/Engine/Plugins/Runtime/Metasound/Source/MetasoundFrontend/Private/MetasoundFrontendTransform.cpp b/Engine/Plugins/Runtime/Metasound/Source/MetasoundFrontend/Private/MetasoundFrontendTransform.cpp
--- a/Engine/Plugins/Runtime/Metasound/Source/MetasoundFrontend/Private/MetasoundFrontendTransform.cpp
+++ b/Engine/Plugins/Runtime/Metasound/Source/MetasoundFrontend/Private/MetasoundFrontendTransform.cpp
@@ -4,9 +4,11 @@
 #include "Algo/Transform.h"
 #include "Interfaces/MetasoundFrontendInterface.h"
 #include "Interfaces/MetasoundFrontendInterfaceRegistry.h"
+#include "NodeTemplates/MetasoundFrontendNodeTemplateInput.h"
 #include "MetasoundAccessPtr.h"
 #include "MetasoundAssetBase.h"
 #include "MetasoundFrontendDocument.h"
+#include "MetasoundFrontendDocumentBuilder.h"
 #include "MetasoundFrontendDocumentController.h"
 #include "MetasoundFrontendDocumentIdGenerator.h"
 #include "MetasoundFrontendRegistries.h"
@@ -54,6 +56,12 @@ namespace Metasound
 
 		bool IDocumentTransform::Transform(FMetasoundFrontendDocument& InOutDocument) const
 		{
+			if (IDocumentBuilderRegistry* DocRegistry = IDocumentBuilderRegistry::Get())
+			{
+				const FMetasoundFrontendClassName& Name = InOutDocument.RootGraph.Metadata.GetClassName();
+				DocRegistry->InvalidateDocumentCache(Name);
+			}
+
 			FDocumentAccessPtr DocAccessPtr = MakeAccessPtr<FDocumentAccessPtr>(InOutDocument.AccessPoint, InOutDocument);
 			return Transform(FDocumentController::CreateDocumentHandle(DocAccessPtr));
 		}
@@ -405,7 +413,12 @@ namespace Metasound
 
 		bool FModifyRootGraphInterfaces::Transform(FMetasoundFrontendDocument& InOutDocument) const
 		{
-			// TODO: Swap implementation to not use access pointers/controllers
+			if (IDocumentBuilderRegistry* DocRegistry = IDocumentBuilderRegistry::Get())
+			{
+				const FMetasoundFrontendClassName& Name = InOutDocument.RootGraph.Metadata.GetClassName();
+				DocRegistry->InvalidateDocumentCache(Name);
+			}
+
 			FDocumentAccessPtr DocAccessPtr = MakeAccessPtr<FDocumentAccessPtr>(InOutDocument.AccessPoint, InOutDocument);
 			return Transform(FDocumentController::CreateDocumentHandle(DocAccessPtr));
 		}
@@ -507,97 +520,6 @@ namespace Metasound
 		}
 #endif // WITH_EDITOR
 
-		bool FUpdateRootGraphInterface::Transform(FDocumentHandle InDocument) const
-		{
-			bool bDidEdit = false;
-
-			if (!ensure(InDocument->IsValid()))
-			{
-				return bDidEdit;
-			}
-
-			// Find registered target interface.
-			FMetasoundFrontendInterface TargetInterface;
-			bool bFoundTargetInterface = ISearchEngine::Get().FindInterfaceWithHighestVersion(InterfaceVersion.Name, TargetInterface);
-			if (!bFoundTargetInterface)
-			{
-				UE_LOG(LogMetaSound, Warning, TEXT("Could not check for interface updates. Target interface is not registered [InterfaceVersion:%s] when attempting to update root graph of asset (%s). Ensure that the module which registers the interface has been loaded before the asset is loaded."), *InterfaceVersion.ToString(), *OwningAssetName);
-				return false;
-			}
-
-			if (TargetInterface.Version == InterfaceVersion)
-			{
-				return false;
-			}
-
-			// Attempt to upgrade
-			TArray<const IInterfaceRegistryEntry*> UpgradePath;
-			GetUpdatePathForDocument(InterfaceVersion, TargetInterface.Version, UpgradePath);
-			return UpdateDocumentInterface(UpgradePath, InDocument);
-		}
-
-		void FUpdateRootGraphInterface::GetUpdatePathForDocument(const FMetasoundFrontendVersion& InCurrentVersion, const FMetasoundFrontendVersion& InTargetVersion, TArray<const IInterfaceRegistryEntry*>& OutUpgradePath) const
-		{
-			if (InCurrentVersion.Name == InTargetVersion.Name)
-			{
-				// Get all associated registered interfaces
-				TArray<FMetasoundFrontendVersion> RegisteredVersions = ISearchEngine::Get().FindAllRegisteredInterfacesWithName(InTargetVersion.Name);
-
-				// Filter registry entries that exist between current version and target version
-				auto FilterRegistryEntries = [&InCurrentVersion, &InTargetVersion](const FMetasoundFrontendVersion& InVersion)
-				{
-					const bool bIsGreaterThanCurrent = InVersion.Number > InCurrentVersion.Number;
-					const bool bIsLessThanOrEqualToTarget = InVersion.Number <= InTargetVersion.Number;
-
-					return bIsGreaterThanCurrent && bIsLessThanOrEqualToTarget;
-				};
-				RegisteredVersions = RegisteredVersions.FilterByPredicate(FilterRegistryEntries);
-
-				// sort registry entries to create an ordered upgrade path.
-				RegisteredVersions.Sort();
-
-				// Get registry entries from registry keys.
-				auto GetRegistryEntry = [](const FMetasoundFrontendVersion& InVersion)
-				{
-					FInterfaceRegistryKey Key = GetInterfaceRegistryKey(InVersion);
-					return IInterfaceRegistry::Get().FindInterfaceRegistryEntry(Key);
-				};
-				Algo::Transform(RegisteredVersions, OutUpgradePath, GetRegistryEntry);
-			}
-		}
-
-		bool FUpdateRootGraphInterface::UpdateDocumentInterface(const TArray<const IInterfaceRegistryEntry*>& InUpgradePath, FDocumentHandle InDocument) const
-		{
-			const FMetasoundFrontendVersionNumber* LastVersionUpdated = nullptr;
-			for (const IInterfaceRegistryEntry* Entry : InUpgradePath)
-			{
-				if (ensure(nullptr != Entry))
-				{
-					if (Entry->UpdateRootGraphInterface(InDocument))
-					{
-						LastVersionUpdated = &Entry->GetInterface().Version.Number;
-					}
-				}
-			}
-
-			if (LastVersionUpdated)
-			{
-#if WITH_EDITOR
-				const FString AssetName = *InDocument->GetRootGraphClass().Metadata.GetDisplayName().ToString();
-#else
-				const FString AssetName = *InDocument->GetRootGraphClass().Metadata.GetClassName().ToString();
-#endif // !WITH_EDITOR
-				UE_LOG(LogMetaSound, Display, TEXT("Asset '%s' interface '%s' updated: '%s' --> '%s'"),
-					*AssetName,
-					*InterfaceVersion.Name.ToString(),
-					*InterfaceVersion.Number.ToString(),
-					*LastVersionUpdated->ToString());
-				return true;
-			}
-
-			return false;
-		}
-
 		bool FAutoUpdateRootGraph::Transform(FDocumentHandle InDocument) const
 		{
 			METASOUND_TRACE_CPUPROFILER_EVENT_SCOPE(FAutoUpdateRootGraph::Transform);
@@ -680,7 +602,9 @@ namespace Metasound
 						FMetasoundAssetBase* PresetMetaSoundAsset = IMetaSoundAssetManager::GetChecked().TryLoadAssetFromKey(RegistryKey);
 						if (ensure(PresetMetaSoundAsset))
 						{
-							PresetMetaSoundAsset->ConformObjectDataToInterfaces();
+							TScriptInterface<IMetaSoundDocumentInterface> PresetInterface = PresetMetaSoundAsset->GetOwningAsset();
+							check(PresetInterface);
+							PresetInterface->ConformObjectToDocument();
 						}
 
 						InDocument->RemoveUnreferencedDependencies();
@@ -799,7 +723,8 @@ namespace Metasound
 			FMetasoundFrontendNodeStyle RefNodeStyle;
 
 			// Offset to be to the right of input nodes
-			RefNodeStyle.Display.Locations.Add(FGuid(), DisplayStyle::NodeLayout::DefaultOffsetX);
+			const FGuid EdNodeGuid = FGuid::NewGuid(); // EdNodes are now never serialized and are transient, so just assign here
+			RefNodeStyle.Display.Locations.Add(EdNodeGuid, DisplayStyle::NodeLayout::DefaultOffsetX);
 			ReferencedNodeHandle->SetNodeStyle(RefNodeStyle);
 #endif // WITH_EDITOR
 
@@ -831,6 +756,8 @@ namespace Metasound
 
 			FConstGraphHandle ReferencedGraphHandle = ReferencedDocument->GetRootGraph();
 
+			const INodeTemplate* InputTemplate = INodeTemplateRegistry::Get().FindTemplate(FInputNodeTemplate::ClassName);
+			check(InputTemplate);
 			TArray<FNodeHandle> NodeHandles;
 			for (const FMetasoundFrontendClassInput& ClassInput : InClassInputs)
 			{
@@ -842,6 +769,14 @@ namespace Metasound
 					FInputHandle InputToConnect = InReferencedNode->GetInputWithVertexName(ClassInput.Name);
 					ensure(OutputToConnect->Connect(*InputToConnect));
 					NodeHandles.Add(MoveTemp(InputNode));
+
+					// template node takes on data type of concrete input node's output type
+					const FName DataType = InputNode->GetOutputs().Last()->GetDataType();
+
+					FNodeTemplateGenerateInterfaceParams Params { { }, { DataType }};
+					FNodeHandle TemplateNodeHandle = InPresetGraphHandle->AddTemplateNode(*InputTemplate, MoveTemp(Params));
+					TemplateNodeHandle->GetInputs().Last()->Connect(*OutputToConnect);
+					TemplateNodeHandle->GetOutputs().Last()->Connect(*InputToConnect);
 				}
 			}
 
@@ -853,12 +788,17 @@ namespace Metasound
 
 			Style.SortDefaults(NodeHandles, DocumentTransform::GetNodeDisplayNameProjection());
 
-			for (const FNodeHandle& InputNode : NodeHandles)
+			InputNodeLocation = FVector2D::ZeroVector;
+			for (const FNodeHandle& NodeHandle : NodeHandles)
 			{
-				// Set input node location
+				// Create input template node and set location
 				FMetasoundFrontendNodeStyle NodeStyle;
-				NodeStyle.Display.Locations.Add(FGuid(), InputNodeLocation);
-				InputNode->SetNodeStyle(NodeStyle);
+				const FGuid EdNodeGuid = FGuid::NewGuid(); // EdNodes are now never serialized and are transient, so just assign here
+				NodeStyle.Display.Locations.Add(EdNodeGuid, InputNodeLocation);
+				FOutputHandle InputNodeOutputHandle = NodeHandle->GetOutputs().Last();
+				FInputHandle InputTemplateNodeInputHandle = InputNodeOutputHandle->GetConnectedInputs().Last();
+				FNodeHandle TemplateNodeHandle = InputTemplateNodeInputHandle->GetOwningNode();
+				TemplateNodeHandle->SetNodeStyle(NodeStyle);
 				InputNodeLocation += DisplayStyle::NodeLayout::DefaultOffsetY;
 			}
 #endif // WITH_EDITOR
@@ -897,7 +837,8 @@ namespace Metasound
 			for (const FNodeHandle& OutputNode : NodeHandles)
 			{
 				FMetasoundFrontendNodeStyle NodeStyle;
-				NodeStyle.Display.Locations.Add(FGuid(), OutputNodeLocation);
+				const FGuid EdNodeGuid = FGuid::NewGuid(); // EdNodes are now never serialized and are transient, so just assign here
+				NodeStyle.Display.Locations.Add(EdNodeGuid, OutputNodeLocation);
 				OutputNode->SetNodeStyle(NodeStyle);
 				OutputNodeLocation += DisplayStyle::NodeLayout::DefaultOffsetY;
 			}
diff --git a/Engine/Plugins/Runtime/Metasound/Source/MetasoundFrontend/Private/NodeTemplates/MetasoundFrontendNodeTemplateInput.cpp b/Engine/Plugins/Runtime/Metasound/Source/MetasoundFrontend/Private/NodeTemplates/MetasoundFrontendNodeTemplateInput.cpp
new file mode 100644
--- /dev/null
+++ b/Engine/Plugins/Runtime/Metasound/Source/MetasoundFrontend/Private/NodeTemplates/MetasoundFrontendNodeTemplateInput.cpp
@@ -0,0 +1,282 @@
+// Copyright Epic Games, Inc. All Rights Reserved.
+
+#include "NodeTemplates/MetasoundFrontendNodeTemplateInput.h"
+
+#include "Algo/AnyOf.h"
+#include "Internationalization/Text.h"
+#include "MetasoundFrontendDataTypeRegistry.h"
+#include "MetasoundFrontendDocumentBuilder.h"
+#include "MetasoundFrontendNodeTemplateRegistry.h"
+#include "MetasoundFrontendRegistries.h"
+#include "MetasoundFrontendRegistryKey.h"
+#include "MetasoundFrontendTransform.h"
+#include "MetasoundInputNode.h"
+
+#if WITH_EDITOR
+#include "MetasoundFrontendController.h"
+#endif // WITH_EDITOR
+
+
+namespace Metasound::Frontend
+{
+	namespace InputNodeTemplatePrivate
+	{
+		// Creates an input template node, sets node position (should only ever be one in style location) from and connects it to the associated input with the given name.
+		const FMetasoundFrontendNode* InitTemplateNode(const INodeTemplate& InTemplate, FName InputName, FMetaSoundFrontendDocumentBuilder& InOutBuilder)
+		{
+			FMetasoundFrontendEdge NewEdge;
+			FName TypeName;
+#if WITH_EDITORONLY_DATA
+			TMap<FGuid, FVector2D> Locations;
+#endif // WITH_EDITORONLY_DATA
+
+			// Cache data from Input node pointer as needed as subsequent call to create new template node may invalidate the input node's pointer
+			{
+				const FMetasoundFrontendNode* InputNode = InOutBuilder.FindGraphInputNode(InputName);
+				check(InputNode);
+				TypeName = InputNode->Interface.Outputs.Last().TypeName;
+				NewEdge.FromNodeID = InputNode->GetID(),
+				NewEdge.FromVertexID = InputNode->Interface.Outputs.Last().VertexID;
+
+#if WITH_EDITORONLY_DATA
+				Locations = InputNode->Style.Display.Locations;
+				ensure(Locations.Num() <= 1);
+#endif // WITH_EDITORONLY_DATA
+			}
+
+			FNodeTemplateGenerateInterfaceParams Params { { }, { TypeName } };
+			const FMetasoundFrontendNode* NewNode = InOutBuilder.AddNodeByTemplate(InTemplate, MoveTemp(Params));
+			check(NewNode);
+			NewEdge.ToNodeID = NewNode->GetID();
+			NewEdge.ToVertexID = NewNode->Interface.Inputs.Last().VertexID;
+
+#if WITH_EDITORONLY_DATA
+			for (const TPair<FGuid, FVector2D>& Pair : Locations)
+			{
+				InOutBuilder.SetNodeLocation(NewEdge.ToNodeID, Pair.Value);
+			}
+#endif // WITH_EDITORONLY_DATA
+
+			// Add edge between input node and new template node
+			InOutBuilder.AddEdge(MoveTemp(NewEdge));
+
+			return NewNode;
+		};
+	} // namespace InputNodeTemplatePrivate
+
+	const FMetasoundFrontendClassName FInputNodeTemplate::ClassName { "UE", "Input", "Template" };
+
+	const FMetasoundFrontendVersionNumber FInputNodeTemplate::VersionNumber = { 1, 0 } ;
+
+	const FMetasoundFrontendClassName& FInputNodeTemplate::GetClassName() const
+	{
+		return FInputNodeTemplate::ClassName;
+	}
+
+#if WITH_EDITOR
+	FText FInputNodeTemplate::GetNodeDisplayName(const IMetaSoundDocumentInterface& Interface, const FGuid& InNodeID) const
+	{
+		return { };
+	}
+#endif // WITH_EDITOR
+
+	const FMetasoundFrontendClass& FInputNodeTemplate::GetFrontendClass() const
+	{
+		auto CreateFrontendClass = []()
+		{
+			using namespace Metasound;
+
+			FMetasoundFrontendClass Class;
+			Class.Metadata.SetClassName(ClassName);
+
+#if WITH_EDITOR
+			Class.Metadata.SetSerializeText(false);
+			Class.Metadata.SetAuthor(PluginAuthor);
+			Class.Metadata.SetDescription(FInputNode::GetInputDescription());
+
+			FMetasoundFrontendClassStyleDisplay& StyleDisplay = Class.Style.Display;
+			StyleDisplay.bShowInputNames = false;
+			StyleDisplay.bShowOutputNames = true;
+			StyleDisplay.bShowLiterals = false;
+			StyleDisplay.bShowName = true;
+#endif // WITH_EDITOR
+
+			Class.Metadata.SetType(EMetasoundFrontendClassType::Template);
+			Class.Metadata.SetVersion(VersionNumber);
+
+			return Class;
+		};
+
+		static const FMetasoundFrontendClass FrontendClass = CreateFrontendClass();
+		return FrontendClass;
+	}
+
+	const FInputNodeTemplate& FInputNodeTemplate::GetChecked()
+	{
+		const INodeTemplate* Template = INodeTemplateRegistry::Get().FindTemplate(GetRegistryKey());
+		checkf(Template, TEXT("Failed to find InputNodeTemplate, which is required for migrating editor document data"));
+		return static_cast<const FInputNodeTemplate&>(*Template);
+	}
+
+	const FNodeRegistryKey& FInputNodeTemplate::GetRegistryKey()
+	{
+		static const FNodeRegistryKey RegistryKey(EMetasoundFrontendClassType::Template, ClassName, VersionNumber);
+		return RegistryKey;
+	}
+
+	EMetasoundFrontendVertexAccessType FInputNodeTemplate::GetNodeInputAccessType(const FMetaSoundFrontendDocumentBuilder& InBuilder, const FGuid& InNodeID, const FGuid& InVertexID) const
+	{
+		const FMetasoundFrontendNode* ConnectedInputNode = nullptr;
+		if (const FMetasoundFrontendVertex* ConnectedInputOutput = InBuilder.FindNodeOutputConnectedToNodeInput(InNodeID, InVertexID, &ConnectedInputNode))
+		{
+			const FMetasoundFrontendClass* InputClass = InBuilder.FindDependency(ConnectedInputNode->ClassID);
+			check(InputClass);
+			return (InputClass->Interface.Outputs.Last().AccessType);
+		}
+
+		return EMetasoundFrontendVertexAccessType::Unset;
+	}
+
+	EMetasoundFrontendVertexAccessType FInputNodeTemplate::GetNodeOutputAccessType(const FMetaSoundFrontendDocumentBuilder& InBuilder, const FGuid& InNodeID, const FGuid& InVertexID) const
+	{
+		if (const FMetasoundFrontendNode* Node = InBuilder.FindNode(InNodeID))
+		{
+			const FMetasoundFrontendVertex& Input = Node->Interface.Inputs.Last();
+			return GetNodeInputAccessType(InBuilder, InNodeID, Input.VertexID);
+		}
+
+		return EMetasoundFrontendVertexAccessType::Unset;
+	}
+
+#if WITH_EDITOR
+	FText FInputNodeTemplate::GetOutputPinDisplayName(const IOutputController& InOutput) const
+	{
+		FConstNodeHandle OwningNode = InOutput.GetOwningNode();
+		FConstOutputHandle ConnectedOutput = OwningNode->GetConstInputs().Last()->GetConnectedOutput();
+		if (ensureMsgf(ConnectedOutput->IsValid(), TEXT("Input template node should always be connected to associated input node's only output")))
+		{
+			FConstNodeHandle ConnectedInputNode = ConnectedOutput->GetOwningNode();
+
+			FName NodeName = ConnectedInputNode->GetNodeName();
+			FName Namespace;
+			FName ParameterName;
+			Audio::FParameterPath::SplitName(NodeName, Namespace, ParameterName);
+
+			FText DisplayName = ConnectedInputNode->GetDisplayName();
+			if (DisplayName.IsEmpty())
+			{
+				DisplayName = FText::FromName(NodeName);
+			}
+
+			if (!Namespace.IsNone())
+			{
+				return FText::Format(NSLOCTEXT("MetasoundFrontend", "InputNodeTemplate_DisplayNameWithNamespaceFormat", "{0} ({1})"), DisplayName, FText::FromName(Namespace));
+			}
+
+			return DisplayName;
+		}
+
+		return FRerouteNodeTemplate::GetOutputPinDisplayName(InOutput);
+	}
+
+	bool FInputNodeTemplate::HasRequiredConnections(const FMetaSoundFrontendDocumentBuilder& InBuilder, const FGuid& InNodeID, FString* OutMessage) const
+	{
+		return true;
+	}
+#endif // WITH_EDITOR
+
+	bool FInputNodeTemplate::IsInputAccessTypeDynamic() const
+	{
+		return true;
+	}
+
+	bool FInputNodeTemplate::IsInputConnectionUserModifiable() const
+	{
+		return false;
+	}
+
+	bool FInputNodeTemplate::IsOutputAccessTypeDynamic() const
+	{
+		return true;
+	}
+
+#if WITH_EDITOR
+	bool FInputNodeTemplate::Inject(FMetaSoundFrontendDocumentBuilder& InOutBuilder, bool bForceNodeCreation) const
+	{
+		bool bInjectedNodes = false;
+
+		const FMetaSoundFrontendDocumentBuilder& ConstBuilder = static_cast<const FMetaSoundFrontendDocumentBuilder&>(InOutBuilder);
+		const TArray<FMetasoundFrontendClassInput>& Inputs = ConstBuilder.GetDocument().RootGraph.Interface.Inputs;
+		for (const FMetasoundFrontendClassInput& Input : Inputs)
+		{
+			// Potentially not used input, which is perfectly valid so early out
+			const FMetasoundFrontendNode* InputNode = InOutBuilder.FindGraphInputNode(Input.Name);
+			if (!InputNode)
+			{
+				continue;
+			}
+
+			const FGuid& InputNodeOutputVertexID = InputNode->Interface.Outputs.Last().VertexID;
+
+			TArray<const FMetasoundFrontendNode*> ConnectedInputNodes;
+			TArray<const FMetasoundFrontendVertex*> ConnectedInputVertices = InOutBuilder.FindNodeInputsConnectedToNodeOutput(InputNode->GetID(), InputNodeOutputVertexID, &ConnectedInputNodes);
+
+			FMetasoundFrontendEdge EdgeToRemove { InputNode->GetID(), InputNodeOutputVertexID };
+
+			bool bHasTemplateConnection = false;
+
+			TArray<FMetasoundFrontendVertexHandle> ConnectedVertices;
+			for (int32 Index = 0; Index < ConnectedInputVertices.Num(); ++Index)
+			{
+				// Ignore edges already connected to input template nodes & cache connected vertex pair
+				// as adding a template node in the subsequent step may invalidate these connected node/vertex
+				// pointers.
+				const FMetasoundFrontendVertex* ConnectedVertex = ConnectedInputVertices[Index];
+				const FMetasoundFrontendNode* ConnectedNode = ConnectedInputNodes[Index];
+				const FMetasoundFrontendClass* Class = InOutBuilder.FindDependency(ConnectedNode->ClassID);
+				if (Class->Metadata.GetClassName() == FInputNodeTemplate::ClassName)
+				{
+					bHasTemplateConnection = true;
+				}
+				else
+				{
+					ConnectedVertices.Add(FMetasoundFrontendVertexHandle { ConnectedNode->GetID(), ConnectedVertex->VertexID });
+				}
+			}
+
+			if (ConnectedVertices.IsEmpty())
+			{
+				if (bForceNodeCreation && !bHasTemplateConnection)
+				{
+					bInjectedNodes = true;
+					InputNodeTemplatePrivate::InitTemplateNode(*this, Input.Name, InOutBuilder);
+				}
+			}
+			else
+			{
+				bInjectedNodes = true;
+				const FMetasoundFrontendNode* TemplateNode = InputNodeTemplatePrivate::InitTemplateNode(*this, Input.Name, InOutBuilder);
+				check(TemplateNode);
+
+				for (const FMetasoundFrontendVertexHandle& ConnectedVertex : ConnectedVertices)
+				{
+					// Swap connections from input node to connected node to now be from template node to connected node
+					EdgeToRemove.ToNodeID = ConnectedVertex.NodeID;
+					EdgeToRemove.ToVertexID = ConnectedVertex.VertexID;
+					InOutBuilder.RemoveEdge(EdgeToRemove);
+
+					InOutBuilder.AddEdge(FMetasoundFrontendEdge
+					{
+						TemplateNode->GetID(),
+						TemplateNode->Interface.Outputs.Last().VertexID,
+						ConnectedVertex.NodeID,
+						ConnectedVertex.VertexID
+					});
+				}
+			}
+		}
+
+		return bInjectedNodes;
+	}
+#endif // WITH_EDITOR
+} // namespace Metasound::Frontend
\ No newline at end of file
diff --git a/Engine/Plugins/Runtime/Metasound/Source/MetasoundFrontend/Private/NodeTemplates/MetasoundFrontendNodeTemplateReroute.cpp b/Engine/Plugins/Runtime/Metasound/Source/MetasoundFrontend/Private/NodeTemplates/MetasoundFrontendNodeTemplateReroute.cpp
--- a/Engine/Plugins/Runtime/Metasound/Source/MetasoundFrontend/Private/NodeTemplates/MetasoundFrontendNodeTemplateReroute.cpp
+++ b/Engine/Plugins/Runtime/Metasound/Source/MetasoundFrontend/Private/NodeTemplates/MetasoundFrontendNodeTemplateReroute.cpp
@@ -16,7 +16,7 @@ namespace Metasound::Frontend
 {
 	namespace ReroutePrivate
 	{
-		class FRerouteNodeTemplateTransform : public INodeTransform
+		class FRerouteNodeTemplateTransform : public INodeTemplateTransform
 		{
 		public:
 			FRerouteNodeTemplateTransform() = default;
@@ -109,6 +109,8 @@ namespace Metasound::Frontend
 							}
 						}
 					}
+
+					return nullptr;
 				}
 			}
 
@@ -202,12 +204,13 @@ namespace Metasound::Frontend
 		NewInterface.Outputs.Add(FMetasoundFrontendVertex { VertexName, DataType, FGuid::NewGuid() });
 
 		return NewInterface;
+
 	}
 
-	TUniquePtr<INodeTransform> FRerouteNodeTemplate::GenerateNodeTransform() const
+	TUniquePtr<INodeTemplateTransform> FRerouteNodeTemplate::GenerateNodeTransform() const
 	{
 		using namespace ReroutePrivate;
-		return TUniquePtr<INodeTransform>(new FRerouteNodeTemplateTransform());
+		return TUniquePtr<INodeTemplateTransform>(new FRerouteNodeTemplateTransform());
 	}
 
 	const FMetasoundFrontendClass& FRerouteNodeTemplate::GetFrontendClass() const
@@ -250,14 +253,14 @@ namespace Metasound::Frontend
 			const FMetasoundFrontendVertex& RerouteOutput = Node->Interface.Outputs.Last();
 
 			TArray<const FMetasoundFrontendNode*> ConnectedNodes;
-			TArray<const FMetasoundFrontendVertex*> ConnectedInputs = InBuilder.FindNodeInputsConnectedToNodeOutput(InNodeID, InVertexID, &ConnectedNodes);
+			TArray<const FMetasoundFrontendVertex*> ConnectedInputs = InBuilder.FindNodeInputsConnectedToNodeOutput(InNodeID, RerouteOutput.VertexID, &ConnectedNodes);
 			for (int32 Index = 0; Index < ConnectedNodes.Num(); ++Index)
 			{
 				const FMetasoundFrontendNode* ConnectedNode = ConnectedNodes[Index];
 				if (const FMetasoundFrontendClass* ConnectedNodeClass = InBuilder.FindDependency(ConnectedNode->ClassID))
 				{
 					const FMetasoundFrontendVertex* ConnectedInput = ConnectedInputs[Index];
-					if (ConnectedNodeClass->Metadata.GetClassName() == ClassName)
+					if (ConnectedNodeClass->Metadata.GetClassName() == GetClassName())
 					{
 						return this->GetNodeInputAccessType(InBuilder, ConnectedNode->GetID(), ConnectedInput->VertexID);
 					}
diff --git a/Engine/Plugins/Runtime/Metasound/Source/MetasoundFrontend/Public/MetasoundAssetBase.h b/Engine/Plugins/Runtime/Metasound/Source/MetasoundFrontend/Public/MetasoundAssetBase.h
--- a/Engine/Plugins/Runtime/Metasound/Source/MetasoundFrontend/Public/MetasoundAssetBase.h
+++ b/Engine/Plugins/Runtime/Metasound/Source/MetasoundFrontend/Public/MetasoundAssetBase.h
@@ -67,6 +67,7 @@ public:
 	virtual const UEdGraph* GetGraph() const = 0;
 	virtual UEdGraph& GetGraphChecked() = 0;
 	virtual const UEdGraph& GetGraphChecked() const = 0;
+	virtual void MigrateEditorGraph(FMetaSoundFrontendDocumentBuilder& OutBuilder) = 0;
 
 	// Sets the graph associated with this Metasound. Graph is required to be referenced on
 	// Metasound UObject for editor serialization purposes.
@@ -78,10 +79,8 @@ public:
 	virtual void SetRegistryAssetClassInfo(const Metasound::Frontend::FNodeClassInfo& InClassInfo) = 0;
 #endif // WITH_EDITORONLY_DATA
 
-	// Called when the interface is changed, presenting the opportunity for
-	// any reflected object data to be updated based on the new interface.
-	// Returns whether or not any edits were made.
-	virtual bool ConformObjectDataToInterfaces() = 0;
+	UE_DEPRECATED(5.5, "Moved to IMetaSoundDocumentInterface::ConformObjectToDocument")
+	virtual bool ConformObjectDataToInterfaces();
 
 	// Registers the root graph of the given asset with the MetaSound Frontend.
 	virtual void RegisterGraphWithFrontend(Metasound::Frontend::FMetaSoundAssetRegistrationOptions InRegistrationOptions = Metasound::Frontend::FMetaSoundAssetRegistrationOptions());
@@ -92,19 +91,11 @@ public:
 	// Cooks this MetaSound and recursively checks and cooks referenced MetaSounds if necessary. Cook includes autoupdating and resolving the document, which is then registered with the MetaSound Frontend.
 	void CookMetaSound();
 
-	// Sets/overwrites the root class metadata
-	UE_DEPRECATED(5.3, "Directly setting graph class Metadata is no longer be supported. Use the FMetaSoundFrontendDocumentBuilder to modify class data.")
-	virtual void SetMetadata(FMetasoundFrontendClassMetadata& InMetadata);
-
 #if WITH_EDITOR
 	// Rebuild dependent asset classes
 	void RebuildReferencedAssetClasses();
 #endif // WITH_EDITOR
 
-	// Returns the interface entries declared by the given asset's document from the InterfaceRegistry.
-	UE_DEPRECATED(5.3, "Use static FMetaSoundFrontendDocumentBuilder 'FindDeclaredInterfaces instead.")
-	bool GetDeclaredInterfaces(TArray<const Metasound::Frontend::IInterfaceRegistryEntry*>& OutInterfaces) const;
-
 	// Returns whether an interface with the given version is declared by the given asset's document.
 	bool IsInterfaceDeclared(const FMetasoundFrontendVersion& InVersion) const;
 
@@ -129,9 +120,6 @@ public:
 	bool AddingReferenceCausesLoop(const FSoftObjectPath& InReferencePath) const;
 	bool IsReferencedAsset(const FMetasoundAssetBase& InAssetToCheck) const;
 
-	UE_DEPRECATED(5.3, "ConvertFromPreset moved to FMetaSoundFrontendDocumentBuilder.")
-	void ConvertFromPreset();
-
 	bool IsRegistered() const;
 
 	// Imports data from a JSON string directly
@@ -148,12 +136,16 @@ public:
 	Metasound::Frontend::FGraphHandle GetRootGraphHandle();
 	Metasound::Frontend::FConstGraphHandle GetRootGraphHandle() const;
 
-	// Overwrites the existing document. If the document's interface is not supported,
-	// the FMetasoundAssetBase be while queried for a new one using `GetPreferredInterface`.
-	void SetDocument(const FMetasoundFrontendDocument& InDocument, bool bMarkDirty = true);
-	void SetDocument(FMetasoundFrontendDocument&& InDocument, bool bMarkDirty = true);
+	UE_DEPRECATED(5.5, "Direct mutation of the document is no longer supported via AssetBase.")
+	void SetDocument(FMetasoundFrontendDocument InDocument, bool bMarkDirty = true);
+
+	const FMetasoundFrontendDocument& GetConstDocumentChecked() const;
 
+	// Soft deprecated.  Document layer should not be directly mutated via asset base in anticipation
+	// of moving to all mutation via the Frontend Document & Builder Subsystem API.
 	FMetasoundFrontendDocument& GetDocumentChecked();
+
+	UE_DEPRECATED(5.5, "Use GetConstDocumentChecked instead.")
 	const FMetasoundFrontendDocument& GetDocumentChecked() const;
 
 	const Metasound::Frontend::FGraphRegistryKey& GetGraphRegistryKey() const;
@@ -161,9 +153,6 @@ public:
 	UE_DEPRECATED(5.4, "Use GetGraphRegistryKey instead.")
 	const Metasound::Frontend::FNodeRegistryKey& GetRegistryKey() const;
 
-	UE_DEPRECATED(5.3, "AddDefaultInterfaces is included in now applied via FMetaSoundFrontendDocumentBuilder::InitDocument and no longer directly supported via this function.")
-	void AddDefaultInterfaces();
-
 	bool VersionAsset();
 
 #if WITH_EDITOR
@@ -176,7 +165,11 @@ public:
 	void CacheRegistryMetadata();
 
 	FMetasoundFrontendDocumentModifyContext& GetModifyContext();
+
+	UE_DEPRECATED(5.5, "Use GetConstModifyContext")
 	const FMetasoundFrontendDocumentModifyContext& GetModifyContext() const;
+
+	const FMetasoundFrontendDocumentModifyContext& GetConstModifyContext() const;
 #endif // WITH_EDITOR
 
 	// Calls the outermost package and marks it dirty.
@@ -196,6 +189,18 @@ public:
 
 	FString GetOwningAssetName() const;
 
+#if WITH_EDITORONLY_DATA
+	void ClearVersionedOnLoad();
+	bool GetVersionedOnLoad() const;
+	void SetVersionedOnLoad();
+#endif // WITH_EDITORONLY_DATA
+
+	// Returns true if the IMetaSoundDocumentInterface is currently has an active builder.
+	//
+	// If true, calls to register the graph will be performed synchronously in order to avoid
+	// race conditions with an active builder.
+	virtual bool IsBuilderActive() const = 0;
+
 protected:
 	void OnNotifyBeginDestroy();
 
@@ -222,8 +227,6 @@ protected:
 	// When a graph is registered, the underlying IMetaSoundDocumentInterface may be accessed on an
 	// async tasks. If modifications need to be made to the IMetaSoundDocumentInterface, callers should
 	// wait for the inflight graph registration to complete by calling this method.
-
-protected:
 	
 	// Container for runtime data of MetaSound graph.
 	struct FRuntimeData_DEPRECATED
@@ -251,10 +254,6 @@ protected:
 	const FRuntimeData& GetRuntimeData() const;
 	PRAGMA_ENABLE_DEPRECATION_WARNINGS
 
-	// Returns all public class inputs.  This is a potentially expensive.
-	// Prefer accessing public class inputs using CacheRuntimeData.
-	TArray<FMetasoundFrontendClassInput> GetPublicClassInputs() const;
-
 	bool AutoUpdate(bool bInLogWarningsOnDroppedConnection);
 	void CookReferencedMetaSounds();
 
@@ -267,13 +266,16 @@ protected:
 private:
 #if WITH_EDITORONLY_DATA
 	void UpdateAssetRegistry();
+
+	bool bVersionedOnLoad = false;
 #endif
 
-	// Returns true if the IMetaSoundDocumentInterface is currently has an active builder.
-	//
-	// If true, calls to register the graph will be performed synchronously in order to avoid
-	// race conditions with an active builder.
-	virtual bool IsBuilderActive() const = 0;
+	// Checks if version is up-to-date. If so, returns true. If false, updates the interfaces within the given asset's document to the most recent version.
+	bool TryUpdateInterfaceFromVersion(const FMetasoundFrontendVersion& Version);
+
+	// Returns new interface to be versioned to from the given version. If no interface versioning is
+	// required, returns invalid interface (interface with no name and invalid version number).
+	FMetasoundFrontendInterface GetInterfaceToVersion(const FMetasoundFrontendVersion& InterfaceVersion, Metasound::Frontend::FConstDocumentHandle InDocument) const;
 
 	Metasound::Frontend::FGraphRegistryKey GraphRegistryKey;
 };
diff --git a/Engine/Plugins/Runtime/Metasound/Source/MetasoundFrontend/Public/MetasoundAssetManager.h b/Engine/Plugins/Runtime/Metasound/Source/MetasoundFrontend/Public/MetasoundAssetManager.h
--- a/Engine/Plugins/Runtime/Metasound/Source/MetasoundFrontend/Public/MetasoundAssetManager.h
+++ b/Engine/Plugins/Runtime/Metasound/Source/MetasoundFrontend/Public/MetasoundAssetManager.h
@@ -12,6 +12,7 @@ class FMetasoundAssetBase;
 class UEdGraph;
 
 struct FMetasoundFrontendClassName;
+struct FMetaSoundFrontendDocumentBuilder;
 
 namespace Metasound
 {
@@ -131,6 +132,8 @@ namespace Metasound
 			// Add or Update a MetaSound Asset's entry data
 			virtual Metasound::Frontend::FNodeRegistryKey AddOrUpdateAsset(const UObject& InObject) = 0;
 
+			virtual FMetaSoundFrontendDocumentBuilder& AttachDocumentBuilderChecked(UObject& InObject) const = 0;
+
 			// Whether or not the class is eligible for auto-update
 			virtual bool CanAutoUpdate(const FMetasoundFrontendClassName& InClassName) const = 0;
 
diff --git a/Engine/Plugins/Runtime/Metasound/Source/MetasoundFrontend/Public/MetasoundDocumentInterface.h b/Engine/Plugins/Runtime/Metasound/Source/MetasoundFrontend/Public/MetasoundDocumentInterface.h
--- a/Engine/Plugins/Runtime/Metasound/Source/MetasoundFrontend/Public/MetasoundDocumentInterface.h
+++ b/Engine/Plugins/Runtime/Metasound/Source/MetasoundFrontend/Public/MetasoundDocumentInterface.h
@@ -45,6 +45,10 @@ public:
 	// Returns the parent class registered with the MetaSound UObject registry.
 	virtual const UClass& GetBaseMetaSoundUClass() const = 0;
 
+	// Conforms UProperty data outside the Frontend Document Model to the document's data.
+	// Returns whether or not object data was modified.
+	virtual bool ConformObjectToDocument() = 0;
+
 private:
 	virtual FMetasoundFrontendDocument& GetDocument() = 0;
 
diff --git a/Engine/Plugins/Runtime/Metasound/Source/MetasoundFrontend/Public/MetasoundFrontendController.h b/Engine/Plugins/Runtime/Metasound/Source/MetasoundFrontend/Public/MetasoundFrontendController.h
--- a/Engine/Plugins/Runtime/Metasound/Source/MetasoundFrontend/Public/MetasoundFrontendController.h
+++ b/Engine/Plugins/Runtime/Metasound/Source/MetasoundFrontend/Public/MetasoundFrontendController.h
@@ -1134,11 +1134,13 @@ namespace Metasound
 			virtual FString ExportToJSON() const = 0;
 		};
 
+		UE_DEPRECATED(5.5, "Reroute output recursion is now privately implemented in the reroute template node class. Use template registry to query reroute node input/output attributes")
 		METASOUNDFRONTEND_API FConstOutputHandle FindReroutedOutput(FConstOutputHandle InOutputHandle);
 
-		UE_DEPRECATED(5.5, "Reroute input recursion is now privately implemented in the reroute template node class.")
+		UE_DEPRECATED(5.5, "Reroute input recursion is now privately implemented in the reroute template node class. Use template registry to query reroute node input/output attributes")
 		METASOUNDFRONTEND_API void FindReroutedInputs(FConstInputHandle InHandleToCheck, TArray<FConstInputHandle>& InOutInputHandles);
 
+		UE_DEPRECATED(5.5, "Reroute input recursion is now privately implemented in the reroute template node class. Use template registry to query reroute node input/output attributes")
 		METASOUNDFRONTEND_API void IterateReroutedInputs(FConstInputHandle InHandleToCheck, TFunctionRef<void(FConstInputHandle)> Func);
 	} // namespace Frontend
 } // namespace Metasound
diff --git a/Engine/Plugins/Runtime/Metasound/Source/MetasoundFrontend/Public/MetasoundFrontendDocument.h b/Engine/Plugins/Runtime/Metasound/Source/MetasoundFrontend/Public/MetasoundFrontendDocument.h
--- a/Engine/Plugins/Runtime/Metasound/Source/MetasoundFrontend/Public/MetasoundFrontendDocument.h
+++ b/Engine/Plugins/Runtime/Metasound/Source/MetasoundFrontend/Public/MetasoundFrontendDocument.h
@@ -134,10 +134,10 @@ enum class EMetasoundFrontendClassType : uint8
 	// The MetaSound class is an output from a graph in the containing document.
 	Output,
 
-	// The MetaSound class is an literal requiring an literal value to construct.
+	// The MetaSound class is an literal requiring a literal value to construct.
 	Literal,
 
-	// The MetaSound class is an variable requiring an literal value to construct.
+	// The MetaSound class is an variable requiring a literal value to construct.
 	Variable,
 
 	// The MetaSound class accesses variables.
@@ -1652,6 +1652,24 @@ struct METASOUNDFRONTEND_API FMetasoundFrontendGraphClass : public FMetasoundFro
 	FMetasoundFrontendGraphClassPresetOptions PresetOptions;
 };
 
+UCLASS()
+class METASOUNDFRONTEND_API UMetaSoundFrontendMemberMetadata : public UObject
+{
+	GENERATED_BODY()
+
+#if WITH_EDITOR
+public:
+	virtual void ForceRefresh() { }
+	virtual FMetasoundFrontendLiteral GetDefault() const { return FMetasoundFrontendLiteral(); }
+	virtual EMetasoundFrontendLiteralType GetLiteralType() const { return EMetasoundFrontendLiteralType::None; }
+	virtual void SetFromLiteral(const FMetasoundFrontendLiteral& InLiteral) { }
+#endif // WITH_EDITOR
+
+#if WITH_EDITORONLY_DATA
+	UPROPERTY()
+	FGuid MemberID;
+#endif // WITH_EDITORONLY_DATA
+};
 
 USTRUCT()
 struct METASOUNDFRONTEND_API FMetasoundFrontendDocumentMetadata
@@ -1663,6 +1681,12 @@ struct METASOUNDFRONTEND_API FMetasoundFrontendDocumentMetadata
 
 #if WITH_EDITORONLY_DATA
 	FMetasoundFrontendDocumentModifyContext ModifyContext;
+
+	// Map of MemberID to metadata used to constrain how literals can be manipulated
+	// with the editor context. This can be used to implement things like numeric ranges,
+	// hardware control parameters, etc.
+	UPROPERTY()
+	TMap<FGuid, TObjectPtr<UMetaSoundFrontendMemberMetadata>> MemberMetadata;
 #endif // WITH_EDITORONLY_DATA
 };
 
@@ -1699,36 +1723,41 @@ public:
 	}
 
 private:
+#if WITH_EDITORONLY_DATA
 	UPROPERTY(meta = (DeprecatedProperty, DeprecationMessage = "5.0 - ArchetypeVersion has been migrated to InterfaceVersions array."))
 	FMetasoundFrontendVersion ArchetypeVersion;
 
 	UPROPERTY(meta = (DeprecatedProperty, DeprecationMessage = "5.0 - InterfaceVersions has been migrated to Interfaces set."))
 	TArray<FMetasoundFrontendVersion> InterfaceVersions;
+#endif // WITH_EDITORONLY_DATA
 
 	// Used for generating deterministic IDs per document
 	mutable uint32 IdCounter = 1;
 
 public:
+#if WITH_EDITORONLY_DATA
+	bool RequiresInterfaceVersioning() const
+	{
+		return ArchetypeVersion.IsValid() || !InterfaceVersions.IsEmpty();
+	}
+
 	// Data migration for 5.0 Early Access data. ArchetypeVersion/InterfaceVersions properties can be removed post 5.0 release
 	// and this fix-up can be removed post 5.0 release.
-	bool VersionInterfaces()
+	void VersionInterfaces()
 	{
-		bool bDidEdit = false;
 		if (ArchetypeVersion.IsValid())
 		{
 			Interfaces.Add(ArchetypeVersion);
 			ArchetypeVersion = FMetasoundFrontendVersion::GetInvalid();
-			bDidEdit = true;
 		}
+
 		if (!InterfaceVersions.IsEmpty())
 		{
 			Interfaces.Append(InterfaceVersions);
 			InterfaceVersions.Reset();
-			bDidEdit = true;
 		}
-
-		return bDidEdit;
 	}
+#endif // WITH_EDITORONLY_DATA
 };
 
 METASOUNDFRONTEND_API const TCHAR* LexToString(EMetasoundFrontendClassType InClassType);
diff --git a/Engine/Plugins/Runtime/Metasound/Source/MetasoundFrontend/Public/MetasoundFrontendDocumentBuilder.h b/Engine/Plugins/Runtime/Metasound/Source/MetasoundFrontend/Public/MetasoundFrontendDocumentBuilder.h
--- a/Engine/Plugins/Runtime/Metasound/Source/MetasoundFrontend/Public/MetasoundFrontendDocumentBuilder.h
+++ b/Engine/Plugins/Runtime/Metasound/Source/MetasoundFrontend/Public/MetasoundFrontendDocumentBuilder.h
@@ -101,6 +101,8 @@ public:
 	// Create and return a valid builder document which copies the provided interface's document & class
 	static UMetaSoundBuilderDocument& Create(const IMetaSoundDocumentInterface& InDocToCopy);
 
+	virtual bool ConformObjectToDocument() override;
+
 	// Returns the document
 	virtual const FMetasoundFrontendDocument& GetConstDocument() const override;
 
@@ -133,11 +135,8 @@ struct METASOUNDFRONTEND_API FMetaSoundFrontendDocumentBuilder
 	GENERATED_BODY()
 
 public:
-	// Exists only to make UObject reflection happy.  Should never be
-	// used directly as builder interface (and optionally delegates) should be specified on construction.
-	FMetaSoundFrontendDocumentBuilder();
-	FMetaSoundFrontendDocumentBuilder(TScriptInterface<IMetaSoundDocumentInterface> InDocumentInterface);
-	FMetaSoundFrontendDocumentBuilder(TScriptInterface<IMetaSoundDocumentInterface> InDocumentInterface, TSharedRef<Metasound::Frontend::FDocumentModifyDelegates> InDocumentDelegates);
+	// Default ctor should typically never be used directly as builder interface (and optionally delegates) should be specified on construction (Default exists only to make UObject reflection happy).
+	FMetaSoundFrontendDocumentBuilder(TScriptInterface<IMetaSoundDocumentInterface> InDocumentInterface = { }, TSharedPtr<Metasound::Frontend::FDocumentModifyDelegates> InDocumentDelegates = { }, bool bPrimeCache = false);
 	virtual ~FMetaSoundFrontendDocumentBuilder();
 
 	// Call when the builder will no longer modify the IMetaSoundDocumentInterface
@@ -150,30 +149,39 @@ public:
 	bool AddEdgesFromMatchingInterfaceNodeOutputsToGraphOutputs(const FGuid& InNodeID, TArray<const FMetasoundFrontendEdge*>& OutEdgesCreated, bool bReplaceExistingConnections = true);
 	bool AddEdgesFromMatchingInterfaceNodeInputsToGraphInputs(const FGuid& InNodeID, TArray<const FMetasoundFrontendEdge*>& OutEdgesCreated, bool bReplaceExistingConnections = true);
 	const FMetasoundFrontendNode* AddGraphInput(const FMetasoundFrontendClassInput& InClassInput);
+	const FMetasoundFrontendNode* AddGraphNode(const FMetasoundFrontendGraphClass& InClass, FGuid InNodeID = FGuid::NewGuid());
 	const FMetasoundFrontendNode* AddGraphOutput(const FMetasoundFrontendClassOutput& InClassOutput);
+
 	bool AddInterface(FName InterfaceName);
 
-	const FMetasoundFrontendNode* AddGraphNode(const FMetasoundFrontendGraphClass& InClass, FGuid InNodeID = FGuid::NewGuid());
 	const FMetasoundFrontendNode* AddNodeByClassName(const FMetasoundFrontendClassName& InClassName, int32 InMajorVersion = 1, FGuid InNodeID = FGuid::NewGuid());
+	const FMetasoundFrontendNode* AddNodeByTemplate(const Metasound::Frontend::INodeTemplate& InTemplate, Metasound::Frontend::FNodeTemplateGenerateInterfaceParams Params, FGuid InNodeID = FGuid::NewGuid());
 
 	// Returns whether or not the given edge can be added, which requires that its input
 	// is not already connected and the edge is valid (see function 'IsValidEdge').
 	bool CanAddEdge(const FMetasoundFrontendEdge& InEdge) const;
 
-	void ClearGraph();
+	// Clears graph completely of all nodes, edges, dependencies, & member metadata.
+	// Reloads the builder state, so external delegates must be relinked if desired.
+	void ClearGraph(TSharedRef<Metasound::Frontend::FDocumentModifyDelegates> ModifyDelegates = { });
+
+#if WITH_EDITORONLY_DATA
+	bool ClearMemberMetadata(const FGuid& InMemberID);
+#endif // WITH_EDITORONLY_DATA
 
 	bool ContainsDependencyOfType(EMetasoundFrontendClassType ClassType) const;
 	bool ContainsEdge(const FMetasoundFrontendEdge& InEdge) const;
 	bool ContainsNode(const FGuid& InNodeID) const;
 
 	bool ConvertFromPreset();
-	bool ConvertToPreset(const FMetasoundFrontendDocument& InReferencedDocument);
+	bool ConvertToPreset(const FMetasoundFrontendDocument& InReferencedDocument, TSharedRef<Metasound::Frontend::FDocumentModifyDelegates> ModifyDelegates = { });
 
-#if WITH_EDITOR
+#if WITH_EDITORONLY_DATA
 	const FMetaSoundFrontendGraphComment* FindGraphComment(const FGuid& InCommentID) const;
 	FMetaSoundFrontendGraphComment* FindGraphComment(const FGuid& InCommentID);
 	FMetaSoundFrontendGraphComment& FindOrAddGraphComment(const FGuid& InCommentID);
-#endif // WITH_EDITOR
+	UMetaSoundFrontendMemberMetadata* FindMemberMetadata(const FGuid& InMemberID);
+#endif // WITH_EDITORONLY_DATA
 
 	static bool FindDeclaredInterfaces(const FMetasoundFrontendDocument& InDocument, TArray<const Metasound::Frontend::IInterfaceRegistryEntry*>& OutInterfaces);
 	bool FindDeclaredInterfaces(TArray<const Metasound::Frontend::IInterfaceRegistryEntry*>& OutInterfaces) const;
@@ -182,14 +190,14 @@ public:
 	const FMetasoundFrontendClass* FindDependency(const FMetasoundFrontendClassMetadata& InMetadata) const;
 	TArray<const FMetasoundFrontendEdge*> FindEdges(const FGuid& InNodeID, const FGuid& InVertexID) const;
 
-	bool FindInterfaceInputNodes(FName InterfaceName, TArray<const FMetasoundFrontendNode*>& OutInputs) const;
-	bool FindInterfaceOutputNodes(FName InterfaceName, TArray<const FMetasoundFrontendNode*>& OutOutputs) const;
-
 	const FMetasoundFrontendClassInput* FindGraphInput(FName InputName) const;
 	const FMetasoundFrontendNode* FindGraphInputNode(FName InputName) const;
 	const FMetasoundFrontendClassOutput* FindGraphOutput(FName OutputName) const;
 	const FMetasoundFrontendNode* FindGraphOutputNode(FName OutputName) const;
 
+	bool FindInterfaceInputNodes(FName InterfaceName, TArray<const FMetasoundFrontendNode*>& OutInputs) const;
+	bool FindInterfaceOutputNodes(FName InterfaceName, TArray<const FMetasoundFrontendNode*>& OutOutputs) const;
+
 	const FMetasoundFrontendNode* FindNode(const FGuid& InNodeID) const;
 
 	const FMetasoundFrontendVertex* FindNodeInput(const FGuid& InNodeID, const FGuid& InVertexID) const;
@@ -221,21 +229,28 @@ public:
 
 	const Metasound::Frontend::FDocumentModifyDelegates& GetDocumentDelegates() const;
 	const IMetaSoundDocumentInterface& GetDocumentInterface() const;
+	FMetasoundAssetBase& GetMetasoundAsset();
+	const FMetasoundAssetBase& GetMetasoundAsset() const;
 
 	int32 GetTransactionCount() const;
 
+	TArray<const FMetasoundFrontendNode*> GetGraphInputTemplateNodes(FName InInputName);
 	EMetasoundFrontendVertexAccessType GetNodeInputAccessType(const FGuid& InNodeID, const FGuid& InVertexID) const;
 	const FMetasoundFrontendLiteral* GetNodeInputClassDefault(const FGuid& InNodeID, const FGuid& InVertexID) const;
 	const FMetasoundFrontendLiteral* GetNodeInputDefault(const FGuid& InNodeID, const FGuid& InVertexID) const;
 	EMetasoundFrontendVertexAccessType GetNodeOutputAccessType(const FGuid& InNodeID, const FGuid& InVertexID) const;
 
+	// Initializes the builder's document, using the (optional) provided document template, (optional) class name, and (optionally) whether or not to reset the existing class version.
+	void InitDocument(const FMetasoundFrontendDocument* InDocumentTemplate = nullptr, const FMetasoundFrontendClassName* InNewClassName = nullptr, bool bResetVersion = true);
+
 	// Initializes GraphClass Metadata, optionally resetting the version back to 1.0 and/or creating a unique class name if a name is not provided.
 	static void InitGraphClassMetadata(FMetasoundFrontendClassMetadata& InOutMetadata, bool bResetVersion = false, const FMetasoundFrontendClassName* NewClassName = nullptr);
-	void InitDocument();
+	void InitGraphClassMetadata(bool bResetVersion, const FMetasoundFrontendClassName* NewClassName);
+
 	void InitNodeLocations();
 
-	// Clears the cached data enabling fast document queries. Generally discouraged unless document has been unavoidably mutated outside of the builder.
-	void InvalidateCache();
+	UE_DEPRECATED(5.5, "Use invalidate overload that is provided new version of modify delegates")
+	void InvalidateCache() { }
 
 	bool IsDependencyReferenced(const FGuid& InClassID) const;
 	bool IsNodeInputConnected(const FGuid& InNodeID, const FGuid& InVertexID) const;
@@ -245,14 +260,19 @@ public:
 	bool IsInterfaceDeclared(const FMetasoundFrontendVersion& InInterfaceVersion) const;
 	bool IsPreset() const;
 
+	// Returns whether or not builder is attached to a DocumentInterface and is valid to build or act on a document.
+	bool IsValid() const;
+	bool IsVariableClass(EMetasoundFrontendClassType ClassType) const;
+
 	// Returns whether or not the given edge is valid (i.e. represents an input and output that equate in data and access types) or malformed.
 	// Note that this does not return whether or not the given edge exists, but rather if it could be legally applied to the given edge vertices.
 	Metasound::Frontend::EInvalidEdgeReason IsValidEdge(const FMetasoundFrontendEdge& InEdge) const;
 
 	bool ModifyInterfaces(Metasound::Frontend::FModifyInterfaceOptions&& InOptions);
 
-	// Reloads the cached data enabling fast document queries. Can be expensive and is generally discouraged in favor of using InvalidateCache
-	// only when document has been known to be mutated outside of the builder (also, invalidate doesn't prime the cache and only reloads piecemeal upon request).
+	UE_DEPRECATED(5.5,
+		"Cache invalidation may require new copy of delegates. In addition, re-priming is discouraged. "
+		"To enforce this, new recommended pattern is to construct a new builder instead")
 	void ReloadCache();
 
 	bool RemoveDependency(const FGuid& InClassID);
@@ -276,6 +296,8 @@ public:
 	int32 RemoveNodeLocation(const FGuid& InNodeID, const FGuid* InLocationGuid = nullptr);
 #endif // WITH_EDITOR
 
+	void Reload(TSharedPtr<Metasound::Frontend::FDocumentModifyDelegates> Delegates = {}, bool bPrimeCache = false);
+
 	bool RemoveNodeInputDefault(const FGuid& InNodeID, const FGuid& InVertexID);
 	bool RemoveUnusedDependencies();
 	bool RenameRootGraphClass(const FMetasoundFrontendClassName& InName);
@@ -283,6 +305,18 @@ public:
 #if WITH_EDITOR
 	void SetAuthor(const FString& InAuthor);
 
+#endif // WITH_EDITOR
+
+	bool SetGraphInputAccessType(FName InputName, EMetasoundFrontendVertexAccessType AccessType);
+	bool SetGraphInputDataType(FName InputName, FName DataType);
+	bool SetGraphInputDefault(FName InputName, const FMetasoundFrontendLiteral& InDefaultLiteral);
+
+	bool SetGraphOutputAccessType(FName OutputName, EMetasoundFrontendVertexAccessType AccessType);
+	bool SetGraphOutputDataType(FName OutputName, FName DataType);
+
+#if WITH_EDITOR
+	void SetMemberMetadata(UMetaSoundFrontendMemberMetadata& NewMetadata);
+
 	// Sets the editor-only comment to the provided value.
 	// Returns true if the node was found and the comment was updated, false if not.
 	bool SetNodeComment(const FGuid& InNodeID, FString&& InNewComment);
@@ -296,7 +330,6 @@ public:
 	bool SetNodeLocation(const FGuid& InNodeID, const FVector2D& InLocation, const FGuid* InLocationGuid = nullptr);
 #endif // WITH_EDITOR
 
-	bool SetGraphInputDefault(FName InputName, const FMetasoundFrontendLiteral& InDefaultLiteral);
 	bool SetNodeInputDefault(const FGuid& InNodeID, const FGuid& InVertexID, const FMetasoundFrontendLiteral& InLiteral);
 
 	// Sets the document's version number.  Should only be called by document versioning.
@@ -313,24 +346,36 @@ public:
 private:
 	using FFinalizeNodeFunctionRef = TFunctionRef<void(FMetasoundFrontendNode&, const Metasound::Frontend::FNodeRegistryKey&)>;
 
-	FMetasoundFrontendNode* AddNodeInternal(const FMetasoundFrontendClassMetadata& InClassMetadata, Metasound::Frontend::FFinalizeNodeFunctionRef FinalizeNode, FGuid InNodeID = FGuid::NewGuid());
-	FMetasoundFrontendNode* FindNodeInternal(const FGuid& InNodeID);
+	FMetasoundFrontendNode* AddNodeInternal(const FMetasoundFrontendClassMetadata& InClassMetadata, Metasound::Frontend::FFinalizeNodeFunctionRef FinalizeNode, FGuid InNodeID = FGuid::NewGuid(), int32* NewNodeIndex = nullptr);
 
-	const FTopLevelAssetPath GetBuilderClassPath() const;
-	FMetasoundFrontendDocument& GetDocument();
+	// Conforms GraphOutput node's ClassID, Access & Data Type with the GraphOutput.
+	// creating and removing dependencies as necessary within the document dependency array. Does *NOT*
+	// modify edge data (i.e. if the DataType is changed on the given node and it has corresponding
+	// edges, edges may then be invalid due to access type/DataType incompatibility).
+	bool ConformGraphInputNodeToClass(const FMetasoundFrontendClassInput& GraphInput);
+
+	// Conforms GraphOutput node's ClassID, Access & Data Type with the GraphOutput.
+	// Creates and removes dependencies as necessary within the document dependency array. Does *NOT*
+	// modify edge data (i.e. if the DataType is changed on the given node and it has corresponding
+	// edges, edges may then be invalid due to access type/DataType incompatibility).
+	bool ConformGraphOutputNodeToClass(const FMetasoundFrontendClassOutput& GraphOutput);
 
 	bool FindNodeClassInterfaces(const FGuid& InNodeID, TSet<FMetasoundFrontendVersion>& OutInterfaces) const;
+	FMetasoundFrontendNode* FindNodeInternal(const FGuid& InNodeID);
+
+	void IterateNodesConnectedWithVertex(const FMetasoundFrontendVertexHandle& Vertex, TFunctionRef<void(const FMetasoundFrontendEdge&, FMetasoundFrontendNode&)> NodeIndexIterFunc);
 
-	void InitCacheInternal(bool bPrimeCache = false);
+	const FTopLevelAssetPath GetBuilderClassPath() const;
+	FMetasoundFrontendDocument& GetDocument();
 
 	bool SetGraphInputInheritsDefault(FName InName, bool bInputInheritsDefault);
+	bool UnlinkVariableNode(const FGuid& InNodeID);
 
-	// Called to mark that the builder is beginning modifications to the IMetaSoundDocumentInterface
-	void BeginBuilding();
+	void BeginBuilding(TSharedPtr<Metasound::Frontend::FDocumentModifyDelegates> Delegates = {}, bool bPrimeCache = false);
 
 	UPROPERTY(Transient)
 	TScriptInterface<IMetaSoundDocumentInterface> DocumentInterface;
 
 	TSharedPtr<Metasound::Frontend::IDocumentCache> DocumentCache;
-	TSharedRef<Metasound::Frontend::FDocumentModifyDelegates> DocumentDelegates;
+	TSharedPtr<Metasound::Frontend::FDocumentModifyDelegates> DocumentDelegates;
 };
diff --git a/Engine/Plugins/Runtime/Metasound/Source/MetasoundFrontend/Public/MetasoundFrontendDocumentCacheInterface.h b/Engine/Plugins/Runtime/Metasound/Source/MetasoundFrontend/Public/MetasoundFrontendDocumentCacheInterface.h
--- a/Engine/Plugins/Runtime/Metasound/Source/MetasoundFrontend/Public/MetasoundFrontendDocumentCacheInterface.h
+++ b/Engine/Plugins/Runtime/Metasound/Source/MetasoundFrontend/Public/MetasoundFrontendDocumentCacheInterface.h
@@ -88,7 +88,9 @@ namespace Metasound::Frontend
 		virtual ~IDocumentGraphInterfaceCache() = default;
 
 		virtual const FMetasoundFrontendClassInput* FindInput(FName InputName) const = 0;
+		virtual const int32* FindInputIndex(FName InputName) const = 0;
 		virtual const FMetasoundFrontendClassOutput* FindOutput(FName OutputName) const = 0;
+		virtual const int32* FindOutputIndex(FName OutputName) const = 0;
 	};
 
 	/** Interface for querying cached document dependencies. */
diff --git a/Engine/Plugins/Runtime/Metasound/Source/MetasoundFrontend/Public/MetasoundFrontendNodeTemplateRegistry.h b/Engine/Plugins/Runtime/Metasound/Source/MetasoundFrontend/Public/MetasoundFrontendNodeTemplateRegistry.h
--- a/Engine/Plugins/Runtime/Metasound/Source/MetasoundFrontend/Public/MetasoundFrontendNodeTemplateRegistry.h
+++ b/Engine/Plugins/Runtime/Metasound/Source/MetasoundFrontend/Public/MetasoundFrontendNodeTemplateRegistry.h
@@ -29,12 +29,22 @@ namespace Metasound::Frontend
 
 	struct METASOUNDFRONTEND_API FNodeTemplateGenerateInterfaceParams
 	{
-		// Currently just DataType FName. Subsequent change will replace this with a VertexHandle
-		// and will add builder reference to once builder supports template nodes.
+		// TODO: Currently just DataType FName. Subsequent change will replace this with a VertexHandle
+		// and will add builder reference to once builder supports template nodes and controllers
+		// are no longer used to add template nodes from editor code.
 		TArray<FName> InputsToConnect;
 		TArray<FName> OutputsToConnect;
 	};
 
+	class METASOUNDFRONTEND_API INodeTemplateTransform
+	{
+	public:
+		virtual ~INodeTemplateTransform() = default;
+
+		/** Return true if the builder was modified, false otherwise. */
+		virtual bool Transform(const FGuid& InNodeID, FMetaSoundFrontendDocumentBuilder& OutBuilder) const = 0;
+	};
+
 	/**
 		* Base interface for a node template, which acts in place of frontend node class and respective instance(s).
 		* Instances are preprocessed, allowing for custom graph manipulation prior to generating a respective runtime
@@ -60,7 +70,7 @@ namespace Metasound::Frontend
 #endif // WITH_EDITOR
 
 		// Generates node transform that is used to preprocess nodes.
-		virtual TUniquePtr<INodeTransform> GenerateNodeTransform() const = 0;
+		virtual TUniquePtr<INodeTemplateTransform> GenerateNodeTransform() const = 0;
 
 		// Returns the class definition for the given node class template.
 		virtual const FMetasoundFrontendClass& GetFrontendClass() const = 0;
diff --git a/Engine/Plugins/Runtime/Metasound/Source/MetasoundFrontend/Public/MetasoundFrontendTransform.h b/Engine/Plugins/Runtime/Metasound/Source/MetasoundFrontend/Public/MetasoundFrontendTransform.h
--- a/Engine/Plugins/Runtime/Metasound/Source/MetasoundFrontend/Public/MetasoundFrontendTransform.h
+++ b/Engine/Plugins/Runtime/Metasound/Source/MetasoundFrontend/Public/MetasoundFrontendTransform.h
@@ -131,24 +131,15 @@ namespace Metasound::Frontend
 
 	};
 
-	/** Updates document's given interface to the most recent version. */
 	class METASOUNDFRONTEND_API FUpdateRootGraphInterface : public IDocumentTransform
 	{
 	public:
 		FUpdateRootGraphInterface(const FMetasoundFrontendVersion& InInterfaceVersion, const FString& InOwningAssetName=FString(TEXT("Unknown")))
-			: InterfaceVersion(InInterfaceVersion)
-			, OwningAssetName(InOwningAssetName)
 		{
 		}
 
-		bool Transform(FDocumentHandle InDocument) const override;
-
-	private:
-		void GetUpdatePathForDocument(const FMetasoundFrontendVersion& InCurrentVersion, const FMetasoundFrontendVersion& InTargetVersion, TArray<const IInterfaceRegistryEntry*>& OutUpgradePath) const;
-		bool UpdateDocumentInterface(const TArray<const IInterfaceRegistryEntry*>& InUpgradePath, FDocumentHandle InDocument) const;
-
-		FMetasoundFrontendVersion InterfaceVersion;
-		FString OwningAssetName;
+		UE_DEPRECATED(5.5, "RootGraph update is now handled privately by internal MetaSound asset management")
+		virtual bool Transform(FDocumentHandle InDocument) const override { return false; }
 	};
 
 	/** Completely rebuilds the graph connecting a preset's inputs to the reference
diff --git a/Engine/Plugins/Runtime/Metasound/Source/MetasoundFrontend/Public/NodeTemplates/MetasoundFrontendNodeTemplateInput.h b/Engine/Plugins/Runtime/Metasound/Source/MetasoundFrontend/Public/NodeTemplates/MetasoundFrontendNodeTemplateInput.h
new file mode 100644
--- /dev/null
+++ b/Engine/Plugins/Runtime/Metasound/Source/MetasoundFrontend/Public/NodeTemplates/MetasoundFrontendNodeTemplateInput.h
@@ -0,0 +1,60 @@
+// Copyright Epic Games, Inc. All Rights Reserved.
+#pragma once
+
+#include "MetasoundFrontendController.h"
+#include "MetasoundFrontendDocument.h"
+#include "MetasoundFrontendNodeTemplateRegistry.h"
+#include "MetasoundFrontendNodeTemplateReroute.h"
+#include "MetasoundFrontendTransform.h"
+#include "Templates/SharedPointer.h"
+#include "UObject/NoExportTypes.h"
+
+
+// Forward Declarations
+struct FMetaSoundFrontendDocumentBuilder;
+
+
+namespace Metasound::Frontend
+{
+	// Specialized node that connects an input node's single output to various input destinations.
+	// While similar to reroute nodes, primarily exists to visually distinguish an input having multiple
+	// locations in a visual graph while sharing implementation at runtime, while also differentiating
+	// general input style from a typical reroute.
+	class METASOUNDFRONTEND_API FInputNodeTemplate : public FRerouteNodeTemplate
+	{
+	public:
+		static const FMetasoundFrontendClassName ClassName;
+		static const FMetasoundFrontendVersionNumber VersionNumber;
+
+		static const FInputNodeTemplate& GetChecked();
+		static const FNodeRegistryKey& GetRegistryKey();
+
+		virtual ~FInputNodeTemplate() = default;
+
+		virtual const FMetasoundFrontendClassName& GetClassName() const override;
+
+#if WITH_EDITOR
+		virtual FText GetNodeDisplayName(const IMetaSoundDocumentInterface& Interface, const FGuid& InNodeID) const override;
+#endif // WITH_EDITOR
+
+		virtual const FMetasoundFrontendClass& GetFrontendClass() const override;
+		virtual EMetasoundFrontendVertexAccessType GetNodeInputAccessType(const FMetaSoundFrontendDocumentBuilder& InBuilder, const FGuid& InNodeID, const FGuid& InVertexID) const override;
+		virtual EMetasoundFrontendVertexAccessType GetNodeOutputAccessType(const FMetaSoundFrontendDocumentBuilder& InBuilder, const FGuid& InNodeID, const FGuid& InVertexID) const override;
+
+#if WITH_EDITOR
+		virtual FText GetOutputPinDisplayName(const IOutputController& InOutput) const override;
+		virtual bool HasRequiredConnections(const FMetaSoundFrontendDocumentBuilder& InBuilder, const FGuid& InNodeID, FString* OutMessage = nullptr) const override;
+
+		// Injects template nodes between builder's document inputs not connected
+		// to existing template inputs, copying locational data from the represented
+		// input metadata. If bForceNodeCreation is false, only generates a template
+		// input node if a connection between the input and other nodes exists. If true,
+		// will inject template node irrespective of whether or not the input has connections.
+		bool Inject(FMetaSoundFrontendDocumentBuilder& InOutBuilder, bool bForceNodeCreation = false) const;
+#endif // WITH_EDITOR
+
+		virtual bool IsInputAccessTypeDynamic() const override;
+		virtual bool IsInputConnectionUserModifiable() const override;
+		virtual bool IsOutputAccessTypeDynamic() const override;
+	};
+} // namespace Metasound::Frontend
diff --git a/Engine/Plugins/Runtime/Metasound/Source/MetasoundFrontend/Public/NodeTemplates/MetasoundFrontendNodeTemplateReroute.h b/Engine/Plugins/Runtime/Metasound/Source/MetasoundFrontend/Public/NodeTemplates/MetasoundFrontendNodeTemplateReroute.h
--- a/Engine/Plugins/Runtime/Metasound/Source/MetasoundFrontend/Public/NodeTemplates/MetasoundFrontendNodeTemplateReroute.h
+++ b/Engine/Plugins/Runtime/Metasound/Source/MetasoundFrontend/Public/NodeTemplates/MetasoundFrontendNodeTemplateReroute.h
@@ -36,7 +36,7 @@ namespace Metasound::Frontend
 		virtual TUniquePtr<INodeTransform> GenerateNodeTransform(FMetasoundFrontendDocument& InPreprocessedDocument) const override { return { }; }
 
 		virtual FMetasoundFrontendNodeInterface GenerateNodeInterface(FNodeTemplateGenerateInterfaceParams InParams) const override;
-		virtual TUniquePtr<INodeTransform> GenerateNodeTransform() const override;
+		virtual TUniquePtr<INodeTemplateTransform> GenerateNodeTransform() const override;
 		virtual const FMetasoundFrontendClass& GetFrontendClass() const override;
 		virtual EMetasoundFrontendVertexAccessType GetNodeInputAccessType(const FMetaSoundFrontendDocumentBuilder& InBuilder, const FGuid& InNodeID, const FGuid& InVertexID) const override;
 		virtual EMetasoundFrontendVertexAccessType GetNodeOutputAccessType(const FMetaSoundFrontendDocumentBuilder& InBuilder, const FGuid& InNodeID, const FGuid& InVertexID) const override;
diff --git a/Engine/Source/Runtime/Engine/Classes/EdGraph/EdGraph.h b/Engine/Source/Runtime/Engine/Classes/EdGraph/EdGraph.h
--- a/Engine/Source/Runtime/Engine/Classes/EdGraph/EdGraph.h
+++ b/Engine/Source/Runtime/Engine/Classes/EdGraph/EdGraph.h
@@ -163,7 +163,7 @@ public:
 	* 
 	* @return True if the node has been removed from the graph
 	*/
-	ENGINE_API bool RemoveNode( UEdGraphNode* NodeToRemove, bool bBreakAllLinks = true );
+	ENGINE_API bool RemoveNode( UEdGraphNode* NodeToRemove, bool bBreakAllLinks = true, bool bMarkDirty = true );
 
 	/** Signal to listeners that the graph has changed - prefer to use NotifyNodeChanged when updating a single node */
 	ENGINE_API virtual void NotifyGraphChanged();
diff --git a/Engine/Source/Runtime/Engine/Classes/EdGraph/EdGraphPin.h b/Engine/Source/Runtime/Engine/Classes/EdGraph/EdGraphPin.h
--- a/Engine/Source/Runtime/Engine/Classes/EdGraph/EdGraphPin.h
+++ b/Engine/Source/Runtime/Engine/Classes/EdGraph/EdGraphPin.h
@@ -398,13 +398,13 @@ public:
 	ENGINE_API static UEdGraphPin* CreatePin(UEdGraphNode* InOwningNode);
 
 	/** Create a link. Note, this does not check that schema allows it, and will not break any existing connections */
-	ENGINE_API void MakeLinkTo(UEdGraphPin* ToPin);
+	ENGINE_API void MakeLinkTo(UEdGraphPin* ToPin, bool bAlwaysMarkDirty = true);
 
 	/** Break a link to the specified pin (if present) */
-	ENGINE_API void BreakLinkTo(UEdGraphPin* ToPin);
+	ENGINE_API void BreakLinkTo(UEdGraphPin* ToPin, bool bAlwaysMarkDirty = true);
 
 	/** Break all links from this pin */
-	ENGINE_API void BreakAllPinLinks(bool bNotifyNodes = false);
+	ENGINE_API void BreakAllPinLinks(bool bNotifyNodes = false, bool bAlwaysMarkDirty = true);
 
 	/**
 	* Moves the persistent data (across a node refresh) from the SourcePin.
@@ -577,7 +577,7 @@ private:
 	bool Serialize(FArchive& Ar);
 
 	// Helper functions
-	static void ConvertConnectedGhostNodesToRealNodes(UEdGraphNode* InNode);
+	static void ConvertConnectedGhostNodesToRealNodes(UEdGraphNode* InNode, bool bAlwaysMarkDirty = true);
 	static void ResolveReferencesToPin(UEdGraphPin* Pin, bool bStrictValidation = true);
 	static void SerializePinArray(FArchive& Ar, TArray<UEdGraphPin*>& ArrayRef, UEdGraphPin* RequestingPin, EPinResolveType ResolveType);
 	static bool SerializePin(FArchive& Ar, UEdGraphPin*& PinRef, int32 ArrayIdx, UEdGraphPin* RequestingPin, EPinResolveType ResolveType, TArray<UEdGraphPin*>& OldPins);
diff --git a/Engine/Source/Runtime/Engine/Private/EdGraph/EdGraph.cpp b/Engine/Source/Runtime/Engine/Private/EdGraph/EdGraph.cpp
--- a/Engine/Source/Runtime/Engine/Private/EdGraph/EdGraph.cpp
+++ b/Engine/Source/Runtime/Engine/Private/EdGraph/EdGraph.cpp
@@ -257,9 +257,9 @@ void UEdGraph::SelectNodeSet(TSet<const UEdGraphNode*> NodeSelection, bool bFrom
 	NotifyGraphChanged(SelectionAction);
 }
 
-bool UEdGraph::RemoveNode( UEdGraphNode* NodeToRemove, bool bBreakAllLinks )
+bool UEdGraph::RemoveNode( UEdGraphNode* NodeToRemove, bool bBreakAllLinks, bool bMarkDirty )
 {
-	Modify();
+	Modify(bMarkDirty);
 
 	int32 NumTimesNodeRemoved = Nodes.Remove(NodeToRemove);
 #if WITH_EDITOR
diff --git a/Engine/Source/Runtime/Engine/Private/EdGraph/EdGraphPin.cpp b/Engine/Source/Runtime/Engine/Private/EdGraph/EdGraphPin.cpp
--- a/Engine/Source/Runtime/Engine/Private/EdGraph/EdGraphPin.cpp
+++ b/Engine/Source/Runtime/Engine/Private/EdGraph/EdGraphPin.cpp
@@ -509,7 +509,7 @@ UEdGraphPin* UEdGraphPin::CreatePin(UEdGraphNode* InOwningNode)
 	return NewPin;
 }
 
-void UEdGraphPin::MakeLinkTo(UEdGraphPin* ToPin)
+void UEdGraphPin::MakeLinkTo(UEdGraphPin* ToPin, bool bAlwaysMarkDirty)
 {
 	if (ToPin)
 	{
@@ -525,32 +525,32 @@ void UEdGraphPin::MakeLinkTo(UEdGraphPin* ToPin)
 			ensureMsgf(MyNode->GetOuter() == ToPin->GetOwningNode()->GetOuter(), TEXT("%s"), *GetLinkInfoString( LOCTEXT("MakeLinkTo", "MakeLinkTo").ToString(), LOCTEXT("OuterMismatch", "has a different outer than pin").ToString(), ToPin)); // Ensure both pins belong to the same graph
 
 			// Notify owning nodes about upcoming change
-			Modify();
-			ToPin->Modify();
+			Modify(bAlwaysMarkDirty);
+			ToPin->Modify(bAlwaysMarkDirty);
 
 			// Add to both lists
 			LinkedTo.Add(ToPin);
 			ToPin->LinkedTo.Add(this);
 
 			// If either node was a pre-placed ghost, turn it into a real thing
-			UEdGraphPin::ConvertConnectedGhostNodesToRealNodes(MyNode);
-			UEdGraphPin::ConvertConnectedGhostNodesToRealNodes(ToPin->GetOwningNode());
+			UEdGraphPin::ConvertConnectedGhostNodesToRealNodes(MyNode, bAlwaysMarkDirty);
+			UEdGraphPin::ConvertConnectedGhostNodesToRealNodes(ToPin->GetOwningNode(), bAlwaysMarkDirty);
 		}
 	}
 }
 
-void UEdGraphPin::BreakLinkTo(UEdGraphPin* ToPin)
+void UEdGraphPin::BreakLinkTo(UEdGraphPin* ToPin, bool bAlwaysMarkDirty)
 {
 	if (ToPin)
 	{
 		// If we do indeed link to the passed in pin...
 		if (LinkedTo.Contains(ToPin))
 		{
-			Modify();
+			Modify(bAlwaysMarkDirty);
 
 			if (ToPin->LinkedTo.Contains(this))
 			{
-				ToPin->Modify();
+				ToPin->Modify(bAlwaysMarkDirty);
 				ToPin->LinkedTo.Remove(this);
 			}
 			else if (OwningNode && !OwningNode->HasAnyFlags(RF_BeginDestroyed))
@@ -569,13 +569,13 @@ void UEdGraphPin::BreakLinkTo(UEdGraphPin* ToPin)
 	}
 }
 
-void UEdGraphPin::BreakAllPinLinks(const bool bNotifyNodes)
+void UEdGraphPin::BreakAllPinLinks(bool bNotifyNodes, bool bAlwaysMarkDirty)
 {
 	TArray<UEdGraphPin*> LinkedToCopy = LinkedTo;
 
 	for (UEdGraphPin* LinkedToPin : LinkedToCopy)
 	{
-		BreakLinkTo(LinkedToPin);
+		BreakLinkTo(LinkedToPin, bAlwaysMarkDirty);
 #if WITH_EDITOR
 		if (bNotifyNodes)
 		{
@@ -1799,12 +1799,12 @@ void UEdGraphPin::DeclarePinCustomVersions(FArchive& Ar)
 }
 #endif
 
-void UEdGraphPin::ConvertConnectedGhostNodesToRealNodes(UEdGraphNode* InNode)
+void UEdGraphPin::ConvertConnectedGhostNodesToRealNodes(UEdGraphNode* InNode, bool bAlwaysMarkDirty)
 {
 	if (InNode && InNode->IsAutomaticallyPlacedGhostNode())
 	{
 		// Enable the node and clear the comment
-		InNode->Modify();
+		InNode->Modify(bAlwaysMarkDirty);
 		InNode->SetEnabledState(ENodeEnabledState::Enabled, /*bUserAction=*/ false);
 		InNode->NodeComment.Empty();
 #if WITH_EDITORONLY_DATA
@@ -1816,7 +1816,7 @@ void UEdGraphPin::ConvertConnectedGhostNodesToRealNodes(UEdGraphNode* InNode)
 		{
 			for (UEdGraphPin* OtherPin : Pin->LinkedTo)
 			{
-				ConvertConnectedGhostNodesToRealNodes(OtherPin->GetOwningNode());
+				ConvertConnectedGhostNodesToRealNodes(OtherPin->GetOwningNode(), bAlwaysMarkDirty);
 			}
 		}
 	}
