Don't ensure if location guid is invalid to fix automated build
#tests local compile

diff --git a/Engine/Plugins/Runtime/Metasound/Source/MetasoundFrontend/Private/MetasoundFrontendDocumentBuilder.cpp b/Engine/Plugins/Runtime/Metasound/Source/MetasoundFrontend/Private/MetasoundFrontendDocumentBuilder.cpp
--- a/Engine/Plugins/Runtime/Metasound/Source/MetasoundFrontend/Private/MetasoundFrontendDocumentBuilder.cpp
+++ b/Engine/Plugins/Runtime/Metasound/Source/MetasoundFrontend/Private/MetasoundFrontendDocumentBuilder.cpp
@@ -3410,10 +3410,16 @@ bool FMetaSoundFrontendDocumentBuilder::SetNodeLocation(const FGuid& InNodeID, c
 		FMetasoundFrontendNodeStyle& Style = Node->Style;
 		if (InLocationGuid)
 		{
-			ensureMsgf(InLocationGuid->IsValid(), TEXT("Invalid Location Guid no longer supported."));
-			Style.Display.Locations.FindOrAdd(*InLocationGuid) = InLocation;
+			if (InLocationGuid->IsValid())
+			{
+				Style.Display.Locations.FindOrAdd(*InLocationGuid) = InLocation;
+				return true;
+			}
+
+			UE_LOG(LogMetaSound, Display, TEXT("Invalid Location Guid no longer supported, reseting display location for node with ID '%s'"), *InNodeID.ToString());
 		}
-		else if (Style.Display.Locations.IsEmpty())
+
+		if (Style.Display.Locations.IsEmpty())
 		{
 			Style.Display.Locations = { { FGuid::NewGuid(), InLocation } };
 		}
