Reverting early return speculative fix for warnings and removing warning for now when attempting to cache MetaSound runtime data to fix sequencer warnings.
#rnx

diff --git a/Engine/Plugins/Runtime/Metasound/Source/MetasoundEngine/Private/MetasoundSource.cpp b/Engine/Plugins/Runtime/Metasound/Source/MetasoundEngine/Private/MetasoundSource.cpp
--- a/Engine/Plugins/Runtime/Metasound/Source/MetasoundEngine/Private/MetasoundSource.cpp
+++ b/Engine/Plugins/Runtime/Metasound/Source/MetasoundEngine/Private/MetasoundSource.cpp
@@ -1738,14 +1738,6 @@ void UMetaSoundSource::CacheRuntimeInputData()
 {
 	using namespace Metasound;
 
-	if (bIsBuilderActive)
-	{
-#if !WITH_EDITOR
-		UE_LOG(LogMetaSound, Warning, TEXT("Skipping caching of runtime inputs for UMetaSoundSource %s because there is an active builder"), *GetOwningAssetName());
-#endif // !WITH_EDITOR
-		return;
-	}
-
 	constexpr bool bCreateUObjectProxies = true; 
 	RuntimeInputData.InputMap = CreateRuntimeInputMap(bCreateUObjectProxies);
 
