Fix for bad initialization of unconnected metasound nodes in dynamic graph (from the ghost of Phil)

#tests Ran EngineTest MetaSoundBuilder tests locally that were failing, now they're beautiful and green

diff --git a/Engine/Plugins/Runtime/Metasound/Source/MetasoundGraphCore/Private/MetasoundDynamicOperatorTransactor.cpp b/Engine/Plugins/Runtime/Metasound/Source/MetasoundGraphCore/Private/MetasoundDynamicOperatorTransactor.cpp
--- a/Engine/Plugins/Runtime/Metasound/Source/MetasoundGraphCore/Private/MetasoundDynamicOperatorTransactor.cpp
+++ b/Engine/Plugins/Runtime/Metasound/Source/MetasoundGraphCore/Private/MetasoundDynamicOperatorTransactor.cpp
@@ -285,7 +285,7 @@ namespace Metasound
 		{
 		}
 
-		FDynamicGraphIncrementalSorter::FDynamicGraphIncrementalSorter(const IGraph& InGraph)
+		FDynamicGraphIncrementalSorter::FDynamicGraphIncrementalSorter(const FGraph& InGraph)
 		{
 			Init(InGraph);
 		}
@@ -435,7 +435,7 @@ namespace Metasound
 		}
 
 
-		void FDynamicGraphIncrementalSorter::Init(const IGraph& InGraph)
+		void FDynamicGraphIncrementalSorter::Init(const FGraph& InGraph)
 		{
 			/* determine new operator order. */
 			TArray<const INode*> NodeOrder;
@@ -466,6 +466,19 @@ namespace Metasound
 				OperatorMap[FromOperatorID].Outputs.Add(ToOperatorID);
 				OperatorMap[ToOperatorID].Inputs.Add(FromOperatorID);
 			}
+
+			// Unconnected nodes are intentionally skipped by the operator builder
+			// in order to reduce perf. In a dynamic operator, these nodes may be
+			// connected in the future. We add them here so that they exist on in
+			// incremental sorter in the case they are connected at a later time.
+			TArray<TPair<FGuid, const INode*>> UnconnectedNodes;
+			if (InGraph.FindUnconnectedNodes(UnconnectedNodes) > 0)
+			{
+				for (const TPair<FGuid, const INode*>& GuidAndNode : UnconnectedNodes)
+				{
+					InsertOperator(DirectedGraphAlgo::GetOperatorID(GuidAndNode.Value), FDynamicGraphIncrementalSorter::EInsertLocation::Last);
+				}
+			}
 		}
 
 		void FDynamicGraphIncrementalSorter::IncrementalTopologicalSortForNewEdge(FOperatorID InFromOperatorID, int32 InFromOrdinal, FOperatorID InToOperatorID, int32 InToOrdinal, TArray<FOrdinalSwap>& OutUpdates)
diff --git a/Engine/Plugins/Runtime/Metasound/Source/MetasoundGraphCore/Public/MetasoundDynamicOperatorTransactor.h b/Engine/Plugins/Runtime/Metasound/Source/MetasoundGraphCore/Public/MetasoundDynamicOperatorTransactor.h
--- a/Engine/Plugins/Runtime/Metasound/Source/MetasoundGraphCore/Public/MetasoundDynamicOperatorTransactor.h
+++ b/Engine/Plugins/Runtime/Metasound/Source/MetasoundGraphCore/Public/MetasoundDynamicOperatorTransactor.h
@@ -112,7 +112,7 @@ namespace Metasound
 
 			FDynamicGraphIncrementalSorter();
 
-			FDynamicGraphIncrementalSorter(const IGraph& InGraph);
+			FDynamicGraphIncrementalSorter(const FGraph& InGraph);
 
 			/** Insert a node into the graph.
 			 * @return Ordinal of added operator.
@@ -161,7 +161,7 @@ namespace Metasound
 
 			void IncrementalTopologicalSortForNewEdge(FOperatorID InFromOperatorID, int32 InFromOrdinal, FOperatorID InToOperatorID, int32 InToOdinal, TArray<FOrdinalSwap>& OutUpdates);
 
-			void Init(const IGraph& InGraph);
+			void Init(const FGraph& InGraph);
 
 			int32 MaxOrdinal = 0;
 			int32 MinOrdinal = 0;
