Fix for automated tests crash when running out-of-editor
#rnx

diff --git a/Engine/Plugins/Runtime/Metasound/Source/MetasoundFrontend/Private/MetasoundFrontendRegistryContainerImpl.cpp b/Engine/Plugins/Runtime/Metasound/Source/MetasoundFrontend/Private/MetasoundFrontendRegistryContainerImpl.cpp
--- a/Engine/Plugins/Runtime/Metasound/Source/MetasoundFrontend/Private/MetasoundFrontendRegistryContainerImpl.cpp
+++ b/Engine/Plugins/Runtime/Metasound/Source/MetasoundFrontend/Private/MetasoundFrontendRegistryContainerImpl.cpp
@@ -64,13 +64,17 @@ namespace Metasound::Frontend
 			return DocumentInterface;
 #else
 	#if !NO_LOGGING
-			const FMetaSoundFrontendDocumentBuilder& OriginalDocBuilder = IMetaSoundAssetManager::GetChecked().AttachDocumentBuilderChecked(*DocObject);
-			const bool bContainsTemplateDependency = OriginalDocBuilder.ContainsDependencyOfType(EMetasoundFrontendClassType::Template);
-			if (bContainsTemplateDependency)
+			// Only assets require template node processing and support document attachment
+			if (DocObject->IsAsset()) 
 			{
-				UE_LOG(LogMetaSound, Error,
-					TEXT("Template node processing disabled but provided asset class at '%s' to register contains template nodes. Runtime graph will fail to build."),
-					*OriginalDocBuilder.GetDebugName());
+				const FMetaSoundFrontendDocumentBuilder& OriginalDocBuilder = IMetaSoundAssetManager::GetChecked().AttachDocumentBuilderChecked(*DocObject);
+				const bool bContainsTemplateDependency = OriginalDocBuilder.ContainsDependencyOfType(EMetasoundFrontendClassType::Template);
+				if (bContainsTemplateDependency)
+				{
+					UE_LOG(LogMetaSound, Error,
+						TEXT("Template node processing disabled but provided asset class at '%s' to register contains template nodes. Runtime graph will fail to build."),
+						*OriginalDocBuilder.GetDebugName());
+				}
 			}
 	#endif // !NO_LOGGING
 
