[Backout] - CL33084850

diff --git a/Engine/Plugins/Runtime/Metasound/Source/MetasoundEditor/Private/MetasoundEditorGraphBuilder.cpp b/Engine/Plugins/Runtime/Metasound/Source/MetasoundEditor/Private/MetasoundEditorGraphBuilder.cpp
--- a/Engine/Plugins/Runtime/Metasound/Source/MetasoundEditor/Private/MetasoundEditorGraphBuilder.cpp
+++ b/Engine/Plugins/Runtime/Metasound/Source/MetasoundEditor/Private/MetasoundEditorGraphBuilder.cpp
@@ -1,6 +1,7 @@
 // Copyright Epic Games, Inc. All Rights Reserved.
 #include "MetasoundEditorGraphBuilder.h"
 
+#include "MetasoundEditorGraphBuilder.h"
 #include "Algo/AnyOf.h"
 #include "Algo/Sort.h"
 #include "Algo/Transform.h"
@@ -10,15 +11,12 @@
 #include "EdGraph/EdGraphPin.h"
 #include "Framework/Notifications/NotificationManager.h"
 #include "GraphEditor.h"
-#include "GraphEditorSettings.h"
 #include "Kismet/KismetSystemLibrary.h"
 #include "Metasound.h"
 #include "MetasoundAssetBase.h"
 #include "MetasoundAudioBuffer.h"
-#include "MetasoundBuilderSubsystem.h"
 #include "MetasoundEditor.h"
 #include "MetasoundEditorGraph.h"
-#include "MetasoundEditorGraphCommentNode.h"
 #include "MetasoundEditorGraphInputNode.h"
 #include "MetasoundEditorGraphMemberDefaults.h"
 #include "MetasoundEditorGraphNode.h"
@@ -44,7 +42,6 @@
 #include "MetasoundVertex.h"
 #include "MetasoundWaveTable.h"
 #include "Modules/ModuleManager.h"
-#include "NodeTemplates/MetasoundFrontendNodeTemplateInput.h"
 #include "NodeTemplates/MetasoundFrontendNodeTemplateReroute.h"
 #include "Templates/Tuple.h"
 #include "Toolkits/ToolkitManager.h"
@@ -60,34 +57,16 @@ namespace Metasound
 	{
 		namespace GraphBuilderPrivate
 		{
-			template <
-				typename TEdGraphNode,
-				typename FInitMetaSoundNodeFunc = TFunctionRef<void(UMetasoundEditorGraph&, TEdGraphNode&)>>
-			TEdGraphNode* AddNode(UObject& InMetaSound, FInitMetaSoundNodeFunc InitNodeFunc, bool bInSelectNewNode)
+			void DeleteNode(UObject& InMetaSound, Frontend::FNodeHandle InNodeHandle)
 			{
-				FMetasoundAssetBase* MetaSoundAsset = IMetasoundUObjectRegistry::Get().GetObjectAsAssetBase(&InMetaSound);
-				check(MetaSoundAsset);
-				UEdGraph& Graph = MetaSoundAsset->GetGraphChecked();
-				FGraphNodeCreator<TEdGraphNode> NodeCreator(Graph);
-				if (TEdGraphNode* NewGraphNode = NodeCreator.CreateNode(bInSelectNewNode))
+				if (InNodeHandle->IsValid())
 				{
-					// Required to happen prior to caching title and syncing location in case underlying type requires
-					// additional logic to initialize state in order to cache the title/sync the location
-					UMetasoundEditorGraph* MetasoundGraph = CastChecked<UMetasoundEditorGraph>(&Graph);
-					check(MetasoundGraph);
-					InitNodeFunc(*MetasoundGraph, *NewGraphNode);
-
-					NodeCreator.Finalize();
-
-					NewGraphNode->CacheTitle();
-
-					// Override finalized EdGraphNode ID if a location is found.
-					constexpr bool bUpdateEditorNodeID = true;
-					NewGraphNode->SyncLocationFromFrontendNode(bUpdateEditorNodeID);
-					return NewGraphNode;
+					Frontend::FGraphHandle GraphHandle = InNodeHandle->GetOwningGraph();
+					if (GraphHandle->IsValid())
+					{
+						GraphHandle->RemoveNode(*InNodeHandle);
+					}
 				}
-
-				return nullptr;
 			}
 
 			FName GenerateUniqueName(const TArray<FName>& InExistingNames, const FString& InBaseName)
@@ -120,7 +99,7 @@ namespace Metasound
 				}
 			};
 
-			void SynchronizeGraphRecursively(UObject& InMetaSound)
+			bool SynchronizeGraphRecursively(UObject& InMetaSound, bool bEditorGraphModified = false)
 			{
 				using namespace Frontend;
 
@@ -134,25 +113,26 @@ namespace Metasound
 				for (FMetasoundAssetBase* Reference : References)
 				{
 					check(Reference);
-					SynchronizeGraphRecursively(*Reference->GetOwningAsset());
+					bEditorGraphModified |= SynchronizeGraphRecursively(*Reference->GetOwningAsset(), bEditorGraphModified);
 				}
 
-				if (!MetaSoundAsset->GetConstModifyContext().GetDocumentModified())
+				if (!MetaSoundAsset->GetModifyContext().GetDocumentModified())
 				{
-					return;
+					return bEditorGraphModified;
 				}
 
 				// If no graph is set, MetaSound has been created outside of asset factory, so initialize it here.
 				if (!MetaSoundAsset->GetGraph())
 				{
-					UMetaSoundEditorSubsystem::GetChecked().InitEdGraph(InMetaSound);
+					UMetaSoundEditorSubsystem::GetChecked().InitAsset(InMetaSound);
 				}
 
-				FGraphBuilder::SynchronizeComments(InMetaSound);
-				FGraphBuilder::SynchronizeGraphMembers(InMetaSound);
-				FGraphBuilder::SynchronizeOutputNodes(InMetaSound);
-				FGraphBuilder::SynchronizeNodes(InMetaSound);
-				FGraphBuilder::SynchronizeConnections(InMetaSound);
+				bEditorGraphModified |= FGraphBuilder::SynchronizeGraphMembers(InMetaSound);
+				bEditorGraphModified |= FGraphBuilder::SynchronizeNodeMembers(InMetaSound);
+				bEditorGraphModified |= FGraphBuilder::SynchronizeNodes(InMetaSound);
+				bEditorGraphModified |= FGraphBuilder::SynchronizeConnections(InMetaSound);
+
+				return bEditorGraphModified;
 			}
 		} // namespace GraphBuilderPrivate
 
@@ -252,22 +232,6 @@ namespace Metasound
 			return true;
 		}
 
-		UMetasoundEditorGraphCommentNode* FGraphBuilder::CreateCommentNode(UObject& InMetaSound, bool bInSelectNewNode, FGuid InCommentID)
-		{
-			FMetasoundAssetBase* MetaSoundAsset = IMetasoundUObjectRegistry::Get().GetObjectAsAssetBase(&InMetaSound);
-			check(MetaSoundAsset);
-			UEdGraph& Graph = MetaSoundAsset->GetGraphChecked();
-			FGraphNodeCreator<UMetasoundEditorGraphCommentNode> NodeCreator(Graph);
-			if (UMetasoundEditorGraphCommentNode* NewNode = NodeCreator.CreateNode(bInSelectNewNode))
-			{
-				NewNode->CommentID = InCommentID;
-				NodeCreator.Finalize();
-				return NewNode;
-			}
-
-			return nullptr;
-		}
-
 		FText FGraphBuilder::GetDisplayName(const FMetasoundFrontendClassMetadata& InClassMetadata, FName InNodeName, bool bInIncludeNamespace)
 		{
 			using namespace Frontend;
@@ -408,7 +372,6 @@ namespace Metasound
 					return OwningNode->GetNodeName();
 				}
 
-				case EMetasoundFrontendClassType::Template:
 				default:
 				{
 					return InFrontendOutput.GetName();
@@ -443,7 +406,6 @@ namespace Metasound
 					return OwningNode->GetNodeName();
 				}
 
-				case EMetasoundFrontendClassType::Template:
 				default:
 				{
 					return InFrontendInput.GetName();
@@ -451,157 +413,158 @@ namespace Metasound
 			}
 		}
 
-		UMetasoundEditorGraphExternalNode* FGraphBuilder::AddTemplateNode(UObject& InMetaSound, const FGuid& InNodeID, const FMetasoundFrontendClassMetadata& InMetadata, bool bInSelectNewNode)
+		UMetasoundEditorGraphExternalNode* FGraphBuilder::AddExternalNode(UObject& InMetaSound, Frontend::FNodeHandle& InNodeHandle, FVector2D InLocation, bool bInSelectNewNode)
 		{
-			const EMetasoundFrontendClassType ClassType = InMetadata.GetType();
-			if (ensureMsgf(ClassType == EMetasoundFrontendClassType::Template, TEXT("Cannot call 'AddTemplateNode' with node of class type '%s'."), LexToString(ClassType)))
+			using namespace Frontend;
+
+			UMetasoundEditorGraphExternalNode* NewGraphNode = nullptr;
+
+			const EMetasoundFrontendClassType ClassType = InNodeHandle->GetClassMetadata().GetType();
+			const bool bIsExternalNode = ClassType == EMetasoundFrontendClassType::External;
+			const bool bIsTemplateNode = ClassType == EMetasoundFrontendClassType::Template;
+			if (!ensure(bIsExternalNode || bIsTemplateNode))
 			{
-				auto InitNodeFunc = [&InMetadata, &InNodeID](UMetasoundEditorGraph&, UMetasoundEditorGraphExternalNode& NewGraphNode)
-				{
-					NewGraphNode.NodeID = InNodeID;
-					NewGraphNode.ClassName = InMetadata.GetClassName();
-				};
-				return GraphBuilderPrivate::AddNode<UMetasoundEditorGraphExternalNode>(InMetaSound, InitNodeFunc, bInSelectNewNode);
+				return nullptr;
 			}
 
-			return nullptr;
-		}
-
-		UMetasoundEditorGraphExternalNode* FGraphBuilder::AddExternalNode(UObject& InMetaSound, const FGuid& InNodeID, const FMetasoundFrontendClassMetadata& InMetadata, bool bInSelectNewNode)
-		{
-			using namespace Frontend;
+			FMetasoundAssetBase* MetaSoundAsset = IMetasoundUObjectRegistry::Get().GetObjectAsAssetBase(&InMetaSound);
+			check(MetaSoundAsset);
+			UEdGraph& Graph = MetaSoundAsset->GetGraphChecked();
+			FGraphNodeCreator<UMetasoundEditorGraphExternalNode> NodeCreator(Graph);
 
-			const EMetasoundFrontendClassType ClassType = InMetadata.GetType();
-			if (ensureMsgf(ClassType == EMetasoundFrontendClassType::External, TEXT("Cannot call 'AddExternalNode' with node of class type '%s'."), LexToString(ClassType)))
+			NewGraphNode = NodeCreator.CreateNode(bInSelectNewNode);
+			if (ensure(NewGraphNode))
 			{
-				auto InitNodeFunc = [&InMetadata, &InNodeID](UMetasoundEditorGraph& MetasoundGraph, UMetasoundEditorGraphExternalNode& NewGraphNode)
-				{
-					const FNodeRegistryKey RegistryKey(InMetadata);
-					NewGraphNode.bIsClassNative = FMetasoundFrontendRegistryContainer::Get()->IsNodeNative(RegistryKey);
-					NewGraphNode.NodeID = InNodeID;
-					NewGraphNode.ClassName = InMetadata.GetClassName();
+				const FNodeRegistryKey RegistryKey = FNodeRegistryKey(InNodeHandle->GetClassMetadata());
+				NewGraphNode->bIsClassNative = FMetasoundFrontendRegistryContainer::Get()->IsNodeNative(RegistryKey);
+				NewGraphNode->ClassName = InNodeHandle->GetClassMetadata().GetClassName();
+				NewGraphNode->CacheTitle();
 
-					if (!NewGraphNode.bIsClassNative)
-					{
-						FMetasoundAssetBase* MetaSoundAsset = IMetasoundUObjectRegistry::Get().GetObjectAsAssetBase(MetasoundGraph.GetMetasound());
-						check(MetaSoundAsset);
-						MetaSoundAsset->RebuildReferencedAssetClasses();
-					}
-				};
+				NodeCreator.Finalize();
+				InitGraphNode(InNodeHandle, NewGraphNode, InMetaSound);
+				NewGraphNode->SetNodeLocation(InLocation);
 
-				return GraphBuilderPrivate::AddNode<UMetasoundEditorGraphExternalNode>(InMetaSound, InitNodeFunc, bInSelectNewNode);
+				// Adding external node may introduce referenced asset so rebuild referenced keys.
+				MetaSoundAsset->RebuildReferencedAssetClasses();
 			}
 
-			return nullptr;
+			return NewGraphNode;
 		}
 
-		UMetasoundEditorGraphExternalNode* FGraphBuilder::AddExternalNode(UObject& InMetaSound, const FMetasoundFrontendClassMetadata& InMetadata, bool bInSelectNewNode)
+		UMetasoundEditorGraphExternalNode* FGraphBuilder::AddExternalNode(UObject& InMetaSound, const FMetasoundFrontendClassMetadata& InMetadata, FVector2D InLocation, bool bInSelectNewNode)
 		{
 			FMetasoundAssetBase* MetaSoundAsset = IMetasoundUObjectRegistry::Get().GetObjectAsAssetBase(&InMetaSound);
 			check(MetaSoundAsset);
 
-			Frontend::FConstNodeHandle NodeHandle = MetaSoundAsset->GetRootGraphHandle()->AddNode(InMetadata);
-			return AddExternalNode(InMetaSound, NodeHandle->GetID(), NodeHandle->GetClassMetadata(), bInSelectNewNode);
-		}
-
-		Frontend::FNodeHandle FGraphBuilder::AddExternalNodeHandle(UObject& InMetaSound, const FMetasoundFrontendClassName& InClassName)
-		{
-			using namespace Frontend;
-
-			FMetasoundFrontendClass FrontendClass;
-			bool bDidFindClassWithName = ISearchEngine::Get().FindClassWithHighestVersion(InClassName, FrontendClass);
-			if (ensure(bDidFindClassWithName))
-			{
-				FMetasoundAssetBase* MetaSoundAsset = IMetasoundUObjectRegistry::Get().GetObjectAsAssetBase(&InMetaSound);
-				check(MetaSoundAsset);
-				return MetaSoundAsset->GetRootGraphHandle()->AddNode(FrontendClass.Metadata);
-			}
-
-			return INodeController::GetInvalidHandle();
+			Frontend::FNodeHandle NodeHandle = MetaSoundAsset->GetRootGraphHandle()->AddNode(InMetadata);
+			return AddExternalNode(InMetaSound, NodeHandle, InLocation, bInSelectNewNode);
 		}
 
-		UMetasoundEditorGraphVariableNode* FGraphBuilder::AddVariableNode(UObject& InMetaSound, const Frontend::FConstNodeHandle& InNodeHandle, bool bInSelectNewNode)
+		UMetasoundEditorGraphVariableNode* FGraphBuilder::AddVariableNode(UObject& InMetaSound, Frontend::FNodeHandle& InNodeHandle, FVector2D InLocation, bool bInSelectNewNode)
 		{
 			using namespace Frontend;
 
-			const FMetasoundFrontendClassMetadata& Metadata = InNodeHandle->GetClassMetadata();
-			EMetasoundFrontendClassType ClassType = Metadata.GetType();
+			EMetasoundFrontendClassType ClassType = InNodeHandle->GetClassMetadata().GetType();
 			const bool bIsSupportedClassType = (ClassType == EMetasoundFrontendClassType::VariableAccessor) 
 				|| (ClassType == EMetasoundFrontendClassType::VariableDeferredAccessor)
 				|| (ClassType == EMetasoundFrontendClassType::VariableMutator);
 
-			if (ensureMsgf(bIsSupportedClassType, TEXT("Cannot call 'AddVariableNode' with node of class type '%s'"), LexToString(ClassType)))
+			if (!ensure(bIsSupportedClassType))
 			{
-				FConstVariableHandle FrontendVariable = InNodeHandle->GetOwningGraph()->FindVariableContainingNode(InNodeHandle->GetID());
-				if (ensure(FrontendVariable->IsValid()))
+				return nullptr;
+			}
+
+			FConstVariableHandle FrontendVariable = InNodeHandle->GetOwningGraph()->FindVariableContainingNode(InNodeHandle->GetID());
+			if (!ensure(FrontendVariable->IsValid()))
+			{
+				return nullptr;
+			}
+
+			UMetasoundEditorGraphVariableNode* NewGraphNode = nullptr;
+			FMetasoundAssetBase* MetaSoundAsset = IMetasoundUObjectRegistry::Get().GetObjectAsAssetBase(&InMetaSound);
+			if (ensure(nullptr != MetaSoundAsset))
+			{
+				if (UMetasoundEditorGraph* MetasoundGraph = CastChecked<UMetasoundEditorGraph>(MetaSoundAsset->GetGraph()))
 				{
-					auto InitNodeFunc = [&FrontendVariable, &InNodeHandle, &Metadata](UMetasoundEditorGraph& MetasoundGraph, UMetasoundEditorGraphVariableNode& NewGraphNode)
+					FGraphNodeCreator<UMetasoundEditorGraphVariableNode> NodeCreator(*MetasoundGraph);
+
+					NewGraphNode = NodeCreator.CreateNode(bInSelectNewNode);
+					if (ensure(NewGraphNode))
 					{
-						UMetasoundEditorGraphVariable* Variable = MetasoundGraph.FindOrAddVariable(FrontendVariable);
+						NewGraphNode->ClassName = InNodeHandle->GetClassMetadata().GetClassName();
+						NewGraphNode->ClassType = ClassType;
+						InitGraphNode(InNodeHandle, NewGraphNode, InMetaSound);
+						NodeCreator.Finalize();
+
+						UMetasoundEditorGraphVariable* Variable = MetasoundGraph->FindOrAddVariable(FrontendVariable);
 						if (ensure(Variable))
 						{
-							NewGraphNode.Variable = Variable;
-							NewGraphNode.NodeID = InNodeHandle->GetID();
-							NewGraphNode.ClassName = Metadata.GetClassName();
-							NewGraphNode.ClassType = Metadata.GetType();
+							NewGraphNode->Variable = Variable;
+
+							// Ensures the variable node value is synced with the editor literal value should it be set
+							constexpr bool bPostTransaction = false;
+							Variable->UpdateFrontendDefaultLiteral(bPostTransaction);
 						}
-					};
 
-					return GraphBuilderPrivate::AddNode<UMetasoundEditorGraphVariableNode>(InMetaSound, InitNodeFunc, bInSelectNewNode);
+						MetasoundGraph->GetModifyContext().AddNodeIDsModified({ NewGraphNode->GetNodeID() });
+						NewGraphNode->SetNodeLocation(InLocation);
+					}
 				}
 			}
 
-			return nullptr;
+			return NewGraphNode;
 		}
 
-		UMetasoundEditorGraphInputNode* FGraphBuilder::AddInputNode(UObject& InMetaSound, const Frontend::FConstNodeHandle& InNodeHandle, bool bInSelectNewNode)
+		UMetasoundEditorGraphOutputNode* FGraphBuilder::AddOutputNode(UObject& InMetaSound, Frontend::FNodeHandle& InNodeHandle, FVector2D InLocation, bool bInSelectNewNode)
 		{
 			using namespace Frontend;
 
-			const FMetasoundFrontendClassMetadata& InputMetadata = InNodeHandle->GetClassMetadata();
-			if (ensureMsgf(InputMetadata.GetType() == EMetasoundFrontendClassType::Template, TEXT("Cannot call 'AddInputNode' with node of class type '%s': Must be input template."), LexToString(InputMetadata.GetType())))
+			UMetasoundEditorGraphOutputNode* NewGraphNode = nullptr;
+			if (!ensure(InNodeHandle->GetClassMetadata().GetType() == EMetasoundFrontendClassType::Output))
 			{
-				if (ensureMsgf(InputMetadata.GetClassName() == FInputNodeTemplate::ClassName, TEXT("Cannot call 'AddInputNode with node handle that is not of class '%s'"), *InputMetadata.GetClassName().ToString()))
+				return nullptr;
+			}
+
+			FMetasoundAssetBase* MetaSoundAsset = IMetasoundUObjectRegistry::Get().GetObjectAsAssetBase(&InMetaSound);
+			check(MetaSoundAsset);
+			UEdGraph& Graph = MetaSoundAsset->GetGraphChecked();
+			FGraphNodeCreator<UMetasoundEditorGraphOutputNode> NodeCreator(Graph);
+
+			NewGraphNode = NodeCreator.CreateNode(bInSelectNewNode);
+			if (ensure(NewGraphNode))
+			{
+				UMetasoundEditorGraph* MetasoundGraph = CastChecked<UMetasoundEditorGraph>(&Graph);
+
+				UMetasoundEditorGraphOutput* Output = MetasoundGraph->FindOrAddOutput(InNodeHandle);
+				if (ensure(Output))
 				{
-					auto InitNodeFunc = [&InNodeHandle](UMetasoundEditorGraph& MetasoundGraph, UMetasoundEditorGraphInputNode& NewGraphNode)
-					{
-						FConstInputHandle TemplateInput = InNodeHandle->GetConstInputs().Last();
-						FConstOutputHandle InputNodeOutput = TemplateInput->GetConnectedOutput();
-						UMetasoundEditorGraphInput* Input = MetasoundGraph.FindOrAddInput(InputNodeOutput->GetOwningNode());
-						if (ensure(Input))
-						{
-							NewGraphNode.Input = Input;
-							NewGraphNode.NodeID = InNodeHandle->GetID();
-						}
-					};
+					NewGraphNode->Output = Output;
+					NodeCreator.Finalize();
+					InitGraphNode(InNodeHandle, NewGraphNode, InMetaSound);
+
+					// Ensures the output node value is synced with the editor literal value should it be set
+					constexpr bool bPostTransaction = false;
+					Output->UpdateFrontendDefaultLiteral(bPostTransaction);
 
-					return GraphBuilderPrivate::AddNode<UMetasoundEditorGraphInputNode>(InMetaSound, InitNodeFunc, bInSelectNewNode);
+					MetasoundGraph->GetModifyContext().AddNodeIDsModified({ NewGraphNode->GetNodeID() });
 				}
+
+				NewGraphNode->CacheTitle();
+				NewGraphNode->SetNodeLocation(InLocation);
 			}
 
-			return nullptr;
+			return NewGraphNode;
 		}
 
-		UMetasoundEditorGraphOutputNode* FGraphBuilder::AddOutputNode(UObject& InMetaSound, const Frontend::FConstNodeHandle& InNodeHandle, bool bInSelectNewNode)
+		void FGraphBuilder::InitGraphNode(Frontend::FNodeHandle& InNodeHandle, UMetasoundEditorGraphNode* NewGraphNode, UObject& InMetaSound)
 		{
-			using namespace Frontend;
-
-			const EMetasoundFrontendClassType ClassType = InNodeHandle->GetClassMetadata().GetType();
-			if (ensureMsgf(ClassType == EMetasoundFrontendClassType::Output, TEXT("Cannot call 'AddOutputNode' with node of class type '%s'"), LexToString(ClassType)))
-			{
-				auto InitNodeFunc = [&InNodeHandle](UMetasoundEditorGraph& MetasoundGraph, UMetasoundEditorGraphOutputNode& NewGraphNode)
-				{
-					UMetasoundEditorGraphOutput* Output = MetasoundGraph.FindOrAddOutput(InNodeHandle);
-					if (ensure(Output))
-					{
-						NewGraphNode.Output = Output;
-					}
-				};
-
-				return GraphBuilderPrivate::AddNode<UMetasoundEditorGraphOutputNode>(InMetaSound, InitNodeFunc, bInSelectNewNode);
-			}
+			NewGraphNode->SetNodeID(InNodeHandle->GetID());
+			RebuildNodePins(*NewGraphNode);
+		}
 
-			return nullptr;
+		void FGraphBuilder::InitGraphNodeIDFromNodeHandle(const Frontend::FConstNodeHandle& InNodeHandle, UMetasoundEditorGraphNode* NewGraphNode)
+		{
+			NewGraphNode->SetNodeID(InNodeHandle->GetID());
 		}
 
 		FGraphValidationResults FGraphBuilder::ValidateGraph(UObject& InMetaSound)
@@ -1013,26 +976,7 @@ namespace Metasound
 
 		Frontend::FConstOutputHandle FGraphBuilder::GetConstOutputHandleFromPin(const UEdGraphPin* InPin)
 		{
-			using namespace Frontend;
-			using namespace VariableNames;
-
-			if (InPin && ensure(InPin->Direction == EGPD_Output))
-			{
-				if (const UMetasoundEditorGraphVariableNode* EdVariableNode = Cast<UMetasoundEditorGraphVariableNode>(InPin->GetOwningNode()))
-				{
-					// UEdGraphPins on variable nodes use the variable's name for display
-					// purposes instead of the underlying vertex's name. The frontend vertices
-					// of a variable node have consistent names no matter what the 
-					// variable is named.
-					return EdVariableNode->GetConstNodeHandle()->GetConstOutputWithVertexName(METASOUND_GET_PARAM_NAME(OutputData));
-				}
-				else if (const UMetasoundEditorGraphNode* EdNode = CastChecked<UMetasoundEditorGraphNode>(InPin->GetOwningNode()))
-				{
-					return EdNode->GetConstNodeHandle()->GetConstOutputWithVertexName(InPin->GetFName());
-				}
-			}
-
-			return IOutputController::GetInvalidHandle();
+			return GetOutputHandleFromPin(InPin);
 		}
 
 		UEdGraphPin* FGraphBuilder::FindReroutedOutputPin(UEdGraphPin* OutputPin)
@@ -1043,7 +987,7 @@ namespace Metasound
 			{
 				if (UMetasoundEditorGraphExternalNode* ExternalNode = Cast<UMetasoundEditorGraphExternalNode>(OutputPin->GetOwningNode()))
 				{
-					if (ExternalNode->GetBreadcrumb().ClassName == FRerouteNodeTemplate::ClassName)
+					if (ExternalNode->GetClassName() == FRerouteNodeTemplate::ClassName)
 					{
 						auto IsInput = [](UEdGraphPin* Pin) { check(Pin); return Pin->Direction == EGPD_Input; };
 						if (UEdGraphPin* RerouteInput = *ExternalNode->Pins.FindByPredicate(IsInput))
@@ -1070,7 +1014,7 @@ namespace Metasound
 			{
 				if (UMetasoundEditorGraphExternalNode* ExternalNode = Cast<UMetasoundEditorGraphExternalNode>(OutputPin->GetOwningNode()))
 				{
-					if (ExternalNode->GetBreadcrumb().ClassName == FRerouteNodeTemplate::ClassName)
+					if (ExternalNode->GetClassName() == FRerouteNodeTemplate::ClassName)
 					{
 						auto IsInput = [](const UEdGraphPin* Pin) { check(Pin); return Pin->Direction == EGPD_Input; };
 						if (const UEdGraphPin* RerouteInput = *ExternalNode->Pins.FindByPredicate(IsInput))
@@ -1097,7 +1041,7 @@ namespace Metasound
 			{
 				if (UMetasoundEditorGraphExternalNode* ExternalNode = Cast<UMetasoundEditorGraphExternalNode>(OutputPin->GetOwningNode()))
 				{
-					if (ExternalNode->GetBreadcrumb().ClassName == FRerouteNodeTemplate::ClassName)
+					if (ExternalNode->GetClassName() == FRerouteNodeTemplate::ClassName)
 					{
 						auto IsInput = [](UEdGraphPin* Pin) { check(Pin); return Pin->Direction == EGPD_Input; };
 						if (UEdGraphPin* RerouteInput = *ExternalNode->Pins.FindByPredicate(IsInput))
@@ -1118,33 +1062,9 @@ namespace Metasound
 			return IOutputController::GetInvalidHandle();
 		}
 
-		Frontend::FConstOutputHandle FGraphBuilder::FindReroutedConstOutputHandleFromPin(const UEdGraphPin* OutputPin)
+		Frontend::FConstOutputHandle FGraphBuilder::FindReroutedConstOutputHandleFromPin(const UEdGraphPin* InPin)
 		{
-			using namespace Frontend;
-
-			if (OutputPin)
-			{
-				if (const UMetasoundEditorGraphExternalNode* ExternalNode = Cast<UMetasoundEditorGraphExternalNode>(OutputPin->GetOwningNode()))
-				{
-					if (ExternalNode->GetBreadcrumb().ClassName == FRerouteNodeTemplate::ClassName)
-					{
-						auto IsInput = [](UEdGraphPin* Pin) { check(Pin); return Pin->Direction == EGPD_Input; };
-						if (const UEdGraphPin* RerouteInput = *ExternalNode->Pins.FindByPredicate(IsInput))
-						{
-							const TArray<UEdGraphPin*>& LinkedTo = RerouteInput->LinkedTo;
-							if (!LinkedTo.IsEmpty())
-							{
-								const UEdGraphPin* ReroutedOutput = LinkedTo.Last();
-								return FindReroutedConstOutputHandleFromPin(ReroutedOutput);
-							}
-						}
-					}
-				}
-
-				return GetConstOutputHandleFromPin(OutputPin);
-			}
-
-			return IOutputController::GetInvalidHandle();
+			return FindReroutedOutputHandleFromPin(InPin);
 		}
 
 		void FGraphBuilder::FindReroutedInputPins(UEdGraphPin* InPinToCheck, TArray<UEdGraphPin*>& InOutInputPins)
@@ -1155,7 +1075,7 @@ namespace Metasound
 			{
 				if (UMetasoundEditorGraphExternalNode* ExternalNode = Cast<UMetasoundEditorGraphExternalNode>(InPinToCheck->GetOwningNode()))
 				{
-					if (ExternalNode->GetBreadcrumb().ClassName == FRerouteNodeTemplate::ClassName)
+					if (ExternalNode->GetClassName() == FRerouteNodeTemplate::ClassName)
 					{
 						for (UEdGraphPin* Pin : ExternalNode->Pins)
 						{
@@ -1191,43 +1111,60 @@ namespace Metasound
 			return Algo::AnyOf(EditorNodes, [](const UMetasoundEditorGraphNode* Node) { return Node->ErrorType == EMessageSeverity::Error; });
 		}
 
-		Frontend::FNodeHandle FGraphBuilder::AddInputTemplateNodeHandle(UObject& InMetaSound, Frontend::FNodeHandle& InputNodeHandle)
+		bool FGraphBuilder::SynchronizeNodeLocation(const Frontend::FConstNodeHandle& InNode, UMetasoundEditorGraphNode& OutGraphNode)
 		{
-			using namespace Frontend;
+			bool bModified = false;
 
-			FMetasoundAssetBase* MetaSoundAsset = IMetasoundUObjectRegistry::Get().GetObjectAsAssetBase(&InMetaSound);
-			check(MetaSoundAsset);
+			const FMetasoundFrontendNodeStyle& Style = InNode->GetNodeStyle();
 
-			UMetasoundEditorGraph* Graph = Cast<UMetasoundEditorGraph>(MetaSoundAsset->GetGraph());
-			if (!ensure(Graph))
+			const FVector2D* Location = Style.Display.Locations.Find(OutGraphNode.NodeGuid);
+			if (!Location)
 			{
-				return INodeController::GetInvalidHandle();
+				// If no specific location found, use default location if provided (zero guid
+				// for example, provided by preset defaults.)
+				Location = Style.Display.Locations.Find({ });
 			}
 
-			TArray<FConstOutputHandle> NodeOutputs = InputNodeHandle->GetConstOutputs();
-			if (!ensure(!NodeOutputs.IsEmpty()))
+			if (Location)
 			{
-				return INodeController::GetInvalidHandle();
+				const int32 LocX = FMath::TruncToInt(Location->X);
+				const int32 LocY = FMath::TruncToInt(Location->Y);
+				const bool bXChanged = static_cast<bool>(LocX - OutGraphNode.NodePosX);
+				const bool bYChanged = static_cast<bool>(LocY - OutGraphNode.NodePosY);
+				if (bXChanged || bYChanged)
+				{
+					OutGraphNode.NodePosX = LocX;
+					OutGraphNode.NodePosY = LocY;
+					bModified = true;
+				}
 			}
 
-			if (!ensure(InputNodeHandle->GetClassMetadata().GetType() == EMetasoundFrontendClassType::Input))
-			{
-				return INodeController::GetInvalidHandle();
-			}
+			return bModified;
+		}
 
-			const INodeTemplate* InputTemplate = INodeTemplateRegistry::Get().FindTemplate(FInputNodeTemplate::ClassName);
-			check(InputTemplate);
+		UMetasoundEditorGraphInputNode* FGraphBuilder::AddInputNode(UObject& InMetaSound, Frontend::FNodeHandle InNodeHandle, FVector2D InLocation, bool bInSelectNewNode)
+		{
+			using namespace Frontend;
 
-			const FName DataType = NodeOutputs.Last()->GetDataType();
-			FNodeTemplateGenerateInterfaceParams Params { { DataType }, { } };
+			FMetasoundAssetBase* MetaSoundAsset = IMetasoundUObjectRegistry::Get().GetObjectAsAssetBase(&InMetaSound);
+			check(MetaSoundAsset);
 
-			FGraphHandle GraphHandle = Graph->GetGraphHandle();
-			FNodeHandle TemplateNodeHandle = GraphHandle->AddTemplateNode(*InputTemplate, MoveTemp(Params));
+			UMetasoundEditorGraph* MetasoundGraph = Cast<UMetasoundEditorGraph>(MetaSoundAsset->GetGraph());
+			if (!ensure(MetasoundGraph))
+			{
+				return nullptr;
+			}
 
-			TemplateNodeHandle->GetInputs().Last()->Connect(*InputNodeHandle->GetOutputs().Last());
-			FMetasoundFrontendDocumentModifyContext& ModifyContext = GetOutermostMetaSoundChecked(*Graph).GetModifyContext();
-			ModifyContext.AddNodeIDsModified({ TemplateNodeHandle->GetID() });
-			return TemplateNodeHandle;
+			UMetasoundEditorGraphInputNode* NewGraphNode = MetasoundGraph->CreateInputNode(InNodeHandle, bInSelectNewNode);
+			if (ensure(NewGraphNode))
+			{
+				NewGraphNode->SetNodeLocation(InLocation);
+				RebuildNodePins(*NewGraphNode);
+				MetasoundGraph->GetModifyContext().AddNodeIDsModified({ NewGraphNode->GetNodeID() });
+				return NewGraphNode;
+			}
+
+			return nullptr;
 		}
 
 		bool FGraphBuilder::GetPinLiteral(UEdGraphPin& InInputPin, FMetasoundFrontendLiteral& OutDefaultLiteral)
@@ -1236,7 +1173,7 @@ namespace Metasound
 
 			IMetasoundEditorModule& EditorModule = FModuleManager::GetModuleChecked<IMetasoundEditorModule>("MetaSoundEditor");
 
-			FConstInputHandle InputHandle = GetConstInputHandleFromPin(&InInputPin);
+			FInputHandle InputHandle = GetInputHandleFromPin(&InInputPin);
 			if (!ensure(InputHandle->IsValid()))
 			{
 				return false;
@@ -1389,6 +1326,78 @@ namespace Metasound
 			return true;
 		}
 
+		Frontend::FNodeHandle FGraphBuilder::AddNodeHandle(UObject& InMetaSound, UMetasoundEditorGraphNode& InGraphNode)
+		{
+			using namespace Frontend;
+
+			FNodeHandle NodeHandle = INodeController::GetInvalidHandle();
+			if (UMetasoundEditorGraphInputNode* InputNode = Cast<UMetasoundEditorGraphInputNode>(&InGraphNode))
+			{
+				const TArray<UEdGraphPin*>& Pins = InGraphNode.GetAllPins();
+				const UEdGraphPin* Pin = Pins.IsEmpty() ? nullptr : Pins[0];
+				if (ensure(Pin) && ensure(Pin->Direction == EGPD_Output))
+				{
+					UMetasoundEditorGraphInput* Input = InputNode->Input;
+					if (ensure(Input))
+					{
+						const FName PinName = Pin->GetFName();
+						FCreateNodeVertexParams VertexParams;
+						VertexParams.DataType = Input->GetDataType();
+
+						NodeHandle = AddInputNodeHandle(InMetaSound, VertexParams, nullptr, &PinName);
+						NodeHandle->SetDescription(InGraphNode.GetTooltipText());
+					}
+				}
+			}
+
+			else if (UMetasoundEditorGraphOutputNode* OutputNode = Cast<UMetasoundEditorGraphOutputNode>(&InGraphNode))
+			{
+				const TArray<UEdGraphPin*>& Pins = InGraphNode.GetAllPins();
+				const UEdGraphPin* Pin = Pins.IsEmpty() ? nullptr : Pins[0];
+				if (ensure(Pin) && ensure(Pin->Direction == EGPD_Input))
+				{
+					UMetasoundEditorGraphOutput* Output = OutputNode->Output;
+					if (ensure(Output))
+					{
+						const FName PinName = Pin->GetFName();
+						FCreateNodeVertexParams VertexParams;
+						VertexParams.DataType = Output->GetDataType();
+
+						NodeHandle = FGraphBuilder::AddOutputNodeHandle(InMetaSound, VertexParams, &PinName);
+						NodeHandle->SetDescription(InGraphNode.GetTooltipText());
+					}
+				}
+			}
+			else if (UMetasoundEditorGraphVariableNode* VariableNode = Cast<UMetasoundEditorGraphVariableNode>(&InGraphNode))
+			{
+				NodeHandle = FGraphBuilder::AddVariableNodeHandle(InMetaSound, VariableNode->Variable->GetVariableID(), VariableNode->GetClassName().ToNodeClassName());
+			}
+			else if (UMetasoundEditorGraphExternalNode* ExternalNode = Cast<UMetasoundEditorGraphExternalNode>(&InGraphNode))
+			{
+				FMetasoundFrontendClass FrontendClass;
+				bool bDidFindClassWithName = ISearchEngine::Get().FindClassWithHighestVersion(ExternalNode->ClassName.ToNodeClassName(), FrontendClass);
+				if (ensure(bDidFindClassWithName))
+				{
+					FMetasoundAssetBase* MetaSoundAsset = IMetasoundUObjectRegistry::Get().GetObjectAsAssetBase(&InMetaSound);
+					check(MetaSoundAsset);
+
+					Frontend::FNodeHandle NewNode = MetaSoundAsset->GetRootGraphHandle()->AddNode(FrontendClass.Metadata);
+					ExternalNode->SetNodeID(NewNode->GetID());
+
+					NodeHandle = NewNode;
+				}
+			}
+
+			if (NodeHandle->IsValid())
+			{
+				FMetasoundFrontendNodeStyle Style = NodeHandle->GetNodeStyle();
+				Style.Display.Locations.Add(InGraphNode.NodeGuid, FVector2D(InGraphNode.NodePosX, InGraphNode.NodePosY));
+				NodeHandle->SetNodeStyle(Style);
+			}
+
+			return NodeHandle;
+		}
+
 		Frontend::FNodeHandle FGraphBuilder::AddInputNodeHandle(UObject& InMetaSound, const FCreateNodeVertexParams& InParams, const FMetasoundFrontendLiteral* InDefaultValue, const FName* InNameBase)
 		{
 			FMetasoundAssetBase* MetaSoundAsset = IMetasoundUObjectRegistry::Get().GetObjectAsAssetBase(&InMetaSound);
@@ -1448,20 +1457,6 @@ namespace Metasound
 			return GraphBuilderPrivate::GenerateUniqueName(ExistingVariableNames, InBaseName);
 		}
 
-		FMetasoundAssetBase& FGraphBuilder::GetOutermostMetaSoundChecked(UObject& InSubObject)
-		{
-			FMetasoundAssetBase* MetaSoundAsset = IMetasoundUObjectRegistry::Get().GetObjectAsAssetBase(InSubObject.GetOutermostObject());
-			check(MetaSoundAsset);
-			return *MetaSoundAsset;
-		}
-
-		const FMetasoundAssetBase& FGraphBuilder::GetOutermostConstMetaSoundChecked(const UObject& InSubObject)
-		{
-			const FMetasoundAssetBase* MetaSoundAsset = IMetasoundUObjectRegistry::Get().GetObjectAsAssetBase(InSubObject.GetOutermostObject());
-			check(MetaSoundAsset);
-			return *MetaSoundAsset;
-		}
-
 		Frontend::FVariableHandle FGraphBuilder::AddVariableHandle(UObject& InMetaSound, const FName& InTypeName)
 		{
 			using namespace Frontend;
@@ -1510,24 +1505,24 @@ namespace Metasound
 							break;
 
 						case EMetasoundFrontendClassType::VariableMutator:
-						{
-							FConstVariableHandle Variable = Graph->FindVariable(InVariableID);
-							FConstNodeHandle ExistingMutator = Variable->FindMutatorNode();
-							if (!ExistingMutator->IsValid())
-							{
-								FrontendNode = Graph->FindOrAddVariableMutatorNode(InVariableID);
-							}
-							else
 							{
-								UE_LOG(LogMetaSound, Error, TEXT("Cannot add node because \"%s\" already exists for variable \"%s\""), *ExistingMutator->GetDisplayName().ToString(), *Variable->GetDisplayName().ToString());
+								FConstVariableHandle Variable = Graph->FindVariable(InVariableID);
+								FConstNodeHandle ExistingMutator = Variable->FindMutatorNode();
+								if (!ExistingMutator->IsValid())
+								{
+									FrontendNode = Graph->FindOrAddVariableMutatorNode(InVariableID);
+								}
+								else
+								{
+									UE_LOG(LogMetaSound, Error, TEXT("Cannot add node because \"%s\" already exists for variable \"%s\""), *ExistingMutator->GetDisplayName().ToString(), *Variable->GetDisplayName().ToString());
+								}
 							}
-						}
-						break;
+							break;
 
 						default:
-						{
-							checkNoEntry();
-						}
+							{
+								checkNoEntry();
+							}
 					}
 				}
 			}
@@ -1536,13 +1531,60 @@ namespace Metasound
 			{
 				InVariableNode->ClassName = FrontendNode->GetClassMetadata().GetClassName();
 				InVariableNode->ClassType = FrontendNode->GetClassMetadata().GetType();
-				InVariableNode->NodeID = FrontendNode->GetID();
-				MetaSoundAsset->GetModifyContext().AddNodeIDsModified({ InVariableNode->GetNodeID() });
+				InVariableNode->SetNodeID(FrontendNode->GetID());
 			}
 
 			return FrontendNode;
 		}
 
+		UMetasoundEditorGraphNode* FGraphBuilder::AddNode(UObject& InMetaSound, Frontend::FNodeHandle InNodeHandle, FVector2D InLocation, bool bInSelectNewNode)
+		{
+			switch (InNodeHandle->GetClassMetadata().GetType())
+			{
+				case EMetasoundFrontendClassType::Input:
+				{
+					return CastChecked<UMetasoundEditorGraphNode>(AddInputNode(InMetaSound, InNodeHandle, InLocation, bInSelectNewNode));
+				}
+				break;
+
+				case EMetasoundFrontendClassType::External:
+				case EMetasoundFrontendClassType::Template:
+				{
+					return CastChecked<UMetasoundEditorGraphNode>(AddExternalNode(InMetaSound, InNodeHandle, InLocation, bInSelectNewNode));
+				}
+				break;
+
+				case EMetasoundFrontendClassType::Output:
+				{
+					return CastChecked<UMetasoundEditorGraphNode>(AddOutputNode(InMetaSound, InNodeHandle, InLocation, bInSelectNewNode));
+				}
+				break;
+
+				case EMetasoundFrontendClassType::VariableMutator:
+				case EMetasoundFrontendClassType::VariableAccessor:
+				case EMetasoundFrontendClassType::VariableDeferredAccessor:
+				case EMetasoundFrontendClassType::Variable:
+				{
+					return CastChecked<UMetasoundEditorGraphNode>(AddVariableNode(InMetaSound, InNodeHandle, InLocation, bInSelectNewNode));
+				}
+				break;
+
+				case EMetasoundFrontendClassType::Invalid:
+				case EMetasoundFrontendClassType::Graph:
+				
+				case EMetasoundFrontendClassType::Literal: // Not yet supported in editor
+				
+				default:
+				{
+					checkNoEntry();
+					static_assert(static_cast<int32>(EMetasoundFrontendClassType::Invalid) == 10, "Possible missing FMetasoundFrontendClassType case coverage");
+				}
+				break;
+			}
+
+			return nullptr;
+		}
+
 		bool FGraphBuilder::ConnectNodes(UEdGraphPin& InInputPin, UEdGraphPin& InOutputPin, bool bInConnectEdPins)
 		{
 			using namespace Frontend;
@@ -1642,45 +1684,196 @@ namespace Metasound
 			MetaSoundAsset->GetModifyContext().SetDocumentModified();
 		}
 
-		bool FGraphBuilder::DeleteMemberNodes(const UMetasoundEditorGraphMember& InGraphMember)
+		void FGraphBuilder::InitMetaSound(UObject& InMetaSound, const FString& InAuthor)
 		{
-			bool bSuccess = true;
-			const TArray<UMetasoundEditorGraphMemberNode*> Nodes = InGraphMember.GetNodes();
-			for (UMetasoundEditorGraphMemberNode* Node : Nodes)
+			using namespace Frontend;
+			using namespace GraphBuilderPrivate;
+
+			FMetasoundAssetBase* MetaSoundAsset = IMetasoundUObjectRegistry::Get().GetObjectAsAssetBase(&InMetaSound);
+			check(MetaSoundAsset);
+			FMetasoundFrontendDocument& Doc = MetaSoundAsset->GetDocumentChecked();
+
+			// 1. Set default class Metadata
+			Doc.RootGraph.Metadata.SetClassName(FMetasoundFrontendClassName(FName(), *FGuid::NewGuid().ToString(), FName()));
+			Doc.RootGraph.Metadata.SetVersion({ 1, 0 });
+			Doc.RootGraph.Metadata.SetType(EMetasoundFrontendClassType::Graph);
+			Doc.RootGraph.Metadata.SetAuthor(InAuthor);
+
+			if (Doc.RootGraph.Metadata.GetType() != EMetasoundFrontendClassType::Graph)
 			{
-				if (ensure(Node))
+				UE_LOG(LogMetaSound, Display, TEXT("Forcing class type to EMetasoundFrontendClassType::Graph on root graph metadata"));
+				Doc.RootGraph.Metadata.SetType(EMetasoundFrontendClassType::Graph);
+			}
+
+			// 2. Set default doc version Metadata
+			Doc.Metadata.Version.Number = FMetasoundFrontendDocument::GetMaxVersion();
+
+			// 3. Add default interfaces
+			{
+				UClass* AssetClass = InMetaSound.GetClass();
+				check(AssetClass);
+
+				TArray<FMetasoundFrontendVersion> InitVersions = ISearchEngine::Get().FindUClassDefaultInterfaceVersions(AssetClass->GetClassPathName());
+				FModifyRootGraphInterfaces({ }, InitVersions).Transform(Doc);
+			}
+
+			FGraphHandle GraphHandle = MetaSoundAsset->GetRootGraphHandle();
+			FVector2D InputNodeLocation = FVector2D::ZeroVector;
+			FVector2D ExternalNodeLocation = InputNodeLocation + DisplayStyle::NodeLayout::DefaultOffsetX;
+			FVector2D OutputNodeLocation = ExternalNodeLocation + DisplayStyle::NodeLayout::DefaultOffsetX;
+
+			TArray<FNodeHandle> NodeHandles = GraphHandle->GetNodes();
+			for (FNodeHandle& NodeHandle : NodeHandles)
+			{
+				const EMetasoundFrontendClassType NodeType = NodeHandle->GetClassMetadata().GetType();
+				FVector2D NewLocation;
+				if (NodeType == EMetasoundFrontendClassType::Input)
 				{
-					bSuccess &= Metasound::Editor::FGraphBuilder::DeleteNode(*Node);
+					NewLocation = InputNodeLocation;
+					InputNodeLocation += DisplayStyle::NodeLayout::DefaultOffsetY;
 				}
+				else if (NodeType == EMetasoundFrontendClassType::Output)
+				{
+					NewLocation = OutputNodeLocation;
+					OutputNodeLocation += DisplayStyle::NodeLayout::DefaultOffsetY;
+				}
+				else
+				{
+					NewLocation = ExternalNodeLocation;
+					ExternalNodeLocation += DisplayStyle::NodeLayout::DefaultOffsetY;
+				}
+				FMetasoundFrontendNodeStyle Style = NodeHandle->GetNodeStyle();
+				// TODO: Find consistent location for controlling node locations.
+				// Currently it is split between MetasoundEditor and MetasoundFrontend modules.
+				Style.Display.Locations = {{FGuid::NewGuid(), NewLocation}};
+				NodeHandle->SetNodeStyle(Style);
 			}
-			return bSuccess;
 		}
 
-		bool FGraphBuilder::DeleteNode(UEdGraphNode& InNode, bool bRemoveUnusedDependencies)
+		void FGraphBuilder::InitMetaSoundPreset(UObject& InMetaSoundReferenced, UObject& InMetaSoundPreset)
 		{
 			using namespace Frontend;
+			using namespace GraphBuilderPrivate;
+
+			FMetasoundAssetBase* PresetAsset = IMetasoundUObjectRegistry::Get().GetObjectAsAssetBase(&InMetaSoundPreset);
+			check(PresetAsset);
+
+			// Mark preset as auto-update and non-editable
+			FGraphHandle PresetGraphHandle = PresetAsset->GetRootGraphHandle();
+			FMetasoundFrontendGraphStyle Style = PresetGraphHandle->GetGraphStyle();
+			Style.bIsGraphEditable = false;
+			PresetGraphHandle->SetGraphStyle(Style);
+
+			// Mark all inputs as inherited by default
+			TSet<FName> InputsInheritingDefault;
+			Algo::Transform(PresetGraphHandle->GetInputNodes(), InputsInheritingDefault, [](FConstNodeHandle NodeHandle)
+			{
+				return NodeHandle->GetNodeName();
+			});
+			PresetGraphHandle->SetInputsInheritingDefault(MoveTemp(InputsInheritingDefault));
+
+			FGraphBuilder::RegisterGraphWithFrontend(InMetaSoundReferenced);
+
+			const FMetasoundAssetBase* ReferencedAsset = IMetasoundUObjectRegistry::Get().GetObjectAsAssetBase(&InMetaSoundReferenced);
+			check(ReferencedAsset);
+
+			FRebuildPresetRootGraph(ReferencedAsset->GetDocumentHandle()).Transform(PresetAsset->GetDocumentHandle());
+
+			PresetAsset->ConformObjectDataToInterfaces();
+		}
+
+		bool FGraphBuilder::DeleteNode(UEdGraphNode& InNode)
+		{
+			using namespace Frontend;
+
+			if (!InNode.CanUserDeleteNode())
+			{
+				return false;
+			}
 
+			// If node isn't a MetasoundEditorGraphNode, just remove and return (ex. comment nodes)
+			UMetasoundEditorGraphNode* Node = Cast<UMetasoundEditorGraphNode>(&InNode);
 			UMetasoundEditorGraph* Graph = CastChecked<UMetasoundEditorGraph>(InNode.GetGraph());
-			UObject& MetaSound = Graph->GetMetasoundChecked();
-			UMetaSoundBuilderBase& Builder = UMetaSoundBuilderSubsystem::GetChecked().AttachBuilderToAssetChecked(MetaSound);
+			if (!Node)
+			{
+				Graph->RemoveNode(&InNode);
+				return true;
+			}
+
 
-			if (UMetasoundEditorGraphCommentNode* Node = Cast<UMetasoundEditorGraphCommentNode>(&InNode))
+			// Remove connects only to pins associated with this EdGraph node
+			// only (Iterate pins and not Frontend representation to preserve
+			// other input/output EditorGraph reference node associations)
+			Node->IteratePins([](UEdGraphPin& Pin, int32 Index)
 			{
-				const bool bFrontendCommentRemoved = Builder.RemoveGraphComment(Node->CommentID);
-				if (ensure(bFrontendCommentRemoved))
+				// Only add literal inputs for output pins as adding when disconnecting
+				// inputs would immediately orphan them on EditorGraph node removal below.
+				const bool bAddLiteralInputs = Pin.Direction == EGPD_Output;
+				FGraphBuilder::DisconnectPinVertex(Pin, bAddLiteralInputs);
+			});
+
+			FNodeHandle NodeHandle = Node->GetNodeHandle();
+			Frontend::FGraphHandle GraphHandle = NodeHandle->GetOwningGraph();
+
+			auto RemoveNodeLocation = [](FNodeHandle InNodeHandle, const FGuid& InNodeGuid)
+			{
+				FMetasoundFrontendNodeStyle Style = InNodeHandle->GetNodeStyle();
+				Style.Display.Locations.Remove(InNodeGuid);
+				InNodeHandle->SetNodeStyle(Style);
+			};	
+
+			auto RemoveNodeHandle = [] (FGraphHandle InGraphHandle, FNodeHandle InNodeHandle)
+			{
+				if (ensure(InGraphHandle->RemoveNode(*InNodeHandle)))
 				{
-					return ensure(Graph->RemoveNode(&InNode));
+					InGraphHandle->GetOwningDocument()->RemoveUnreferencedDependencies();
 				}
-			}
+			};
 
-			UMetasoundEditorGraphNode* Node = Cast<UMetasoundEditorGraphNode>(&InNode);
-			if (ensure(Node))
+			if (GraphHandle->IsValid())
 			{
-//				Need to split out delete ed vs. frontend impl which will happen in subsequent change.
-//				For now, just ignore removal failure as a hack as this gets called from ed graph sync.
-				EMetaSoundBuilderResult Result = EMetaSoundBuilderResult::Failed;
-				Builder.RemoveNode({ Node->GetNodeID() }, Result, bRemoveUnusedDependencies);
-// 				ensure(Result == EMetaSoundBuilderResult::Succeeded);
+				const EMetasoundFrontendClassType ClassType = NodeHandle->GetClassMetadata().GetType();
+				switch (ClassType)
+				{
+					// NodeHandle does not get removed in these cases as EdGraph Inputs/Outputs
+					// Frontend node is represented by the editor graph as a respective member
+					// (not a node) on the MetasoundGraph. Therefore, just the editor position
+					// data is removed.
+					case EMetasoundFrontendClassType::Output:
+					case EMetasoundFrontendClassType::Input:
+					{
+						RemoveNodeLocation(NodeHandle, InNode.NodeGuid);
+					}
+					break;
+					
+					// NodeHandle is only removed for variable accessors if the editor graph
+					// no longer contains nodes representing the given accessor on the MetasoundGraph.
+					// Therefore, just the editor position data is removed unless no location remains
+					// on the Frontend node.
+					case EMetasoundFrontendClassType::VariableAccessor:
+					case EMetasoundFrontendClassType::VariableDeferredAccessor:
+					{
+						RemoveNodeLocation(NodeHandle, InNode.NodeGuid);
+						if (NodeHandle->GetNodeStyle().Display.Locations.IsEmpty())
+						{
+							RemoveNodeHandle(GraphHandle, NodeHandle);
+						}
+					}
+					break;
+
+					case EMetasoundFrontendClassType::Graph:
+					case EMetasoundFrontendClassType::Literal:
+					case EMetasoundFrontendClassType::VariableMutator:
+					case EMetasoundFrontendClassType::Variable:
+					case EMetasoundFrontendClassType::External:
+					case EMetasoundFrontendClassType::Template:
+					default:
+					{
+						static_assert(static_cast<int32>(EMetasoundFrontendClassType::Invalid) == 10, "Possible missing MetasoundFrontendClassType switch case coverage.");
+						RemoveNodeHandle(GraphHandle, NodeHandle);
+					}
+					break;
+				}
 			}
 
 			return ensure(Graph->RemoveNode(&InNode));
@@ -1695,10 +1888,17 @@ namespace Metasound
 				InGraphNode.RemovePin(InGraphNode.Pins[i]);
 			}
 
-			FConstNodeHandle NodeHandle = InGraphNode.GetNodeHandle();
-			TArray<FConstInputHandle> InputHandles = NodeHandle->GetConstInputs();
-			NodeHandle->GetInputStyle().SortDefaults(InputHandles, [](const Frontend::FConstInputHandle& Handle) { return FGraphBuilder::GetDisplayName(*Handle); });
-			for (const FConstInputHandle& InputHandle : InputHandles)
+			// TODO: Make this a utility in Frontend (ClearInputLiterals())
+			FNodeHandle NodeHandle = InGraphNode.GetNodeHandle();
+			TArray<FInputHandle> Inputs = NodeHandle->GetInputs();
+			for (FInputHandle& Input : Inputs)
+			{
+				NodeHandle->ClearInputLiteral(Input->GetID());
+			}
+
+			TArray<FInputHandle> InputHandles = NodeHandle->GetInputs();
+			NodeHandle->GetInputStyle().SortDefaults(InputHandles, [](const Frontend::FInputHandle& Handle) { return FGraphBuilder::GetDisplayName(*Handle); });
+			for (const FInputHandle& InputHandle : InputHandles)
 			{
 				// Only add pins of the node if the connection is user modifiable. 
 				// Connections which the user cannot modify are controlled elsewhere.
@@ -1708,9 +1908,9 @@ namespace Metasound
 				}
 			}
 
-			TArray<FConstOutputHandle> OutputHandles = NodeHandle->GetConstOutputs();
-			NodeHandle->GetOutputStyle().SortDefaults(OutputHandles, [](const Frontend::FConstOutputHandle& Handle) { return FGraphBuilder::GetDisplayName(*Handle); });
-			for (const FConstOutputHandle& OutputHandle : OutputHandles)
+			TArray<FOutputHandle> OutputHandles = NodeHandle->GetOutputs();
+			NodeHandle->GetOutputStyle().SortDefaults(OutputHandles, [](const Frontend::FOutputHandle& Handle) { return FGraphBuilder::GetDisplayName(*Handle); });
+			for (const FOutputHandle& OutputHandle : OutputHandles)
 			{
 				// Only add pins of the node if the connection is user modifiable. 
 				// Connections which the user cannot modify are controlled elsewhere.
@@ -1825,7 +2025,7 @@ namespace Metasound
 				return false;
 			}
 
-			Frontend::FConstInputHandle PinInputHandle = GetConstInputHandleFromPin(&InEditorPin);
+			Frontend::FInputHandle PinInputHandle = GetInputHandleFromPin(&InEditorPin);
 			if (PinInputHandle->GetID() == InInputHandle->GetID())
 			{
 				return true;
@@ -1841,7 +2041,7 @@ namespace Metasound
 				return false;
 			}
 
-			Frontend::FConstOutputHandle PinOutputHandle = GetConstOutputHandleFromPin(&InEditorPin);
+			Frontend::FOutputHandle PinOutputHandle = GetOutputHandleFromPin(&InEditorPin);
 			if (PinOutputHandle->GetID() == InOutputHandle->GetID())
 			{
 				return true;
@@ -1921,7 +2121,7 @@ namespace Metasound
 		{
 			using namespace Metasound::Frontend;
 
-			if (InAssetBase.GetConstModifyContext().GetDocumentModified())
+			if (InAssetBase.GetModifyContext().GetDocumentModified())
 			{
 				return true;
 			}
@@ -1963,97 +2163,26 @@ namespace Metasound
 			return false;
 		}
 
-		bool FGraphBuilder::SynchronizeComments(UObject& InMetaSound)
+		bool FGraphBuilder::SynchronizeConnections(UObject& InMetaSound)
 		{
-			bool bModified = false;
+			using namespace Frontend;
+
+			bool bIsGraphDirty = false;
 
 			FMetasoundAssetBase* MetaSoundAsset = IMetasoundUObjectRegistry::Get().GetObjectAsAssetBase(&InMetaSound);
 			check(MetaSoundAsset);
 
-			UMetasoundEditorGraph* EditorGraph = CastChecked<UMetasoundEditorGraph>(MetaSoundAsset->GetGraph());
-			check(EditorGraph);
-
-			TMap<FGuid, UMetasoundEditorGraphCommentNode*> CommentIDToEdNode;
-			Algo::TransformIf(EditorGraph->Nodes, CommentIDToEdNode,
-				[](const UEdGraphNode* Node)
-				{
-					return Node->IsA<UMetasoundEditorGraphCommentNode>();
-				},
-				[](UEdGraphNode* Node)
-				{
-					UMetasoundEditorGraphCommentNode* CommentNode = CastChecked<UMetasoundEditorGraphCommentNode>(Node);
-					return TPair<FGuid, UMetasoundEditorGraphCommentNode*>(CommentNode->GetCommentID(), CommentNode);
-				});
-
-			const FMetasoundFrontendDocument& Document = MetaSoundAsset->GetConstDocumentChecked();
-			const TMap<FGuid, FMetaSoundFrontendGraphComment>& Comments = Document.RootGraph.Graph.Style.Comments;
-			for (const TPair<FGuid, FMetaSoundFrontendGraphComment>& Pair : Comments)
-			{
-				UMetasoundEditorGraphCommentNode* CommentNode = nullptr;
-				if (!CommentIDToEdNode.RemoveAndCopyValue(Pair.Key, CommentNode))
-				{
-					bModified = true;
-					constexpr bool bSelectNode = false;
-
-					// Can't use spawn node action because it modifies the transaction stack, so just generate from CDO and backport frontend data accordingly.
-					UClass* CommentClass = UMetasoundEditorGraphCommentNode::StaticClass();
-					check(CommentClass);
-					UMetasoundEditorGraphCommentNode* CommentCDO = CommentClass->GetDefaultObject<UMetasoundEditorGraphCommentNode>();
-					check(CommentCDO);
-
-					const UGraphEditorSettings* GraphEditorSettings = GetDefault<UGraphEditorSettings>();
-					check(GraphEditorSettings);
-
-					UMetasoundEditorGraphCommentNode* NewNode = CastChecked<UMetasoundEditorGraphCommentNode>(DuplicateObject<UEdGraphNode>(CommentCDO, EditorGraph));
-					NewNode->SetFlags(RF_Transactional);
-					EditorGraph->AddNode(NewNode, false, false);
-					NewNode->CreateNewGuid();
-					NewNode->bCommentBubbleVisible_InDetailsPanel = GraphEditorSettings->bShowCommentBubbleWhenZoomedOut;
-
-					// Pull position, color, etc. from the existing frontend data.
-					NewNode->CommentID = Pair.Key;
-					UMetasoundEditorGraphCommentNode::ConvertFromFrontendComment(Pair.Value, *NewNode);
-				}
-			}
-
-			// Remaining items are stale, so they are removed
-			bModified |= !CommentIDToEdNode.IsEmpty();
-			for (const TPair<FGuid, UMetasoundEditorGraphCommentNode*>& Comment : CommentIDToEdNode)
-			{
-				constexpr bool bMarkDirty = false;
-				constexpr bool bBreakAllLinks = true;
-				EditorGraph->RemoveNode(Comment.Value, bBreakAllLinks, bMarkDirty);
-			}
-
-			return bModified;
-		}
-
-		bool FGraphBuilder::SynchronizeConnections(UObject& InMetaSound)
-		{
-			using namespace Frontend;
+			FGraphHandle GraphHandle = MetaSoundAsset->GetRootGraphHandle();
 
-			bool bIsGraphDirty = false;
+			UMetasoundEditorGraph* EditorGraph = CastChecked<UMetasoundEditorGraph>(MetaSoundAsset->GetGraph());
 
 			TArray<UMetasoundEditorGraphNode*> EditorNodes;
-			{
-				FMetasoundAssetBase* MetaSoundAsset = IMetasoundUObjectRegistry::Get().GetObjectAsAssetBase(&InMetaSound);
-				check(MetaSoundAsset);
-				UMetasoundEditorGraph* EditorGraph = CastChecked<UMetasoundEditorGraph>(MetaSoundAsset->GetGraph());
-				EditorGraph->GetNodesOfClass(EditorNodes);
-			}
+			EditorGraph->GetNodesOfClass(EditorNodes);
 
-			TMap<FGuid, UMetasoundEditorGraphNode*> EditorNodesByFrontendID;
+			TMap<FGuid, TArray<UMetasoundEditorGraphNode*>> EditorNodesByFrontendID;
 			for (UMetasoundEditorGraphNode* EditorNode : EditorNodes)
 			{
-				const FGuid& NodeID = EditorNode->GetNodeID();
-				if (EditorNodesByFrontendID.Contains(NodeID))
-				{
-					UE_LOG(LogMetasoundEditor, Error, TEXT("Multiple editor nodes associated with FrontendDocument node with ID '%s'"), *NodeID.ToString());
-				}
-				else
-				{
-					EditorNodesByFrontendID.Add(NodeID, EditorNode);
-				}
+				EditorNodesByFrontendID.FindOrAdd(EditorNode->GetNodeID()).Add(EditorNode);
 			}
 
 			// Iterate through all nodes in metasound editor graph and synchronize connections.
@@ -2061,7 +2190,7 @@ namespace Metasound
 			{
 				bool bIsNodeDirty = false;
 
-				FConstNodeHandle Node = EditorNode->GetConstNodeHandle();
+				FConstNodeHandle Node = EditorNode->GetNodeHandle();
 
 				TArray<UEdGraphPin*> Pins = EditorNode->GetAllPins();
 				TArray<FConstInputHandle> NodeInputs = Node->GetConstInputs();
@@ -2107,30 +2236,25 @@ namespace Metasound
 						else if (!IsMatchingOutputHandleAndPin(OutputHandle, *MatchingPin->LinkedTo[0]))
 						{
 							// The wrong link exists.
-							constexpr bool bNotifyNodes = false;
-							constexpr bool bMarkDirty = false;
-							MatchingPin->BreakAllPinLinks(bNotifyNodes, bMarkDirty);
+							MatchingPin->BreakAllPinLinks();
 							bAddLink = true;
 						}
 
 						if (bAddLink)
 						{
 							const FGuid NodeID = OutputHandle->GetOwningNodeID();
-							const FText& OwningNodeName = EditorNode->GetDisplayName();
-							UMetasoundEditorGraphNode* OutputEditorNode = EditorNodesByFrontendID.FindRef(NodeID);
-							if (OutputEditorNode)
+							TArray<UMetasoundEditorGraphNode*>* OutputEditorNode = EditorNodesByFrontendID.Find(NodeID);
+							if (ensure(OutputEditorNode))
 							{
-								UEdGraphPin* OutputPin = OutputEditorNode->FindPinChecked(OutputHandle->GetName(), EEdGraphPinDirection::EGPD_Output);
-
-								UE_LOG(LogMetasoundEditor, Verbose, TEXT("Synchronizing Node '%s' Connection: Linking Pin '%s' to '%s'"), *OwningNodeName.ToString(), *MatchingPin->GetName(), *OutputPin->GetName());
+								if (ensure(!OutputEditorNode->IsEmpty()))
+								{
+									UEdGraphPin* OutputPin = (*OutputEditorNode)[0]->FindPinChecked(OutputHandle->GetName(), EEdGraphPinDirection::EGPD_Output);
+									const FText& OwningNodeName = EditorNode->GetDisplayName();
 
-								constexpr bool bAlwaysMarkDirty = false;
-								MatchingPin->MakeLinkTo(OutputPin, bAlwaysMarkDirty);
-								bIsNodeDirty = true;
-							}
-							else
-							{
-								UE_LOG(LogMetasoundEditor, Warning, TEXT("Failed to synchronize Frontend Node '%s' Connection: Pin '%s'"), *OwningNodeName.ToString(), *MatchingPin->GetName());
+									UE_LOG(LogMetasoundEditor, Verbose, TEXT("Synchronizing Node '%s' Connection: Linking Pin '%s' to '%s'"), *OwningNodeName.ToString(), *MatchingPin->GetName(), *OutputPin->GetName());
+									MatchingPin->MakeLinkTo(OutputPin);
+									bIsNodeDirty = true;
+								}
 							}
 						}
 					}
@@ -2139,9 +2263,7 @@ namespace Metasound
 						// No link should exist.
 						if (!MatchingPin->LinkedTo.IsEmpty())
 						{
-							constexpr bool bNotifyNodes = false;
-							constexpr bool bMarkDirty = false;
-							MatchingPin->BreakAllPinLinks(bNotifyNodes, bMarkDirty);
+							MatchingPin->BreakAllPinLinks();
 							const FText OwningNodeName = EditorNode->GetDisplayName();
 							const FText InputName = FGraphBuilder::GetDisplayName(*NodeInput);
 							UE_LOG(LogMetasoundEditor, Verbose, TEXT("Synchronizing Node '%s' Connection: Breaking all pin links to '%s'"), *OwningNodeName.ToString(), *InputName.ToString());
@@ -2169,7 +2291,7 @@ namespace Metasound
 
 			if (RecurseGetDocumentModified(*MetaSoundAsset))
 			{
-				TSet<FMetasoundAssetBase*> EditedReferencingMetaSounds;
+				TArray<FMetasoundAssetBase*> EditedReferencingMetaSounds;
 				if (GEditor)
 				{
 					TArray<UObject*> EditedAssets = GEditor->GetEditorSubsystem<UAssetEditorSubsystem>()->GetAllEditedAssets();
@@ -2191,7 +2313,7 @@ namespace Metasound
 				if (EditedReferencingMetaSounds.IsEmpty())
 				{
 					MetaSoundAsset->CacheRegistryMetadata();
-					GraphBuilderPrivate::SynchronizeGraphRecursively(InMetaSound);
+					const bool bEditorGraphModified = GraphBuilderPrivate::SynchronizeGraphRecursively(InMetaSound);
 					GraphBuilderPrivate::RecurseClearDocumentModified(*MetaSoundAsset);
 				}
 				else
@@ -2211,9 +2333,9 @@ namespace Metasound
 			return false;
 		}
 
-		bool FGraphBuilder::SynchronizeOutputNodes(UObject& InMetaSound)
+		bool FGraphBuilder::SynchronizeNodeMembers(UObject& InMetaSound)
 		{
-			TRACE_CPUPROFILER_EVENT_SCOPE(Metasound::Editor::FGraphBuilder::SynchronizeOutputNodes);
+			TRACE_CPUPROFILER_EVENT_SCOPE(Metasound::Editor::FGraphBuilder::SynchronizeNodeMembers);
 
 			using namespace Frontend;
 
@@ -2221,9 +2343,41 @@ namespace Metasound
 
 			FMetasoundAssetBase* MetaSoundAsset = IMetasoundUObjectRegistry::Get().GetObjectAsAssetBase(&InMetaSound);
 			check(MetaSoundAsset);
-			FConstGraphHandle GraphHandle = MetaSoundAsset->GetRootGraphHandle();
+			FGraphHandle GraphHandle = MetaSoundAsset->GetRootGraphHandle();
 			UMetasoundEditorGraph* EditorGraph = CastChecked<UMetasoundEditorGraph>(MetaSoundAsset->GetGraph());
 
+			TArray<UMetasoundEditorGraphInputNode*> InputNodes;
+			EditorGraph->GetNodesOfClassEx<UMetasoundEditorGraphInputNode>(InputNodes);
+			for (UMetasoundEditorGraphInputNode* Node : InputNodes)
+			{
+				check(Node);
+				FConstNodeHandle NodeHandle = Node->GetConstNodeHandle();
+				if (!NodeHandle->IsValid())
+				{
+					for (UEdGraphPin* Pin : Node->Pins)
+					{
+						check(Pin);
+
+						FConstClassInputAccessPtr ClassInputPtr = GraphHandle->FindClassInputWithName(Pin->PinName);
+						if (const FMetasoundFrontendClassInput* Input = ClassInputPtr.Get())
+						{
+							const FGuid& InitialID = Node->GetNodeID();
+							if (Node->GetNodeHandle()->GetID() != Input->NodeID)
+							{
+								Node->SetNodeID(Input->NodeID);
+
+								// Requery handle as the id has been fixed up
+								NodeHandle = Node->GetConstNodeHandle();
+								FText InputDisplayName = Node->GetDisplayName();
+								UE_LOG(LogMetasoundEditor, Verbose, TEXT("Editor Input Node '%s' interface versioned"), *InputDisplayName.ToString());
+
+								bEditorGraphModified = true;
+							}
+						}
+					}
+				}
+			}
+
 			TArray<UMetasoundEditorGraphOutputNode*> OutputNodes;
 			EditorGraph->GetNodesOfClassEx<UMetasoundEditorGraphOutputNode>(OutputNodes);
 			for (UMetasoundEditorGraphOutputNode* Node : OutputNodes)
@@ -2257,22 +2411,30 @@ namespace Metasound
 			return bEditorGraphModified;
 		}
 
-		void FGraphBuilder::SynchronizeNodes(UObject& InMetaSound)
+		bool FGraphBuilder::SynchronizeNodes(UObject& InMetaSound)
 		{
 			TRACE_CPUPROFILER_EVENT_SCOPE(Metasound::Editor::FGraphBuilder::SynchronizeNodes);
 
 			using namespace Frontend;
 
+			bool bEditorGraphModified = false;
+
 			// Get all external nodes from Frontend graph.  Input and output references will only be added/synchronized
 			// if required when synchronizing connections (as they are not required to inhabit editor graph).
 			FMetasoundAssetBase* MetaSoundAsset = IMetasoundUObjectRegistry::Get().GetObjectAsAssetBase(&InMetaSound);
 			check(MetaSoundAsset);
-			FConstGraphHandle GraphHandle = static_cast<const FMetasoundAssetBase*>(MetaSoundAsset)->GetRootGraphHandle();
-			TArray<FConstNodeHandle> FrontendNodes = GraphHandle->GetConstNodes();
+			FGraphHandle GraphHandle = MetaSoundAsset->GetRootGraphHandle();
+			TArray<FNodeHandle> FrontendNodes = GraphHandle->GetNodes();
 			UMetasoundEditorGraph* EditorGraph = CastChecked<UMetasoundEditorGraph>(MetaSoundAsset->GetGraph());
 			TArray<UMetasoundEditorGraphNode*> EditorNodes;
 			EditorGraph->GetNodesOfClass(EditorNodes);
 
+			TMap<FGuid, UMetasoundEditorGraphNode*> EditorNodesByEdNodeGuid;
+			for (UMetasoundEditorGraphNode* Node : EditorNodes)
+			{
+				EditorNodesByEdNodeGuid.Add(Node->NodeGuid, Node);
+			}
+
 			// Find existing array of editor nodes associated with Frontend node
 			struct FAssociatedNodes
 			{
@@ -2281,13 +2443,6 @@ namespace Metasound
 			};
 			TMap<FGuid, FAssociatedNodes> AssociatedNodes;
 
-			// Input nodes use template input nodes to visually distinguish location, so ignore them.
-			FrontendNodes.RemoveAllSwap([](const FConstNodeHandle& Node)
-			{
-				return Node->GetClassMetadata().GetType() == EMetasoundFrontendClassType::Input;
-			},
-			EAllowShrinking::No);
-
 			// Reverse iterate so paired nodes can safely be removed from the array.
 			for (int32 i = FrontendNodes.Num() - 1; i >= 0; i--)
 			{
@@ -2309,8 +2464,7 @@ namespace Metasound
 							AssociatedNodeData.Node = Node;
 						}
 
-						EditorNode->SyncLocationFromFrontendNode();
-						EditorNode->SyncCommentFromFrontendNode();
+						bEditorGraphModified |= SynchronizeNodeLocation(Node, *EditorNode);
 						AssociatedNodeData.EditorNodes.Add(EditorNode);
 						EditorNodes.RemoveAtSwap(j, EAllowShrinking::No);
 					}
@@ -2330,78 +2484,27 @@ namespace Metasound
 			// exist before attempting to synchronize connections.
 			for (UMetasoundEditorGraphNode* EditorNode : EditorNodes)
 			{
-				constexpr bool bMarkDirty = false;
-				constexpr bool bBreakAllLinks = true;
-				EditorGraph->RemoveNode(EditorNode, bBreakAllLinks, bMarkDirty);
+				bEditorGraphModified |= EditorGraph->RemoveNode(EditorNode);
 			}
 
 			// Add missing editor nodes marked as visible.
-			for (const FConstNodeHandle& Node : FrontendNodes)
+			for (const FNodeHandle& Node : FrontendNodes)
 			{
 				const FMetasoundFrontendNodeStyle& CurrentStyle = Node->GetNodeStyle();
-				for (const TPair<FGuid, FVector2D>& Location : CurrentStyle.Display.Locations)
+				if (CurrentStyle.Display.Locations.IsEmpty())
 				{
-					UMetasoundEditorGraphNode* NewGraphNode = nullptr;
-					const FMetasoundFrontendClassMetadata& ClassMetadata = Node->GetClassMetadata();
-					switch (ClassMetadata.GetType())
-					{
-						case EMetasoundFrontendClassType::External:
-						{
-							NewGraphNode = CastChecked<UMetasoundEditorGraphNode>(AddExternalNode(InMetaSound, Node->GetID(), ClassMetadata, false));
-						}
-						break;
-
-						case EMetasoundFrontendClassType::Template:
-						{
-							if (Node->GetClassMetadata().GetClassName() == FInputNodeTemplate::ClassName)
-							{
-								NewGraphNode = CastChecked<UMetasoundEditorGraphNode>(AddInputNode(InMetaSound, Node, false));
-							}
-							else
-							{
-								NewGraphNode = CastChecked<UMetasoundEditorGraphNode>(AddTemplateNode(InMetaSound, Node->GetID(), Node->GetClassMetadata(), false));
-							}
-						}
-						break;
-
-						case EMetasoundFrontendClassType::Output:
-						{
-							NewGraphNode = CastChecked<UMetasoundEditorGraphNode>(AddOutputNode(InMetaSound, Node, false));
-						}
-						break;
-
-						case EMetasoundFrontendClassType::VariableMutator:
-						case EMetasoundFrontendClassType::VariableAccessor:
-						case EMetasoundFrontendClassType::VariableDeferredAccessor:
-						case EMetasoundFrontendClassType::Variable:
-						{
-							NewGraphNode = CastChecked<UMetasoundEditorGraphNode>(AddVariableNode(InMetaSound, Node, false));
-						}
-						break;
-
-						case EMetasoundFrontendClassType::Invalid:
-
-						// Class type needs to be deprecated 
-						case EMetasoundFrontendClassType::Graph:
-
-						// Not supported in editor
-						case EMetasoundFrontendClassType::Literal:
+					continue;
+				}
 
-						// Since MetaSound Document v1.12 update, the editor uses template input nodes, so no direct node representation of inputs no longer exists
-						case EMetasoundFrontendClassType::Input: 
-						default:
-						{
-							checkNoEntry();
-							static_assert(static_cast<int32>(EMetasoundFrontendClassType::Invalid) == 10, "Possible missing FMetasoundFrontendClassType case coverage");
-						}
-						break;
-					}
+				FMetasoundFrontendNodeStyle NewStyle = CurrentStyle;
+				bEditorGraphModified = true;
 
-					if (ensure(NewGraphNode))
+				TArray<UMetasoundEditorGraphNode*> AddedNodes;
+				for (const TPair<FGuid, FVector2D>& Location : NewStyle.Display.Locations)
+				{
+					UMetasoundEditorGraphNode* NewNode = AddNode(InMetaSound, Node, Location.Value, false /* bInSelectNewNode */);
+					if (ensure(NewNode))
 					{
-						ensureMsgf(NewGraphNode->NodeGuid.IsValid(), TEXT("New editor NodeGuid must be valid."));
-						ensureMsgf(NewGraphNode->NodeGuid == Location.Key, TEXT("New editor NodeGuid must match location key"));
-
 						FAssociatedNodes& AssociatedNodeData = AssociatedNodes.FindOrAdd(Node->GetID());
 						if (AssociatedNodeData.Node->IsValid())
 						{
@@ -2412,9 +2515,17 @@ namespace Metasound
 							AssociatedNodeData.Node = Node;
 						}
 
-						AssociatedNodeData.EditorNodes.Add(NewGraphNode);
+						AddedNodes.Add(NewNode);
+						AssociatedNodeData.EditorNodes.Add(NewNode);
 					}
 				}
+
+				NewStyle.Display.Locations.Reset();
+				for (UMetasoundEditorGraphNode* EditorNode : AddedNodes)
+				{
+					NewStyle.Display.Locations.Add(EditorNode->NodeGuid, FVector2D(EditorNode->NodePosX, EditorNode->NodePosY));
+				}
+				Node->SetNodeStyle(NewStyle);
 			}
 
 			// Synchronize pins on node associations.
@@ -2422,9 +2533,11 @@ namespace Metasound
 			{
 				for (UMetasoundEditorGraphNode* EditorNode : IdNodePair.Value.EditorNodes)
 				{
-					SynchronizeNodePins(*EditorNode, IdNodePair.Value.Node);
+					bEditorGraphModified |= SynchronizeNodePins(*EditorNode, IdNodePair.Value.Node);
 				}
 			}
+
+			return bEditorGraphModified;
 		}
 
 		bool FGraphBuilder::SynchronizeNodePins(UMetasoundEditorGraphNode& InEditorNode, Frontend::FConstNodeHandle InNode, bool bRemoveUnusedPins, bool bLogChanges)
@@ -2606,7 +2719,7 @@ namespace Metasound
 
 			const FString OldValue = InPin.DefaultValue;
 
-			FConstInputHandle InputHandle = GetConstInputHandleFromPin(&InPin);
+			FInputHandle InputHandle = GetInputHandleFromPin(&InPin);
 			if (const FMetasoundFrontendLiteral* NodeDefaultLiteral = InputHandle->GetLiteral())
 			{
 				InPin.DefaultValue = NodeDefaultLiteral->ToString();
@@ -2659,27 +2772,6 @@ namespace Metasound
 				bEditorGraphModified = true;
 			}, EMetasoundFrontendClassType::Input);
 
-			// Remove stale inputs.
-			{
-				for (int32 Index = Graph->Inputs.Num() - 1; Index >= 0; --Index)
-				{
-					const UMetasoundEditorGraphInput* Input = Graph->Inputs[Index];
-					check(Input);
-					if (!Inputs.Contains(Input))
-					{
-						TArray<UMetasoundEditorGraphMemberNode*> InputNodes = Input->GetNodes();
-						for (UMetasoundEditorGraphNode* Node : InputNodes)
-						{
-							FGraphBuilder::DeleteMemberNodes(*Input);
-						}
-
-						bEditorGraphModified = true;
-						UE_LOG(LogMetasoundEditor, Verbose, TEXT("Synchronizing Outputs: Removing stale output '%s'."), *Input->GetName());
-						Graph->Inputs.RemoveAtSwap(Index, EAllowShrinking::No);
-					}
-				}
-			}
-
 			// Collect all editor graph outputs with corresponding frontend outputs. 
 			GraphHandle->IterateNodes([&](FNodeHandle NodeHandle)
 			{
@@ -2697,19 +2789,30 @@ namespace Metasound
 				bEditorGraphModified = true;
 			}, EMetasoundFrontendClassType::Output);
 
-			// Remove stale outputs.
+			// Collect editor inputs and outputs to remove which have no corresponding frontend input or output.
+			TArray<UMetasoundEditorGraphMember*> ToRemove;
+			Graph->IterateInputs([&](UMetasoundEditorGraphInput& Input)
 			{
-				for (int32 Index = Graph->Outputs.Num() - 1; Index >= 0; --Index)
+				if (!Inputs.Contains(&Input))
 				{
-					const UMetasoundEditorGraphOutput* Output = Graph->Outputs[Index];
-					check(Output);
-					if (!Outputs.Contains(Output))
-					{
-						bEditorGraphModified = true;
-						UE_LOG(LogMetasoundEditor, Verbose, TEXT("Synchronizing Outputs: Removing stale output '%s'."), *Output->GetName());
-						Graph->Outputs.RemoveAtSwap(Index, EAllowShrinking::No);
-					}
+					UE_LOG(LogMetasoundEditor, Verbose, TEXT("Synchronizing Inputs: Removing stale input '%s'."), *Input.GetName());
+					ToRemove.Add(&Input);
 				}
+			});
+			Graph->IterateOutputs([&](UMetasoundEditorGraphOutput& Output)
+			{
+				if (!Outputs.Contains(&Output))
+				{
+					UE_LOG(LogMetasoundEditor, Verbose, TEXT("Synchronizing Outputs: Removing stale output '%s'."), *Output.GetName());
+					ToRemove.Add(&Output);
+				}
+			});
+
+			// Remove stale inputs and outputs.
+			bEditorGraphModified |= !ToRemove.IsEmpty();
+			for (UMetasoundEditorGraphMember* GraphMember: ToRemove)
+			{
+				Graph->RemoveMember(*GraphMember);
 			}
 
 			UMetasoundEditorGraphMember* Member = nullptr;
@@ -2770,10 +2873,6 @@ namespace Metasound
 							}
 						}
 					}
-					else
-					{
-						Input->InitializeLiteral();
-					}
 				}
 			}, EMetasoundFrontendClassType::Input);
 
@@ -2784,7 +2883,7 @@ namespace Metasound
 				{
 					SynchronizeMemberDataType(*Output);
 
-					// Fix up members with no literal set
+					// Fix up corrupted assets with no literal set 
 					if (!Output->GetLiteral())
 					{
 						Output->InitializeLiteral();
@@ -2792,32 +2891,10 @@ namespace Metasound
 				}
 			}, EMetasoundFrontendClassType::Output);
 
-			// Enforce all variables exists in the edgraph
-			TSet<FGuid> VariableIDs;
-			Algo::Transform(Graph->Variables, VariableIDs, [](const TObjectPtr<UMetasoundEditorGraphVariable>& Variable) { return Variable->GetVariableID(); });
-			TArray<FVariableHandle> FrontendVariables = GraphHandle->GetVariables();
-			for (const FVariableHandle& VariableHandle : FrontendVariables)
-			{
-				const FGuid VariableID = VariableHandle->GetID();
-				if (VariableIDs.Remove(VariableID) == 0)
-				{
-					UMetasoundEditorGraphVariable* NewVariable = NewObject<UMetasoundEditorGraphVariable>(Graph, FName(), RF_Transactional);
-					if (ensure(NewVariable))
-					{
-						const FMetasoundFrontendLiteral DefaultLiteral = VariableHandle->GetLiteral();
-						NewVariable->InitMember(VariableHandle->GetDataType(), DefaultLiteral, VariableID);
-						Graph->Variables.Add(NewVariable);
-					}
-				}
-			}
-
 			// Remove empty entries
 			bEditorGraphModified |= Graph->Inputs.RemoveAllSwap([](const TObjectPtr<UMetasoundEditorGraphInput>& Input) { return !Input; }) > 0;
 			bEditorGraphModified |= Graph->Outputs.RemoveAllSwap([](const TObjectPtr<UMetasoundEditorGraphOutput>& Output) { return !Output; }) > 0;
-			bEditorGraphModified |= Graph->Variables.RemoveAllSwap([&VariableIDs](const TObjectPtr<UMetasoundEditorGraphVariable>& Variable)
-			{
-				return !Variable || VariableIDs.Contains(Variable->GetVariableID());
-			}) > 0;
+			bEditorGraphModified |= Graph->Variables.RemoveAllSwap([](const TObjectPtr<UMetasoundEditorGraphVariable>& Variable) { return !Variable; }) > 0;
 
 			return bEditorGraphModified;
 		}
diff --git a/Engine/Plugins/Runtime/Metasound/Source/MetasoundEditor/Private/MetasoundEditorGraphBuilder.h b/Engine/Plugins/Runtime/Metasound/Source/MetasoundEditor/Private/MetasoundEditorGraphBuilder.h
--- a/Engine/Plugins/Runtime/Metasound/Source/MetasoundEditor/Private/MetasoundEditorGraphBuilder.h
+++ b/Engine/Plugins/Runtime/Metasound/Source/MetasoundEditor/Private/MetasoundEditorGraphBuilder.h
@@ -20,7 +20,6 @@ class UEdGraph;
 class UEdGraphNode;
 class UEdGraphPin;
 class UMetaSoundPatch;
-class UMetasoundEditorGraphCommentNode;
 class UMetasoundEditorGraphExternalNode;
 class UMetasoundEditorGraphNode;
 class UMetasoundEditorGraphInputNode;
@@ -62,6 +61,12 @@ namespace Metasound
 			static const FText FunctionMenuName;
 			static const FText GraphMenuName;
 
+			static void InitGraphNode(Frontend::FNodeHandle& InNodeHandle, UMetasoundEditorGraphNode* NewGraphNode, UObject& InMetaSound);
+			static void InitGraphNodeIDFromNodeHandle(const Frontend::FConstNodeHandle& InNodeHandle, UMetasoundEditorGraphNode* NewGraphNode);
+
+			// Adds an EdGraph node to mirror the provided FNodeHandle.
+			static UMetasoundEditorGraphNode* AddNode(UObject& InMetaSound, Frontend::FNodeHandle InNodeHandle, FVector2D InLocation, bool bInSelectNewNode = true);
+
 			// Convenience functions for retrieving the editor for the given UObject
 			static TSharedPtr<FEditor> GetEditorForMetasound(const UObject& InMetaSound);
 			static TSharedPtr<FEditor> GetEditorForGraph(const UEdGraph& InEdGraph);
@@ -118,11 +123,12 @@ namespace Metasound
 			// Returns the PinName for an IInputController.
 			static FName GetPinName(const Frontend::IInputController& InFrontendInput);
 
-			// Adds a new EdGraph comment node associated with the given MetaSoundFrontendGraph comment ID
-			static UMetasoundEditorGraphCommentNode* CreateCommentNode(UObject& InMetaSound, bool bInSelectNewNode = true, FGuid InCommentID = FGuid::NewGuid());
+			// Adds a node handle to mirror the provided graph node and binds to it.  Does *NOT* mirror existing EdGraph connections
+			// nor does it remove existing bound Frontend Node (if set) from associated Frontend Graph.
+			static Frontend::FNodeHandle AddNodeHandle(UObject& InMetaSound, UMetasoundEditorGraphNode& InGraphNode);
 
-			// Adds an input template node to the given MetaSound object's document model and returns an associated node handle
-			static Frontend::FNodeHandle AddInputTemplateNodeHandle(UObject& InMetaSound, Frontend::FNodeHandle& InputNodeHandle);
+			// Adds a corresponding UMetasoundEditorGraphInputNode for the provided node handle.
+			static UMetasoundEditorGraphInputNode* AddInputNode(UObject& InMetaSound, Frontend::FNodeHandle InNodeHandle, FVector2D InLocation, bool bInSelectNewNode = true);
 
 			// Generates FNodeHandle for the given external node data. Does not bind or create EdGraph representation of given node.
 			static Frontend::FNodeHandle AddInputNodeHandle(
@@ -132,25 +138,17 @@ namespace Metasound
 				const FName* InNameBase = nullptr);
 
 			// Adds a corresponding UMetasoundEditorGraphExternalNode for the provided node handle.
-			static UMetasoundEditorGraphExternalNode* AddExternalNode(UObject& InMetaSound, const FGuid& InNodeID, const FMetasoundFrontendClassMetadata& InMetadata, bool bInSelectNewNode = true);
-
-			// Adds an editor graph node that corresponds with an instance of a node that is defined by an external MetaSound node class.
-			static UMetasoundEditorGraphExternalNode* AddTemplateNode(UObject& InMetaSound, const FGuid& InNodeID, const FMetasoundFrontendClassMetadata& InMetadata, bool bInSelectNewNode = true);
+			static UMetasoundEditorGraphExternalNode* AddExternalNode(UObject& InMetaSound, Frontend::FNodeHandle& InNodeHandle, FVector2D InLocation, bool bInSelectNewNode = true);
 
 			// Adds an externally-defined node with the given class info to both the editor and document graphs.
 			// Generates analogous FNodeHandle.
-			static UMetasoundEditorGraphExternalNode* AddExternalNode(UObject& InMetaSound, const FMetasoundFrontendClassMetadata& InMetadata, bool bInSelectNewNode = true);
-
-			static Frontend::FNodeHandle AddExternalNodeHandle(UObject& InMetaSound, const FMetasoundFrontendClassName& InClassName);
+			static UMetasoundEditorGraphExternalNode* AddExternalNode(UObject& InMetaSound, const FMetasoundFrontendClassMetadata& InMetadata, FVector2D InLocation, bool bInSelectNewNode = true);
 
 			// Adds an variable node with the given node handle to the editor graph.
-			static UMetasoundEditorGraphVariableNode* AddVariableNode(UObject& InMetaSound, const Frontend::FConstNodeHandle& InNodeHandle, bool bInSelectNewNode = true);
-
-			// Adds an input node to the editor graph that corresponds to the provided input template node handle.
-			static UMetasoundEditorGraphInputNode* AddInputNode(UObject& InMetaSound, const Frontend::FConstNodeHandle& InInputTemplateNode, bool bInSelectNewNode = true);
+			static UMetasoundEditorGraphVariableNode* AddVariableNode(UObject& InMetaSound, Frontend::FNodeHandle& InNodeHandle, FVector2D InLocation, bool bInSelectNewNode = true);
 
 			// Adds an output node to the editor graph that corresponds to the provided node handle.
-			static UMetasoundEditorGraphOutputNode* AddOutputNode(UObject& InMetaSound, const Frontend::FConstNodeHandle& InNodeHandle, bool bInSelectNewNode = true);
+			static UMetasoundEditorGraphOutputNode* AddOutputNode(UObject& InMetaSound, Frontend::FNodeHandle& InNodeHandle, FVector2D InLocation, bool bInSelectNewNode = true);
 
 			// Generates analogous FNodeHandle for the given internal node data. Does not bind nor create EdGraph representation of given node.
 			static Frontend::FNodeHandle AddOutputNodeHandle(UObject& InMetaSound, const FCreateNodeVertexParams& InParams, const FName* InNameBase = nullptr);
@@ -158,10 +156,6 @@ namespace Metasound
 			// Create a unique name for the variable.
 			static FName GenerateUniqueVariableName(const Frontend::FConstGraphHandle& InFrontendGraph, const FString& InBaseName);
 
-			// Convenience method for walking to the outermost object and transforming to a base MetaSound
-			static FMetasoundAssetBase& GetOutermostMetaSoundChecked(UObject& InSubObject);
-			static const FMetasoundAssetBase& GetOutermostConstMetaSoundChecked(const UObject& InSubObject);
-
 			// Adds a frontend variable to the root graph of the MetaSound
 			//
 			// @param InMetaSound - FMetasoundAssetBase derived UObject.
@@ -233,14 +227,20 @@ namespace Metasound
 			// Retrieves the proper pin color for the given PinType
 			static FLinearColor GetPinCategoryColor(const FEdGraphPinType& PinType);
 
+			// Initializes MetaSound with default inputs & outputs.
+			UE_DEPRECATED(5.3, "Use FMetaSoundFrontendDocumentBuilder::InitDocument instead.")
+			static void InitMetaSound(UObject& InMetaSound, const FString& InAuthor);
+
+			// Initializes a MetaSound Preset using the provided ReferencedMetaSound asset's
+			// root graph as the sole, encapsulated topology.
+			UE_DEPRECATED(5.3, "Use FMetaSoundFrontendDocumentBuilder::ConvertToPreset instead.")
+			static void InitMetaSoundPreset(UObject& InMetaSoundReferenced, UObject& InMetaSoundPreset);
+
 			// Rebuilds all editor node pins based on the provided node handle's class definition.
 			static void RebuildNodePins(UMetasoundEditorGraphNode& InGraphNode);
 
-			// Deletes all nodes from a parent graph that are linked to a given member.
-			static bool DeleteMemberNodes(const UMetasoundEditorGraphMember& InGraphMember);
-
 			// Deletes both the editor graph & frontend nodes from respective graphs
-			static bool DeleteNode(UEdGraphNode& InNode, bool bRemoveUnusedDependencies = true);
+			static bool DeleteNode(UEdGraphNode& InNode);
 
 			// Adds an Input UEdGraphPin to a UMetasoundEditorGraphNode
 			static UEdGraphPin* AddPinToNode(UMetasoundEditorGraphNode& InEditorNode, Frontend::FConstInputHandle InInputHandle);
@@ -259,20 +259,25 @@ namespace Metasound
 			static bool SynchronizeGraph(UObject& InMetaSound);
 
 			// Synchronizes editor nodes with frontend nodes, removing editor nodes that are not represented in the frontend, and adding editor nodes to represent missing frontend nodes.
-			static void SynchronizeNodes(UObject& InMetaSound);
+			//
+			// @return True if the UMetasoundEditorGraphNode was altered. False otherwise.
+			static bool SynchronizeNodes(UObject& InMetaSound);
 
-			// Synchronizes and reports to log whether or not an output node's associated FrontendNode ID has changed and therefore been updated through node versioning.
+			// Synchronizes and reports to log whether or not an editor member node's associated FrontendNode ID has changed and therefore been updated through node versioning.
 			//
 			// @return True if the UMetasoundEditorGraphNode was altered. False otherwise.
-			static bool SynchronizeOutputNodes(UObject& InMetaSound);
+			static bool SynchronizeNodeMembers(UObject& InMetaSound);
+
+			// Synchronizes node location data using handle's data
+			//
+			// @return True if the UMetasoundEditorGraphNode was altered. False otherwise.
+			static bool SynchronizeNodeLocation(const Frontend::FConstNodeHandle& InNode, UMetasoundEditorGraphNode& OutGraphNode);
 
 			// Adds and removes pins so that the UMetasoundEditorGraphNode matches the InNode.
 			//
 			// @return True if the UMetasoundEditorGraphNode was altered. False otherwise.
 			static bool SynchronizeNodePins(UMetasoundEditorGraphNode& InEditorNode, Frontend::FConstNodeHandle InNode, bool bRemoveUnusedPins = true, bool bLogChanges = true);
 
-			static bool SynchronizeComments(UObject& InMetaSound);
-
 			// Adds and removes connections so that the UEdGraph of the MetaSound has the same
 			// connections as the FMetasoundFrontendDocument graph.
 			//
diff --git a/Engine/Plugins/Runtime/Metasound/Source/MetasoundEditor/Private/MetasoundEditorGraphCommentNode.cpp b/Engine/Plugins/Runtime/Metasound/Source/MetasoundEditor/Private/MetasoundEditorGraphCommentNode.cpp
--- a/Engine/Plugins/Runtime/Metasound/Source/MetasoundEditor/Private/MetasoundEditorGraphCommentNode.cpp
+++ b/Engine/Plugins/Runtime/Metasound/Source/MetasoundEditor/Private/MetasoundEditorGraphCommentNode.cpp
@@ -6,7 +6,6 @@
 #include "Kismet2/Kismet2NameValidators.h"
 #include "Layout/SlateRect.h"
 #include "MetasoundBuilderSubsystem.h"
-#include "MetasoundEditorGraph.h"
 #include "MetasoundUObjectRegistry.h"
 #include "Styling/AppStyle.h"
 #include "Templates/Casts.h"
@@ -30,11 +29,6 @@ FMetasoundAssetBase& UMetasoundEditorGraphCommentNode::GetAssetChecked()
 	return *MetaSound;
 }
 
-bool UMetasoundEditorGraphCommentNode::CanUserDeleteNode() const
-{
-	return true;
-}
-
 void UMetasoundEditorGraphCommentNode::ConvertToFrontendComment(const UEdGraphNode_Comment& InEdNode, FMetaSoundFrontendGraphComment& OutComment)
 {
 	OutComment.bColorBubble = InEdNode.bColorCommentBubble != 0;
@@ -65,9 +59,31 @@ void UMetasoundEditorGraphCommentNode::PostEditChangeProperty(FPropertyChangedEv
 {
 	Super::PostEditChangeProperty(PropertyChangedEvent);
 
+	// TODO: Check property name and only change that which changed
 	FMetaSoundFrontendGraphComment Comment;
 	ConvertToFrontendComment(*this, Comment);
-	GetBuilderChecked().FindOrAddGraphComment(CommentID) = MoveTemp(Comment);
+
+	UMetaSoundBuilderBase& Builder = UMetaSoundBuilderSubsystem::GetChecked().AttachBuilderToAssetChecked(*GetOutermostObject());
+	Builder.FindOrAddGraphComment(CommentID) = MoveTemp(Comment);
+}
+
+void UMetasoundEditorGraphCommentNode::PostPlacedNewNode()
+{
+	Super::PostPlacedNewNode();
+
+	UMetaSoundBuilderBase& Builder = UMetaSoundBuilderSubsystem::GetChecked().AttachBuilderToAssetChecked(*GetOutermostObject());
+
+	// If the frontend version predates this new node, sync this node's comment to it...
+	if (const FMetaSoundFrontendGraphComment* Comment = Builder.FindGraphComment(CommentID))
+	{
+		NodeComment = Comment->Comment;
+	}
+
+	// otherwise, initialize the frontend version with the ed graph's default comment.
+	else
+	{
+		Builder.FindOrAddGraphComment(CommentID).Comment = NodeComment;
+	}
 }
 
 void UMetasoundEditorGraphCommentNode::ResizeNode(const FVector2D& NewSize)
@@ -75,33 +91,22 @@ void UMetasoundEditorGraphCommentNode::ResizeNode(const FVector2D& NewSize)
 	Super::ResizeNode(NewSize);
 	if (bCanResizeNode) 
 	{
-		UMetasoundEditorGraph& Graph = *CastChecked<UMetasoundEditorGraph>(GetGraph());
-		UMetaSoundBuilderBase& Builder = UMetaSoundBuilderSubsystem::GetChecked().AttachBuilderToAssetChecked(Graph.GetMetasoundChecked());
+		UMetaSoundBuilderBase& Builder = UMetaSoundBuilderSubsystem::GetChecked().AttachBuilderToAssetChecked(*GetOutermostObject());
 		Builder.FindOrAddGraphComment(CommentID).Size = FVector2D(NodeWidth, NodeHeight);
 	}
 }
 
-UMetaSoundBuilderBase& UMetasoundEditorGraphCommentNode::GetBuilderChecked() const
-{
-	UMetasoundEditorGraph* EdGraph = CastChecked<UMetasoundEditorGraph>(GetGraph());
-	return EdGraph->GetBuilderChecked();
-}
-
 FGuid UMetasoundEditorGraphCommentNode::GetCommentID() const
 {
 	return CommentID;
 }
 
-bool UMetasoundEditorGraphCommentNode::RemoveFromDocument() const
-{
-	return GetBuilderChecked().RemoveGraphComment(CommentID);
-}
-
 void UMetasoundEditorGraphCommentNode::SetBounds(const class FSlateRect& Rect)
 {
 	Super::SetBounds(Rect);
 
-	FMetaSoundFrontendGraphComment& FrontendComment = GetBuilderChecked().FindOrAddGraphComment(CommentID);
+	UMetaSoundBuilderBase& Builder = UMetaSoundBuilderSubsystem::GetChecked().AttachBuilderToAssetChecked(*GetOutermostObject());
+	FMetaSoundFrontendGraphComment& FrontendComment = Builder.FindOrAddGraphComment(CommentID);
 	FrontendComment.Position = FVector2D(NodePosX, NodePosY);
 	FrontendComment.Size = FVector2D(NodeWidth, NodeHeight);
 }
@@ -110,7 +115,7 @@ void UMetasoundEditorGraphCommentNode::OnRenameNode(const FString& NewName)
 {
 	Super::OnRenameNode(NewName);
 
-	UMetaSoundBuilderBase& Builder = GetBuilderChecked();
+	UMetaSoundBuilderBase& Builder = UMetaSoundBuilderSubsystem::GetChecked().AttachBuilderToAssetChecked(*GetOutermostObject());
 	Builder.FindOrAddGraphComment(CommentID).Comment = NewName;
 }
 #undef LOCTEXT_NAMESPACE
diff --git a/Engine/Plugins/Runtime/Metasound/Source/MetasoundEditor/Private/MetasoundEditorGraphInputNode.cpp b/Engine/Plugins/Runtime/Metasound/Source/MetasoundEditor/Private/MetasoundEditorGraphInputNode.cpp
--- a/Engine/Plugins/Runtime/Metasound/Source/MetasoundEditor/Private/MetasoundEditorGraphInputNode.cpp
+++ b/Engine/Plugins/Runtime/Metasound/Source/MetasoundEditor/Private/MetasoundEditorGraphInputNode.cpp
@@ -25,53 +25,6 @@
 
 #include UE_INLINE_GENERATED_CPP_BY_NAME(MetasoundEditorGraphInputNode)
 
-void UMetasoundEditorGraphInputNode::CacheTitle()
-{
-	using namespace Metasound::Frontend;
-
-	if (Input)
-	{
-		FConstNodeHandle NodeHandle = Input->GetNodeHandle();
-		CachedTitle = NodeHandle->GetDisplayTitle();
-	}
-	else
-	{
-		CachedTitle = FText::FromName(Breadcrumb.MemberName);
-	}
-}
-
-const FMetasoundEditorGraphVertexNodeBreadcrumb& UMetasoundEditorGraphInputNode::GetBreadcrumb() const
-{
-	return Breadcrumb;
-}
-
-void UMetasoundEditorGraphInputNode::CacheBreadcrumb()
-{
-	using namespace Metasound::Frontend;
-
-	Breadcrumb = { };
-
-	// Take data from associated input as pasted graph may not be same as local graph
-	// and associated input will not be copied with given node.  Need the following data
-	// to associate or create new associated input.
-	if (Input)
-	{
-		FConstNodeHandle NodeHandle = Input->GetConstNodeHandle();
-
-		Breadcrumb.MemberName = NodeHandle->GetNodeName();
-		Breadcrumb.ClassName = NodeHandle->GetClassMetadata().GetClassName();
-
-		FConstOutputHandle OutputHandle = NodeHandle->GetConstOutputs().Last();
-		Breadcrumb.AccessType = OutputHandle->GetVertexAccessType();
-		Breadcrumb.DataType = OutputHandle->GetDataType();
-
-		if (const UMetasoundEditorGraphMemberDefaultLiteral* Literal = Input->GetLiteral())
-		{
-			Breadcrumb.DefaultLiteral = Literal->GetDefault();
-		}
-	}
-}
-
 UMetasoundEditorGraphMember* UMetasoundEditorGraphInputNode::GetMember() const
 {
 	return Input;
@@ -79,15 +32,12 @@ UMetasoundEditorGraphMember* UMetasoundEditorGraphInputNode::GetMember() const
 
 FMetasoundFrontendClassName UMetasoundEditorGraphInputNode::GetClassName() const
 {
-	using namespace Metasound::Frontend;
-
-	if (Input)
+	if (ensure(Input))
 	{
-		FConstNodeHandle NodeHandle = Input->GetConstNodeHandle();
-		return NodeHandle->GetClassMetadata().GetClassName();
+		return Input->ClassName;
 	}
 
-	return Breadcrumb.ClassName;
+	return Super::GetClassName();
 }
 
 void UMetasoundEditorGraphInputNode::UpdatePreviewInstance(const Metasound::FVertexName& InParameterName, TScriptInterface<IAudioParameterControllerInterface>& InParameterInterface) const
@@ -103,7 +53,12 @@ void UMetasoundEditorGraphInputNode::UpdatePreviewInstance(const Metasound::FVer
 
 FGuid UMetasoundEditorGraphInputNode::GetNodeID() const
 {
-	return NodeID;
+	if (Input)
+	{
+		return Input->NodeID;
+	}
+
+	return Super::GetNodeID();
 }
 
 FLinearColor UMetasoundEditorGraphInputNode::GetNodeTitleColor() const
@@ -122,40 +77,6 @@ FSlateIcon UMetasoundEditorGraphInputNode::GetNodeTitleIcon() const
 	return FSlateIcon("MetaSoundStyle", NativeIconName);
 }
 
-void UMetasoundEditorGraphInputNode::GetPinHoverText(const UEdGraphPin& Pin, FString& OutHoverText) const
-{
-	using namespace Metasound::Editor;
-	using namespace Metasound::Frontend;
-
-	if (ensure(Pin.Direction == EGPD_Output)) // Should never display input pin for input node hover
-	{
-		if (ensure(Input))
-		{
-			FConstNodeHandle InputNode = Input->GetConstNodeHandle();
-			OutHoverText = InputNode->GetDescription().ToString();
-			if (ShowNodeDebugData())
-			{
-				FConstOutputHandle OutputHandle = FGraphBuilder::FindReroutedConstOutputHandleFromPin(&Pin);
-				OutHoverText = FString::Format(TEXT("{0}\nVertex Name: {1}\nDataType: {2}\nID: {3}"),
-				{
-					OutHoverText,
-					OutputHandle->GetName().ToString(),
-					OutputHandle->GetDataType().ToString(),
-					OutputHandle->GetID().ToString(),
-				});
-			}
-		}
-	}
-}
-
-void UMetasoundEditorGraphInputNode::ReconstructNode()
-{
-	using namespace Metasound::Editor;
-	using namespace Metasound::Frontend;
-
-	Super::ReconstructNode();
-}
-
 void UMetasoundEditorGraphInputNode::Validate(Metasound::Editor::FGraphNodeValidationResult& OutResult)
 {
 #if WITH_EDITOR
@@ -190,3 +111,12 @@ void UMetasoundEditorGraphInputNode::Validate(Metasound::Editor::FGraphNodeValid
 	}
 #endif // #if WITH_EDITOR
 }
+
+void UMetasoundEditorGraphInputNode::SetNodeID(FGuid InNodeID)
+{
+	if (Input)
+	{
+		Input->NodeID = InNodeID;
+	}
+}
+
