[Backout] - CL33084850

diff --git a/Engine/Plugins/Runtime/Metasound/Source/MetasoundEditor/Private/MetasoundEditorGraphNode.cpp b/Engine/Plugins/Runtime/Metasound/Source/MetasoundEditor/Private/MetasoundEditorGraphNode.cpp
--- a/Engine/Plugins/Runtime/Metasound/Source/MetasoundEditor/Private/MetasoundEditorGraphNode.cpp
+++ b/Engine/Plugins/Runtime/Metasound/Source/MetasoundEditor/Private/MetasoundEditorGraphNode.cpp
@@ -65,84 +65,20 @@ void UMetasoundEditorGraphNode::UpdateFrontendNodeLocation(const FVector2D& InLo
 {
 	using namespace Metasound::Frontend;
 	
-	ensureMsgf(NodeGuid.IsValid(), TEXT("Cannot update frontend node location prior to node guid being finalized."));
-
 	FNodeHandle NodeHandle = GetNodeHandle();
 	FMetasoundFrontendNodeStyle Style = NodeHandle->GetNodeStyle();
 	Style.Display.Locations.FindOrAdd(NodeGuid) = InLocation;
 	NodeHandle->SetNodeStyle(Style);
 }
 
-bool UMetasoundEditorGraphNode::ShowNodeDebugData()
-{
-	using namespace Metasound::Editor;
-
-	return GraphNodePrivate::ShowNodeDebugData != 0;
-}
-
-bool UMetasoundEditorGraphNode::RemoveFromDocument() const
-{
-	UMetaSoundBuilderBase& Builder = GetBuilderChecked();
-	EMetaSoundBuilderResult Result = EMetaSoundBuilderResult::Failed;
-	Builder.RemoveNode({ GetNodeID() }, Result);
-	return Result == EMetaSoundBuilderResult::Succeeded;
-}
-
 void UMetasoundEditorGraphNode::SetNodeLocation(const FVector2D& InLocation)
 {
-	NodePosX = UE::LWC::FloatToIntCastChecked<int32>(InLocation.X);
-	NodePosY = UE::LWC::FloatToIntCastChecked<int32>(InLocation.Y);
+	NodePosX = FMath::TruncToInt(InLocation.X);
+	NodePosY = FMath::TruncToInt(InLocation.Y);
 
 	UpdateFrontendNodeLocation(InLocation);
 }
 
-void UMetasoundEditorGraphNode::SyncCommentFromFrontendNode()
-{
-	using namespace Metasound::Frontend;
-
-	FConstNodeHandle NodeHandle = GetNodeHandle();
-	const FMetasoundFrontendNodeStyle& Style = NodeHandle->GetNodeStyle();
-	NodeComment = Style.Display.Comment;
-	bCommentBubbleMakeVisible = Style.Display.bCommentVisible;
-}
-
-bool UMetasoundEditorGraphNode::SyncLocationFromFrontendNode(bool bUpdateEditorNodeID)
-{
-	using namespace Metasound::Frontend;
-
-	FConstNodeHandle NodeHandle = GetConstNodeHandle();
-	const FMetasoundFrontendNodeStyle& Style = NodeHandle->GetNodeStyle();
-
-	if (!Style.Display.Locations.IsEmpty())
-	{
-		if (bUpdateEditorNodeID)
-		{
-			if (ensureMsgf(Style.Display.Locations.Num() == 1, TEXT("Style location map should only contain one location (legacy support for multiple location values should be migrated by this point)")))
-			{
-				for (const TPair<FGuid, FVector2D>& IDLocationPair : Style.Display.Locations)
-				{
-					if (ensureMsgf(IDLocationPair.Key.IsValid(), TEXT("Location cannot be updated for ed node with invalid guid")))
-					{
-						NodeGuid = IDLocationPair.Key;
-						NodePosX = UE::LWC::FloatToIntCastChecked<int32>(IDLocationPair.Value.X);
-						NodePosY = UE::LWC::FloatToIntCastChecked<int32>(IDLocationPair.Value.Y);
-						return true;
-					}
-				}
-			}
-		}
-		if (const FVector2D* Location = Style.Display.Locations.Find(NodeGuid))
-		{
-			NodePosX = UE::LWC::FloatToIntCastChecked<int32>(Location->X);
-			NodePosY = UE::LWC::FloatToIntCastChecked<int32>(Location->Y);
-
-			return true;
-		}
-	}
-
-	return false;
-}
-
 void UMetasoundEditorGraphNode::PostLoad()
 {
 	Super::PostLoad();
@@ -198,10 +134,13 @@ int32 UMetasoundEditorGraphNode::EstimateNodeWidth() const
 	}
 }
 
-UMetaSoundBuilderBase& UMetasoundEditorGraphNode::GetBuilderChecked() const
+const FMetaSoundFrontendDocumentBuilder& UMetasoundEditorGraphNode::GetFrontendBuilderChecked() const
 {
-	UMetasoundEditorGraph* EdGraph = CastChecked<UMetasoundEditorGraph>(GetGraph());
-	return UMetaSoundBuilderSubsystem::GetChecked().AttachBuilderToAssetChecked(EdGraph->GetMetasoundChecked());
+	UObject* Outermost = GetOutermostObject();
+	check(Outermost);
+
+	const UMetaSoundBuilderBase& BuilderBase = UMetaSoundBuilderSubsystem::GetChecked().AttachBuilderToAssetChecked(*Outermost);
+	return BuilderBase.GetConstBuilder();
 }
 
 const FMetasoundFrontendNode* UMetasoundEditorGraphNode::GetFrontendNode() const
@@ -229,12 +168,6 @@ const FMetasoundFrontendNode& UMetasoundEditorGraphNode::GetFrontendNodeChecked(
 	return *FrontendNode;
 }
 
-const FMetasoundEditorGraphNodeBreadcrumb& UMetasoundEditorGraphNode::GetBreadcrumb() const
-{
-	static const FMetasoundEditorGraphNodeBreadcrumb StubCrumb;
-	return StubCrumb;
-}
-
 UObject& UMetasoundEditorGraphNode::GetMetasoundChecked()
 {
 	UMetasoundEditorGraph* EdGraph = CastChecked<UMetasoundEditorGraph>(GetGraph());
@@ -250,7 +183,10 @@ const UObject& UMetasoundEditorGraphNode::GetMetasoundChecked() const
 Metasound::Frontend::FConstGraphHandle UMetasoundEditorGraphNode::GetConstRootGraphHandle() const
 {
 	const FMetasoundAssetBase* ConstMetasoundAsset = Metasound::IMetasoundUObjectRegistry::Get().GetObjectAsAssetBase(&GetMetasoundChecked());
-	return ConstMetasoundAsset->GetRootGraphHandle();
+	FMetasoundAssetBase* MetasoundAsset = const_cast<FMetasoundAssetBase*>(ConstMetasoundAsset);
+	check(MetasoundAsset);
+
+	return MetasoundAsset->GetRootGraphHandle();
 }
 
 Metasound::Frontend::FGraphHandle UMetasoundEditorGraphNode::GetRootGraphHandle() const
@@ -319,7 +255,7 @@ void UMetasoundEditorGraphNode::Validate(Metasound::Editor::FGraphNodeValidation
 
 #if WITH_EDITOR
 	// Validate that non-reroute inputs are connected to "real" outputs
-	if (GetBreadcrumb().ClassName != FRerouteNodeTemplate::ClassName)
+	if (GetClassName() != FRerouteNodeTemplate::ClassName)
 	{
 		for (UEdGraphPin* Pin : Pins)
 		{
@@ -333,7 +269,7 @@ void UMetasoundEditorGraphNode::Validate(Metasound::Editor::FGraphNodeValidation
 					{
 						if (UMetasoundEditorGraphExternalNode* ExternalNode = Cast<UMetasoundEditorGraphExternalNode>(ReroutedPin->GetOwningNode()))
 						{
-							if (ExternalNode->GetBreadcrumb().ClassName == FRerouteNodeTemplate::ClassName)
+							if (ExternalNode->GetClassName() == FRerouteNodeTemplate::ClassName)
 							{
 								FConstInputHandle InputHandle = FGraphBuilder::GetConstInputHandleFromPin(Pin);
 								const FString Msg = FString::Format(*GraphNodePrivate::MissingConcreteOutputConnectionFormat, { InputHandle->GetDisplayName().ToString() });
@@ -507,28 +443,17 @@ Metasound::Frontend::FDataTypeRegistryInfo UMetasoundEditorGraphNode::GetPinData
 	using namespace Metasound::Editor;
 	using namespace Metasound::Frontend;
 
- 	FDataTypeRegistryInfo DataTypeInfo;
+ 	Metasound::Frontend::FDataTypeRegistryInfo DataTypeInfo;
 
-	const FMetaSoundFrontendDocumentBuilder& Builder = GetBuilderChecked().GetConstBuilder();
-	const FMetasoundFrontendVertexHandle Handle = FGraphBuilder::GetPinVertexHandle(Builder, &InPin);
-	if (Handle.IsSet())
+	if (InPin.Direction == EGPD_Input)
 	{
-		if (InPin.Direction == EGPD_Input)
-		{
-			const FMetasoundFrontendVertex* Vertex = Builder.FindNodeInput(Handle.NodeID, Handle.VertexID);
-			if (Vertex)
-			{
-				IDataTypeRegistry::Get().GetDataTypeInfo(Vertex->TypeName, DataTypeInfo);
-			}
-		}
-		else // InPin.Direction == EGPD_Output
-		{
-			const FMetasoundFrontendVertex* Vertex = Builder.FindNodeOutput(Handle.NodeID, Handle.VertexID);
-			if (Vertex)
-			{
-				IDataTypeRegistry::Get().GetDataTypeInfo(Vertex->TypeName, DataTypeInfo);
-			}
-		}
+		FConstInputHandle Handle = FGraphBuilder::GetConstInputHandleFromPin(&InPin);
+		ensure(IDataTypeRegistry::Get().GetDataTypeInfo(Handle->GetDataType(), DataTypeInfo));
+	}
+	else // InPin.Direction == EGPD_Output
+	{
+		FConstOutputHandle Handle = FGraphBuilder::GetConstOutputHandleFromPin(&InPin);
+		ensure(IDataTypeRegistry::Get().GetDataTypeInfo(Handle->GetDataType(), DataTypeInfo));
 	}
 
 	return DataTypeInfo;
@@ -605,22 +530,6 @@ FString UMetasoundEditorGraphNode::GetPinMetaData(FName InPinName, FName InKey)
 	return Super::GetPinMetaData(InPinName, InKey);
 }
 
-void UMetasoundEditorGraphNode::OnUpdateCommentText(const FString& NewComment)
-{
-	using namespace Metasound::Frontend;
-
-	if (!NodeComment.Equals(NewComment))
-	{
-		const FScopedTransaction Transaction(LOCTEXT("CommentCommitted", "Comment Changed"));
-		Modify();
-		NodeComment = NewComment;
-		FNodeHandle NodeHandle = GetNodeHandle();
-		FMetasoundFrontendNodeStyle Style = NodeHandle->GetNodeStyle();
-		Style.Display.Comment = NewComment;
-		NodeHandle->SetNodeStyle(Style);
-	}
-}
-
 void UMetasoundEditorGraphNode::PreSave(FObjectPreSaveContext InSaveContext)
 {
 	using namespace Metasound::Editor;
@@ -769,7 +678,7 @@ void UMetasoundEditorGraphOutputNode::PinDefaultValueChanged(UEdGraphPin* InPin)
 		UObject& MetaSound = GetMetasoundChecked();
 		MetaSound.Modify();
 		
-		FConstInputHandle InputHandle = FGraphBuilder::GetConstInputHandleFromPin(InPin);
+		FInputHandle InputHandle = FGraphBuilder::GetInputHandleFromPin(InPin);
 		if (InputHandle->IsValid())
 		{
 			FMetasoundFrontendLiteral LiteralValue;
@@ -808,46 +717,6 @@ void UMetasoundEditorGraphOutputNode::PinDefaultValueChanged(UEdGraphPin* InPin)
 	}
 }
 
-void UMetasoundEditorGraphOutputNode::ReconstructNode()
-{
-	using namespace Metasound::Frontend;
-
-	if (!Output)
-	{
-		UMetasoundEditorGraph* Graph = CastChecked<UMetasoundEditorGraph>(GetGraph());
-		if (ensure(Graph))
-		{
-			Graph->IterateOutputs([&Graph, this](UMetasoundEditorGraphOutput& TestOutput)
-			{
-				FConstNodeHandle OutputHandle = TestOutput.GetConstNodeHandle();
-				FConstInputHandle TestInput = OutputHandle->GetConstInputs().Last();
-				const bool bTypeMatches = TestInput->GetDataType() == Breadcrumb.DataType;
-				const bool bAccessMatches = TestInput->GetVertexAccessType() == Breadcrumb.AccessType;
-				const bool bNameMatches = OutputHandle->GetNodeName() == Breadcrumb.MemberName;
-				if (bTypeMatches && bAccessMatches && bNameMatches)
-				{
-					Output = &TestOutput;
-				}
-			});
-		}
-	}
-
-	Super::ReconstructNode();
-}
-
-bool UMetasoundEditorGraphOutputNode::RemoveFromDocument() const
-{
-	if (Output)
-	{
-		UMetaSoundBuilderBase& Builder = GetBuilderChecked();
-		EMetaSoundBuilderResult Result = EMetaSoundBuilderResult::Failed;
-		Builder.RemoveGraphOutput(Output->GetMemberName(), Result);
-		return Result == EMetaSoundBuilderResult::Succeeded;
-	}
-
-	return false;
-}
-
 bool UMetasoundEditorGraphOutputNode::EnableInteractWidgets() const
 {
 	using namespace Metasound::Frontend;
@@ -865,6 +734,7 @@ bool UMetasoundEditorGraphOutputNode::EnableInteractWidgets() const
 
 void UMetasoundEditorGraphOutputNode::Validate(Metasound::Editor::FGraphNodeValidationResult& OutResult)
 {
+#if WITH_EDITOR
 	using namespace Metasound::Editor;
 	using namespace Metasound::Frontend;
 
@@ -893,43 +763,16 @@ void UMetasoundEditorGraphOutputNode::Validate(Metasound::Editor::FGraphNodeVali
 			}
 		}
 	}
-}
-
-const FMetasoundEditorGraphNodeBreadcrumb& UMetasoundEditorGraphOutputNode::GetBreadcrumb() const
-{
-	return Breadcrumb;
-}
-
-void UMetasoundEditorGraphOutputNode::CacheBreadcrumb()
-{
-	using namespace Metasound::Frontend;
-	FConstNodeHandle NodeHandle = GetConstNodeHandle();
-
-	Breadcrumb.MemberName = NodeHandle->GetNodeName();
-	Breadcrumb.ClassName = NodeHandle->GetClassMetadata().GetClassName();
-
-	FConstInputHandle InputHandle = NodeHandle->GetConstInputs().Last();
-	Breadcrumb.AccessType = InputHandle->GetVertexAccessType();
-	Breadcrumb.DataType = InputHandle->GetDataType();
-
-
-	if (const UMetasoundEditorGraphMemberDefaultLiteral* Literal = Output->GetLiteral())
-	{
-		Breadcrumb.DefaultLiteral = Literal->GetDefault();
-	}
+#endif // #if WITH_EDITOR
 }
 
 FMetasoundFrontendClassName UMetasoundEditorGraphOutputNode::GetClassName() const
 {
-	using namespace Metasound::Frontend;
-
-	if (Output)
+	if (ensure(Output))
 	{
-		FConstNodeHandle NodeHandle = Output->GetConstNodeHandle();
-		return NodeHandle->GetClassMetadata().GetClassName();
+		return Output->ClassName;
 	}
-
-	return Breadcrumb.ClassName;
+	return FMetasoundFrontendClassName();
 }
 
 FGuid UMetasoundEditorGraphOutputNode::GetNodeID() const
@@ -1007,15 +850,6 @@ bool UMetasoundEditorGraphExternalNode::CanAutoUpdate() const
 	return GetConstNodeHandle()->CanAutoUpdate(InterfaceUpdates);
 }
 
-void UMetasoundEditorGraphExternalNode::CacheBreadcrumb()
-{
-	using namespace Metasound::Frontend;
-	FConstNodeHandle NodeHandle = GetConstNodeHandle();
-	const FNodeRegistryKey RegistryKey(NodeHandle->GetClassMetadata());
-	Breadcrumb.bIsClassNative = FMetasoundFrontendRegistryContainer::Get()->IsNodeNative(RegistryKey);
-	Breadcrumb.ClassName = NodeHandle->GetClassMetadata().GetClassName();
-}
-
 void UMetasoundEditorGraphExternalNode::CacheTitle()
 {
 	using namespace Metasound::Editor;
@@ -1105,7 +939,7 @@ void UMetasoundEditorGraphExternalNode::Validate(Metasound::Editor::FGraphNodeVa
 			{
 #if WITH_EDITOR
 				FString Message;
-				if (!Template->HasRequiredConnections(GetBuilderChecked().GetConstBuilder(), GetNodeID(), &Message))
+				if (!Template->HasRequiredConnections(GetFrontendBuilderChecked(), GetNodeID(), &Message))
 				{
 					OutResult.SetMessage(EMessageSeverity::Warning, Message);
 				}
@@ -1191,8 +1025,8 @@ void UMetasoundEditorGraphExternalNode::Validate(Metasound::Editor::FGraphNodeVa
 		bool bWasRemoved = false;
 		if (Pin->Direction == EGPD_Input)
 		{
-			FConstInputHandle Input = FGraphBuilder::GetConstInputHandleFromPin(Pin);
-			bWasRemoved |= InterfaceUpdates.RemovedInputs.ContainsByPredicate([&Input](const FMetasoundFrontendClassInput* ClassInput)
+			FInputHandle Input = FGraphBuilder::GetInputHandleFromPin(Pin);
+			bWasRemoved |= InterfaceUpdates.RemovedInputs.ContainsByPredicate([&](const FMetasoundFrontendClassInput* ClassInput)
 			{
 				return Input->GetName() == ClassInput->Name && Input->GetDataType() == ClassInput->TypeName;
 			});
@@ -1200,8 +1034,8 @@ void UMetasoundEditorGraphExternalNode::Validate(Metasound::Editor::FGraphNodeVa
 
 		if (Pin->Direction == EGPD_Output)
 		{
-			FConstOutputHandle Output = FGraphBuilder::GetConstOutputHandleFromPin(Pin);
-			bWasRemoved |= InterfaceUpdates.RemovedOutputs.ContainsByPredicate([&Output](const FMetasoundFrontendClassOutput* ClassOutput)
+			FOutputHandle Output = FGraphBuilder::GetOutputHandleFromPin(Pin);
+			bWasRemoved |= InterfaceUpdates.RemovedOutputs.ContainsByPredicate([&](const FMetasoundFrontendClassOutput* ClassOutput)
 			{
 				return Output->GetName() == ClassOutput->Name && Output->GetDataType() == ClassOutput->TypeName;
 			});
@@ -1301,11 +1135,6 @@ void UMetasoundEditorGraphExternalNode::Validate(Metasound::Editor::FGraphNodeVa
 #endif // WITH_EDITOR
 }
 
-const FMetasoundEditorGraphNodeBreadcrumb& UMetasoundEditorGraphExternalNode::GetBreadcrumb() const
-{
-	return Breadcrumb;
-}
-
 FLinearColor UMetasoundEditorGraphExternalNode::GetNodeTitleColor() const
 {
 	if (const UMetasoundEditorSettings* EditorSettings = GetDefault<UMetasoundEditorSettings>())
@@ -1337,7 +1166,7 @@ bool UMetasoundEditorGraphExternalNode::ShouldDrawNodeAsControlPointOnly(int32&
 {
 	using namespace Metasound::Frontend;
 
-	if (GetBreadcrumb().ClassName == FRerouteNodeTemplate::ClassName)
+	if (GetClassName() == FRerouteNodeTemplate::ClassName)
 	{
 		OutInputPinIndex = 0;
 		OutOutputPinIndex = 1;
@@ -1347,25 +1176,6 @@ bool UMetasoundEditorGraphExternalNode::ShouldDrawNodeAsControlPointOnly(int32&
 	return Super::ShouldDrawNodeAsControlPointOnly(OutInputPinIndex, OutOutputPinIndex);
 }
 
-void UMetasoundEditorGraphVariableNode::CacheBreadcrumb()
-{
-	using namespace Metasound::Frontend;
-
-	Breadcrumb = { };
-	if (Variable)
-	{
-		FConstVariableHandle VariableHandle = Variable->GetConstVariableHandle();
-
-		Breadcrumb.MemberName = VariableHandle->GetName();
-		Breadcrumb.DataType = VariableHandle->GetDataType();
-	}
-}
-
-const FMetasoundEditorGraphNodeBreadcrumb& UMetasoundEditorGraphVariableNode::GetBreadcrumb() const
-{
-	return Breadcrumb;
-}
-
 UMetasoundEditorGraphMember* UMetasoundEditorGraphVariableNode::GetMember() const
 {
 	return Variable;
@@ -1436,7 +1246,7 @@ void UMetasoundEditorGraphVariableNode::PinDefaultValueChanged(UEdGraphPin* Pin)
 		UObject& MetaSound = GetMetasoundChecked();
 		MetaSound.Modify();
 
-		FConstInputHandle InputHandle = FGraphBuilder::GetConstInputHandleFromPin(Pin);
+		FInputHandle InputHandle = FGraphBuilder::GetInputHandleFromPin(Pin);
 		if (InputHandle->IsValid())
 		{
 			FMetasoundFrontendLiteral LiteralValue;
diff --git a/Engine/Plugins/Runtime/Metasound/Source/MetasoundEditor/Private/MetasoundEditorGraphNodeFactory.h b/Engine/Plugins/Runtime/Metasound/Source/MetasoundEditor/Private/MetasoundEditorGraphNodeFactory.h
--- a/Engine/Plugins/Runtime/Metasound/Source/MetasoundEditor/Private/MetasoundEditorGraphNodeFactory.h
+++ b/Engine/Plugins/Runtime/Metasound/Source/MetasoundEditor/Private/MetasoundEditorGraphNodeFactory.h
@@ -3,7 +3,6 @@
 
 #include "EdGraphUtilities.h"
 #include "EdGraph/EdGraphNode.h"
-#include "MetasoundEditorGraphCommentNode.h"
 #include "MetasoundEditorGraphNode.h"
 #include "MetasoundEditorGraphSchema.h"
 #include "NodeTemplates/MetasoundFrontendNodeTemplateReroute.h"
@@ -34,10 +33,9 @@ class FMetasoundGraphNodeFactory : public FGraphPanelNodeFactory
 			}
 			return SNew(SMetaSoundGraphNode, InNode);
 		}
-		else if (UMetasoundEditorGraphCommentNode* CommentNode = Cast<UMetasoundEditorGraphCommentNode>(InNode))
+		else if (UEdGraphNode_Comment* CommentNode = Cast<UEdGraphNode_Comment>(InNode))
 		{
-			const UEdGraphSchema* EdGraphSchema = CommentNode->GetSchema();
-			if (EdGraphSchema && EdGraphSchema->IsA(UMetasoundEditorGraphSchema::StaticClass()))
+			if (CommentNode->GetSchema()->IsA(UMetasoundEditorGraphSchema::StaticClass()))
 			{
 				return SNew(SMetasoundGraphNodeComment, CommentNode);
 			}
diff --git a/Engine/Plugins/Runtime/Metasound/Source/MetasoundEditor/Private/MetasoundEditorGraphSchema.cpp b/Engine/Plugins/Runtime/Metasound/Source/MetasoundEditor/Private/MetasoundEditorGraphSchema.cpp
--- a/Engine/Plugins/Runtime/Metasound/Source/MetasoundEditor/Private/MetasoundEditorGraphSchema.cpp
+++ b/Engine/Plugins/Runtime/Metasound/Source/MetasoundEditor/Private/MetasoundEditorGraphSchema.cpp
@@ -23,7 +23,6 @@
 #include "MetasoundEditorCommands.h"
 #include "MetasoundEditorGraph.h"
 #include "MetasoundEditorGraphBuilder.h"
-#include "MetasoundEditorGraphCommentNode.h"
 #include "MetasoundEditorGraphInputNode.h"
 #include "MetasoundEditorGraphNode.h"
 #include "MetasoundEditorGraphValidation.h"
@@ -38,13 +37,11 @@
 #include "MetasoundLiteral.h"
 #include "MetasoundNodeInterface.h"
 #include "MetasoundStandardNodesCategories.h"
-#include "MetasoundFrontendNodeTemplateRegistry.h"
 #include "MetasoundUObjectRegistry.h"
 #include "MetasoundVariableNodes.h"
 #include "MetasoundVertex.h"
 #include "NodeTemplates/MetasoundFrontendNodeTemplateReroute.h"
 #include "ScopedTransaction.h"
-#include "Settings/EditorStyleSettings.h"
 #include "Styling/SlateStyleRegistry.h"
 #include "Toolkits/ToolkitManager.h"
 #include "ToolMenus.h"
@@ -165,12 +162,12 @@ PRAGMA_ENABLE_DEPRECATION_WARNINGS
 
 				if (FromPin->Direction == EGPD_Input)
 				{
-					FConstInputHandle InputHandle = FGraphBuilder::GetConstInputHandleFromPin(FromPin);
+					FInputHandle InputHandle = FGraphBuilder::GetInputHandleFromPin(FromPin);
 					for (UEdGraphPin* Pin : NewGraphNode.Pins)
 					{
 						if (Pin->Direction == EGPD_Output)
 						{
-							FConstOutputHandle OutputHandle = FGraphBuilder::GetConstOutputHandleFromPin(Pin);
+							FOutputHandle OutputHandle = FGraphBuilder::GetOutputHandleFromPin(Pin);
 							if (OutputHandle->IsValid() && 
 								InputHandle->CanConnectTo(*OutputHandle).Connectable == FConnectability::EConnectable::Yes)
 							{
@@ -185,12 +182,12 @@ PRAGMA_ENABLE_DEPRECATION_WARNINGS
 
 				if (FromPin->Direction == EGPD_Output)
 				{
-					FConstOutputHandle OutputHandle = FGraphBuilder::GetConstOutputHandleFromPin(FromPin);
+					FOutputHandle OutputHandle = FGraphBuilder::GetOutputHandleFromPin(FromPin);
 					for (UEdGraphPin* Pin : NewGraphNode.Pins)
 					{
 						if (Pin->Direction == EGPD_Input)
 						{
-							FConstInputHandle InputHandle = FGraphBuilder::GetConstInputHandleFromPin(Pin);
+							FInputHandle InputHandle = FGraphBuilder::GetInputHandleFromPin(Pin);
 							if (InputHandle->IsValid() && 
 								InputHandle->CanConnectTo(*OutputHandle).Connectable == FConnectability::EConnectable::Yes)
 							{
@@ -288,13 +285,10 @@ PRAGMA_ENABLE_DEPRECATION_WARNINGS
 						constexpr bool bPostSubTransaction = false;
 						Variable->SetMemberName(NodeName, bPostSubTransaction);
 
-						FConstNodeHandle NodeHandle = FGraphBuilder::AddVariableNodeHandle(ParentMetaSound, Variable->GetVariableID(), ClassName.ToNodeClassName());
+						FNodeHandle NodeHandle = FGraphBuilder::AddVariableNodeHandle(ParentMetaSound, Variable->GetVariableID(), ClassName.ToNodeClassName());
 
-						if (UMetasoundEditorGraphVariableNode* NewGraphNode = FGraphBuilder::AddVariableNode(ParentMetaSound, NodeHandle))
+						if (UMetasoundEditorGraphVariableNode* NewGraphNode = FGraphBuilder::AddVariableNode(ParentMetaSound, NodeHandle, InLocation))
 						{
-							NewGraphNode->UpdateFrontendNodeLocation(InLocation);
-							NewGraphNode->SyncLocationFromFrontendNode();
-
 							// Set the literal using the new value if connecting to an input
 							if (FromPin.Direction == EGPD_Input)
 							{
@@ -457,11 +451,9 @@ UEdGraphNode* FMetasoundGraphSchemaAction_NewNode::PerformAction(UEdGraph* Paren
 	ParentMetasound.Modify();
 	ParentGraph->Modify();
 
-	if (UMetasoundEditorGraphExternalNode* NewGraphNode = FGraphBuilder::AddExternalNode(ParentMetasound, ClassMetadata, bSelectNewNode))
+	if (UMetasoundEditorGraphExternalNode* NewGraphNode = FGraphBuilder::AddExternalNode(ParentMetasound, ClassMetadata, Location, bSelectNewNode))
 	{
 		NewGraphNode->Modify();
-		NewGraphNode->UpdateFrontendNodeLocation(Location);
-		NewGraphNode->SyncLocationFromFrontendNode();
 		SchemaPrivate::TryConnectNewNodeToMatchingDataTypePin(*NewGraphNode, FromPin);
 		SchemaPrivate::SelecteNodeInEditor(*MetaSoundGraph, *NewGraphNode);
 		return NewGraphNode;
@@ -506,8 +498,8 @@ UEdGraphNode* FMetasoundGraphSchemaAction_NewInput::PerformAction(UEdGraph* Pare
 		return nullptr;
 	}
 
-	FNodeHandle InputNodeHandle = Input->GetNodeHandle();
-	if (!ensure(InputNodeHandle->IsValid()))
+	FNodeHandle NodeHandle = Input->GetNodeHandle();
+	if (!ensure(NodeHandle->IsValid()))
 	{
 		return nullptr;
 	}
@@ -517,12 +509,9 @@ UEdGraphNode* FMetasoundGraphSchemaAction_NewInput::PerformAction(UEdGraph* Pare
 	MetasoundGraph->Modify();
 	Input->Modify();
 
-	FConstNodeHandle TemplateNodeHandle = FGraphBuilder::AddInputTemplateNodeHandle(ParentMetasound, InputNodeHandle);
-	if (UMetasoundEditorGraphNode* NewGraphNode = FGraphBuilder::AddInputNode(ParentMetasound, TemplateNodeHandle))
+	if (UMetasoundEditorGraphInputNode* NewGraphNode = FGraphBuilder::AddInputNode(ParentMetasound, NodeHandle, InLocation))
 	{
 		NewGraphNode->Modify();
-		NewGraphNode->UpdateFrontendNodeLocation(InLocation);
-		NewGraphNode->SyncLocationFromFrontendNode();
 		SchemaPrivate::TryConnectNewNodeToMatchingDataTypePin(*NewGraphNode, FromPin);
 		FGraphBuilder::RegisterGraphWithFrontend(ParentMetasound);
 		return NewGraphNode;
@@ -564,17 +553,14 @@ UEdGraphNode* FMetasoundGraphSchemaAction_PromoteToInput::PerformAction(UEdGraph
 	// The promoted input must have the same vertex access type in order for it to be connectible
 	const FCreateNodeVertexParams VertexParams = { InputHandle->GetDataType(), InputHandle->GetVertexAccessType() };
 
-	FNodeHandle InputNodeHandle = FGraphBuilder::AddInputNodeHandle(ParentMetasound, VertexParams, &DefaultValue, &InputName);
-	if (ensure(InputNodeHandle->IsValid()))
+	FNodeHandle NodeHandle = FGraphBuilder::AddInputNodeHandle(ParentMetasound, VertexParams, &DefaultValue, &InputName);
+	if (ensure(NodeHandle->IsValid()))
 	{
-		UMetasoundEditorGraphInput* Input = MetasoundGraph->FindOrAddInput(InputNodeHandle);
+		UMetasoundEditorGraphInput* Input = MetasoundGraph->FindOrAddInput(NodeHandle);
 		if (ensure(Input))
 		{
-			FConstNodeHandle TemplateNodeHandle = FGraphBuilder::AddInputTemplateNodeHandle(ParentMetasound, InputNodeHandle);
-			if (UMetasoundEditorGraphNode* NewGraphNode = FGraphBuilder::AddInputNode(ParentMetasound, TemplateNodeHandle))
+			if (UMetasoundEditorGraphInputNode* NewGraphNode = FGraphBuilder::AddInputNode(ParentMetasound, NodeHandle, InLocation))
 			{
-				NewGraphNode->UpdateFrontendNodeLocation(InLocation);
-				NewGraphNode->SyncLocationFromFrontendNode();
 				UEdGraphNode* EdGraphNode = CastChecked<UEdGraphNode>(NewGraphNode);
 
 				if (ensure(SchemaPrivate::TryConnectNewNodeToMatchingDataTypePin(*EdGraphNode, FromPin)))
@@ -605,7 +591,7 @@ UEdGraphNode* FMetasoundGraphSchemaAction_PromoteToVariable_AccessorNode::Perfor
 	using namespace Metasound::Editor;
 	using namespace Metasound::Frontend;
 
-	FConstInputHandle InputHandle = FGraphBuilder::GetConstInputHandleFromPin(FromPin);
+	FInputHandle InputHandle = FGraphBuilder::GetInputHandleFromPin(FromPin);
 	if (!ensure(InputHandle->IsValid()))
 	{
 		return nullptr;
@@ -637,7 +623,7 @@ UEdGraphNode* FMetasoundGraphSchemaAction_PromoteToVariable_DeferredAccessorNode
 	using namespace Metasound::Editor;
 	using namespace Metasound::Frontend;
 
-	FConstInputHandle InputHandle = FGraphBuilder::GetConstInputHandleFromPin(FromPin);
+	FInputHandle InputHandle = FGraphBuilder::GetInputHandleFromPin(FromPin);
 	if (!ensure(InputHandle->IsValid()))
 	{
 		return nullptr;
@@ -669,7 +655,7 @@ UEdGraphNode* FMetasoundGraphSchemaAction_PromoteToVariable_MutatorNode::Perform
 	using namespace Metasound::Editor;
 	using namespace Metasound::Frontend;
 
-	FConstOutputHandle OutputHandle = FGraphBuilder::GetConstOutputHandleFromPin(FromPin);
+	FOutputHandle OutputHandle = FGraphBuilder::GetOutputHandleFromPin(FromPin);
 	if (!ensure(OutputHandle->IsValid()))
 	{
 		return nullptr;
@@ -715,13 +701,13 @@ UEdGraphNode* FMetasoundGraphSchemaAction_NewOutput::PerformAction(UEdGraph* Par
 	UMetasoundEditorGraph* MetasoundGraph = CastChecked<UMetasoundEditorGraph>(ParentGraph);
 	UObject& ParentMetasound = MetasoundGraph->GetMetasoundChecked();
 
-	const UMetasoundEditorGraphOutput* Output = MetasoundGraph->FindOutput(NodeID);
+	UMetasoundEditorGraphOutput* Output = MetasoundGraph->FindOutput(NodeID);
 	if (!ensure(Output))
 	{
 		return nullptr;
 	}
 
-	FConstNodeHandle NodeHandle = Output->GetConstNodeHandle();
+	FNodeHandle NodeHandle = Output->GetNodeHandle();
 	if (!ensure(NodeHandle->IsValid()))
 	{
 		return nullptr;
@@ -731,10 +717,8 @@ UEdGraphNode* FMetasoundGraphSchemaAction_NewOutput::PerformAction(UEdGraph* Par
 	ParentMetasound.Modify();
 	ParentGraph->Modify();
 
-	if (UMetasoundEditorGraphOutputNode* NewGraphNode = FGraphBuilder::AddOutputNode(ParentMetasound, NodeHandle, bSelectNewNode))
+	if (UMetasoundEditorGraphOutputNode* NewGraphNode = FGraphBuilder::AddOutputNode(ParentMetasound, NodeHandle, Location, bSelectNewNode))
 	{
-		NewGraphNode->UpdateFrontendNodeLocation(Location);
-		NewGraphNode->SyncLocationFromFrontendNode();
 		SchemaPrivate::TryConnectNewNodeToMatchingDataTypePin(*NewGraphNode, FromPin);
 		FGraphBuilder::RegisterGraphWithFrontend(ParentMetasound);
 		return NewGraphNode;
@@ -758,7 +742,7 @@ UEdGraphNode* FMetasoundGraphSchemaAction_PromoteToOutput::PerformAction(UEdGrap
 	using namespace Metasound::Editor;
 	using namespace Metasound::Frontend;
 
-	FConstOutputHandle OutputHandle = FGraphBuilder::GetConstOutputHandleFromPin(FromPin);
+	FOutputHandle OutputHandle = FGraphBuilder::GetOutputHandleFromPin(FromPin);
 	if (!ensure(OutputHandle->IsValid()))
 	{
 		return nullptr;
@@ -774,17 +758,14 @@ UEdGraphNode* FMetasoundGraphSchemaAction_PromoteToOutput::PerformAction(UEdGrap
 	const FVertexName NewNodeName = FGraphBuilder::GenerateUniqueNameByClassType(ParentMetasound, EMetasoundFrontendClassType::Output, OutputName);
 	const FCreateNodeVertexParams VertexParams = { OutputHandle->GetDataType(), OutputHandle->GetVertexAccessType() };
 
-	FConstNodeHandle NodeHandle = FGraphBuilder::AddOutputNodeHandle(ParentMetasound, VertexParams, &NewNodeName);
+	FNodeHandle NodeHandle = FGraphBuilder::AddOutputNodeHandle(ParentMetasound, VertexParams, &NewNodeName);
 	if (ensure(NodeHandle->IsValid()))
 	{
 		UMetasoundEditorGraphOutput* Output = MetasoundGraph->FindOrAddOutput(NodeHandle);
 		if (ensure(Output))
 		{
-			if (UMetasoundEditorGraphOutputNode* NewGraphNode = FGraphBuilder::AddOutputNode(ParentMetasound, NodeHandle))
+			if (UMetasoundEditorGraphOutputNode* NewGraphNode = FGraphBuilder::AddOutputNode(ParentMetasound, NodeHandle, InLocation))
 			{ 
-				NewGraphNode->UpdateFrontendNodeLocation(InLocation);
-				NewGraphNode->SyncLocationFromFrontendNode();
-
 				UEdGraphNode* EdGraphNode = CastChecked<UEdGraphNode>(NewGraphNode);
 
 				if (ensure(SchemaPrivate::TryConnectNewNodeToMatchingDataTypePin(*EdGraphNode, FromPin)))
@@ -822,14 +803,12 @@ UEdGraphNode* FMetasoundGraphSchemaAction_NewVariableNode::PerformAction(UEdGrap
 				MetasoundGraph->Modify();
 				Variable->Modify();
 				
-				FConstNodeHandle FrontendNode = CreateFrontendVariableNode(MetasoundGraph->GetGraphHandle(), VariableID);
+				FNodeHandle FrontendNode = CreateFrontendVariableNode(MetasoundGraph->GetGraphHandle(), VariableID);
 				if (ensure(FrontendNode->IsValid()))
 				{
-					if (UMetasoundEditorGraphVariableNode* NewGraphNode = FGraphBuilder::AddVariableNode(*ParentMetasound, FrontendNode, bSelectNewNode))
+					if (UMetasoundEditorGraphVariableNode* NewGraphNode = FGraphBuilder::AddVariableNode(*ParentMetasound, FrontendNode, Location, bSelectNewNode))
 					{
 						NewGraphNode->Modify();
-						NewGraphNode->UpdateFrontendNodeLocation(Location);
-						NewGraphNode->SyncLocationFromFrontendNode();
 						SchemaPrivate::TryConnectNewNodeToMatchingDataTypePin(*NewGraphNode, FromPin);
 						return NewGraphNode;
 					}
@@ -949,16 +928,13 @@ UEdGraphNode* FMetasoundGraphSchemaAction_NewReroute::PerformAction(UEdGraph* Pa
 		Params.InputsToConnect.Add(FromVertexDataType);
 	}
 
-	FConstNodeHandle NodeHandle = MetaSoundAsset->GetRootGraphHandle()->AddTemplateNode(*RerouteTemplate, MoveTemp(Params));
+	FNodeHandle NodeHandle = MetaSoundAsset->GetRootGraphHandle()->AddTemplateNode(*RerouteTemplate, MoveTemp(Params));
 
-	if (UMetasoundEditorGraphExternalNode* NewGraphNode = FGraphBuilder::AddTemplateNode(ParentMetasound, NodeHandle->GetID(), NodeHandle->GetClassMetadata(), bSelectNewNode))
+	if (UMetasoundEditorGraphExternalNode* NewGraphNode = FGraphBuilder::AddExternalNode(ParentMetasound, NodeHandle, Location, bSelectNewNode))
 	{
 		NewGraphNode->Modify();
-		NewGraphNode->UpdateFrontendNodeLocation(Location);
-		NewGraphNode->SyncLocationFromFrontendNode();
-
 		SchemaPrivate::TryConnectNewNodeToMatchingDataTypePin(*NewGraphNode, FromPin);
-		MetaSoundAsset->GetModifyContext().AddNodeIDsModified({ NewGraphNode->GetNodeID()});
+		MetaSoundGraph->GetModifyContext().AddNodeIDsModified({ NewGraphNode->GetNodeID()});
 
 		TSharedPtr<FEditor> ParentEditor = FGraphBuilder::GetEditorForMetasound(ParentMetasound);
 		if (ParentEditor.IsValid() && bSelectNewNode)
@@ -986,39 +962,23 @@ UEdGraphNode* FMetasoundGraphSchemaAction_NewComment::PerformAction(UEdGraph* Pa
 {
 	using namespace Metasound::Editor;
 
-	UMetasoundEditorGraph* MetaSoundGraph = CastChecked<UMetasoundEditorGraph>(ParentGraph);
-	UObject& MetaSound = MetaSoundGraph->GetMetasoundChecked();
-
 	const FScopedTransaction Transaction(LOCTEXT("AddNewOutputNode1", "Add Comment to MetaSound Graph"));
-	MetaSoundGraph->Modify();
-	MetaSound.Modify();
+	ParentGraph->Modify();
+
+	UEdGraphNode_Comment* CommentTemplate = NewObject<UEdGraphNode_Comment>();
 
-	// Must cache bounds prior to comment creation as call selects new node and invalidates original selection
 	FSlateRect Bounds;
+	FVector2D SpawnLocation = Location;
 	TSharedPtr<FEditor> MetasoundEditor = FGraphBuilder::GetEditorForGraph(*ParentGraph);
-	const bool bUseBounds = MetasoundEditor.IsValid() && MetasoundEditor->GetBoundsForSelectedNodes(Bounds, 50.0f);
 
-	if (UMetasoundEditorGraphCommentNode* NewComment = FGraphBuilder::CreateCommentNode(MetaSound, bSelectNewNode))
+	if (MetasoundEditor.IsValid() && MetasoundEditor->GetBoundsForSelectedNodes(Bounds, 50.0f))
 	{
-		if (bUseBounds)
-		{
-			NewComment->SetBounds(Bounds);
-		}
-		else
-		{
-			NewComment->NodePosX = Location.X;
-			NewComment->NodePosY = Location.Y;
-			NewComment->SnapToGrid(GetDefault<UEditorStyleSettings>()->GridSnapSize);
-
-			// Applies new snapped node data to frontend doc data
-			FMetaSoundFrontendGraphComment& FrontendComment = MetaSoundGraph->GetBuilderChecked().FindOrAddGraphComment(NewComment->GetCommentID());
-			UMetasoundEditorGraphCommentNode::ConvertToFrontendComment(*NewComment, FrontendComment);
-		}
-
-		return NewComment;
+		CommentTemplate->SetBounds(Bounds);
+		SpawnLocation.X = CommentTemplate->NodePosX;
+		SpawnLocation.Y = CommentTemplate->NodePosY;
 	}
 
-	return nullptr;
+	return FEdGraphSchemaAction_NewNode::SpawnNodeFromTemplate<UEdGraphNode_Comment>(ParentGraph, CommentTemplate, SpawnLocation);
 }
 
 const FSlateBrush* FMetasoundGraphSchemaAction_NewComment::GetIconBrush() const
@@ -1306,8 +1266,8 @@ const FPinConnectionResponse UMetasoundEditorGraphSchema::CanCreateConnection(co
 		}
 	}
 
-	Frontend::FConstInputHandle InputHandle = Editor::FGraphBuilder::GetConstInputHandleFromPin(InputPin);
-	Frontend::FConstOutputHandle OutputHandle = Editor::FGraphBuilder::GetConstOutputHandleFromPin(OutputPin);
+	Frontend::FInputHandle InputHandle = Editor::FGraphBuilder::GetInputHandleFromPin(InputPin);
+	Frontend::FOutputHandle OutputHandle = Editor::FGraphBuilder::GetOutputHandleFromPin(OutputPin);
 
 	const bool bInputValid = InputHandle->IsValid();
 	const bool bOutputValid = OutputHandle->IsValid();
@@ -1427,7 +1387,6 @@ void UMetasoundEditorGraphSchema::OnPinConnectionDoubleCicked(UEdGraphPin* PinA,
 
 bool UMetasoundEditorGraphSchema::TryCreateConnection(UEdGraphPin* PinA, UEdGraphPin* PinB) const
 {
-	using namespace Metasound::Editor;
 	using namespace Metasound::Frontend;
 
 	if (!ensure(PinA && PinB))
@@ -1462,12 +1421,12 @@ bool UMetasoundEditorGraphSchema::TryCreateConnection(UEdGraphPin* PinA, UEdGrap
 		return false;
 	}
 
-	if (!FGraphBuilder::ConnectNodes(*InputPin, *OutputPin, false /* bConnectEdPins */))
+	if (!Metasound::Editor::FGraphBuilder::ConnectNodes(*InputPin, *OutputPin, false /* bConnectEdPins */))
 	{
 		return false;
 	}
 
-	FGraphBuilder::GetOutermostMetaSoundChecked(*Graph).GetModifyContext().SetDocumentModified();
+	Graph->GetModifyContext().SetDocumentModified();
 
 	return true;
 }
@@ -1516,14 +1475,15 @@ bool UMetasoundEditorGraphSchema::SafeDeleteNodeFromGraph(UEdGraph* Graph, UEdGr
 {
 	using namespace Metasound::Editor;
 
-	UMetasoundEditorGraph* MetaSoundGraph = Cast<UMetasoundEditorGraph>(Graph);
-	if (!InNodeToDelete || !MetaSoundGraph || InNodeToDelete->GetGraph() != Graph)
+	if (!InNodeToDelete || !Graph || InNodeToDelete->GetGraph() != Graph)
 	{
 		return false;
 	}
 
-	UObject& MetaSound = MetaSoundGraph->GetMetasoundChecked();
-	MetaSound.Modify();
+	if (UObject* MetaSound = Graph->GetOutermostObject())
+	{
+		MetaSound->Modify();
+	}
 	Graph->Modify();
 
 	return FGraphBuilder::DeleteNode(*InNodeToDelete);
@@ -1541,13 +1501,13 @@ bool UMetasoundEditorGraphSchema::ShouldHidePinDefaultValue(UEdGraphPin* Pin) co
 
 	if (Pin->Direction == EGPD_Input)
 	{
-		FConstInputHandle InputHandle = FGraphBuilder::GetConstInputHandleFromPin(Pin);
+		FInputHandle InputHandle = FGraphBuilder::GetInputHandleFromPin(Pin);
 		return !InputHandle->GetOwningNode()->GetClassStyle().Display.bShowLiterals;
 	}
 
 	if (Pin->Direction == EGPD_Output)
 	{
-		FConstOutputHandle OutputHandle = FGraphBuilder::GetConstOutputHandleFromPin(Pin);
+		FOutputHandle OutputHandle = FGraphBuilder::GetOutputHandleFromPin(Pin);
 		return !OutputHandle->GetOwningNode()->GetClassStyle().Display.bShowLiterals;
 	}
 
@@ -1564,8 +1524,7 @@ FText UMetasoundEditorGraphSchema::GetPinDisplayName(const UEdGraphPin* Pin) con
 
 	UMetasoundEditorGraphNode* Node = CastChecked<UMetasoundEditorGraphNode>(Pin->GetOwningNode());
 	FConstNodeHandle NodeHandle = Node->GetConstNodeHandle();
-	const FMetasoundFrontendClassMetadata& ClassMetadata = NodeHandle->GetClassMetadata();
-	const EMetasoundFrontendClassType ClassType = ClassMetadata.GetType();
+	const EMetasoundFrontendClassType ClassType = NodeHandle->GetClassMetadata().GetType();
 
 	switch (ClassType)
 	{
@@ -1590,6 +1549,7 @@ FText UMetasoundEditorGraphSchema::GetPinDisplayName(const UEdGraphPin* Pin) con
 
 		case EMetasoundFrontendClassType::Literal:
 		case EMetasoundFrontendClassType::External:
+		case EMetasoundFrontendClassType::Template:
 		{
 			if (Pin->Direction == EGPD_Input)
 			{
@@ -1611,33 +1571,6 @@ FText UMetasoundEditorGraphSchema::GetPinDisplayName(const UEdGraphPin* Pin) con
 			return Super::GetPinDisplayName(Pin);
 		}
 
-		case EMetasoundFrontendClassType::Template:
-		{
-			const INodeTemplate* Template = INodeTemplateRegistry::Get().FindTemplate(ClassMetadata.GetClassName());
-			if (ensure(Template))
-			{
-				if (Pin->Direction == EGPD_Input)
-				{
-					FConstInputHandle InputHandle = NodeHandle->GetConstInputWithVertexName(Pin->GetFName());
-					if (InputHandle->IsValid())
-					{
-						return Template->GetInputPinDisplayName(*InputHandle);
-					}
-				}
-				else
-				{
-					FConstOutputHandle OutputHandle = NodeHandle->GetConstOutputWithVertexName(Pin->GetFName());
-					if (OutputHandle->IsValid())
-					{
-						return Template->GetOutputPinDisplayName(*OutputHandle);
-					}
-				}
-			}
-
-			return Super::GetPinDisplayName(Pin);
-		}
-		break;
-
 		case EMetasoundFrontendClassType::Graph:
 		case EMetasoundFrontendClassType::Invalid:
 		default:
@@ -1874,7 +1807,7 @@ void UMetasoundEditorGraphSchema::DroppedAssetsOnGraph(const TArray<FAssetData>&
 			}
 
 			// This may not be necessary as dropping an asset on the graph may load it, thus triggering the registration from the MetaSoundAssetManager.
-			const FMetasoundFrontendDocument& DroppedDoc = DroppedMetaSoundAsset->GetConstDocumentChecked();
+			const FMetasoundFrontendDocument& DroppedDoc = DroppedMetaSoundAsset->GetDocumentChecked();
 			const FNodeRegistryKey RegistryKey = FNodeRegistryKey(DroppedDoc.RootGraph);
 			if (ensure(RegistryKey.IsValid()))
 			{
@@ -1886,10 +1819,7 @@ void UMetasoundEditorGraphSchema::DroppedAssetsOnGraph(const TArray<FAssetData>&
 			FMetasoundFrontendClass Class;
 			if (ensure(ISearchEngine::Get().FindClassWithHighestVersion(DroppedDoc.RootGraph.Metadata.GetClassName(), Class)))
 			{
-				UMetasoundEditorGraphNode* NewGraphNode = Metasound::Editor::FGraphBuilder::AddExternalNode(MetaSound, Class.Metadata);
-				NewGraphNode->UpdateFrontendNodeLocation(GraphPosition);
-				NewGraphNode->SyncLocationFromFrontendNode();
-
+				Metasound::Editor::FGraphBuilder::AddExternalNode(MetaSound, Class.Metadata, GraphPosition);
 				bTransactionSucceeded = true;
 			}
 		}
diff --git a/Engine/Plugins/Runtime/Metasound/Source/MetasoundEditor/Private/MetasoundEditorSubsystem.cpp b/Engine/Plugins/Runtime/Metasound/Source/MetasoundEditor/Private/MetasoundEditorSubsystem.cpp
--- a/Engine/Plugins/Runtime/Metasound/Source/MetasoundEditor/Private/MetasoundEditorSubsystem.cpp
+++ b/Engine/Plugins/Runtime/Metasound/Source/MetasoundEditor/Private/MetasoundEditorSubsystem.cpp
@@ -10,7 +10,6 @@
 #include "MetasoundEditorSettings.h"
 #include "MetasoundFactory.h"
 #include "MetasoundUObjectRegistry.h"
-#include "NodeTemplates/MetasoundFrontendNodeTemplateInput.h"
 #include "Sound/SoundSourceBusSend.h"
 #include "Sound/SoundSubmixSend.h"
 
@@ -65,6 +64,10 @@ TScriptInterface<IMetaSoundDocumentInterface> UMetaSoundEditorSubsystem::BuildTo
 				const bool bIsSource = &BuilderUClass == UMetaSoundSource::StaticClass();
 				if (InBuilder->IsPreset())
 				{
+					FMetasoundAssetBase* PresetAsset = IMetasoundUObjectRegistry::Get().GetObjectAsAssetBase(NewMetaSound);
+					check(PresetAsset);
+					PresetAsset->ConformObjectDataToInterfaces();
+
 					// Only use referenced UObject's SoundWave settings for sources if not overridden 
 					if (TemplateSoundWave == nullptr && bIsSource)
 					{
@@ -82,11 +85,7 @@ TScriptInterface<IMetaSoundDocumentInterface> UMetaSoundEditorSubsystem::BuildTo
 				}
 			}
 
-			UMetaSoundBuilderBase& NewDocBuilder = UMetaSoundBuilderSubsystem::GetChecked().AttachBuilderToAssetChecked(*NewMetaSound);
-
-			EMetaSoundBuilderResult InjectResult = EMetaSoundBuilderResult::Failed;
-			constexpr bool bForceNodeCreation = true;
-			NewDocBuilder.InjectInputTemplateNodes(bForceNodeCreation, InjectResult);
+			InitEdGraph(*NewMetaSound);
 
 			if (!bWasRooted)
 			{
@@ -106,44 +105,6 @@ TScriptInterface<IMetaSoundDocumentInterface> UMetaSoundEditorSubsystem::BuildTo
 	return nullptr;
 }
 
-bool UMetaSoundEditorSubsystem::BindMemberMetadata(
-	UMetasoundEditorGraphMember& InMember,
-	TSubclassOf<UMetasoundEditorGraphMemberDefaultLiteral> LiteralClass,
-	UMetasoundEditorGraphMemberDefaultLiteral* TemplateObject)
-{
-	const FGuid& MemberID = InMember.GetMemberID();
-
-	UMetasoundEditorGraph* Graph = InMember.GetOwningGraph();
-	check(Graph);
-	UObject& MetaSound = Graph->GetMetasoundChecked();
-	UMetaSoundBuilderBase& Builder = UMetaSoundBuilderSubsystem::GetChecked().AttachBuilderToAssetChecked(MetaSound);
-
-	if (TemplateObject)
-	{
-		Builder.ClearMemberMetadata(MemberID);
-	}
-	else
-	{
-		if (UMetaSoundFrontendMemberMetadata* Literal = Builder.FindMemberMetadata(MemberID))
-		{
-			InMember.Literal = CastChecked<UMetasoundEditorGraphMemberDefaultLiteral>(Literal);
-			return false;
-		}
-	}
-
-	if (UMetasoundEditorGraphMemberDefaultLiteral* NewLiteral = NewObject<UMetasoundEditorGraphMemberDefaultLiteral>(&MetaSound, LiteralClass, FName(), RF_Transactional, TemplateObject))
-	{
-		NewLiteral->MemberID = MemberID;
-
-		Builder.SetMemberMetadata(*NewLiteral);
-		InMember.Literal = NewLiteral;
-		return true;
-	}
-
-	checkNoEntry();
-	return false;
-}
-
 UMetaSoundEditorSubsystem& UMetaSoundEditorSubsystem::GetChecked()
 {
 	checkf(GEditor, TEXT("Cannot access UMetaSoundEditorSubsystem without editor loaded"));
@@ -186,12 +147,10 @@ void UMetaSoundEditorSubsystem::InitAsset(UObject& InNewMetaSound, UObject* InRe
 
 	TScriptInterface<IMetaSoundDocumentInterface> DocInterface = &InNewMetaSound;
 	FMetaSoundFrontendDocumentBuilder Builder(DocInterface);
-
 	Builder.InitDocument();
+#if WITH_EDITORONLY_DATA
 	Builder.InitNodeLocations();
-
-	constexpr bool bForceNodeCreation = true;
-	FInputNodeTemplate::GetChecked().Inject(Builder, bForceNodeCreation);
+#endif // WITH_EDITORONLY_DATA
 
 	const FString& Author = GetDefaultAuthor();
 	Builder.SetAuthor(Author);
@@ -206,22 +165,27 @@ void UMetaSoundEditorSubsystem::InitAsset(UObject& InNewMetaSound, UObject* InRe
 		TScriptInterface<IMetaSoundDocumentInterface> ReferencedDocInterface = InReferencedMetaSound;
 		Builder.ConvertToPreset(ReferencedDocInterface->GetConstDocument());
 
+		// Update asset object data from interfaces 
+		FMetasoundAssetBase* PresetAsset = IMetasoundUObjectRegistry::Get().GetObjectAsAssetBase(&InNewMetaSound);
+		check(PresetAsset);
+		PresetAsset->ConformObjectDataToInterfaces();
+
 		// Copy sound wave settings to preset for sources
 		if (&ReferencedDocInterface->GetBaseMetaSoundUClass() == UMetaSoundSource::StaticClass())
 		{
 			SetSoundWaveSettingsFromTemplate(*CastChecked<USoundWave>(&InNewMetaSound), *CastChecked<USoundWave>(InReferencedMetaSound));
 		}
 	}
+
+	// Initial graph generation is not something to be managed by the transaction
+	// stack, so don't track dirty state until after initial setup if necessary.
+	InitEdGraph(InNewMetaSound);
 }
 
 void UMetaSoundEditorSubsystem::InitEdGraph(UObject& InMetaSound)
 {
 	using namespace Metasound;
 	using namespace Metasound::Editor;
-	using namespace Metasound::Frontend;
-
-	TScriptInterface<IMetaSoundDocumentInterface> DocInterface = &InMetaSound;
-	const FMetasoundFrontendClassName& ClassName = DocInterface->GetConstDocument().RootGraph.Metadata.GetClassName();
 
 	FMetasoundAssetBase* MetaSoundAsset = IMetasoundUObjectRegistry::Get().GetObjectAsAssetBase(&InMetaSound);
 	checkf(MetaSoundAsset, TEXT("EdGraph can only be initialized on registered MetaSoundAsset type"));
diff --git a/Engine/Plugins/Runtime/Metasound/Source/MetasoundEditor/Private/MetasoundNodeDetailCustomization.cpp b/Engine/Plugins/Runtime/Metasound/Source/MetasoundEditor/Private/MetasoundNodeDetailCustomization.cpp
--- a/Engine/Plugins/Runtime/Metasound/Source/MetasoundEditor/Private/MetasoundNodeDetailCustomization.cpp
+++ b/Engine/Plugins/Runtime/Metasound/Source/MetasoundEditor/Private/MetasoundNodeDetailCustomization.cpp
@@ -66,13 +66,10 @@ namespace Metasound
 			static const FText ConstructorPinTooltip = LOCTEXT("ConstructorPinTooltip",
 				"Whether this input or output is a constructor pin. Constructor values are only read on construction (on play), and are not dynamically updated at runtime.");
 
-
-			// Retrieves the data type info if the literal property's member is found. Returns if the associated member is found, false if not.
-			bool GetDataTypeFromElementPropertyHandle(TSharedPtr<IPropertyHandle> ElementPropertyHandle, Frontend::FDataTypeRegistryInfo& OutDataTypeInfo)
+			void GetDataTypeFromElementPropertyHandle(TSharedPtr<IPropertyHandle> ElementPropertyHandle, Frontend::FDataTypeRegistryInfo& OutDataTypeInfo)
 			{
 				using namespace Frontend;
 
-				OutDataTypeInfo = { };
 				TArray<UObject*>OuterObjects;
 				ElementPropertyHandle->GetOuterObjects(OuterObjects);
 				if (OuterObjects.Num() == 1)
@@ -80,7 +77,7 @@ namespace Metasound
 					UObject* Outer = OuterObjects.Last();
 					if (const UMetasoundEditorGraphMemberDefaultLiteral* DefaultLiteral = Cast<UMetasoundEditorGraphMemberDefaultLiteral>(Outer))
 					{
-						if (const UMetasoundEditorGraphMember* Member = DefaultLiteral->FindMember())
+						if (const UMetasoundEditorGraphMember* Member = Cast<UMetasoundEditorGraphMember>(DefaultLiteral->GetOuter()))
 						{
 							FName DataTypeName = Member->GetDataType();
 							ensure(IDataTypeRegistry::Get().GetDataTypeInfo(DataTypeName, OutDataTypeInfo));
@@ -94,13 +91,9 @@ namespace Metasound
 									ensure(IDataTypeRegistry::Get().GetDataTypeInfo(DataTypeName, OutDataTypeInfo));
 								}
 							}
-
-							return true;
 						}
 					}
 				}
-
-				return false;
 			}
 
 			// If DataType is an array type, creates & returns the array's
@@ -127,9 +120,9 @@ namespace Metasound
 					}
 					
 					Frontend::FDataTypeRegistryInfo DataTypeInfo;
-					const bool bMemberFound = MemberCustomizationPrivate::GetDataTypeFromElementPropertyHandle(PropertyHandle, DataTypeInfo);
+					MemberCustomizationPrivate::GetDataTypeFromElementPropertyHandle(PropertyHandle, DataTypeInfo);
 					UClass* ProxyGenClass = DataTypeInfo.ProxyGeneratorClass;
-					if (!bMemberFound || !ProxyGenClass)
+					if (!ProxyGenClass)
 					{
 						return;
 					}
@@ -226,10 +219,8 @@ namespace Metasound
 				PropertyHandle->GetOuterObjects(OuterObjects);
 				if (!OuterObjects.IsEmpty())
 				{
-					if (UMetasoundEditorGraphMemberDefaultLiteral* Literal = Cast<UMetasoundEditorGraphMemberDefaultLiteral>(OuterObjects[0]))
-					{
-						GraphMember = Literal->FindMember();
-					}
+					UObject* Outer = OuterObjects[0]->GetOuter();
+					GraphMember = Cast<UMetasoundEditorGraphMember>(Outer);
 				}
 
 				FUIAction PasteAction;
@@ -319,10 +310,12 @@ namespace Metasound
 					if (FloatLiteral.IsValid())
 					{
 						FloatLiteral->OnClampChanged.Remove(OnClampChangedDelegateHandle);
-						if (const UMetasoundEditorGraphMember* Member = FloatLiteral->FindMember())
+						if (FMetasoundAssetBase* MetasoundAsset = Metasound::IMetasoundUObjectRegistry::Get().GetObjectAsAssetBase(FloatLiteral->GetOutermostObject()))
 						{
-							FMetasoundAssetBase& MetasoundAsset = Editor::FGraphBuilder::GetOutermostMetaSoundChecked(*FloatLiteral);
-							MetasoundAsset.GetModifyContext().AddMemberIDsModified({ Member->GetMemberID() });
+							if (const UMetasoundEditorGraphMember* Member = Cast<UMetasoundEditorGraphMember>(FloatLiteral->GetOuter()))
+							{
+								MetasoundAsset->GetModifyContext().AddMemberIDsModified({ Member->GetMemberID() });
+							}
 						}
 					}
 				});
@@ -339,7 +332,7 @@ namespace Metasound
 
 			// Enable widget options for editable inputs only 
 			bool bShowWidgetOptions = false;
-			if (const UMetasoundEditorGraphInput* ParentMember = Cast <UMetasoundEditorGraphInput>(InLiteral.FindMember()))
+			if (const UMetasoundEditorGraphInput* ParentMember = Cast <UMetasoundEditorGraphInput>(InLiteral.GetParentMember()))
 			{
 				if (const UMetasoundEditorGraph* OwningGraph = ParentMember->GetOwningGraph())
 				{
@@ -394,8 +387,8 @@ namespace Metasound
 
 			// Enable widget options for editable inputs only 
 			bool bShowWidgetOptions = false;
-			if (const UMetasoundEditorGraphInput* ParentMember = Cast <UMetasoundEditorGraphInput>(InLiteral.FindMember()))
-			{
+			if (const UMetasoundEditorGraphInput* ParentMember = Cast <UMetasoundEditorGraphInput>(InLiteral.GetParentMember()))
+			{								
 				if (const UMetasoundEditorGraph* OwningGraph = ParentMember->GetOwningGraph())
 				{
 					// Disable widget options for constructor inputs for now to prevent changing default value via widget while playing 
@@ -446,9 +439,12 @@ namespace Metasound
 				.OnAreAssetsAcceptableForDropWithReason_Lambda([this, DefaultValueHandle](TArrayView<FAssetData> InAssets, FText& OutReason)
 				{
 					Frontend::FDataTypeRegistryInfo DataTypeInfo;
-					const bool bMemberFound = MemberCustomizationPrivate::GetDataTypeFromElementPropertyHandle(DefaultValueHandle, DataTypeInfo);
-					bool bCanDrop = bMemberFound;
-					if (UClass* ProxyGenClass = DataTypeInfo.ProxyGeneratorClass; bCanDrop && bMemberFound)
+					MemberCustomizationPrivate::GetDataTypeFromElementPropertyHandle(DefaultValueHandle, DataTypeInfo);
+
+					const IMetasoundEditorModule& EditorModule = FModuleManager::GetModuleChecked<IMetasoundEditorModule>("MetaSoundEditor");
+					bool bCanDrop = false;
+
+					if (UClass* ProxyGenClass = DataTypeInfo.ProxyGeneratorClass)
 					{
 						bCanDrop = true;
 						for (const FAssetData& AssetData : InAssets)
@@ -456,7 +452,6 @@ namespace Metasound
 							if (UClass* Class = AssetData.GetClass())
 							{
 PRAGMA_DISABLE_DEPRECATION_WARNINGS
-								const IMetasoundEditorModule& EditorModule = FModuleManager::GetModuleChecked<IMetasoundEditorModule>("MetaSoundEditor");
 								if (EditorModule.IsExplicitProxyClass(*DataTypeInfo.ProxyGeneratorClass))
 								{
 									bCanDrop &= Class == DataTypeInfo.ProxyGeneratorClass;
@@ -539,7 +534,7 @@ PRAGMA_ENABLE_DEPRECATION_WARNINGS
 				{
 					if (UMetasoundEditorGraphMemberDefaultLiteral* Literal = Cast<UMetasoundEditorGraphMemberDefaultLiteral>(OuterObjects.Last()))
 					{
-						if (UMetasoundEditorGraphInput* Input = Cast<UMetasoundEditorGraphInput>(Literal->FindMember()))
+						if (UMetasoundEditorGraphInput* Input = Cast<UMetasoundEditorGraphInput>(Literal->GetParentMember()))
 						{
 							// Don't display trigger simulation widget if its a trigger
 							// provided by an interface that does not support transmission.
@@ -566,10 +561,7 @@ PRAGMA_ENABLE_DEPRECATION_WARNINGS
 			using namespace Frontend;
 			using namespace MemberCustomizationPrivate;
 
-			// DataType can be reset during deletion of a literal value.  Customization can repaint briefly before the literal is removed,
-			// so just ignores if DataType is invalid.
-			const bool bIsValidDataType = !DataTypeInfo.DataTypeName.IsNone();
-			if (bIsValidDataType)
+			if (FMetasoundFrontendRegistryContainer* Registry = FMetasoundFrontendRegistryContainer::Get())
 			{
 				TSharedPtr<IPropertyHandle> ValueProperty = StructPropertyHandle->GetChildHandle(GET_MEMBER_NAME_CHECKED(FMetasoundEditorGraphMemberDefaultIntRef, Value));
 				if (ValueProperty.IsValid())
@@ -791,7 +783,7 @@ PRAGMA_ENABLE_DEPRECATION_WARNINGS
 				}
 			}
 
-			const bool bMemberFound = MemberCustomizationPrivate::GetDataTypeFromElementPropertyHandle(ElementPropertyHandle, DataTypeInfo);
+			MemberCustomizationPrivate::GetDataTypeFromElementPropertyHandle(ElementPropertyHandle, DataTypeInfo);
 
 			TSharedRef<SWidget> ValueWidget = CreateValueWidget(ParentPropertyHandleArray, StructPropertyHandle);
 			FDetailWidgetRow& ValueRow = ChildBuilder.AddCustomRow(MemberCustomizationStyle::DefaultPropertyText);
@@ -1648,8 +1640,8 @@ PRAGMA_ENABLE_DEPRECATION_WARNINGS
 
 			if (bIsPreset && !bIsDefaultConstructed && !bIsTriggerDataType)
 			{
-				const UMetasoundEditorGraphInput* Input = Cast<UMetasoundEditorGraphInput>(MemberDefaultLiteral->FindMember());
-				if (Input)
+				const UMetasoundEditorGraphInput* Input = Cast<UMetasoundEditorGraphInput>(MemberDefaultLiteral->GetParentMember());
+				if (ensure(Input))
 				{
 					auto PropertyEnabled = TAttribute<bool>::CreateLambda([this] { return !GetInputInheritsDefault(); });
 					for (IDetailPropertyRow* DefaultPropertyRow : DefaultPropertyRows)
@@ -1661,12 +1653,12 @@ PRAGMA_ENABLE_DEPRECATION_WARNINGS
 						DefaultPropertyRow->OverrideResetToDefault(ResetOverride);
 					}
 				}
-			}
+			} 
 			else if (!bIsPreset)
 			{
 				// Make default value uneditable while playing for constructor inputs
-				const UMetasoundEditorGraphInput* Input = Cast<UMetasoundEditorGraphInput>(MemberDefaultLiteral->FindMember());
-				if (Input)
+				const UMetasoundEditorGraphInput* Input = Cast<UMetasoundEditorGraphInput>(MemberDefaultLiteral->GetParentMember());
+				if (ensure(Input))
 				{
 					auto PropertyEnabled = TAttribute<bool>::CreateLambda([this, Input]
 					{
diff --git a/Engine/Plugins/Runtime/Metasound/Source/MetasoundEditor/Private/SMetasoundGraphNode.cpp b/Engine/Plugins/Runtime/Metasound/Source/MetasoundEditor/Private/SMetasoundGraphNode.cpp
--- a/Engine/Plugins/Runtime/Metasound/Source/MetasoundEditor/Private/SMetasoundGraphNode.cpp
+++ b/Engine/Plugins/Runtime/Metasound/Source/MetasoundEditor/Private/SMetasoundGraphNode.cpp
@@ -17,7 +17,6 @@
 #include "KismetPins/SGraphPinNum.h"
 #include "KismetPins/SGraphPinObject.h"
 #include "KismetPins/SGraphPinString.h"
-#include "MetasoundBuilderSubsystem.h"
 #include "MetasoundEditorGraph.h"
 #include "MetasoundEditorGraphBuilder.h"
 #include "MetasoundEditorGraphInputNode.h"
@@ -35,6 +34,7 @@
 #include "PropertyCustomizationHelpers.h"
 #include "SAudioRadialSlider.h"
 #include "SAudioSlider.h"
+#include "SCommentBubble.h"
 #include "ScopedTransaction.h"
 #include "SGraphNode.h"
 #include "SGraphPinComboBox.h"
@@ -94,7 +94,7 @@ namespace Metasound
 							{
 								constexpr bool bPostTransaction = false;
 								GraphMember->UpdateFrontendDefaultLiteral(bPostTransaction);
-								FGraphBuilder::GetOutermostMetaSoundChecked(*Graph).GetModifyContext().AddMemberIDsModified({ GraphMember->GetMemberID() });
+								Graph->GetModifyContext().AddMemberIDsModified({ GraphMember->GetMemberID() });
 							}
 						}
 						DefaultFloat->OnDefaultValueChanged.Remove(InputSliderOnValueChangedDelegateHandle);
@@ -137,7 +137,7 @@ namespace Metasound
 
 		void SMetaSoundGraphNode::ExecuteTrigger(UMetasoundEditorGraphMemberDefaultLiteral& Literal)
 		{
-			UMetasoundEditorGraphMember* Member = Literal.FindMember();
+			UMetasoundEditorGraphMember* Member = Cast<UMetasoundEditorGraphMember>(Literal.GetOuter());
 			if (!ensure(Member))
 			{
 				return;
@@ -462,43 +462,6 @@ namespace Metasound
 			}
 		}
 
-		void SMetaSoundGraphNode::OnCommentBubbleToggled(bool bInCommentBubbleVisible)
-		{
-			UMetasoundEditorGraphNode& EdNode = GetMetaSoundNode();
-			UObject& MetaSound = EdNode.GetMetasoundChecked();
-			UMetaSoundBuilderBase& Builder = UMetaSoundBuilderSubsystem::GetChecked().AttachBuilderToAssetChecked(MetaSound);
-			if (const FMetasoundFrontendNode* Node = Builder.GetConstBuilder().FindNode(EdNode.GetNodeID()))
-			{
-				if (bInCommentBubbleVisible != Node->Style.Display.bCommentVisible)
-				{
-					const FScopedTransaction Transaction(LOCTEXT("GraphNodeCommentCommitted", "Graph Node Comment Bubble Toggled"));
-					MetaSound.Modify();
-					EMetaSoundBuilderResult Result;
-					EdNode.bCommentBubbleVisible = bInCommentBubbleVisible;
-					Builder.SetNodeCommentVisible(Node->GetID(), bInCommentBubbleVisible, Result);
-				}
-			}
-		}
-
-		void SMetaSoundGraphNode::OnCommentTextCommitted(const FText& NewComment, ETextCommit::Type CommitInfo)
-		{
-			FString NewCommentString = NewComment.ToString();
-			UMetasoundEditorGraphNode& EdNode = GetMetaSoundNode();
-			UObject& MetaSound = EdNode.GetMetasoundChecked();
-			UMetaSoundBuilderBase& Builder = UMetaSoundBuilderSubsystem::GetChecked().AttachBuilderToAssetChecked(MetaSound);
-			if (const FMetasoundFrontendNode* Node = Builder.GetConstBuilder().FindNode(EdNode.GetNodeID()))
-			{
-				if (!Node->Style.Display.Comment.Equals(NewCommentString))
-				{
-					const FScopedTransaction Transaction(LOCTEXT("GraphNodeCommentCommitted", "Graph Node Comment Changed"));
-					MetaSound.Modify();
-					EMetaSoundBuilderResult Result;
-					EdNode.NodeComment = MoveTemp(NewCommentString);
-					Builder.SetNodeComment(Node->GetID(), EdNode.NodeComment, Result);
-				}
-			}
-		}
-
 		FLinearColor SMetaSoundGraphNode::GetNodeTitleColorOverride() const
 		{
 			FLinearColor ReturnTitleColor = GraphNode->IsDeprecated() ? FLinearColor::Red : GetNodeObj()->GetNodeTitleColor();
@@ -590,8 +553,7 @@ namespace Metasound
 
 			UMetasoundEditorGraphNode& Node = GetMetaSoundNode();
 			Node.GetMetasoundChecked().Modify();
-			Node.UpdateFrontendNodeLocation(NewPosition);
-			Node.SyncLocationFromFrontendNode();
+			Node.SetNodeLocation(NewPosition);
 		}
 
 		const FSlateBrush* SMetaSoundGraphNode::GetNodeBodyBrush() const
@@ -765,7 +727,7 @@ namespace Metasound
 
 								if (UMetasoundEditorGraph* Graph = GraphMember->GetOwningGraph())
 								{
-									FGraphBuilder::GetOutermostMetaSoundChecked(*Graph).GetModifyContext().AddMemberIDsModified({ GraphMember->GetMemberID() });
+									Graph->GetModifyContext().AddMemberIDsModified({ GraphMember->GetMemberID() });
 								}
 							}
 						};
@@ -1042,7 +1004,7 @@ namespace Metasound
 								.AutoHeight()
 								[
 									SAssignNew(MaterialButtonWidget, SAudioMaterialButton)
-										.OnBooleanValueChanged_Lambda(OnboolValueChangedLambda)
+										.OnBooleanValueChanged_Lambda(OnboolValueChangedLambda)									
 								];
 
 							MaterialButtonWidget->SetPressedState(DefaultBool->GetDefault());
@@ -1118,8 +1080,7 @@ namespace Metasound
 
 			UMetasoundEditorGraphNode& Node = GetMetaSoundNode();
 			Node.GetMetasoundChecked().Modify();
-			Node.UpdateFrontendNodeLocation(NewPosition);
-			Node.SyncLocationFromFrontendNode();
+			Node.SetNodeLocation(NewPosition);
 		}
 
 		UMetasoundEditorGraphNode& SMetaSoundGraphNodeKnot::GetMetaSoundNode()
diff --git a/Engine/Plugins/Runtime/Metasound/Source/MetasoundEditor/Private/SMetasoundGraphNode.h b/Engine/Plugins/Runtime/Metasound/Source/MetasoundEditor/Private/SMetasoundGraphNode.h
--- a/Engine/Plugins/Runtime/Metasound/Source/MetasoundEditor/Private/SMetasoundGraphNode.h
+++ b/Engine/Plugins/Runtime/Metasound/Source/MetasoundEditor/Private/SMetasoundGraphNode.h
@@ -23,95 +23,93 @@ class UMetasoundEditorGraphMemberNode;
 class UMetasoundEditorGraphNode;
 
 
-namespace Metasound::Editor
+namespace Metasound
 {
-	class SMetaSoundGraphNode : public SGraphNode
+	namespace Editor
 	{
-		public:
-		SLATE_BEGIN_ARGS(SMetaSoundGraphNode)
+		class SMetaSoundGraphNode : public SGraphNode
 		{
-		}
-
-		SLATE_END_ARGS()
-		virtual ~SMetaSoundGraphNode();
-
-		void Construct(const FArguments& InArgs, class UEdGraphNode* InNode);
-
-	protected:
-		bool IsVariableAccessor() const;
-		bool IsVariableMutator() const;
-
-		// SGraphNode Interface
-		virtual const FSlateBrush* GetShadowBrush(bool bSelected) const override;
-		virtual void CreateInputSideAddButton(TSharedPtr<SVerticalBox> InputBox) override;
-		virtual void CreateOutputSideAddButton(TSharedPtr<SVerticalBox> OutputBox) override;
-		virtual TSharedPtr<SGraphPin> CreatePinWidget(UEdGraphPin* InPin) const override;
-		virtual void CreateStandardPinWidget(UEdGraphPin* InPin) override;
-		virtual TSharedRef<SWidget> CreateNodeContentArea() override;
-		virtual TSharedRef<SWidget> CreateTitleWidget(TSharedPtr<SNodeTitle> NodeTitle) override;
-		virtual const FSlateBrush* GetNodeBodyBrush() const override;
-		virtual TSharedRef<SWidget> CreateTitleRightWidget() override;
-		virtual EVisibility IsAddPinButtonVisible() const override;
-		virtual FReply OnAddPin() override;
-		virtual void MoveTo(const FVector2D& NewPosition, FNodeSet& NodeFilter, bool bMarkDirty = true) override;
-		virtual void SetDefaultTitleAreaWidget(TSharedRef<SOverlay> DefaultTitleAreaWidget) override;
-		virtual void GetOverlayBrushes(bool bSelected, const FVector2D WidgetSize, TArray<FOverlayBrushInfo>& Brushes) const override;
-		virtual void OnCommentBubbleToggled(bool bInCommentBubbleVisible) override;
-		virtual void OnCommentTextCommitted(const FText& NewComment, ETextCommit::Type CommitInfo) override;
-
-		FLinearColor GetNodeTitleColorOverride() const;
-
-		FName GetLiteralDataType() const;
-		UMetasoundEditorGraphNode& GetMetaSoundNode();
-		const UMetasoundEditorGraphNode& GetMetaSoundNode() const;
-
-	public:
-		static void ExecuteTrigger(UMetasoundEditorGraphMemberDefaultLiteral& Literal);
-		static TSharedRef<SWidget> CreateTriggerSimulationWidget(UMetasoundEditorGraphMemberDefaultLiteral& Literal, TAttribute<EVisibility>&& InVisibility, TAttribute<bool>&& InEnablement, const FText* InToolTip = nullptr);
-
-	private:
-		// If this node represents a graph member node, returns corresponding member.
-		UMetasoundEditorGraphMember* GetMetaSoundMember();
-
-		// If this node represents a graph member node, returns cast node.
-		UMetasoundEditorGraphMemberNode* GetMetaSoundMemberNode();
-
-		TAttribute<EVisibility> GetSimulationVisibilityAttribute() const;
-
-		// Slider widget for float input
-		TSharedPtr<SAudioInputWidget> InputWidget;
-
-		//Button Widget for bool input.
-		TSharedPtr <SAudioMaterialButton> MaterialButtonWidget;
-
-		// Handle for on value changed delegate for input slider 
-		FDelegateHandle InputSliderOnValueChangedDelegateHandle;
-		
-		// Handle for on input slider range changed  
-		FDelegateHandle InputSliderOnRangeChangedDelegateHandle;
-
-		// Whether the input widget is currently transacting 
-		// for keeping track of transaction state across delegates to only commit transaction on value commit
-		bool bIsInputWidgetTransacting = false;
-
-		EMetasoundFrontendClassType ClassType;
-	};
-
-	class SMetaSoundGraphNodeKnot : public SGraphNodeKnot
-	{
-	public:
-		SLATE_BEGIN_ARGS(SMetaSoundGraphNode)
+			public:
+			SLATE_BEGIN_ARGS(SMetaSoundGraphNode)
+			{
+			}
+
+			SLATE_END_ARGS()
+			virtual ~SMetaSoundGraphNode();
+
+			void Construct(const FArguments& InArgs, class UEdGraphNode* InNode);
+
+		protected:
+			bool IsVariableAccessor() const;
+			bool IsVariableMutator() const;
+
+			// SGraphNode Interface
+			virtual const FSlateBrush* GetShadowBrush(bool bSelected) const override;
+			virtual void CreateInputSideAddButton(TSharedPtr<SVerticalBox> InputBox) override;
+			virtual void CreateOutputSideAddButton(TSharedPtr<SVerticalBox> OutputBox) override;
+			virtual TSharedPtr<SGraphPin> CreatePinWidget(UEdGraphPin* InPin) const override;
+			virtual void CreateStandardPinWidget(UEdGraphPin* InPin) override;
+			virtual TSharedRef<SWidget> CreateNodeContentArea() override;
+			virtual TSharedRef<SWidget> CreateTitleWidget(TSharedPtr<SNodeTitle> NodeTitle) override;
+			virtual const FSlateBrush* GetNodeBodyBrush() const override;
+			virtual TSharedRef<SWidget> CreateTitleRightWidget() override;
+			virtual EVisibility IsAddPinButtonVisible() const override;
+			virtual FReply OnAddPin() override;
+			virtual void MoveTo(const FVector2D& NewPosition, FNodeSet& NodeFilter, bool bMarkDirty = true) override;
+			virtual void SetDefaultTitleAreaWidget(TSharedRef<SOverlay> DefaultTitleAreaWidget) override;
+			virtual void GetOverlayBrushes(bool bSelected, const FVector2D WidgetSize, TArray<FOverlayBrushInfo>& Brushes) const override;
+
+			FLinearColor GetNodeTitleColorOverride() const;
+
+			FName GetLiteralDataType() const;
+			UMetasoundEditorGraphNode& GetMetaSoundNode();
+			const UMetasoundEditorGraphNode& GetMetaSoundNode() const;
+
+		public:
+			static void ExecuteTrigger(UMetasoundEditorGraphMemberDefaultLiteral& Literal);
+			static TSharedRef<SWidget> CreateTriggerSimulationWidget(UMetasoundEditorGraphMemberDefaultLiteral& Literal, TAttribute<EVisibility>&& InVisibility, TAttribute<bool>&& InEnablement, const FText* InToolTip = nullptr);
+
+		private:
+			// If this node represents a graph member node, returns corresponding member.
+			UMetasoundEditorGraphMember* GetMetaSoundMember();
+
+			// If this node represents a graph member node, returns cast node.
+			UMetasoundEditorGraphMemberNode* GetMetaSoundMemberNode();
+
+			TAttribute<EVisibility> GetSimulationVisibilityAttribute() const;
+
+			// Slider widget for float input
+			TSharedPtr<SAudioInputWidget> InputWidget;
+			//Button Widget for bool input.
+			TSharedPtr <SAudioMaterialButton> MaterialButtonWidget;
+			// Handle for on value changed delegate for input slider 
+			FDelegateHandle InputSliderOnValueChangedDelegateHandle;
+			// Handle for on input slider range changed  
+			FDelegateHandle InputSliderOnRangeChangedDelegateHandle;
+
+			// Whether the input widget is currently transacting 
+			// for keeping track of transaction state across delegates to only commit transaction on value commit
+			bool bIsInputWidgetTransacting = false;
+
+			EMetasoundFrontendClassType ClassType;
+		};
+
+		class SMetaSoundGraphNodeKnot : public SGraphNodeKnot
 		{
-		}
+		public:
+			SLATE_BEGIN_ARGS(SMetaSoundGraphNode)
+			{
+			}
 
-		SLATE_END_ARGS()
+			SLATE_END_ARGS()
 
-		void Construct(const FArguments& InArgs, class UEdGraphNode* InNode);
+			void Construct(const FArguments& InArgs, class UEdGraphNode* InNode);
 
-		virtual TSharedPtr<SGraphPin> CreatePinWidget(UEdGraphPin* Pin) const override;
-		virtual void MoveTo(const FVector2D& NewPosition, FNodeSet& NodeFilter, bool bMarkDirty) override;
+			virtual TSharedPtr<SGraphPin> CreatePinWidget(UEdGraphPin* Pin) const override;
+			virtual void MoveTo(const FVector2D& NewPosition, FNodeSet& NodeFilter, bool bMarkDirty) override;
 
-		UMetasoundEditorGraphNode& GetMetaSoundNode();
-		const UMetasoundEditorGraphNode& GetMetaSoundNode() const;
-	};
-} // namespace Metasound::Editor
+			UMetasoundEditorGraphNode& GetMetaSoundNode();
+			const UMetasoundEditorGraphNode& GetMetaSoundNode() const;
+		};
+	} // namespace Editor
+} // namespace Metasound
diff --git a/Engine/Plugins/Runtime/Metasound/Source/MetasoundEditor/Private/SMetasoundGraphNodeComment.cpp b/Engine/Plugins/Runtime/Metasound/Source/MetasoundEditor/Private/SMetasoundGraphNodeComment.cpp
--- a/Engine/Plugins/Runtime/Metasound/Source/MetasoundEditor/Private/SMetasoundGraphNodeComment.cpp
+++ b/Engine/Plugins/Runtime/Metasound/Source/MetasoundEditor/Private/SMetasoundGraphNodeComment.cpp
@@ -3,7 +3,6 @@
 
 #include "EdGraphNode_Comment.h"
 #include "Framework/Application/SlateApplication.h"
-#include "MetasoundEditorGraphCommentNode.h"
 #include "MetasoundEditorGraphNode.h"
 
 namespace Metasound
@@ -20,7 +19,7 @@ namespace Metasound
 			FModifierKeysState KeysState = FSlateApplication::Get().GetModifierKeys();
 			if (!KeysState.IsShiftDown())
 			{
-				UMetasoundEditorGraphCommentNode* CommentNode = Cast<UMetasoundEditorGraphCommentNode>(GraphNode);
+				UEdGraphNode_Comment* CommentNode = Cast<UEdGraphNode_Comment>(GraphNode);
 				if (CommentNode && CommentNode->MoveMode == ECommentBoxMode::GroupMovement)
 				{
 					FVector2D PositionDelta = NewPosition - GetPosition();
diff --git a/Engine/Plugins/Runtime/Metasound/Source/MetasoundEditor/Private/SMetasoundGraphPin.h b/Engine/Plugins/Runtime/Metasound/Source/MetasoundEditor/Private/SMetasoundGraphPin.h
--- a/Engine/Plugins/Runtime/Metasound/Source/MetasoundEditor/Private/SMetasoundGraphPin.h
+++ b/Engine/Plugins/Runtime/Metasound/Source/MetasoundEditor/Private/SMetasoundGraphPin.h
@@ -8,7 +8,6 @@
 #include "KismetPins/SGraphPinInteger.h"
 #include "KismetPins/SGraphPinObject.h"
 #include "KismetPins/SGraphPinString.h"
-#include "MetasoundBuilderSubsystem.h"
 #include "MetasoundEditorGraph.h"
 #include "MetasoundEditorGraphBuilder.h"
 #include "MetasoundEditorGraphMemberDefaults.h"
@@ -43,11 +42,6 @@ namespace Metasound
 			TSharedPtr<FMetasoundPinAudioInspector> PinAudioInspector;
 
 		protected:
-			virtual bool CanInspectPin(const UEdGraphPin* InPin)
-			{
-				return FGraphBuilder::CanInspectPin(InPin);
-			}
-
 			static TWeakPtr<FPinValueInspectorTooltip> OpenPinInspector(UEdGraphPin& InPin, TSharedPtr<SMetasoundPinValueInspector>& OutPinInspector)
 			{
 				TSharedPtr<SMetasoundPinValueInspector> NewPinInspector = SNew(SMetasoundPinValueInspector);
@@ -61,43 +55,37 @@ namespace Metasound
 				return nullptr;
 			}
 
-			void UpdatePinInspector(
+			static void UpdatePinInspector(
 				UEdGraphPin& InPin,
 				const bool bIsHoveringPin,
 				TSharedPtr<SMetasoundPinValueInspector>& OutPinInspector,
 				TWeakPtr<FPinValueInspectorTooltip>& OutInspectorTooltip,
 				TFunctionRef<void(FVector2D&)> InGetTooltipLocation)
 			{
-				if (bIsHoveringPin)
+				const bool bCanInspectPin = FGraphBuilder::CanInspectPin(&InPin);
+				if (bIsHoveringPin && bCanInspectPin)
 				{
-					const bool bCanInspectPin = CanInspectPin(&InPin);
-					if (bCanInspectPin)
+					if (OutPinInspector.IsValid())
 					{
-						if (OutPinInspector.IsValid())
+						const UEdGraphPin* InspectedPin = OutPinInspector->GetPinRef().Get();
+						if (InspectedPin == &InPin)
 						{
-							const UEdGraphPin* InspectedPin = OutPinInspector->GetPinRef().Get();
-							if (InspectedPin == &InPin)
-							{
-								OutPinInspector->UpdateMessage();
-							}
+							OutPinInspector->UpdateMessage();
 						}
-						else
+					}
+					else
+					{
+						OutInspectorTooltip = OpenPinInspector(InPin, OutPinInspector);
+						TSharedPtr<FPinValueInspectorTooltip> NewTooltip = OutInspectorTooltip.Pin();
+						if (NewTooltip.IsValid())
 						{
-							OutInspectorTooltip = OpenPinInspector(InPin, OutPinInspector);
-							TSharedPtr<FPinValueInspectorTooltip> NewTooltip = OutInspectorTooltip.Pin();
-							if (NewTooltip.IsValid())
-							{
-								FVector2D TooltipLocation;
-								InGetTooltipLocation(TooltipLocation);
-								NewTooltip->MoveTooltip(TooltipLocation);
-							}
+							FVector2D TooltipLocation;
+							InGetTooltipLocation(TooltipLocation);
+							NewTooltip->MoveTooltip(TooltipLocation);
 						}
-
-						return;
 					}
 				}
-
-				if (OutPinInspector.IsValid())
+				else if (OutPinInspector.IsValid())
 				{
 					TSharedPtr<FPinValueInspectorTooltip> InspectorTooltip = OutInspectorTooltip.Pin();
 					if (InspectorTooltip.IsValid())
@@ -131,34 +119,29 @@ namespace Metasound
 				return nullptr;
 			}
 
-			void UpdatePinAudioInspector(UEdGraphPin& InPin,
+			static void UpdatePinAudioInspector(UEdGraphPin& InPin,
 				const bool bIsHoveringPin,
 				TSharedPtr<FMetasoundPinAudioInspector>& OutPinAudioInspector,
 				TWeakPtr<FPinValueInspectorTooltip>& OutInspectorTooltip,
 				TFunctionRef<void(FVector2D&)> InGetTooltipLocation)
 			{
-				if (bIsHoveringPin)
+				const bool bCanInspectPin = FGraphBuilder::CanInspectPin(&InPin);
+
+				if (bIsHoveringPin && bCanInspectPin)
 				{
-					const bool bCanInspectPin = CanInspectPin(&InPin);
-					if (bCanInspectPin)
+					if (!OutPinAudioInspector.IsValid())
 					{
-						if (!OutPinAudioInspector.IsValid())
+						OutInspectorTooltip = OpenPinAudioInspector(InPin, OutPinAudioInspector);
+						TSharedPtr<FPinValueInspectorTooltip> NewTooltip = OutInspectorTooltip.Pin();
+						if (NewTooltip.IsValid())
 						{
-							OutInspectorTooltip = OpenPinAudioInspector(InPin, OutPinAudioInspector);
-							TSharedPtr<FPinValueInspectorTooltip> NewTooltip = OutInspectorTooltip.Pin();
-							if (NewTooltip.IsValid())
-							{
-								FVector2D TooltipLocation;
-								InGetTooltipLocation(TooltipLocation);
-								NewTooltip->MoveTooltip(TooltipLocation);
-							}
+							FVector2D TooltipLocation;
+							InGetTooltipLocation(TooltipLocation);
+							NewTooltip->MoveTooltip(TooltipLocation);
 						}
-
-						return;
 					}
 				}
-
-				if (OutPinAudioInspector.IsValid())
+				else if (OutPinAudioInspector.IsValid())
 				{
 					TSharedPtr<FPinValueInspectorTooltip> InspectorTooltip = OutInspectorTooltip.Pin();
 					if (InspectorTooltip.IsValid())
@@ -197,12 +180,12 @@ namespace Metasound
 						{
 							if (Pin->Direction == EGPD_Input)
 							{
-								Frontend::FConstInputHandle InputHandle = FGraphBuilder::GetConstInputHandleFromPin(Pin);
+								Frontend::FInputHandle InputHandle = FGraphBuilder::GetInputHandleFromPin(Pin);
 								AccessType = InputHandle->GetVertexAccessType();
 							}
 							else if (Pin->Direction == EGPD_Output)
 							{
-								Frontend::FConstOutputHandle OutputHandle = FGraphBuilder::GetConstOutputHandleFromPin(Pin);
+								Frontend::FOutputHandle OutputHandle = FGraphBuilder::GetOutputHandleFromPin(Pin);
 								AccessType = OutputHandle->GetVertexAccessType();
 							}
 						}
@@ -267,56 +250,6 @@ namespace Metasound
 				return IInputController::GetInvalidHandle();
 			}
 
-			const UMetasoundEditorGraphNode* GetOwningMetaSoundNode() const
-			{
-				if (const UEdGraphPin* Pin = ParentPinType::GetPinObj())
-				{
-					UObject* Node = Pin->GetOwningNode();
-					return Cast<UMetasoundEditorGraphNode>(Node);
-				}
-
-				return nullptr;
-			}
-
-			UMetaSoundBuilderBase& GetBuilderChecked() const
-			{
-				const UMetasoundEditorGraphNode* Node = GetOwningMetaSoundNode();
-				check(Node);
-				UObject* Outermost = Node->GetOutermostObject();
-				check(Outermost);
-				return UMetaSoundBuilderSubsystem::GetChecked().AttachBuilderToAssetChecked(*Outermost);
-			}
-
-			const FMetasoundFrontendNode* GetFrontendNode() const
-			{
-				if (const UMetasoundEditorGraphNode* Node = GetOwningMetaSoundNode())
-				{
-					if (UObject* Outermost = Node->GetOutermostObject())
-					{
-						const FGuid NodeID = Node->GetNodeID();
-						const UMetaSoundBuilderBase& Builder = UMetaSoundBuilderSubsystem::GetChecked().AttachBuilderToAssetChecked(*Outermost);
-						return Builder.GetConstBuilder().FindNode(NodeID);
-					}
-				}
-
-				return nullptr;
-			}
-
-			const FMetasoundFrontendNode& GetFrontendNodeChecked() const
-			{
-				const UMetasoundEditorGraphNode* Node = GetOwningMetaSoundNode();
-				check(Node);
-				UObject* Outermost = Node->GetOutermostObject();
-				check(Outermost);
-
-				const FGuid NodeID = Node->GetNodeID();
-				const UMetaSoundBuilderBase& Builder = UMetaSoundBuilderSubsystem::GetChecked().AttachBuilderToAssetChecked(*Outermost);
-
-				const FMetasoundFrontendNode* FrontendNode = Builder.GetConstBuilder().FindNode(NodeID);
-				check(FrontendNode);
-				return *FrontendNode;
-			}
-
 			bool ShowDefaultValueWidget() const
 			{
 				UEdGraphPin* Pin = ParentPinType::GetPinObj();
@@ -394,7 +327,7 @@ namespace Metasound
 									{
 										if (UMetasoundEditorGraphExternalNode* Node = Cast<UMetasoundEditorGraphExternalNode>(Pin->GetOwningNode()))
 										{
-											if (Node->GetBreadcrumb().ClassName == FRerouteNodeTemplate::ClassName)
+											if (Node->GetClassName() == FRerouteNodeTemplate::ClassName)
 											{
 												bIsRerouteNode = true;
 											}
@@ -412,7 +345,6 @@ namespace Metasound
 						}))
 						.OnClicked(FOnClicked::CreateLambda([this]()
 						{
-							using namespace Editor;
 							using namespace Frontend;
 
 							if (UEdGraphPin* Pin = ParentPinType::GetPinObj())
@@ -422,7 +354,7 @@ namespace Metasound
 									if (UMetasoundEditorGraph* MetaSoundGraph = CastChecked<UMetasoundEditorGraph>(Node->GetGraph()))
 									{
 										UObject& MetaSound = MetaSoundGraph->GetMetasoundChecked();
-										FMetasoundFrontendDocumentModifyContext& ModifyContext = FGraphBuilder::GetOutermostMetaSoundChecked(MetaSound).GetModifyContext();
+
 										{
 											const FScopedTransaction Transaction(LOCTEXT("MetaSoundEditorResetToClassDefault", "Reset to Class Default"));
 											MetaSound.Modify();
@@ -434,18 +366,18 @@ namespace Metasound
 												if (ensure(Member))
 												{
 													Member->ResetToClassDefault();
-													ModifyContext.AddMemberIDsModified({ Member->GetMemberID() });
+													MetaSoundGraph->GetModifyContext().AddMemberIDsModified({ Member->GetMemberID() });
 												}
 												else
 												{
-													ModifyContext.SetDocumentModified();
+													MetaSoundGraph->GetModifyContext().SetDocumentModified();
 												}
 											}
 											else
 											{
 												FInputHandle InputHandle = GetInputHandle();
 												InputHandle->ClearLiteral();
-												ModifyContext.SetDocumentModified();
+												MetaSoundGraph->GetModifyContext().SetDocumentModified();
 											}
 										}
 									}
@@ -612,8 +544,6 @@ namespace Metasound
 			virtual const FSlateBrush* GetPinIcon() const override;
 
 		protected:
-			virtual bool CanInspectPin(const UEdGraphPin* InPin) override;
-
 			void CacheHasRequiredConnections();
 
 			bool bHasRequiredConnections = false;
diff --git a/Engine/Plugins/Runtime/Metasound/Source/MetasoundEditor/Private/SMetasoundPinValueInspector.cpp b/Engine/Plugins/Runtime/Metasound/Source/MetasoundEditor/Private/SMetasoundPinValueInspector.cpp
--- a/Engine/Plugins/Runtime/Metasound/Source/MetasoundEditor/Private/SMetasoundPinValueInspector.cpp
+++ b/Engine/Plugins/Runtime/Metasound/Source/MetasoundEditor/Private/SMetasoundPinValueInspector.cpp
@@ -171,7 +171,7 @@ namespace Metasound
 			return nullptr;
 		}
 
-		Frontend::FConstOutputHandle FMetasoundNumericDebugLineItem::GetConstReroutedOutputHandle() const
+		Frontend::FConstOutputHandle FMetasoundNumericDebugLineItem::GetReroutedOutputHandle() const
 		{
 			return FGraphBuilder::FindReroutedConstOutputHandleFromPin(GraphPinObj);
 		}
@@ -336,9 +336,9 @@ namespace Metasound
 			return OutputHandle->GetOwningNode()->GetOwningGraph();
 		}
 
-		Frontend::FConstGraphHandle FMetasoundNumericDebugLineItem::GetConstGraphHandle() const
+		Frontend::FConstGraphHandle FMetasoundNumericDebugLineItem::GetGraphHandle() const
 		{
-			Frontend::FConstOutputHandle OutputHandle = FGraphBuilder::GetConstOutputHandleFromPin(GraphPinObj);
+			Frontend::FConstOutputHandle OutputHandle = FGraphBuilder::GetOutputHandleFromPin(GraphPinObj);
 			return OutputHandle->GetOwningNode()->GetOwningGraph();
 		}
 
@@ -360,21 +360,21 @@ namespace Metasound
 		{
 			using namespace Frontend;
 
-			FConstGraphHandle OwningGraph = GetConstGraphHandle();
+			FConstGraphHandle OwningGraph = GetGraphHandle();
 			if (!OwningGraph->IsValid())
 			{
 				return nullptr;
 			}
 
-			const FConstOutputHandle OutputHandle = GetConstReroutedOutputHandle();
+			const FConstOutputHandle OutputHandle = GetReroutedOutputHandle();
 			const FGuid NodeID = OutputHandle->GetOwningNodeID();
 			const FName OutputName = OutputHandle->GetName();
 			const FMetasoundFrontendGraphStyle& Style = OwningGraph->GetGraphStyle();
 
 			return Style.EdgeStyles.FindByPredicate([&NodeID, &OutputName](FMetasoundFrontendEdgeStyle& EdgeStyle)
-			{
-				return EdgeStyle.NodeID == NodeID && EdgeStyle.OutputName == OutputName;
-			});
+				{
+					return EdgeStyle.NodeID == NodeID && EdgeStyle.OutputName == OutputName;
+				});
 		}
 
 		void FMetasoundNumericDebugLineItem::Update()
@@ -383,7 +383,7 @@ namespace Metasound
 
 			bIsValueColorizationEnabled = false;
 
-			const FConstOutputHandle OutputHandle = GetConstReroutedOutputHandle();
+			const FConstOutputHandle OutputHandle = GetReroutedOutputHandle();
 			const FGuid NodeID = OutputHandle->GetOwningNodeID();
 			const FName OutputName = OutputHandle->GetName();
 			if (const FMetasoundFrontendEdgeStyle* EdgeStyle = GetEdgeStyle())
@@ -413,6 +413,12 @@ namespace Metasound
 				return FReply::Unhandled();
 			}
 
+			FGraphHandle OwningGraph = GetGraphHandle();
+			if (!OwningGraph->IsValid())
+			{
+				return FReply::Unhandled();
+			}
+
 			const FMetasoundFrontendEdgeStyle* EdgeStyle = GetEdgeStyle();
 			if (!EdgeStyle)
 			{
diff --git a/Engine/Plugins/Runtime/Metasound/Source/MetasoundEditor/Private/SMetasoundPinValueInspector.h b/Engine/Plugins/Runtime/Metasound/Source/MetasoundEditor/Private/SMetasoundPinValueInspector.h
--- a/Engine/Plugins/Runtime/Metasound/Source/MetasoundEditor/Private/SMetasoundPinValueInspector.h
+++ b/Engine/Plugins/Runtime/Metasound/Source/MetasoundEditor/Private/SMetasoundPinValueInspector.h
@@ -61,13 +61,13 @@ namespace Metasound
 
 			UObject* GetOutermostObject();
 
-			Frontend::FConstOutputHandle GetConstReroutedOutputHandle() const;
+			Frontend::FConstOutputHandle GetReroutedOutputHandle() const;
 
 			Frontend::FOutputHandle GetReroutedOutputHandle();
 
 			Frontend::FGraphHandle GetGraphHandle();
 
-			Frontend::FConstGraphHandle GetConstGraphHandle() const;
+			Frontend::FConstGraphHandle GetGraphHandle() const;
 
 			FGraphConnectionManager* GetConnectionManager();
 
diff --git a/Engine/Plugins/Runtime/Metasound/Source/MetasoundEditor/Public/MetasoundEditorGraph.h b/Engine/Plugins/Runtime/Metasound/Source/MetasoundEditor/Public/MetasoundEditorGraph.h
--- a/Engine/Plugins/Runtime/Metasound/Source/MetasoundEditor/Public/MetasoundEditorGraph.h
+++ b/Engine/Plugins/Runtime/Metasound/Source/MetasoundEditor/Public/MetasoundEditorGraph.h
@@ -2,7 +2,6 @@
 #pragma once
 
 #include "AudioParameterControllerInterface.h"
-#include "MetasoundBuilderSubsystem.h"
 #include "MetasoundEditor.h"
 #include "MetasoundEditorGraphNode.h"
 #include "MetasoundFrontendController.h"
@@ -47,7 +46,7 @@ DECLARE_MULTICAST_DELEGATE_OneParam(FOnMetasoundMemberNameChanged, FGuid /* Node
 DECLARE_MULTICAST_DELEGATE(FOnMetasoundMemberRenameRequested);
 
 UCLASS()
-class METASOUNDEDITOR_API UMetasoundEditorGraphMemberDefaultLiteral : public UMetaSoundFrontendMemberMetadata
+class METASOUNDEDITOR_API UMetasoundEditorGraphMemberDefaultLiteral : public UObject
 {
 	GENERATED_BODY()
 
@@ -56,22 +55,35 @@ public:
 	{
 	}
 
+	virtual FMetasoundFrontendLiteral GetDefault() const
+	{
+		return FMetasoundFrontendLiteral();
+	}
+
+	virtual EMetasoundFrontendLiteralType GetLiteralType() const
+	{
+		return EMetasoundFrontendLiteralType::None;
+	}
+
+	virtual void SetFromLiteral(const FMetasoundFrontendLiteral& InLiteral)
+	{
+	}
+
+	virtual void ForceRefresh()
+	{
+	}
+
 #if WITH_EDITOR
 	virtual void PostEditChangeProperty(FPropertyChangedEvent& InPropertyChangedEvent) override;
 	virtual void PostEditChangeChainProperty(FPropertyChangedChainEvent& InPropertyChangedEvent) override;
 	virtual void PostEditUndo() override;
 #endif // WITH_EDITOR
 
-	const UMetasoundEditorGraphMember* FindMember() const;
-	UMetasoundEditorGraphMember* FindMember();
-
-	UE_DEPRECATED(5.5, "Due to serialization optimization, literals now inherit from LiteralMetadata and are no longer parented under members. Use FindMember instead")
 	const UMetasoundEditorGraphMember* GetParentMember() const
 	{
 		return Cast<UMetasoundEditorGraphMember>(GetOuter());
 	}
 
-	UE_DEPRECATED(5.5, "Due to serialization optimization, literals now inherit from LiteralMetadata and are no longer parented under members. Use FindMember instead")
 	UMetasoundEditorGraphMember* GetParentMember()
 	{
 		return Cast<UMetasoundEditorGraphMember>(GetOuter());
@@ -164,8 +176,6 @@ public:
 	void InitializeLiteral();
 
 protected:
-	virtual UMetasoundEditorGraphNode* AddNode(Metasound::Frontend::FNodeHandle InNodeHandle, bool bInSelectNewNode) PURE_VIRTUAL(UMetasoundEditorGraphMember::AddNode, return nullptr;);
-
 	/** Default literal value of member */
 	UPROPERTY()
 	TObjectPtr<UMetasoundEditorGraphMemberDefaultLiteral> Literal;
@@ -174,8 +184,7 @@ protected:
 	UPROPERTY()
 	FName TypeName;
 
-	friend class UMetasoundEditorGraph;
-	friend class UMetaSoundEditorSubsystem;
+friend class UMetasoundEditorGraph;
 };
 
 /** Base class for an input or output of the graph. */
@@ -185,6 +194,9 @@ class METASOUNDEDITOR_API UMetasoundEditorGraphVertex : public UMetasoundEditorG
 	GENERATED_BODY()
 
 protected:
+	UE_DEPRECATED(5.1, "Use AddNodeHandle with FCreateNodeVertexParams instead.")
+	virtual Metasound::Frontend::FNodeHandle AddNodeHandle(const FName& InNodeName, FName InDataType) PURE_VIRTUAL(UMetasoundEditorGraphVertex::AddNodeHandle, return Metasound::Frontend::INodeController::GetInvalidHandle(); )
+
 	/** Adds the node handle for a newly created vertex. */
 	virtual Metasound::Frontend::FNodeHandle AddNodeHandle(const FName& InNodeName, const Metasound::Editor::FCreateNodeVertexParams& InParams) PURE_VIRTUAL(UMetasoundEditorGraphVertex::AddNodeHandle, return Metasound::Frontend::INodeController::GetInvalidHandle(); )
 
@@ -211,6 +223,7 @@ public:
 	virtual void SetDescription(const FText& InDescription, bool bPostTransaction) override;
 	virtual void SetMemberName(const FName& InNewName, bool bPostTransaction) override;
 	virtual void SetDisplayName(const FText& InNewName, bool bPostTransaction) override;
+	virtual void SetDataType(FName InNewType, bool bPostTransaction = true) override;
 	/* ~End UMetasoundEditorGraphMember interface */
 
 	/** Version of interface membership, or invalid version if not an interface member. */
@@ -220,19 +233,15 @@ public:
 	virtual bool IsInterfaceMember() const;
 
 	/** Returns the Metasound class type of the associated node */
-	virtual EMetasoundFrontendClassType GetClassType() const PURE_VIRTUAL(UMetasoundEditorGraphVertex::GetClassType, return EMetasoundFrontendClassType::Invalid; )
+	virtual EMetasoundFrontendClassType GetClassType() const PURE_VIRTUAL(UMetasoundEditorGraphMember::GetClassType, return EMetasoundFrontendClassType::Invalid; )
 
 	/** Returns the SortOrderIndex assigned to this member. */
-	virtual int32 GetSortOrderIndex() const PURE_VIRTUAL(UMetasoundEditorGraphVertex::GetSortOrderIndex, return 0; )
-
-	/** Sets the DataType assigned to this member. */
-	virtual void SetDataType(FName DataType, bool bPostTransaction = true) PURE_VIRTUAL(UMetasoundEditorGraphVertex::SetDataType, )
+	virtual int32 GetSortOrderIndex() const PURE_VIRTUAL(UMetasoundEditorGraphMember::GetSortOrderIndex, return 0; )
 
 	/** Sets the SortOrderIndex assigned to this member. */
-	virtual void SetSortOrderIndex(int32 InSortOrderIndex) PURE_VIRTUAL(UMetasoundEditorGraphVertex::SetSortOrderIndex, )
+	virtual void SetSortOrderIndex(int32 InSortOrderIndex) PURE_VIRTUAL(UMetasoundEditorGraphMember::SetSortOrderIndex, )
 	
-	/** Sets the VertexAccessType assigned to this member. */
-	virtual void SetVertexAccessType(EMetasoundFrontendVertexAccessType InNewAccessType, bool bPostTransaction = true) PURE_VIRTUAL(UMetasoundEditorGraphVertex::SetVertexAccessType, )
+	void SetVertexAccessType(EMetasoundFrontendVertexAccessType InNewAccessType, bool bPostTransaction = true);
 
 	/** Returns the node handle associated with the vertex. */
 	Metasound::Frontend::FNodeHandle GetNodeHandle();
@@ -243,8 +252,6 @@ public:
 	virtual EMetasoundFrontendVertexAccessType GetVertexAccessType() const PURE_VIRTUAL(UMetasoundEditorGraphVertex::GetVertexAccessType, return EMetasoundFrontendVertexAccessType::Reference; );
 
 	virtual bool CanRename() const override;
-
-	friend class Metasound::Editor::FEditor;
 };
 
 UCLASS()
@@ -255,20 +262,16 @@ class METASOUNDEDITOR_API UMetasoundEditorGraphInput : public UMetasoundEditorGr
 public:
 	virtual int32 GetSortOrderIndex() const override;
 	virtual void SetSortOrderIndex(int32 InSortOrderIndex) override;
-	virtual TArray<UMetasoundEditorGraphMemberNode*> GetNodes() const override;
 
 	virtual const FText& GetGraphMemberLabel() const override;
 	virtual void ResetToClassDefault() override;
-	virtual void SetDataType(FName InNewType, bool bPostTransaction) override;
-	virtual void SetVertexAccessType(EMetasoundFrontendVertexAccessType InNewAccessType, bool bPostTransaction) override;
 	virtual void SetMemberName(const FName& InNewName, bool bPostTransaction) override;
 	virtual void UpdateFrontendDefaultLiteral(bool bPostTransaction) override;
 	virtual EMetasoundFrontendVertexAccessType GetVertexAccessType() const override;
 
 protected:
-	virtual UMetasoundEditorGraphNode* AddNode(Metasound::Frontend::FNodeHandle InNodeHandle, bool bInSelectNewNode) override;
+	virtual Metasound::Frontend::FNodeHandle AddNodeHandle(const FName& InNodeName, FName InDataType) override;
 	virtual Metasound::Frontend::FNodeHandle AddNodeHandle(const FName& InNodeName, const Metasound::Editor::FCreateNodeVertexParams& InParams) override;
-
 	virtual EMetasoundFrontendClassType GetClassType() const override { return EMetasoundFrontendClassType::Input; }
 	virtual Metasound::Editor::ENodeSection GetSectionID() const override;
 };
@@ -284,13 +287,11 @@ public:
 	virtual void SetSortOrderIndex(int32 InSortOrderIndex) override;
 	virtual const FText& GetGraphMemberLabel() const override;
 	virtual void ResetToClassDefault() override;
-	virtual void SetDataType(FName InNewType, bool bPostTransaction) override;
-	virtual void SetVertexAccessType(EMetasoundFrontendVertexAccessType InNewAccessType, bool bPostTransaction) override;
 	virtual void UpdateFrontendDefaultLiteral(bool bPostTransaction) override;
 	virtual EMetasoundFrontendVertexAccessType GetVertexAccessType() const override;
 
 protected:
-	virtual UMetasoundEditorGraphNode* AddNode(Metasound::Frontend::FNodeHandle InNodeHandle, bool bInSelectNewNode) override;
+	virtual Metasound::Frontend::FNodeHandle AddNodeHandle(const FName& InNodeName, FName InDataType) override;
 	virtual Metasound::Frontend::FNodeHandle AddNodeHandle(const FName& InNodeName, const Metasound::Editor::FCreateNodeVertexParams& InParams) override;
 	virtual EMetasoundFrontendClassType GetClassType() const override { return EMetasoundFrontendClassType::Output; }
 	virtual Metasound::Editor::ENodeSection GetSectionID() const override;
@@ -337,9 +338,6 @@ public:
 
 	virtual bool CanRename() const override;
 
-protected:
-	virtual UMetasoundEditorGraphNode* AddNode(Metasound::Frontend::FNodeHandle InNodeHandle, bool bInSelectNewNode) override;
-
 private:
 	struct FVariableEditorNodes
 	{
@@ -367,7 +365,6 @@ class METASOUNDEDITOR_API UMetasoundEditorGraph : public UMetasoundEditorGraphBa
 	GENERATED_BODY()
 
 public:
-	UE_DEPRECATED(5.5, "Input node creation is no longer supported publically via the MetasoundEditorGraph.")
 	UMetasoundEditorGraphInputNode* CreateInputNode(Metasound::Frontend::FNodeHandle InNodeHandle, bool bInSelectNewNode);
 
 	Metasound::Frontend::FDocumentHandle GetDocumentHandle();
@@ -377,7 +374,6 @@ public:
 
 	virtual void PreSave(FObjectPreSaveContext InSaveContext) override;
 
-	UMetaSoundBuilderBase& GetBuilderChecked();
 	UObject* GetMetasound();
 	const UObject* GetMetasound() const;
 	UObject& GetMetasoundChecked();
@@ -398,15 +394,25 @@ public:
 
 	// UMetasoundEditorGraphBase Implementation
 	virtual void RegisterGraphWithFrontend() override;
-	virtual void MigrateEditorDocumentData(FMetaSoundFrontendDocumentBuilder& OutBuilder) override;
+	virtual FMetasoundFrontendDocumentModifyContext& GetModifyContext() override;
+	virtual const FMetasoundFrontendDocumentModifyContext& GetModifyContext() const override;
+
+	virtual void ClearVersionedOnLoad() override;
+	virtual bool GetVersionedOnLoad() const override;
+	virtual void SetVersionedOnLoad() override;
 
 private:
+	bool RemoveFrontendInput(UMetasoundEditorGraphInput& Input);
+	bool RemoveFrontendOutput(UMetasoundEditorGraphOutput& Output);
+	bool RemoveFrontendVariable(UMetasoundEditorGraphVariable& Variable);
 	void ValidateInternal(Metasound::Editor::FGraphValidationResults& OutResults);
 
 	// Preview ID is the Unique ID provided by the UObject that implements
 	// a sound's ParameterInterface when a sound begins playing.
 	uint32 PreviewID = INDEX_NONE;
 
+	bool bVersionedOnLoad = false;
+
 	// Used as a means of forcing the graph to rebuild nodes on next tick.
 	// TODO: Will no longer require this once all editor metadata is migrated
 	// to the frontend & the system can adequately rely on the changeIDs as a
@@ -425,26 +431,21 @@ private:
 public:
 	UMetasoundEditorGraphInput* FindInput(FGuid InNodeID) const;
 	UMetasoundEditorGraphInput* FindInput(FName InName) const;
-	UMetasoundEditorGraphInput* FindOrAddInput(Metasound::Frontend::FConstNodeHandle InNodeHandle);
+	UMetasoundEditorGraphInput* FindOrAddInput(Metasound::Frontend::FNodeHandle InNodeHandle);
 
 	UMetasoundEditorGraphOutput* FindOutput(FGuid InNodeID) const;
 	UMetasoundEditorGraphOutput* FindOutput(FName InName) const;
-	UMetasoundEditorGraphOutput* FindOrAddOutput(Metasound::Frontend::FConstNodeHandle InNodeHandle);
+	UMetasoundEditorGraphOutput* FindOrAddOutput(Metasound::Frontend::FNodeHandle InNodeHandle);
 
 	UMetasoundEditorGraphVariable* FindVariable(const FGuid& InVariableID) const;
 	UMetasoundEditorGraphVariable* FindOrAddVariable(const Metasound::Frontend::FConstVariableHandle& InVariableHandle);
 
-	UMetasoundEditorGraphMember* FindMember(FGuid InMemberID) const;
+	UMetasoundEditorGraphMember* FindMember(FGuid InNodeID) const;
 	UMetasoundEditorGraphMember* FindAdjacentMember(const UMetasoundEditorGraphMember& InMember);
 
-	UE_DEPRECATED(5.5, "Use the associated Frontend builder and synchronize the graph remove a given member.")
-	bool RemoveMember(UMetasoundEditorGraphMember& InGraphMember) { return false; }
-
-	UE_DEPRECATED(5.5, "Use the associated Frontend builder and synchronize the graph remove all member nodes.")
-	bool RemoveMemberNodes(UMetasoundEditorGraphMember& InGraphMember) { return false; }
-
-	UE_DEPRECATED(5.5, "Use RemoveMember instead")
-	bool RemoveFrontendMember(UMetasoundEditorGraphMember& InGraphMember) { return false; }
+	bool RemoveMember(UMetasoundEditorGraphMember& InGraphMember);
+	bool RemoveMemberNodes(UMetasoundEditorGraphMember& InGraphMember);
+	bool RemoveFrontendMember(UMetasoundEditorGraphMember& InMember);
 
 	friend class UMetaSoundFactory;
 	friend class UMetaSoundSourceFactory;
diff --git a/Engine/Plugins/Runtime/Metasound/Source/MetasoundEditor/Public/MetasoundEditorGraphCommentNode.h b/Engine/Plugins/Runtime/Metasound/Source/MetasoundEditor/Public/MetasoundEditorGraphCommentNode.h
--- a/Engine/Plugins/Runtime/Metasound/Source/MetasoundEditor/Public/MetasoundEditorGraphCommentNode.h
+++ b/Engine/Plugins/Runtime/Metasound/Source/MetasoundEditor/Public/MetasoundEditorGraphCommentNode.h
@@ -14,7 +14,6 @@
 
 struct FPropertyChangedEvent;
 class FMetasoundAssetBase;
-class UMetaSoundBuilderBase;
 
 
 // Forward Declarations
@@ -41,21 +40,17 @@ public:
 	//~ End UObject Interface
 
 	//~ Begin UEdGraphNode Interface
-	METASOUNDEDITOR_API virtual bool CanUserDeleteNode() const override;
-	METASOUNDEDITOR_API virtual void OnRenameNode(const FString& NewName) override;
 	METASOUNDEDITOR_API virtual void ResizeNode(const FVector2D& NewSize) override;
+	METASOUNDEDITOR_API virtual void PostPlacedNewNode() override;
+	METASOUNDEDITOR_API virtual void OnRenameNode(const FString& NewName) override;
 	//~ End UEdGraphNode Interface
 
 	/** Set the Bounds for the comment node */
 	METASOUNDEDITOR_API void SetBounds(const class FSlateRect& Rect);
 
-	UMetaSoundBuilderBase& GetBuilderChecked() const;
-
 	// Returns frontend comment ID
 	FGuid GetCommentID() const;
 
-	bool RemoveFromDocument() const;
-
 private:
 	FMetasoundAssetBase& GetAssetChecked();
 
diff --git a/Engine/Plugins/Runtime/Metasound/Source/MetasoundEditor/Public/MetasoundEditorGraphInputNode.h b/Engine/Plugins/Runtime/Metasound/Source/MetasoundEditor/Public/MetasoundEditorGraphInputNode.h
--- a/Engine/Plugins/Runtime/Metasound/Source/MetasoundEditor/Public/MetasoundEditorGraphInputNode.h
+++ b/Engine/Plugins/Runtime/Metasound/Source/MetasoundEditor/Public/MetasoundEditorGraphInputNode.h
@@ -25,31 +25,19 @@ public:
 	UPROPERTY()
 	TObjectPtr<UMetasoundEditorGraphInput> Input;
 
-	UPROPERTY()
-	FGuid NodeID;
-
 public:
-	const FMetasoundEditorGraphVertexNodeBreadcrumb& GetBreadcrumb() const;
-
-	virtual void CacheTitle() override;
-
-	virtual void CacheBreadcrumb() override;
 	virtual UMetasoundEditorGraphMember* GetMember() const override;
-	virtual void ReconstructNode() override;
-
 	virtual void UpdatePreviewInstance(const Metasound::FVertexName& InParameterName, TScriptInterface<IAudioParameterControllerInterface>& InParameterInterface) const;
+
 	virtual FMetasoundFrontendClassName GetClassName() const;
 
 	virtual FGuid GetNodeID() const override;
 	virtual FLinearColor GetNodeTitleColor() const override;
 	virtual FSlateIcon GetNodeTitleIcon() const override;
-	virtual void GetPinHoverText(const UEdGraphPin& Pin, FString& OutHoverText) const override;
 	virtual void Validate(Metasound::Editor::FGraphNodeValidationResult& OutResult) override;
 
 protected:
-	// Breadcrumb used if associated FrontendNode cannot be found or has been unlinked
-	UPROPERTY()
-	FMetasoundEditorGraphVertexNodeBreadcrumb Breadcrumb;
+	virtual void SetNodeID(FGuid InNodeID) override;
 
 	friend class Metasound::Editor::FGraphBuilder;
 };
