[Backout] - CL33084850

diff --git a/Engine/Plugins/Runtime/Metasound/Source/MetasoundEditor/Public/MetasoundEditorGraphNode.h b/Engine/Plugins/Runtime/Metasound/Source/MetasoundEditor/Public/MetasoundEditorGraphNode.h
--- a/Engine/Plugins/Runtime/Metasound/Source/MetasoundEditor/Public/MetasoundEditorGraphNode.h
+++ b/Engine/Plugins/Runtime/Metasound/Source/MetasoundEditor/Public/MetasoundEditorGraphNode.h
@@ -17,62 +17,24 @@
 
 // Forward Declarations
 class UEdGraphPin;
-class UMetaSoundBuilderBase;
 class UMetaSoundPatch;
 class UMetasoundEditorGraphOutput;
 class UMetasoundEditorGraphMember;
 class UMetasoundEditorGraphVariable;
 class UMetasoundEditorGraphMemberDefaultFloat;
-
 namespace EMessageSeverity { enum Type : int; }
 
-namespace Metasound::Editor
-{
-	struct FDocumentClipboardUtils;
-	struct FGraphNodeValidationResult;
-
-	class FGraphBuilder;
-
-	// Map of class names to sorted array of registered version numbers
-	using FSortedClassVersionMap = TMap<FName, TArray<FMetasoundFrontendVersionNumber>>;
-} // namespace Metasound::Editor
-
-
-USTRUCT()
-struct FMetasoundEditorGraphNodeBreadcrumb
-{
-	GENERATED_BODY()
-
-	UPROPERTY()
-	FMetasoundFrontendClassName ClassName;
-
-	UPROPERTY()
-	bool bIsClassNative = true;
-};
-
-USTRUCT()
-struct FMetasoundEditorGraphMemberNodeBreadcrumb : public FMetasoundEditorGraphNodeBreadcrumb
+namespace Metasound
 {
-	GENERATED_BODY()
-
-	UPROPERTY()
-	FName MemberName;
-
-	UPROPERTY()
-	FName DataType;
-
-	UPROPERTY()
-	FMetasoundFrontendLiteral DefaultLiteral;
-};
-
-USTRUCT()
-struct FMetasoundEditorGraphVertexNodeBreadcrumb : public FMetasoundEditorGraphMemberNodeBreadcrumb
-{
-	GENERATED_BODY()
+	namespace Editor
+	{
+		struct FGraphNodeValidationResult;
+		class FGraphBuilder;
 
-	UPROPERTY()
-	EMetasoundFrontendVertexAccessType AccessType = EMetasoundFrontendVertexAccessType::Unset;
-};
+		// Map of class names to sorted array of registered version numbers
+		using FSortedClassVersionMap = TMap<FName, TArray<FMetasoundFrontendVersionNumber>>;
+	} // namespace Editor
+} // namespace Metasound
 
 UCLASS(MinimalAPI)
 class UMetasoundEditorGraphNode : public UEdGraphNode
@@ -102,7 +64,6 @@ public:
 	virtual void PinDefaultValueChanged(UEdGraphPin* Pin) override;
 	virtual void ReconstructNode() override;
 	virtual FString GetPinMetaData(FName InPinName, FName InKey) override;
-	virtual void OnUpdateCommentText(const FString& NewComment) override;
 	// End of UEdGraphNode interface
 
 	// UObject interface
@@ -117,36 +78,21 @@ public:
 
 	virtual FSlateIcon GetNodeTitleIcon() const { return FSlateIcon(); }
 	virtual FName GetCornerIcon() const { return FName(); }
-	virtual bool CanAddInputPin() const { return false; }
+	virtual bool CanAddInputPin() const
+	{
+		return false;
+	}
 
-	UMetaSoundBuilderBase& GetBuilderChecked() const;
+	const FMetaSoundFrontendDocumentBuilder& GetFrontendBuilderChecked() const;
 	const FMetasoundFrontendNode* GetFrontendNode() const;
 	const FMetasoundFrontendNode& GetFrontendNodeChecked() const;
 
-	virtual const FMetasoundEditorGraphNodeBreadcrumb& GetBreadcrumb() const;
-
-	// Caches any "breadcrumb" data associated with a particular MetaSound Editor node. Called
-	// when copying edgraph data to the clipboard or validating for fast access. Also generally
-	// provides a mechanism for MetaSound nodes to cache Frontend data for use to look-up Frontend
-	// data if re-associated should the associated document data/node becomes unlinked.
-	virtual void CacheBreadcrumb() { }
-
 	UObject& GetMetasoundChecked();
 	const UObject& GetMetasoundChecked() const;
 
-	virtual bool RemoveFromDocument() const;
-
-	UE_DEPRECATED(5.4, "Use UpdateFrontendNodeLocation and/or SyncLocationFromFrontendNode")
+	// Sets the node's location, both on this graph member node and on the frontend handle
 	void SetNodeLocation(const FVector2D& InLocation);
-
-	// Finds the associated node with the given ID and sets this EdGraphNode's comment and comment visibility boolean.
-	void SyncCommentFromFrontendNode();
-
-	// Finds the associated node with the given ID and sets this EdGraphNode's location.
-	// Returns whether or not node ID entry exists and if location was set.
-	bool SyncLocationFromFrontendNode(bool bUpdateEditorNodeID = false);
-
-	// Helper function that sets the associated frontend node's location. Does NOT set this EdGraphNode's location.
+	// Helper function to update node location on frontend handle
 	void UpdateFrontendNodeLocation(const FVector2D& InLocation);
 
 	Metasound::Frontend::FGraphHandle GetRootGraphHandle() const;
@@ -157,9 +103,7 @@ public:
 
 	TSet<FString> GetDisallowedPinClassNames(const UEdGraphPin& InPin) const;
 
-	UE_DEPRECATED(5.4, "Use the Frontend node or Breadcrumb directly to get the class name.")
 	virtual FMetasoundFrontendClassName GetClassName() const { return FMetasoundFrontendClassName(); }
-
 	virtual FGuid GetNodeID() const { return FGuid(); }
 	virtual FText GetDisplayName() const;
 	virtual void CacheTitle();
@@ -168,6 +112,7 @@ public:
 	// Mark node for refresh
 	void SyncChangeIDs();
 
+
 	FText GetCachedTitle() const { return CachedTitle; }
 
 	// Returns whether or not the class interface, metadata, or style has been changed since the last node refresh
@@ -182,10 +127,9 @@ protected:
 	// a new name if the external definition changes between application sessions.
 	FText CachedTitle;
 
-	static bool ShowNodeDebugData();
-
-	UE_DEPRECATED(5.4, "Now set directly on implementing nodes")
 	virtual void SetNodeID(FGuid InNodeID) { }
+
+	friend class Metasound::Editor::FGraphBuilder;
 };
 
 /** Node that represents a graph member */
@@ -203,7 +147,7 @@ public:
 	{
 		return true;
 	}
-
+protected:
 	// Clamp float literal value based on the given default float literal. 
 	// Returns whether the literal was clamped. 
 	static bool ClampFloatLiteral(const UMetasoundEditorGraphMemberDefaultFloat* DefaultFloatLiteral, FMetasoundFrontendLiteral& LiteralValue);
@@ -219,12 +163,7 @@ public:
 	UPROPERTY()
 	TObjectPtr<UMetasoundEditorGraphOutput> Output;
 
-	virtual const FMetasoundEditorGraphNodeBreadcrumb& GetBreadcrumb() const override;
-	virtual void CacheBreadcrumb() override;
-
-	UE_DEPRECATED(5.4, "Use the Frontend node or Breadcrumb directly to get the class name.")
 	virtual FMetasoundFrontendClassName GetClassName() const override;
-
 	virtual FGuid GetNodeID() const override;
 	virtual UMetasoundEditorGraphMember* GetMember() const override;
 
@@ -238,8 +177,6 @@ public:
 	virtual bool CanUserDeleteNode() const override;
 
 	virtual void PinDefaultValueChanged(UEdGraphPin* InPin) override;
-	virtual void ReconstructNode() override;
-	virtual bool RemoveFromDocument() const override;
 
 	// Disables interact widgets (ex. sliders, knobs) when input is connected
 	virtual bool EnableInteractWidgets() const override;
@@ -247,16 +184,10 @@ public:
 	virtual void Validate(Metasound::Editor::FGraphNodeValidationResult& OutResult) override;
 
 protected:
-	// Breadcrumb used if associated FrontendNode cannot be found or has been unlinked
-	UPROPERTY()
-	FMetasoundEditorGraphVertexNodeBreadcrumb Breadcrumb;
-
 	virtual FLinearColor GetNodeTitleColor() const override;
 	virtual FSlateIcon GetNodeTitleIcon() const override;
 	virtual void SetNodeID(FGuid InNodeID) override;
 
-	// Friended to enable mutation of Frontend NodeID & direct breadcrumb access
-	friend struct Metasound::Editor::FDocumentClipboardUtils;
 	friend class Metasound::Editor::FGraphBuilder;
 };
 
@@ -267,29 +198,25 @@ class UMetasoundEditorGraphExternalNode : public UMetasoundEditorGraphNode
 
 protected:
 	UPROPERTY()
-	FMetasoundEditorGraphNodeBreadcrumb Breadcrumb;
-
-	UPROPERTY(meta = (Deprecated = "5.4", DeprecationMessage = "Use Breadcrumb value when manipulating clipboard data or validating. Otherwise, look-up frontend node's associated class directly"))
 	FMetasoundFrontendClassName ClassName;
 
 	UPROPERTY()
 	FGuid NodeID;
 
-	UPROPERTY(meta = (Deprecated = "5.4", DeprecationMessage = "Use Breadcrumb value when manipulating clipboard data or validating, Otherwise, look-up frontend node's associated class directly"))
+	// Whether or not the referenced class is natively defined
+	// (false if defined in another asset). Cached from node
+	// implementation for fast access when validated.
+	UPROPERTY()
 	bool bIsClassNative = true;
 
 public:
-	UE_DEPRECATED(5.4, "Use the Frontend node or Breadcrumb directly to get the class name.")
-	virtual FMetasoundFrontendClassName GetClassName() const override { return Breadcrumb.ClassName; }
-
-	virtual const FMetasoundEditorGraphNodeBreadcrumb& GetBreadcrumb() const;
+	virtual FMetasoundFrontendClassName GetClassName() const override { return ClassName; }
 	virtual FGuid GetNodeID() const override { return NodeID; }
 	virtual FLinearColor GetNodeTitleColor() const override;
 	virtual FSlateIcon GetNodeTitleIcon() const override;
 	virtual bool ShouldDrawNodeAsControlPointOnly(int32& OutInputPinIndex, int32& OutOutputPinIndex) const override;
 
 	virtual void ReconstructNode() override;
-	virtual void CacheBreadcrumb() override;
 	virtual void CacheTitle() override;
 	virtual void GetPinHoverText(const UEdGraphPin& Pin, FString& OutHoverText) const override;
 
@@ -301,8 +228,11 @@ public:
 
 
 protected:
-	// Friended to enable mutation of Frontend NodeID & direct breadcrumb access
-	friend struct Metasound::Editor::FDocumentClipboardUtils;
+	virtual void SetNodeID(FGuid InNodeID) override
+	{
+		NodeID = InNodeID;
+	}
+
 	friend class Metasound::Editor::FGraphBuilder;
 };
 
@@ -333,13 +263,8 @@ public:
 	// Variables do not have titles to distinguish more visually from vertex types
 	virtual void CacheTitle() override { }
 
-	virtual const FMetasoundEditorGraphNodeBreadcrumb& GetBreadcrumb() const override;
-
-	virtual void CacheBreadcrumb() override;
 	virtual UMetasoundEditorGraphMember* GetMember() const override;
 	virtual bool EnableInteractWidgets() const override;
-
-	
 	virtual FMetasoundFrontendClassName GetClassName() const override;
 	virtual FGuid GetNodeID() const override;
 	virtual FName GetCornerIcon() const override;
@@ -348,15 +273,9 @@ public:
 	virtual EMetasoundFrontendClassType GetClassType() const;
 
 protected:
-	// Breadcrumb used if associated FrontendNode cannot be found or has been unlinked
-	UPROPERTY()
-	FMetasoundEditorGraphMemberNodeBreadcrumb Breadcrumb;
-
 	virtual FLinearColor GetNodeTitleColor() const override;
 	virtual FSlateIcon GetNodeTitleIcon() const override;
 	virtual void SetNodeID(FGuid InNodeID) override;
 
-	// Friended to enable mutation of Frontend NodeID & direct breadcrumb access
-	friend struct Metasound::Editor::FDocumentClipboardUtils;
 	friend class Metasound::Editor::FGraphBuilder;
 };
diff --git a/Engine/Plugins/Runtime/Metasound/Source/MetasoundEditor/Public/MetasoundEditorSubsystem.h b/Engine/Plugins/Runtime/Metasound/Source/MetasoundEditor/Public/MetasoundEditorSubsystem.h
--- a/Engine/Plugins/Runtime/Metasound/Source/MetasoundEditor/Public/MetasoundEditorSubsystem.h
+++ b/Engine/Plugins/Runtime/Metasound/Source/MetasoundEditor/Public/MetasoundEditorSubsystem.h
@@ -12,8 +12,6 @@
 
 // Forward Declarations
 class UMetaSoundBuilderBase;
-class UMetasoundEditorGraphMember;
-class UMetasoundEditorGraphMemberDefaultLiteral;
 
 /** The subsystem in charge of editor MetaSound functionality */
 UCLASS()
@@ -22,15 +20,6 @@ class METASOUNDEDITOR_API UMetaSoundEditorSubsystem : public UEditorSubsystem
 	GENERATED_BODY()
 
 public:
-	// Binds literal editor Metadata to the given member.  If the literal already exists, adds literal
-	// reference to given member (asserts that existing literal is of similar subclass provided).  If 
-	// it does not exist, or an optional template object is provided, metadata is generated then bound.
-	// Returns true if new literal metadata was generated, false if not. Asserts if bind failed.
-	bool BindMemberMetadata(
-		UMetasoundEditorGraphMember& InMember,
-		TSubclassOf<UMetasoundEditorGraphMemberDefaultLiteral> LiteralClass,
-		UMetasoundEditorGraphMemberDefaultLiteral* TemplateObject = nullptr);
-
 	// Build the given builder to a MetaSound asset
 	// @param Author - Sets the author on the given builder's document.
 	// @param AssetName - Name of the asset to build.
@@ -53,7 +42,7 @@ public:
 		UPARAM(DisplayName = "Node Handle") const FMetaSoundNodeHandle& InNode,
 		UPARAM(DisplayName = "Location") const FVector2D& InLocation,
 		EMetaSoundBuilderResult& OutResult);
-
+	
 	// Initialize the UObject asset, with an optional MetaSound to be referenced if the asset is a preset
 	void InitAsset(UObject& InNewMetaSound, UObject* InReferencedMetaSound = nullptr);
 
diff --git a/Engine/Plugins/Runtime/Metasound/Source/MetasoundEngine/Private/Metasound.cpp b/Engine/Plugins/Runtime/Metasound/Source/MetasoundEngine/Private/Metasound.cpp
--- a/Engine/Plugins/Runtime/Metasound/Source/MetasoundEngine/Private/Metasound.cpp
+++ b/Engine/Plugins/Runtime/Metasound/Source/MetasoundEngine/Private/Metasound.cpp
@@ -11,7 +11,6 @@
 #include "MetasoundEngineEnvironment.h"
 #include "MetasoundEnvironment.h"
 #include "MetasoundFrontendController.h"
-#include "MetasoundFrontendDocument.h"
 #include "MetasoundFrontendQuery.h"
 #include "MetasoundFrontendQuerySteps.h"
 #include "MetasoundFrontendRegistries.h"
@@ -53,6 +52,7 @@ int32 UMetasoundEditorGraphBase::GetHighestMessageSeverity() const
 	return HighestMessageSeverity;
 }
 
+
 UMetaSoundPatch::UMetaSoundPatch(const FObjectInitializer& ObjectInitializer)
 	: Super(ObjectInitializer)
 	, FMetasoundAssetBase()
@@ -135,37 +135,26 @@ void UMetaSoundPatch::Serialize(FArchive& InArchive)
 }
 
 #if WITH_EDITORONLY_DATA
-void UMetaSoundPatch::MigrateEditorGraph(FMetaSoundFrontendDocumentBuilder& OutBuilder)
-{
-	PRAGMA_DISABLE_DEPRECATION_WARNINGS
-	if (Graph)
-	{
-		Graph->MigrateEditorDocumentData(OutBuilder);
-		Graph = nullptr;
-	}
-	PRAGMA_ENABLE_DEPRECATION_WARNINGS
-}
-
 UEdGraph* UMetaSoundPatch::GetGraph()
 {
-	return EditorGraph;
+	return Graph;
 }
 
 const UEdGraph* UMetaSoundPatch::GetGraph() const
 {
-	return EditorGraph;
+	return Graph;
 }
 
 UEdGraph& UMetaSoundPatch::GetGraphChecked()
 {
-	check(EditorGraph);
-	return *EditorGraph;
+	check(Graph);
+	return *Graph;
 }
 
 const UEdGraph& UMetaSoundPatch::GetGraphChecked() const
 {
-	check(EditorGraph);
-	return *EditorGraph;
+	check(Graph);
+	return *Graph;
 }
 FText UMetaSoundPatch::GetDisplayName() const
 {
diff --git a/Engine/Plugins/Runtime/Metasound/Source/MetasoundEngine/Private/MetasoundAssetSubsystem.cpp b/Engine/Plugins/Runtime/Metasound/Source/MetasoundEngine/Private/MetasoundAssetSubsystem.cpp
--- a/Engine/Plugins/Runtime/Metasound/Source/MetasoundEngine/Private/MetasoundAssetSubsystem.cpp
+++ b/Engine/Plugins/Runtime/Metasound/Source/MetasoundEngine/Private/MetasoundAssetSubsystem.cpp
@@ -144,7 +144,7 @@ void UMetaSoundAssetSubsystem::AddAssetReferences(FMetasoundAssetBase& InAssetBa
 	using namespace Metasound;
 	using namespace Metasound::Frontend;
 
-	const FMetasoundFrontendDocument& Document = InAssetBase.GetConstDocumentChecked();
+	const FMetasoundFrontendDocument& Document = InAssetBase.GetDocumentChecked();
 	const FNodeRegistryKey AssetClassKey = FNodeRegistryKey(Document.RootGraph);
 
 	if (!ContainsKey(AssetClassKey))
@@ -179,7 +179,7 @@ void UMetaSoundAssetSubsystem::AddAssetReferences(FMetasoundAssetBase& InAssetBa
 	{
 		if (Asset)
 		{
-			const FMetasoundFrontendDocument& RefDocument = Asset->GetConstDocumentChecked();
+			const FMetasoundFrontendDocument& RefDocument = Asset->GetDocumentChecked();
 			const FNodeRegistryKey ClassKey = FNodeRegistryKey(RefDocument.RootGraph);
 			if (!ContainsKey(ClassKey))
 			{
@@ -214,7 +214,7 @@ Metasound::Frontend::FNodeRegistryKey UMetaSoundAssetSubsystem::AddOrUpdateAsset
 	const FMetasoundAssetBase* MetaSoundAsset = Metasound::IMetasoundUObjectRegistry::Get().GetObjectAsAssetBase(&InObject);
 	check(MetaSoundAsset);
 
-	const FMetasoundFrontendDocument& Document = MetaSoundAsset->GetConstDocumentChecked();
+	const FMetasoundFrontendDocument& Document = MetaSoundAsset->GetDocumentChecked();
 	const FNodeRegistryKey RegistryKey = FNodeRegistryKey(Document.RootGraph);
 
 	if (RegistryKey.IsValid())
@@ -275,11 +275,6 @@ Metasound::Frontend::FNodeRegistryKey UMetaSoundAssetSubsystem::AddOrUpdateAsset
 	return FNodeRegistryKey::GetInvalid();
 }
 
-FMetaSoundFrontendDocumentBuilder& UMetaSoundAssetSubsystem::AttachDocumentBuilderChecked(UObject& InObject) const
-{
-	return UMetaSoundBuilderSubsystem::GetChecked().AttachBuilderToAssetChecked(InObject).Builder;
-}
-
 bool UMetaSoundAssetSubsystem::CanAutoUpdate(const FMetasoundFrontendClassName& InClassName) const
 {
 	const UMetaSoundSettings* Settings = GetDefault<UMetaSoundSettings>();
@@ -349,7 +344,7 @@ TSet<UMetaSoundAssetSubsystem::FAssetInfo> UMetaSoundAssetSubsystem::GetReferenc
 	using namespace Metasound::Frontend;
 
 	TSet<FAssetInfo> OutAssetInfos;
-	const FMetasoundFrontendDocument& Document = InAssetBase.GetConstDocumentChecked();
+	const FMetasoundFrontendDocument& Document = InAssetBase.GetDocumentChecked();
 	for (const FMetasoundFrontendClass& Class : Document.Dependencies)
 	{
 		const FNodeRegistryKey Key = FNodeRegistryKey(Class.Metadata);
@@ -627,10 +622,12 @@ void UMetaSoundAssetSubsystem::RemoveAsset(const UObject& InObject)
 	if (const FMetasoundAssetBase* MetaSoundAsset = Metasound::IMetasoundUObjectRegistry::Get().GetObjectAsAssetBase(&InObject))
 	{
 		const FSoftObjectPath ObjectPath(&InObject);
-		const FMetasoundFrontendDocument& Document = MetaSoundAsset->GetConstDocumentChecked();
+		const FMetasoundFrontendDocument& Document = MetaSoundAsset->GetDocumentChecked();
 		FNodeRegistryKey RegistryKey = FNodeRegistryKey(Document.RootGraph);
-		AssetSubsystemPrivate::RemoveIfExactMatch(PathMap, RegistryKey, ObjectPath);
-		UMetaSoundBuilderSubsystem::GetChecked().DetachBuilderFromAsset(Document.RootGraph.Metadata.GetClassName());
+		if (AssetSubsystemPrivate::RemoveIfExactMatch(PathMap, RegistryKey, ObjectPath))
+		{
+			UMetaSoundBuilderSubsystem::GetChecked().DetachBuilderFromAsset(Document.RootGraph.Metadata.GetClassName());
+		}
 	}
 }
 
@@ -644,8 +641,10 @@ void UMetaSoundAssetSubsystem::RemoveAsset(const FAssetData& InAssetData)
 	{
 		const FNodeRegistryKey RegistryKey(ClassInfo);
 		const FSoftObjectPath ObjectPath = InAssetData.GetSoftObjectPath();
-		AssetSubsystemPrivate::RemoveIfExactMatch(PathMap, RegistryKey, ObjectPath);
-		UMetaSoundBuilderSubsystem::GetChecked().DetachBuilderFromAsset(ClassInfo.ClassName);
+		if (AssetSubsystemPrivate::RemoveIfExactMatch(PathMap, RegistryKey, ObjectPath))
+		{
+			UMetaSoundBuilderSubsystem::GetChecked().DetachBuilderFromAsset(ClassInfo.ClassName);
+		}
 	}
 }
 
diff --git a/Engine/Plugins/Runtime/Metasound/Source/MetasoundEngine/Private/MetasoundBuilderSubsystem.cpp b/Engine/Plugins/Runtime/Metasound/Source/MetasoundEngine/Private/MetasoundBuilderSubsystem.cpp
--- a/Engine/Plugins/Runtime/Metasound/Source/MetasoundEngine/Private/MetasoundBuilderSubsystem.cpp
+++ b/Engine/Plugins/Runtime/Metasound/Source/MetasoundEngine/Private/MetasoundBuilderSubsystem.cpp
@@ -26,7 +26,6 @@
 #include "MetasoundTrace.h"
 #include "MetasoundUObjectRegistry.h"
 #include "MetasoundVertex.h"
-#include "NodeTemplates/MetasoundFrontendNodeTemplateInput.h"
 #include "PerPlatformProperties.h"
 
 #include UE_INLINE_GENERATED_CPP_BY_NAME(MetasoundBuilderSubsystem)
@@ -72,6 +71,8 @@ void UMetaSoundBuilderBase::BeginDestroy()
 {
 	if (UMetaSoundBuilderSubsystem* Subsystem = UMetaSoundBuilderSubsystem::Get())
 	{
+		const FMetasoundFrontendDocument& Document = static_cast<const FMetaSoundFrontendDocumentBuilder&>(Builder).GetDocument();
+		const FMetasoundFrontendClassName& ClassName = Document.RootGraph.Metadata.GetClassName();
 		Subsystem->TransientBuilders.Remove(ClassName);
 	}
 
@@ -79,6 +80,7 @@ void UMetaSoundBuilderBase::BeginDestroy()
 	// often holds a TScriptInterface<IMetaSoundDocumentInterface> of
 	// a UPROPERTY that lives on this or derived objects. 
 	Builder.FinishBuilding();
+
 	Super::BeginDestroy();
 }
 
@@ -228,18 +230,18 @@ FMetaSoundNodeHandle UMetaSoundBuilderBase::AddNode(const TScriptInterface<IMeta
 	return NewHandle;
 }
 
-FMetaSoundNodeHandle UMetaSoundBuilderBase::AddNodeByClassName(const FMetasoundFrontendClassName& InClassName, int32 MajorVersion, EMetaSoundBuilderResult& OutResult)
+FMetaSoundNodeHandle UMetaSoundBuilderBase::AddNodeByClassName(const FMetasoundFrontendClassName& ClassName, int32 MajorVersion, EMetaSoundBuilderResult& OutResult)
 {
-	return AddNodeByClassName(InClassName, OutResult, MajorVersion);
+	return AddNodeByClassName(ClassName, OutResult, MajorVersion);
 }
 
-FMetaSoundNodeHandle UMetaSoundBuilderBase::AddNodeByClassName(const FMetasoundFrontendClassName& InClassName, EMetaSoundBuilderResult& OutResult, int32 MajorVersion)
+FMetaSoundNodeHandle UMetaSoundBuilderBase::AddNodeByClassName(const FMetasoundFrontendClassName& ClassName, EMetaSoundBuilderResult& OutResult, int32 MajorVersion)
 {
 	using namespace Metasound;
 	using namespace Metasound::Frontend;
 
 	FMetaSoundNodeHandle NewHandle;
-	if (const FMetasoundFrontendNode* NewNode = Builder.AddNodeByClassName(InClassName, MajorVersion))
+	if (const FMetasoundFrontendNode* NewNode = Builder.AddNodeByClassName(ClassName, MajorVersion))
 	{
 		NewHandle.NodeID = NewNode->GetID();
 	}
@@ -248,31 +250,6 @@ FMetaSoundNodeHandle UMetaSoundBuilderBase::AddNodeByClassName(const FMetasoundF
 	return NewHandle;
 }
 
-void UMetaSoundBuilderBase::BuildInternal(TScriptInterface<IMetaSoundDocumentInterface> NewMetaSound, const FMetasoundFrontendClassName* InDocClassName) const
-{
-	using namespace Metasound::Frontend;
-
-	// If using existing class name, ensure that a builder does not exist for it to
-	// avoid build active flag conflation with the locally generated frontend builder below.
-	if (InDocClassName)
-	{
-		UMetaSoundBuilderSubsystem::GetChecked().DetachBuilderFromAsset(*InDocClassName);
-	}
-
-	FMetaSoundFrontendDocumentBuilder NewDocBuilder(NewMetaSound);
-
-	constexpr bool bResetVersion = false;
-	NewDocBuilder.InitDocument(&GetConstBuilder().GetDocument(), InDocClassName, bResetVersion);
-	NewMetaSound->ConformObjectToDocument();
-}
-
-#if WITH_EDITOR
-bool UMetaSoundBuilderBase::ClearMemberMetadata(const FGuid& InMemberID)
-{
-	return Builder.ClearMemberMetadata(InMemberID);
-}
-#endif // WITH_EDITOR
-
 bool UMetaSoundBuilderBase::ContainsNode(const FMetaSoundNodeHandle& NodeHandle) const
 {
 	return Builder.ContainsNode(NodeHandle.NodeID);
@@ -443,8 +420,6 @@ void UMetaSoundBuilderBase::ConvertFromPreset(EMetaSoundBuilderResult& OutResult
 
 void UMetaSoundBuilderBase::ConvertToPreset(const TScriptInterface<IMetaSoundDocumentInterface>& ReferencedNodeClass, EMetaSoundBuilderResult& OutResult)
 {
-	using namespace Metasound::Frontend;
-
 	const IMetaSoundDocumentInterface* ReferencedInterface = ReferencedNodeClass.GetInterface();
 	if (!ReferencedInterface)
 	{
@@ -452,7 +427,7 @@ void UMetaSoundBuilderBase::ConvertToPreset(const TScriptInterface<IMetaSoundDoc
 		return;
 	}
 	
-	// Ensure the referenced node class isn't transient
+	// Ensure the referenced node class isn't transient 
 	if (Cast<UMetaSoundBuilderDocument>(ReferencedInterface))
 	{
 		UE_LOG(LogMetaSound, Warning, TEXT("Transient document builders cannot be referenced when converting builder '%s' to a preset. Build the referenced node class an asset first or use an existing asset instead"), *GetName());
@@ -477,12 +452,14 @@ void UMetaSoundBuilderBase::ConvertToPreset(const TScriptInterface<IMetaSoundDoc
 	}
 
 	const FMetasoundFrontendDocument& ReferencedDocument = ReferencedInterface->GetConstDocument();
-
-	TSharedRef<FDocumentModifyDelegates> DocumentDelegates = MakeShared<FDocumentModifyDelegates>();
-	InitDelegates(*DocumentDelegates);
-	const bool bConvertToPreset = Builder.ConvertToPreset(ReferencedDocument, DocumentDelegates);
-	ConformObjectToDocument();
-	OutResult = bConvertToPreset ? EMetaSoundBuilderResult::Succeeded : EMetaSoundBuilderResult::Failed;
+	if (Builder.ConvertToPreset(ReferencedDocument))
+	{
+		OutResult = EMetaSoundBuilderResult::Succeeded;
+	}
+	else
+	{
+		OutResult = EMetaSoundBuilderResult::Failed;
+	}
 }
 
 void UMetaSoundBuilderBase::DisconnectNodes(const FMetaSoundBuilderNodeOutputHandle& NodeOutputHandle, const FMetaSoundBuilderNodeInputHandle& NodeInputHandle, EMetaSoundBuilderResult& OutResult)
@@ -716,13 +693,6 @@ FMetaSoundNodeHandle UMetaSoundBuilderBase::FindGraphOutputNode(FName OutputName
 	return { };
 }
 
-#if WITH_EDITOR
-UMetaSoundFrontendMemberMetadata* UMetaSoundBuilderBase::FindMemberMetadata(const FGuid& InMemberID)
-{
-	return Builder.FindMemberMetadata(InMemberID);
-}
-#endif // WITH_EDITOR
-
 UMetaSoundBuilderDocument* UMetaSoundBuilderBase::CreateTransientDocumentObject() const
 {
 	return &UMetaSoundBuilderDocument::Create(GetBuilderUClass());
@@ -746,14 +716,12 @@ UObject* UMetaSoundBuilderBase::GetReferencedPresetAsset() const
 	auto FindExternalNode = [this](const FMetasoundFrontendNode& Node) 
 	{
 		const FMetasoundFrontendClass* Class = Builder.FindDependency(Node.ClassID);
-		check(Class);
 		return Class->Metadata.GetType() == EMetasoundFrontendClassType::External;
 	};
 	const FMetasoundFrontendNode* Node = Builder.GetDocument().RootGraph.Graph.Nodes.FindByPredicate(FindExternalNode);
 	if (Node != nullptr)
 	{
 		const FMetasoundFrontendClass* NodeClass = Builder.FindDependency(Node->ClassID);
-		check(NodeClass);
 		const FNodeRegistryKey NodeClassRegistryKey = FNodeRegistryKey(NodeClass->Metadata);
 		if (FMetasoundAssetBase* Asset = IMetaSoundAssetManager::GetChecked().TryLoadAssetFromKey(NodeClassRegistryKey))
 		{
@@ -776,28 +744,14 @@ void UMetaSoundBuilderBase::InitNodeLocations()
 	Builder.InitNodeLocations();
 }
 
-#if WITH_EDITOR
-void UMetaSoundBuilderBase::InjectInputTemplateNodes(bool bForceNodeCreation, EMetaSoundBuilderResult& OutResult)
-{
-	using namespace Metasound::Frontend;
-	const bool bTemplateNodesInjected = FInputNodeTemplate::GetChecked().Inject(Builder, bForceNodeCreation);
-	OutResult = bTemplateNodesInjected ? EMetaSoundBuilderResult::Succeeded : EMetaSoundBuilderResult::Failed;
-}
-#endif // WITH_EDITOR
-
 bool UMetaSoundBuilderBase::InterfaceIsDeclared(FName InterfaceName) const
 {
 	return Builder.IsInterfaceDeclared(InterfaceName);
 }
 
-void UMetaSoundBuilderBase::InvalidateCache(bool bPrimeCache)
+void UMetaSoundBuilderBase::InvalidateCache()
 {
-	using namespace Metasound::Frontend;
-
-	TSharedRef<FDocumentModifyDelegates> DocumentDelegates = MakeShared<FDocumentModifyDelegates>();
-	InitDelegates(*DocumentDelegates);
-
-	Builder.Reload(DocumentDelegates, bPrimeCache);
+	Builder.InvalidateCache();
 }
 
 bool UMetaSoundBuilderBase::IsPreset() const
@@ -860,18 +814,9 @@ void UMetaSoundBuilderBase::RegisterGraphIfOutstandingTransactions(UObject& InMe
 	MetaSoundAsset->RegisterGraphWithFrontend(Options);
 }
 
-void UMetaSoundBuilderBase::Reload(bool bPrimeCache)
-{
-	using namespace Metasound::Frontend;
-
-	TSharedRef<FDocumentModifyDelegates> DocumentDelegates = MakeShared<FDocumentModifyDelegates>();
-	InitDelegates(*DocumentDelegates);
-	Builder.Reload(DocumentDelegates, bPrimeCache);
-}
-
 void UMetaSoundBuilderBase::ReloadCache(bool bPrimeCache)
 {
-	Reload(bPrimeCache);
+	Builder.ReloadCache();
 }
 
 #if WITH_EDITOR
@@ -899,13 +844,9 @@ void UMetaSoundBuilderBase::RemoveInterface(FName InterfaceName, EMetaSoundBuild
 	OutResult = bInterfaceRemoved ? EMetaSoundBuilderResult::Succeeded : EMetaSoundBuilderResult::Failed;
 }
 
-void UMetaSoundBuilderBase::RemoveNode(const FMetaSoundNodeHandle& NodeHandle, EMetaSoundBuilderResult& OutResult, bool bRemoveUnusedDependencies)
+void UMetaSoundBuilderBase::RemoveNode(const FMetaSoundNodeHandle& NodeHandle, EMetaSoundBuilderResult& OutResult)
 {
 	const bool bNodeRemoved = Builder.RemoveNode(NodeHandle.NodeID);
-	if (bNodeRemoved && bRemoveUnusedDependencies)
-	{
-		Builder.RemoveUnusedDependencies();
-	}
 	OutResult = bNodeRemoved ? EMetaSoundBuilderResult::Succeeded : EMetaSoundBuilderResult::Failed;
 }
 
@@ -915,11 +856,6 @@ void UMetaSoundBuilderBase::RemoveNodeInputDefault(const FMetaSoundBuilderNodeIn
 	OutResult = bInputDefaultRemoved ? EMetaSoundBuilderResult::Succeeded : EMetaSoundBuilderResult::Failed;
 }
 
-void UMetaSoundBuilderBase::RemoveUnusedDependencies()
-{
-	Builder.RemoveUnusedDependencies();
-}
-
 void UMetaSoundBuilderBase::RenameRootGraphClass(const FMetasoundFrontendClassName& InName)
 {
 	Builder.RenameRootGraphClass(InName);
@@ -932,42 +868,11 @@ void UMetaSoundBuilderBase::SetAuthor(const FString& InAuthor)
 }
 #endif // WITH_EDITOR
 
-void UMetaSoundBuilderBase::SetGraphInputAccessType(FName InputName, EMetasoundFrontendVertexAccessType AccessType, EMetaSoundBuilderResult& OutResult)
-{
-	const bool bSet = Builder.SetGraphInputAccessType(InputName, AccessType);
-	OutResult = bSet ? EMetaSoundBuilderResult::Succeeded : EMetaSoundBuilderResult::Failed;
-}
-
-void UMetaSoundBuilderBase::SetGraphInputDataType(FName InputName, FName DataType, EMetaSoundBuilderResult& OutResult)
-{
-	const bool bSet = Builder.SetGraphInputDataType(InputName, DataType);
-	OutResult = bSet ? EMetaSoundBuilderResult::Succeeded : EMetaSoundBuilderResult::Failed;
-}
-
 void UMetaSoundBuilderBase::SetGraphInputDefault(FName InputName, const FMetasoundFrontendLiteral& Literal, EMetaSoundBuilderResult& OutResult)
 {
-	const bool bSet = Builder.SetGraphInputDefault(InputName, Literal);
-	OutResult = bSet ? EMetaSoundBuilderResult::Succeeded : EMetaSoundBuilderResult::Failed;
-}
-
-void UMetaSoundBuilderBase::SetGraphOutputAccessType(FName OutputName, EMetasoundFrontendVertexAccessType AccessType, EMetaSoundBuilderResult& OutResult)
-{
-	const bool bSet = Builder.SetGraphOutputAccessType(OutputName, AccessType);
-	OutResult = bSet ? EMetaSoundBuilderResult::Succeeded : EMetaSoundBuilderResult::Failed;
-}
-
-void UMetaSoundBuilderBase::SetGraphOutputDataType(FName OutputName, FName DataType, EMetaSoundBuilderResult& OutResult)
-{
-	const bool bSet = Builder.SetGraphOutputDataType(OutputName, DataType);
-	OutResult = bSet ? EMetaSoundBuilderResult::Succeeded : EMetaSoundBuilderResult::Failed;
-}
-
-#if WITH_EDITOR
-void UMetaSoundBuilderBase::SetMemberMetadata(UMetaSoundFrontendMemberMetadata& NewMetadata)
-{
-	Builder.SetMemberMetadata(NewMetadata);
+	const bool bInputDefaultSet = Builder.SetGraphInputDefault(InputName, Literal);
+	OutResult = bInputDefaultSet ? EMetaSoundBuilderResult::Succeeded : EMetaSoundBuilderResult::Failed;
 }
-#endif // WITH_EDITOR
 
 void UMetaSoundBuilderBase::SetNodeInputDefault(const FMetaSoundBuilderNodeInputHandle& InputHandle, const FMetasoundFrontendLiteral& Literal, EMetaSoundBuilderResult& OutResult)
 {
@@ -1119,11 +1024,6 @@ void UMetaSoundBuilderBase::UpdateDependencyClassNames(const TMap<FMetasoundFron
 	Builder.UpdateDependencyClassNames(OldToNewReferencedClassNames);
 }
 
-bool UMetaSoundPatchBuilder::ConformObjectToDocument()
-{
-	return false;
-}
-
 void UMetaSoundPatchBuilder::CreateTransientBuilder()
 {
 	using namespace Metasound::Frontend;
@@ -1206,7 +1106,7 @@ void UMetaSoundSourceBuilder::Audition(UObject* Parent, UAudioComponent* AudioCo
 
 	// Must be called post register as register ensures cached runtime data passed to transactor is up-to-date
 	MetaSoundSource.SetDynamicGeneratorEnabled(MetaSoundSource.GetAssetPathChecked(), bLiveUpdatesEnabled);
-	MetaSoundSource.ConformObjectToDocument();
+	MetaSoundSource.ConformObjectDataToInterfaces();
 
 	AudioComponent->SetSound(&MetaSoundSource);
 
@@ -1256,12 +1156,6 @@ TScriptInterface<IMetaSoundDocumentInterface> UMetaSoundSourceBuilder::Build(UOb
 	return &BuildInternal<UMetaSoundSource>(Parent, InBuilderOptions);
 }
 
-bool UMetaSoundSourceBuilder::ConformObjectToDocument()
-{
-	UMetaSoundSource& Source = GetMetaSoundSource();
-	return Source.ConformObjectToDocument();
-}
-
 void UMetaSoundSourceBuilder::CreateTransientBuilder()
 {
 	using namespace Metasound::Frontend;
@@ -2066,7 +1960,7 @@ void UMetaSoundBuilderSubsystem::InvalidateDocumentCache(const FMetasoundFronten
 	TWeakObjectPtr<UMetaSoundBuilderBase> BuilderPtr = AssetBuilders.FindRef(InClassName);
 	if (BuilderPtr.IsValid())
 	{
-		BuilderPtr->Reload();
+		BuilderPtr->InvalidateCache();
 	}
 }
 
@@ -2081,16 +1975,16 @@ UMetaSoundBuilderBase* UMetaSoundBuilderSubsystem::FindBuilderOfDocument(TScript
 	if (const UObject* MetaSoundObject = InMetaSound.GetObject())
 	{
 		const FMetasoundFrontendDocument& Document = InMetaSound->GetConstDocument();
-		const FMetasoundFrontendClassName& InClassName = Document.RootGraph.Metadata.GetClassName();
+		const FMetasoundFrontendClassName& ClassName = Document.RootGraph.Metadata.GetClassName();
 
-		Builder = AssetBuilders.FindRef(InClassName);
+		Builder = AssetBuilders.FindRef(ClassName);
 		if (Builder.IsValid())
 		{
 			ensureAlwaysMsgf(MetaSoundObject->IsAsset(), TEXT("MetaSound is asset but Builder was registered with Subsystem as transient"));
 		}
 		else
 		{
-			Builder = TransientBuilders.FindRef(InClassName);
+			Builder = TransientBuilders.FindRef(ClassName);
 			if (Builder.IsValid())
 			{
 				ensureAlwaysMsgf(!MetaSoundObject->IsAsset(), TEXT("MetaSound is transient object but Builder was registered with Subsystem as asset"));
@@ -2149,7 +2043,7 @@ void UMetaSoundBuilderSubsystem::PostBuilderAssetTransaction(const FMetasoundFro
 	TWeakObjectPtr<UMetaSoundBuilderBase> BuilderPtr = AssetBuilders.FindRef(InClassName);
 	if (UMetaSoundBuilderBase* Builder = BuilderPtr.Get())
 	{
-		Builder->Reload();
+		Builder->ReloadCache();
 	}
 }
 #endif // WITH_EDITOR
diff --git a/Engine/Plugins/Runtime/Metasound/Source/MetasoundEngine/Private/MetasoundEngineAsset.h b/Engine/Plugins/Runtime/Metasound/Source/MetasoundEngine/Private/MetasoundEngineAsset.h
--- a/Engine/Plugins/Runtime/Metasound/Source/MetasoundEngine/Private/MetasoundEngineAsset.h
+++ b/Engine/Plugins/Runtime/Metasound/Source/MetasoundEngine/Private/MetasoundEngineAsset.h
@@ -31,7 +31,7 @@ namespace Metasound
 		{
 			InMetaSound.GetModifyContext().SetForceRefreshViews();
 
-			const FMetasoundFrontendDocument& Document = InMetaSound.GetConstDocumentChecked();
+			const FMetasoundFrontendDocument& Document = InMetaSound.GetDocumentChecked();
 			const FMetasoundFrontendClassName& ClassName = Document.RootGraph.Metadata.GetClassName();
 			UMetaSoundBuilderSubsystem::GetChecked().PostBuilderAssetTransaction(ClassName);
 
@@ -95,37 +95,43 @@ namespace Metasound
 			return ReferencedAssets;
 		}
 
-		static void PreSaveAsset(FMetasoundAssetBase& InMetaSound, FObjectPreSaveContext InSaveContext)
+		template <typename TMetaSoundObject>
+		static void PreSaveAsset(TMetaSoundObject& InMetaSound, FObjectPreSaveContext InSaveContext)
 		{
 #if WITH_EDITORONLY_DATA
 			using namespace Frontend;
 
-			// Do not call asset manager on CDO objects which may be loaded before asset
+			// Do not call asset manager on CDO objects which may be loaded before asset 
 			// manager is set.
 			if (IMetaSoundAssetManager* AssetManager = IMetaSoundAssetManager::Get())
 			{
 				AssetManager->WaitUntilAsyncLoadReferencedAssetsComplete(InMetaSound);
 			}
 
-			if (InSaveContext.IsCooking() || IsRunningCommandlet())
-			{
-				constexpr bool bIsDeterministic = true;
-				FDocumentIDGenerator::FScopeDeterminism DeterminismScope = FDocumentIDGenerator::FScopeDeterminism(bIsDeterministic);
-				InMetaSound.CookMetaSound();
-			}
- 			else if (FApp::CanEverRenderAudio())
+			if (UMetasoundEditorGraphBase* MetaSoundGraph = Cast<UMetasoundEditorGraphBase>(InMetaSound.GetGraph()))
 			{
-				if (UMetasoundEditorGraphBase* MetaSoundGraph = Cast<UMetasoundEditorGraphBase>(InMetaSound.GetGraph()))
+				if (InSaveContext.IsCooking() || IsRunningCommandlet())
+				{
+					// Use deterministic ID generation so more can be done at cook rather than runtime
+					if (MetaSoundEnableCookDeterministicIDGeneration != 0)
+					{
+						{
+							constexpr bool bIsDeterministic = true;
+							FDocumentIDGenerator::FScopeDeterminism DeterminismScope = FDocumentIDGenerator::FScopeDeterminism(bIsDeterministic);
+							InMetaSound.CookMetaSound();
+						}
+					}
+				}
+ 				else if (FApp::CanEverRenderAudio())
 				{
-					// Uses graph flavor of register with frontend to update editor systems/asset editors in case editor is enabled.
 					MetaSoundGraph->RegisterGraphWithFrontend();
-					InMetaSound.GetModifyContext().SetForceRefreshViews();
+					MetaSoundGraph->GetModifyContext().SetForceRefreshViews();
+				}
+				else
+				{
+					UE_LOG(LogMetaSound, Warning, TEXT("PreSaveAsset for MetaSound: (%s) is doing nothing because InSaveContext.IsCooking, IsRunningCommandlet, and FApp::CanEverRenderAudio were all false")
+						, *InMetaSound.GetPathName());
 				}
-			}
-			else
-			{
-				UE_LOG(LogMetaSound, Warning, TEXT("PreSaveAsset for MetaSound: (%s) is doing nothing because InSaveContext.IsCooking, IsRunningCommandlet, and FApp::CanEverRenderAudio were all false")
-					, *InMetaSound.GetOwningAssetName());
 			}
 #endif // WITH_EDITORONLY_DATA
 		}
@@ -138,7 +144,10 @@ namespace Metasound
 				if (InMetaSound.VersionAsset())
 				{
 #if WITH_EDITORONLY_DATA
-					InMetaSound.SetVersionedOnLoad();
+					if (UMetasoundEditorGraphBase* MetaSoundGraph = Cast<UMetasoundEditorGraphBase>(InMetaSound.GetGraph()))
+					{
+						MetaSoundGraph->SetVersionedOnLoad();
+					}
 #endif // WITH_EDITORONLY_DATA
 				}
 			}
diff --git a/Engine/Plugins/Runtime/Metasound/Source/MetasoundEngine/Private/MetasoundSource.cpp b/Engine/Plugins/Runtime/Metasound/Source/MetasoundEngine/Private/MetasoundSource.cpp
--- a/Engine/Plugins/Runtime/Metasound/Source/MetasoundEngine/Private/MetasoundSource.cpp
+++ b/Engine/Plugins/Runtime/Metasound/Source/MetasoundEngine/Private/MetasoundSource.cpp
@@ -93,9 +93,12 @@ namespace Metasound
 			Frontend::FMetaSoundAssetRegistrationOptions RegOptions;
 			RegOptions.bForceReregister = false;
 #if !WITH_EDITOR 
-			// When without editor, don't AutoUpdate or ResolveDocument at runtime. This only happens at cook or save.
-			// When with editor, those are needed because sounds are not necessarily saved before previewing.
-			RegOptions.bAutoUpdate = false;
+			if (Frontend::MetaSoundEnableCookDeterministicIDGeneration != 0)
+			{
+				// When without editor, don't AutoUpdate or ResolveDocument at runtime. This only happens at cook or save.
+				// When with editor, those are needed because sounds are not necessarily saved before previewing.
+				RegOptions.bAutoUpdate = false;
+			}
 #endif // !WITH_EDITOR
 			if (const UMetaSoundSettings* Settings = GetDefault<UMetaSoundSettings>())
 			{
@@ -430,12 +433,17 @@ void UMetaSoundSource::PostEditChangeOutputFormat()
 		UMetaSoundSourceBuilder* SourceBuilder = BuilderSubsystem.AttachSourceBuilderToAsset(this);
 		check(SourceBuilder);
 		SourceBuilder->SetFormat(OutputFormat, Result);
+
+		// TODO: Once builders are notified of controller changes and can be safely persistent, this
+		// can be removed so builders can be shared and not have to be created for each change output
+		// format mutation transaction.
+		BuilderSubsystem.DetachBuilderFromAsset(GetConstDocument().RootGraph.Metadata.GetClassName());
 	}
 
 	if (Result == EMetaSoundBuilderResult::Succeeded)
 	{
 		// Update the data in this UMetaSoundSource to reflect what is in the metasound document.
-		ConformObjectToDocument();
+		ConformObjectDataToInterfaces();
 
 		// Use the editor form of register to ensure other editors'
 		// MetaSounds are auto-updated if they are referencing this graph.
@@ -464,7 +472,7 @@ void UMetaSoundSource::PostEditChangeQualitySettings()
 }
 #endif // WITH_EDITOR
 
-bool UMetaSoundSource::ConformObjectToDocument()
+bool UMetaSoundSource::ConformObjectDataToInterfaces()
 {
 	using namespace Metasound::Engine;
 	using namespace Metasound::Frontend;
@@ -635,37 +643,27 @@ void UMetaSoundSource::OnAsyncReferencedAssetsLoaded(const TArray<FMetasoundAsse
 }
 
 #if WITH_EDITORONLY_DATA
-void UMetaSoundSource::MigrateEditorGraph(FMetaSoundFrontendDocumentBuilder& OutBuilder)
-{
-	PRAGMA_DISABLE_DEPRECATION_WARNINGS
-		if (Graph)
-		{
-			Graph->MigrateEditorDocumentData(OutBuilder);
-			Graph = nullptr;
-		}
-	PRAGMA_ENABLE_DEPRECATION_WARNINGS
-}
 
 UEdGraph* UMetaSoundSource::GetGraph()
 {
-	return EditorGraph;
+	return Graph;
 }
 
 const UEdGraph* UMetaSoundSource::GetGraph() const
 {
-	return EditorGraph;
+	return Graph;
 }
 
 UEdGraph& UMetaSoundSource::GetGraphChecked()
 {
-	check(EditorGraph);
-	return *EditorGraph;
+	check(Graph);
+	return *Graph;
 }
 
 const UEdGraph& UMetaSoundSource::GetGraphChecked() const
 {
-	check(EditorGraph);
-	return *EditorGraph;
+	check(Graph);
+	return *Graph;
 }
 
 FText UMetaSoundSource::GetDisplayName() const
@@ -922,7 +920,7 @@ Metasound::Frontend::FConstDocumentAccessPtr UMetaSoundSource::GetDocumentConstA
 bool UMetaSoundSource::ImplementsParameterInterface(Audio::FParameterInterfacePtr InInterface) const
 {
 	const FMetasoundFrontendVersion Version { InInterface->GetName(), { InInterface->GetVersion().Major, InInterface->GetVersion().Minor } };
-	return GetConstDocumentChecked().Interfaces.Contains(Version);
+	return GetDocumentChecked().Interfaces.Contains(Version);
 }
 
 ISoundGeneratorPtr UMetaSoundSource::CreateSoundGenerator(const FSoundGeneratorInitParams& InParams, TArray<FAudioParameter>&& InDefaultParameters)
@@ -1219,7 +1217,7 @@ void UMetaSoundSource::InitParametersInternal(const Metasound::TSortedVertexName
 
 bool UMetaSoundSource::IsParameterValid(const FAudioParameter& InParameter) const
 {
-	const TArray<FMetasoundFrontendClassInput>& Inputs = GetConstDocumentChecked().RootGraph.Interface.Inputs;
+	const TArray<FMetasoundFrontendClassInput>& Inputs = GetDocumentChecked().RootGraph.Interface.Inputs;
 	const FMetasoundFrontendVertex* Vertex = Algo::FindByPredicate(Inputs, [&InParameter] (const FMetasoundFrontendClassInput& Input)
 	{
 		return Input.Name == InParameter.ParamName;
@@ -1738,6 +1736,11 @@ void UMetaSoundSource::CacheRuntimeInputData()
 {
 	using namespace Metasound;
 
+	if (bIsBuilderActive)
+	{
+		UE_LOG(LogMetaSound, Warning, TEXT("Skipping caching of runtime inputs for UMetaSoundSource %s because there is an active builder"), *GetOwningAssetName());
+	}
+
 	constexpr bool bCreateUObjectProxies = true; 
 	RuntimeInputData.InputMap = CreateRuntimeInputMap(bCreateUObjectProxies);
 
diff --git a/Engine/Plugins/Runtime/Metasound/Source/MetasoundEngine/Public/Metasound.h b/Engine/Plugins/Runtime/Metasound/Source/MetasoundEngine/Public/Metasound.h
--- a/Engine/Plugins/Runtime/Metasound/Source/MetasoundEngine/Public/Metasound.h
+++ b/Engine/Plugins/Runtime/Metasound/Source/MetasoundEngine/Public/Metasound.h
@@ -31,22 +31,11 @@ public:
 	virtual void RegisterGraphWithFrontend() PURE_VIRTUAL(UMetasoundEditorGraphBase::RegisterGraphWithFrontend(), )
 
 #if WITH_EDITORONLY_DATA
-	UE_DEPRECATED(5.5, "ModifyContext is to be replaced by builder API delegates providing context when items changed and it will be up to the caller to track modification deltas.")
-	virtual FMetasoundFrontendDocumentModifyContext& GetModifyContext() { static FMetasoundFrontendDocumentModifyContext InvalidModifyData; return InvalidModifyData; }
-
-	UE_DEPRECATED(5.5, "ModifyContext is to be replaced by builder API delegates providing context when items changed and it will be up to the caller to track modification deltas.")
-	virtual const FMetasoundFrontendDocumentModifyContext& GetModifyContext() const { static const FMetasoundFrontendDocumentModifyContext InvalidModifyData; return InvalidModifyData; }
-
-	UE_DEPRECATED(5.5, "Editor Graph is now transient, so versioning flag moved to AssetBase.")
-	virtual void ClearVersionedOnLoad() { }
-
-	UE_DEPRECATED(5.5, "Editor Graph is now transient, so versioning flag moved to AssetBase.")
-	virtual bool GetVersionedOnLoad() const { return false; }
-
-	UE_DEPRECATED(5.5, "Editor Graph is now transient, so versioning flag moved to AssetBase.")
-	virtual void SetVersionedOnLoad() {  }
-
-	virtual void MigrateEditorDocumentData(FMetaSoundFrontendDocumentBuilder & OutBuilder) PURE_VIRTUAL(UMetasoundEditorGraphBase::MigrateEditorDocumentData(), )
+	virtual FMetasoundFrontendDocumentModifyContext& GetModifyContext() PURE_VIRTUAL(UMetasoundEditorGraphBase::GetModifyContext, static FMetasoundFrontendDocumentModifyContext InvalidModifyData; return InvalidModifyData; )
+	virtual const FMetasoundFrontendDocumentModifyContext& GetModifyContext() const PURE_VIRTUAL(UMetasoundEditorGraphBase::GetModifyContext, static const FMetasoundFrontendDocumentModifyContext InvalidModifyData; return InvalidModifyData; )
+	virtual void ClearVersionedOnLoad() PURE_VIRTUAL(UMetasoundEditorGraphBase::ClearVersionedOnLoad, )
+	virtual bool GetVersionedOnLoad() const PURE_VIRTUAL(UMetasoundEditorGraphBase::GetVersionedOnLoad, return false; )
+	virtual void SetVersionedOnLoad() PURE_VIRTUAL(UMetasoundEditorGraphBase::SetVersionedOnLoad, )
 #endif // WITH_EDITORONLY_DATA
 
 	int32 GetHighestMessageSeverity() const;
@@ -84,11 +73,8 @@ protected:
 	TSet<FSoftObjectPath> ReferenceAssetClassCache;
 
 #if WITH_EDITORONLY_DATA
-	UPROPERTY(meta=(DeprecatedProperty, DeprecationMessage = "Use EditorGraph instead as it is now transient and generated via the FrontendDocument dynamically."))
+	UPROPERTY()
 	TObjectPtr<UMetasoundEditorGraphBase> Graph;
-
-	UPROPERTY(Transient)
-	TObjectPtr<UMetasoundEditorGraphBase> EditorGraph;
 #endif // WITH_EDITORONLY_DATA
 
 public:
@@ -133,14 +119,13 @@ public:
 	virtual const UEdGraph* GetGraph() const override;
 	virtual UEdGraph& GetGraphChecked() override;
 	virtual const UEdGraph& GetGraphChecked() const override;
-	virtual void MigrateEditorGraph(FMetaSoundFrontendDocumentBuilder& OutBuilder) override;
 
 	// Sets the graph associated with this Metasound. Graph is required to be referenced on
 	// Metasound UObject for editor serialization purposes.
 	// @param Editor graph associated with UMetaSoundSource.
 	virtual void SetGraph(UEdGraph* InGraph) override
 	{
-		EditorGraph = CastChecked<UMetasoundEditorGraphBase>(InGraph);
+		Graph = CastChecked<UMetasoundEditorGraphBase>(InGraph);
 	}
 #endif // #if WITH_EDITORONLY_DATA
 
@@ -158,7 +143,7 @@ public:
 	virtual void Serialize(FArchive& InArchive) override;
 	virtual void PostLoad() override;
 
-	virtual bool ConformObjectToDocument() override { return false; }
+	virtual bool ConformObjectDataToInterfaces() override { return false; }
 
 	virtual const TSet<FString>& GetReferencedAssetClassKeys() const override
 	{
@@ -178,8 +163,6 @@ public:
 		return this;
 	}
 
-	virtual bool IsBuilderActive() const override;
-
 protected:
 #if WITH_EDITOR
 	virtual void SetReferencedAssetClasses(TSet<Metasound::Frontend::IMetaSoundAssetManager::FAssetInfo>&& InAssetClasses) override;
@@ -194,6 +177,7 @@ private:
 		return RootMetaSoundDocument;
 	}
 
+	virtual bool IsBuilderActive() const override;
 	virtual void OnBeginActiveBuilder() override;
 	virtual void OnFinishActiveBuilder() override;
 
diff --git a/Engine/Plugins/Runtime/Metasound/Source/MetasoundEngine/Public/MetasoundAssetSubsystem.h b/Engine/Plugins/Runtime/Metasound/Source/MetasoundEngine/Public/MetasoundAssetSubsystem.h
--- a/Engine/Plugins/Runtime/Metasound/Source/MetasoundEngine/Public/MetasoundAssetSubsystem.h
+++ b/Engine/Plugins/Runtime/Metasound/Source/MetasoundEngine/Public/MetasoundAssetSubsystem.h
@@ -66,9 +66,7 @@ public:
 #if WITH_EDITORONLY_DATA
 	virtual void AddAssetReferences(FMetasoundAssetBase& InAssetBase) override;
 #endif
-	virtual Metasound::Frontend::FNodeRegistryKey AddOrUpdateAsset(const UObject& InObject) override;
-	virtual FMetaSoundFrontendDocumentBuilder& AttachDocumentBuilderChecked(UObject& InObject) const override;
-
+	virtual Metasound::Frontend::FNodeRegistryKey  AddOrUpdateAsset(const UObject& InObject) override;
 	virtual bool CanAutoUpdate(const FMetasoundFrontendClassName& InClassName) const override;
 	virtual bool ContainsKey(const Metasound::Frontend::FNodeRegistryKey& InRegistryKey) const override;
 	virtual const FSoftObjectPath* FindObjectPathFromKey(const Metasound::Frontend::FNodeRegistryKey& RegistryKey) const override;
diff --git a/Engine/Plugins/Runtime/Metasound/Source/MetasoundEngine/Public/MetasoundBuilderSubsystem.h b/Engine/Plugins/Runtime/Metasound/Source/MetasoundEngine/Public/MetasoundBuilderSubsystem.h
--- a/Engine/Plugins/Runtime/Metasound/Source/MetasoundEngine/Public/MetasoundBuilderSubsystem.h
+++ b/Engine/Plugins/Runtime/Metasound/Source/MetasoundEngine/Public/MetasoundBuilderSubsystem.h
@@ -165,10 +165,6 @@ public:
 	UE_DEPRECATED(5.4, "This version of AddNodeByClassName is deprecated. Use the one with a default MajorVersion of 1.")
 	UPARAM(DisplayName = "Node Handle") FMetaSoundNodeHandle AddNodeByClassName(const FMetasoundFrontendClassName& ClassName, int32 MajorVersion, EMetaSoundBuilderResult& OutResult);
 
-#if WITH_EDITOR
-	bool ClearMemberMetadata(const FGuid& InMemberID);
-#endif // WITH_EDITOR
-
 	// Connects node output to a node input. Does *NOT* provide loop detection for performance reasons.  Loop detection is checked on class registration when built or played.
 	// Returns succeeded if connection made, failed if connection already exists with input, the data types do not match, or the connection is not supported due to access type
 	// incompatibility (ex. constructor input to non-constructor input).
@@ -223,7 +219,7 @@ public:
 	// all connections were found and removed, failed if any connections were not.
 	UFUNCTION(BlueprintCallable, Category = "Audio|MetaSound|Builder", meta = (ExpandEnumAsExecs = "OutResult"))
 	void DisconnectNodesByInterfaceBindings(const FMetaSoundNodeHandle& FromNodeHandle, const FMetaSoundNodeHandle& ToNodeHandle, EMetaSoundBuilderResult& OutResult);
-
+	
 	// Returns graph input node by the given name if it exists, or an invalid handle if not found.
 	UFUNCTION(BlueprintCallable, Category = "Audio|MetaSound|Builder", meta = (ExpandEnumAsExecs = "OutResult"))
 	UPARAM(DisplayName = "Node Handle") FMetaSoundNodeHandle FindGraphInputNode(FName InputName, EMetaSoundBuilderResult& OutResult);
@@ -232,10 +228,6 @@ public:
 	UFUNCTION(BlueprintCallable, Category = "Audio|MetaSound|Builder", meta = (ExpandEnumAsExecs = "OutResult"))
 	UPARAM(DisplayName = "Node Handle") FMetaSoundNodeHandle FindGraphOutputNode(FName OutputName, EMetaSoundBuilderResult& OutResult);
 
-#if WITH_EDITOR
-	UMetaSoundFrontendMemberMetadata* FindMemberMetadata(const FGuid& InMemberID);
-#endif // WITH_EDITOR
-
 	// Returns node input by the given name if it exists, or an invalid handle if not found.
 	UFUNCTION(BlueprintCallable, Category = "Audio|MetaSound|Builder", meta = (ExpandEnumAsExecs = "OutResult"))
 	UPARAM(DisplayName = "Input Handle") FMetaSoundBuilderNodeInputHandle FindNodeInputByName(const FMetaSoundNodeHandle& NodeHandle, FName InputName, EMetaSoundBuilderResult& OutResult);
@@ -354,28 +346,20 @@ public:
 
 	// Removes node and any associated connections from the builder's MetaSound.
 	UFUNCTION(BlueprintCallable, Category = "Audio|MetaSound|Builder", meta = (ExpandEnumAsExecs = "OutResult"))
-	void RemoveNode(const FMetaSoundNodeHandle& NodeHandle, EMetaSoundBuilderResult& OutResult, bool bRemoveUnusedDependencies = true);
+	void RemoveNode(const FMetaSoundNodeHandle& NodeHandle, EMetaSoundBuilderResult& OutResult);
 
 	// Removes node input literal default if set, reverting the value to be whatever the node class defaults the value to.
 	// Returns success if value was removed, false if not removed (i.e. wasn't set to begin with).
 	UFUNCTION(BlueprintCallable, Category = "Audio|MetaSound|Builder", meta = (ExpandEnumAsExecs = "OutResult"))
 	void RemoveNodeInputDefault(const FMetaSoundBuilderNodeInputHandle& InputHandle, EMetaSoundBuilderResult& OutResult);
-
-	// Removes dependencies in document that are no longer referenced by nodes
-	UFUNCTION(BlueprintCallable, Category = "Audio|MetaSound|Builder")
-	void RemoveUnusedDependencies();
-
+	
 	// Rename the document's root graph class with a guid and optional namespace and variant
 	void RenameRootGraphClass(const FMetasoundFrontendClassName& InName);
 
-	UE_DEPRECATED(5.5, "Moved to 'Reload' to discourage priming by defaulting to false. "
-		"In addition, internal delegates and respective bindings are now regenerated when reloading, "
-		"effectively resetting the builder's entire state and not just strictly the underlying cache.")
+	// Primarily used by editor transaction stack to avoid corruption when object is
+	// changed outside of the builder API. Generally discouraged for direct use otherwise.
 	void ReloadCache(bool bPrimeCache = true);
 
-	// Reloads the builder, freeing the internal cache and rebuilding delegate bindings
-	void Reload(bool bPrimeCache = false);
-
 #if WITH_EDITOR
 	// Sets the author of the MetaSound.
 	void SetAuthor(const FString& InAuthor);
@@ -385,32 +369,10 @@ public:
 	UFUNCTION(BlueprintCallable, Category = "Audio|MetaSound|Builder", meta = (ExpandEnumAsExecs = "OutResult"))
 	void SetNodeInputDefault(const FMetaSoundBuilderNodeInputHandle& NodeInputHandle, const FMetasoundFrontendLiteral& Literal, EMetaSoundBuilderResult& OutResult);
 
-	// Disconnects the given graph input's respective template nodes and sets the graph input's AccessType should it not match the current AccessType.
-	// Result succeeds if the AccessType was successfully changed or if the provided AccessType is already the input's current AccessType.
-	UFUNCTION(BlueprintCallable, Category = "Audio|MetaSound|Builder", meta = (ExpandEnumAsExecs = "OutResult"))
-	void SetGraphInputAccessType(FName InputName, EMetasoundFrontendVertexAccessType AccessType, EMetaSoundBuilderResult& OutResult);
-
-	// Disconnects the given graph input's respective template nodes and sets the graph input's DataType should it not match the current DataType.
-	// Result succeeds if the DataType was successfully changed or if the provided DataType is already the input's current DataType.
-	UFUNCTION(BlueprintCallable, Category = "Audio|MetaSound|Builder", meta = (ExpandEnumAsExecs = "OutResult"))
-	void SetGraphInputDataType(FName InputName, FName DataType, EMetaSoundBuilderResult& OutResult);
-
 	// Sets the input node's default value, overriding the default provided by the referenced graph if the graph is a preset.
 	UFUNCTION(BlueprintCallable, Category = "Audio|MetaSound|Builder", meta = (ExpandEnumAsExecs = "OutResult"))
 	void SetGraphInputDefault(FName InputName, const FMetasoundFrontendLiteral& Literal, EMetaSoundBuilderResult& OutResult);
 
-	// Disconnects the given graph output's respective template nodes and sets the graph output's AccessType should it not match the current AccessType.
-	// Result succeeds if the AccessType was successfully changed or if the provided AccessType is already the output's current AccessType.
-	UFUNCTION(BlueprintCallable, Category = "Audio|MetaSound|Builder", meta = (ExpandEnumAsExecs = "OutResult"))
-	void SetGraphOutputAccessType(FName OutputName, EMetasoundFrontendVertexAccessType AccessType, EMetaSoundBuilderResult& OutResult);
-
-	// Disconnects the given graph output's respective template nodes and sets the graph output's DataType should it not match the current DataType.
-	// Result succeeds if the DataType was successfully changed or if the provided DataType is already the output's current DataType.
-	UFUNCTION(BlueprintCallable, Category = "Audio|MetaSound|Builder", meta = (ExpandEnumAsExecs = "OutResult"))
-	void SetGraphOutputDataType(FName OutputName, FName DataType, EMetaSoundBuilderResult& OutResult);
-
-	void SetMemberMetadata(UMetaSoundFrontendMemberMetadata& NewMetadata);
-
 	// Update dependency class names given a map of old to new referenced class names. 
 	void UpdateDependencyClassNames(const TMap<FMetasoundFrontendClassName, FMetasoundFrontendClassName>& OldToNewReferencedClassNames);
 
@@ -422,19 +384,10 @@ public:
 	UE_DEPRECATED(5.4, "Moved to protected pure virtual 'CreateTransientBuilder'")
 	virtual void InitFrontendBuilder();
 
-	virtual bool ConformObjectToDocument() PURE_VIRTUAL(UMetaSoundBuilderBase::ConformObjectToDocument, return false; );
-
 	// Initializes and ensures all nodes have a position (required prior to exporting to an asset if expected to be viewed in the editor).
 	virtual void InitNodeLocations();
 
 #if WITH_EDITOR
-	// Injects template nodes between builder's document inputs not connected
-	// to existing template inputs, copying locational data from the represented
-	// input metadata. If bForceNodeCreation is false, only generates a template
-	// input node if a connection between the input and other nodes exists. If true,
-	// will inject template node irrespective of whether or not the input has connections.
-	void InjectInputTemplateNodes(bool bForceNodeCreation, EMetaSoundBuilderResult& OutResult);
-
 	const FMetaSoundFrontendGraphComment* FindGraphComment(const FGuid& InCommentID) const;
 	FMetaSoundFrontendGraphComment* FindGraphComment(const FGuid & InCommentID);
 	FMetaSoundFrontendGraphComment& FindOrAddGraphComment(const FGuid& InCommentID);
@@ -449,14 +402,11 @@ public:
 	const FMetaSoundFrontendDocumentBuilder& GetConstBuilder() const;
 
 protected:
-	virtual void InitDelegates(Metasound::Frontend::FDocumentModifyDelegates & OutDocumentDelegates) { };
-
-	// Creates a FrontendBuilder wrapping a transient MetaSound object of the class supported by the given subsystem builder (See GetBuilderUClass).
+	// Creates a FrontendBuilder wrapping the transient MetaSound object of the class supported by the given subsystem builder (See GetBuilderUClass).
 	// Should only be used when builders is not being attached to an existing, serialized MetaSound asset.
 	virtual void CreateTransientBuilder() PURE_VIRTUAL(UMetaSoundBuilderBase::CreateTransientBuilder, );
 
-	UE_DEPRECATED(5.5, "Moved to 'Reload', to enforce generation of new delegates")
-	void InvalidateCache(bool bPrimeCache = false);
+	void InvalidateCache();
 
 	// Runs build, conforming the document and corresponding object data on a MetaSound UObject to that managed by this builder.
 	template <typename UClassType>
@@ -478,7 +428,7 @@ protected:
 			{
 				// If MetaSound already exists, preserve the class name to avoid
 				// nametable bloat & preserve potentially existing references.
-				const FMetasoundFrontendDocument& ExistingDoc = CastChecked<const UClassType>(MetaSound)->GetConstDocumentChecked();
+				const FMetasoundFrontendDocument& ExistingDoc = CastChecked<const UClassType>(MetaSound)->GetDocumentChecked();
 				if (!BuilderOptions.bForceUniqueClassName)
 				{
 					DocClassName = &ExistingDoc.RootGraph.Metadata.GetClassName();
@@ -504,9 +454,17 @@ protected:
 		}
 
 		checkf(MetaSound, TEXT("Failed to build MetaSound from builder '%s'"), *GetPathName());
+		FMetasoundFrontendDocument NewDocument = GetConstBuilder().GetDocument();
+		{
+			// This is required to ensure the newly build document has a unique class
+			// identifier than that of the given builder's document copy to avoid collisions
+			// if added to the Frontend class registry (either below or at a later point in time).
+			constexpr bool bResetVersion = false;
+			FMetaSoundFrontendDocumentBuilder::InitGraphClassMetadata(NewDocument.RootGraph.Metadata, bResetVersion, DocClassName);
+		}
 
-		TScriptInterface<IMetaSoundDocumentInterface> MetaSoundInterface = MetaSound;
-		BuildInternal(MetaSoundInterface, DocClassName);
+		MetaSound->SetDocument(MoveTemp(NewDocument));
+		MetaSound->ConformObjectDataToInterfaces();
 
 		if (BuilderOptions.bAddToRegistry)
 		{
@@ -528,25 +486,15 @@ protected:
 	UPROPERTY()
 	FMetaSoundFrontendDocumentBuilder Builder;
 
-	// ClassName of associated document object.  Cached in case DocumentObject is deleted but continues to provide an identifier after destruction.
-	UPROPERTY()
-	FMetasoundFrontendClassName ClassName;
-
 #if WITH_EDITORONLY_DATA
 	UPROPERTY(meta = (DeprecatedProperty, DeprecationMessage = "5.4 - All source builders now operate on an underlying document source document that is also used to audition."))
 	bool bIsAttached = false;
 #endif // WITH_EDITORONLY_DATA
 
-
-private:
-	void BuildInternal(TScriptInterface<IMetaSoundDocumentInterface> NewMetaSound, const FMetasoundFrontendClassName* InDocClassName) const;
-
 private:
 	int32 LastTransactionRegistered = 0;
 
 	// Friending allows for swapping the builder in certain circumstances where desired (eg. attaching a builder to an existing asset)
-	// or performing editor-only manipulations to the builder (ex. adding template nodes) and versioning the document data where necessary.
-	friend class UMetaSoundAssetSubsystem;
 	friend class UMetaSoundBuilderSubsystem;
 };
 
@@ -560,7 +508,6 @@ public:
 	UFUNCTION(BlueprintCallable, Category = "Audio|MetaSound|Builder", meta = (WorldContext = "Parent"))
 	virtual UPARAM(DisplayName = "MetaSound") TScriptInterface<IMetaSoundDocumentInterface> Build(UObject* Parent, const FMetaSoundBuilderOptions& Options) const override;
 
-	virtual bool ConformObjectToDocument() override;
 	virtual const UClass& GetBuilderUClass() const override;
 
 protected:
@@ -600,7 +547,6 @@ public:
 
 	const Metasound::Engine::FOutputAudioFormatInfoPair* FindOutputAudioFormatInfo() const;
 
-	virtual bool ConformObjectToDocument() override;
 	virtual const UClass& GetBuilderUClass() const override;
 
 #if WITH_EDITORONLY_DATA
@@ -617,7 +563,6 @@ public:
 
 protected:
 	virtual void CreateTransientBuilder() override;
-	virtual void InitDelegates(Metasound::Frontend::FDocumentModifyDelegates& OutDocumentDelegates) override;
 
 private:
 	static TOptional<Metasound::FAnyDataReference> CreateDataReference(const Metasound::FOperatorSettings& InOperatorSettings, FName DataType, const Metasound::FLiteral& InLiteral, Metasound::EDataReferenceAccessType AccessType);
@@ -625,6 +570,7 @@ private:
 	const UMetaSoundSource& GetMetaSoundSource() const;
 	UMetaSoundSource& GetMetaSoundSource();
 
+	void InitDelegates(Metasound::Frontend::FDocumentModifyDelegates& OutDocumentDelegates);
 
 	void OnEdgeAdded(int32 EdgeIndex) const;
 	void OnInputAdded(int32 InputIndex);
@@ -672,7 +618,7 @@ public:
 	static const UMetaSoundBuilderSubsystem* GetConst();
 	static const UMetaSoundBuilderSubsystem& GetConstChecked();
 
-	UMetaSoundBuilderBase& AttachBuilderToAssetChecked(UObject& InObject) const;
+	UMetaSoundBuilderBase& AttachBuilderToAssetChecked(UObject& InMetaSound) const;
 
 	// Moving toward becoming UFUNCTIONS in the future, but right now unguarded case
 	// where a user could improperly reference transient assets. Also, unsafe where
@@ -794,7 +740,6 @@ private:
 
 		const FMetasoundFrontendDocument& Document = NewBuilder->GetConstBuilder().GetDocument();
 		const FMetasoundFrontendClassName& ClassName = Document.RootGraph.Metadata.GetClassName();
-		NewBuilder->ClassName = ClassName;
 		TransientBuilders.Add(ClassName, NewBuilder);
 		return *NewBuilder.Get();
 	}
@@ -807,29 +752,18 @@ private:
 
 		TScriptInterface<IMetaSoundDocumentInterface> DocInterface = InMetaSoundObject;
 		const FMetasoundFrontendDocument& Document = DocInterface->GetConstDocument();
-		const FMetasoundFrontendClassName& FullClassName = Document.RootGraph.Metadata.GetClassName();
-
-		if (FullClassName.IsValid())
+		const FMetasoundFrontendClassName& ClassName = Document.RootGraph.Metadata.GetClassName();
+		TWeakObjectPtr<UMetaSoundBuilderBase> Builder = AssetBuilders.FindRef(ClassName);
+		if (Builder.IsValid())
 		{
-			TWeakObjectPtr<UMetaSoundBuilderBase> Builder = AssetBuilders.FindRef(FullClassName);
-			if (Builder.IsValid())
-			{
-				return *CastChecked<BuilderClass>(Builder.Get());
-			}
+			return *CastChecked<BuilderClass>(Builder.Get());
 		}
 
 		TObjectPtr<BuilderClass> NewBuilder = NewObject<BuilderClass>(InMetaSoundObject);
 		check(NewBuilder);
 		NewBuilder->Builder = FMetaSoundFrontendDocumentBuilder(DocInterface);
-
-		if (!FullClassName.IsValid())
-		{
-			NewBuilder->Builder.InitDocument();
-		}
-
-		checkf(FullClassName.IsValid(), TEXT("Document initialization must result in a valid class name being generated"));
 		TObjectPtr<UMetaSoundBuilderBase> NewBuilderBase = CastChecked<UMetaSoundBuilderBase>(NewBuilder);
-		AssetBuilders.Add(FullClassName, NewBuilderBase);
+		AssetBuilders.Add(ClassName, NewBuilderBase);
 		return *NewBuilder;
 	}
 
diff --git a/Engine/Plugins/Runtime/Metasound/Source/MetasoundEngine/Public/MetasoundSource.h b/Engine/Plugins/Runtime/Metasound/Source/MetasoundEngine/Public/MetasoundSource.h
--- a/Engine/Plugins/Runtime/Metasound/Source/MetasoundEngine/Public/MetasoundSource.h
+++ b/Engine/Plugins/Runtime/Metasound/Source/MetasoundEngine/Public/MetasoundSource.h
@@ -26,8 +26,6 @@
 class UMetaSoundSettings;
 struct FMetaSoundQualitySettings;
 
-struct FMetaSoundFrontendDocumentBuilder;
-
 namespace Audio
 {
 	using DeviceID = uint32;
@@ -107,11 +105,8 @@ protected:
 	TSet<FSoftObjectPath> ReferenceAssetClassCache;
 
 #if WITH_EDITORONLY_DATA
-	UPROPERTY(meta=(DeprecatedProperty, DeprecationMessage = "Use EditorGraph instead as it is now transient and generated via the FrontendDocument dynamically."))
+	UPROPERTY()
 	TObjectPtr<UMetasoundEditorGraphBase> Graph;
-
-	UPROPERTY(Transient)
-	TObjectPtr<UMetasoundEditorGraphBase> EditorGraph;
 #endif // WITH_EDITORONLY_DATA
 
 public:
@@ -180,14 +175,13 @@ public:
 	virtual const UEdGraph* GetGraph() const override;
 	virtual UEdGraph& GetGraphChecked() override;
 	virtual const UEdGraph& GetGraphChecked() const override;
-	virtual void MigrateEditorGraph(FMetaSoundFrontendDocumentBuilder& OutBuilder) override;
 
 	// Sets the graph associated with this Metasound. Graph is required to be referenced on
 	// Metasound UObject for editor serialization purposes.
 	// @param Editor graph associated with UMetaSoundSource.
 	virtual void SetGraph(UEdGraph* InGraph) override
 	{
-		EditorGraph = CastChecked<UMetasoundEditorGraphBase>(InGraph);
+		Graph = CastChecked<UMetasoundEditorGraphBase>(InGraph);
 	}
 #endif // #if WITH_EDITORONLY_DATA
 
@@ -218,9 +212,10 @@ public:
 private:
 	void PostEditChangeOutputFormat();
 	void PostEditChangeQualitySettings();
+public:
+
 #endif // WITH_EDITOR
 
-public:
 	virtual const TSet<FString>& GetReferencedAssetClassKeys() const override
 	{
 		return ReferencedAssetClassKeys;
@@ -236,7 +231,7 @@ public:
 
 	void PostLoadQualitySettings();
 
-	virtual bool ConformObjectToDocument() override;
+	virtual bool ConformObjectDataToInterfaces() override;
 
 	virtual FTopLevelAssetPath GetAssetPathChecked() const override;
 
@@ -273,7 +268,6 @@ public:
 	Metasound::FOperatorSettings GetOperatorSettings(Metasound::FSampleRate InDeviceSampleRate) const;
 
 	virtual const FMetasoundFrontendDocument& GetConstDocument() const override;
-	virtual bool IsBuilderActive() const override;
 
 protected:
 	Metasound::Frontend::FDocumentAccessPtr GetDocumentAccessPtr() override;
@@ -294,6 +288,7 @@ private:
 		return RootMetasoundDocument;
 	}
 
+	virtual bool IsBuilderActive() const override;
 	virtual void OnBeginActiveBuilder() override;
 	virtual void OnFinishActiveBuilder() override;
 
@@ -356,7 +351,7 @@ private:
 
 	// Quality settings. 
 	bool GetQualitySettings(const FName InPlatformName, Metasound::SourcePrivate::FCookedQualitySettings& OutQualitySettings) const;
-	void ResolveQualitySettings(const UMetaSoundSettings* Settings);
+	void ResolveQualitySettings(const UMetaSoundSettings* Settings);	
 	void SerializeCookedQualitySettings(const FName PlatformName, FArchive& Ar);
 	TPimplPtr<Metasound::SourcePrivate::FCookedQualitySettings> CookedQualitySettings;
 };
diff --git a/Engine/Plugins/Runtime/Metasound/Source/MetasoundEngineTest/Private/EngineTestMetasoundSource.cpp b/Engine/Plugins/Runtime/Metasound/Source/MetasoundEngineTest/Private/EngineTestMetasoundSource.cpp
--- a/Engine/Plugins/Runtime/Metasound/Source/MetasoundEngineTest/Private/EngineTestMetasoundSource.cpp
+++ b/Engine/Plugins/Runtime/Metasound/Source/MetasoundEngineTest/Private/EngineTestMetasoundSource.cpp
@@ -310,7 +310,7 @@ bool FMetaSoundSourceBuilderAuditionLatentCommand::Update()
 			if (const USoundBase* InitSound = AudioComponent->GetSound())
 			{
 				const UMetaSoundSource* InitMetaSound = CastChecked<UMetaSoundSource>(InitSound);
-				return InitMetaSound->GetConstDocumentChecked().RootGraph.Metadata.GetClassName().GetFullName();
+				return InitMetaSound->GetDocumentChecked().RootGraph.Metadata.GetClassName().GetFullName();
 			}
 
 			return { };
@@ -384,10 +384,7 @@ bool FAudioMetasoundSourceTest::RunTest(const FString& Parameters)
 	UMetaSoundSource* MetaSoundSource = NewObject<UMetaSoundSource>(GetTransientPackage(), FName(*LexToString(FGuid::NewGuid())));;
 	if (ensure(nullptr != MetaSoundSource))
 	{
-		// TODO: Move to using builder to swap in doc
-		PRAGMA_DISABLE_DEPRECATION_WARNINGS
 		MetaSoundSource->SetDocument(CreateMonoSourceDocument());
-		PRAGMA_ENABLE_DEPRECATION_WARNINGS
 
 		if (UAudioComponent* AudioComponent = CreateTestComponent(*this, MetaSoundSource))
 		{
diff --git a/Engine/Plugins/Runtime/Metasound/Source/MetasoundFrontend/Private/MetasoundFrontend.cpp b/Engine/Plugins/Runtime/Metasound/Source/MetasoundFrontend/Private/MetasoundFrontend.cpp
--- a/Engine/Plugins/Runtime/Metasound/Source/MetasoundFrontend/Private/MetasoundFrontend.cpp
+++ b/Engine/Plugins/Runtime/Metasound/Source/MetasoundFrontend/Private/MetasoundFrontend.cpp
@@ -18,7 +18,6 @@
 #include "MetasoundRouter.h"
 #include "MetasoundVertex.h"
 #include "Modules/ModuleManager.h"
-#include "NodeTemplates/MetasoundFrontendNodeTemplateInput.h"
 #include "NodeTemplates/MetasoundFrontendNodeTemplateReroute.h"
 #include "StructDeserializer.h"
 #include "StructSerializer.h"
@@ -95,8 +94,8 @@ class FMetasoundFrontendModule : public IModuleInterface
 	{
 		using namespace Metasound::Frontend;
 
-		RegisterNodeTemplate(MakeUnique<FInputNodeTemplate>());
-		RegisterNodeTemplate(MakeUnique<FRerouteNodeTemplate>());
+		TUniquePtr<INodeTemplate> RerouteTemplate = MakeUnique<FRerouteNodeTemplate>();
+		RegisterNodeTemplate(MoveTemp(RerouteTemplate));
 
 		FMetasoundFrontendRegistryContainer* Registry = FMetasoundFrontendRegistryContainer::Get();
 		if (ensure(nullptr != Registry))
@@ -109,7 +108,6 @@ class FMetasoundFrontendModule : public IModuleInterface
 	{
 		using namespace Metasound::Frontend;
 
-		UnregisterNodeTemplate(FInputNodeTemplate::ClassName, FInputNodeTemplate::VersionNumber);
 		UnregisterNodeTemplate(FRerouteNodeTemplate::ClassName, FRerouteNodeTemplate::VersionNumber);
 	}
 };
diff --git a/Engine/Plugins/Runtime/Metasound/Source/MetasoundFrontend/Private/MetasoundFrontendController.cpp b/Engine/Plugins/Runtime/Metasound/Source/MetasoundFrontend/Private/MetasoundFrontendController.cpp
--- a/Engine/Plugins/Runtime/Metasound/Source/MetasoundFrontend/Private/MetasoundFrontendController.cpp
+++ b/Engine/Plugins/Runtime/Metasound/Source/MetasoundFrontend/Private/MetasoundFrontendController.cpp
@@ -4,7 +4,6 @@
 
 #include "MetasoundFrontendInvalidController.h"
 #include "MetasoundFrontendDocumentController.h"
-#include "NodeTemplates/MetasoundFrontendNodeTemplateInput.h"
 #include "NodeTemplates/MetasoundFrontendNodeTemplateReroute.h"
 
 
@@ -12,17 +11,6 @@ namespace Metasound
 {
 	namespace Frontend
 	{
-		namespace ControllerPrivate
-		{
-			// This check for class name is a hack due to the fact that handles do not provide a builder
-			// in order to do a proper look-up of a rerouted output's access/data type. Once calling systems of the
-			// recurse functions utilizing this function are refactored to use the builder API, this can be removed.
-			bool IsRerouteClass(const FMetasoundFrontendClassName& ClassName)
-			{
-				return ClassName == FRerouteNodeTemplate::ClassName || ClassName == FInputNodeTemplate::ClassName;
-			}
-		}
-
 		FOutputHandle IOutputController::GetInvalidHandle()
 		{
 			static FOutputHandle Invalid = MakeShared<FInvalidOutputController>();
@@ -67,15 +55,6 @@ namespace Metasound
 
 		FDocumentHandle IDocumentController::CreateDocumentHandle(FMetasoundFrontendDocument& InDocument)
 		{
-			// Mutation of a document via the soft deprecated access ptr/controller system is not tracked by
-			// the builder registry, so the document cache is invalidated here. It is discouraged to mutate
-			// documents using both systems at the same time as it can corrupt a builder document's cache.
-			if (IDocumentBuilderRegistry* DocRegistry = IDocumentBuilderRegistry::Get())
-			{
-				const FMetasoundFrontendClassName& Name = InDocument.RootGraph.Metadata.GetClassName();
-				DocRegistry->InvalidateDocumentCache(Name);
-			}
-
 			return CreateDocumentHandle(MakeAccessPtr<FDocumentAccessPtr>(InDocument.AccessPoint, InDocument));
 		}
 
@@ -102,13 +81,14 @@ namespace Metasound
 
 		FConstOutputHandle FindReroutedOutput(FConstOutputHandle InOutputHandle)
 		{
+			using namespace Frontend;
+
 			if (InOutputHandle->IsValid())
 			{
 				FConstNodeHandle NodeHandle = InOutputHandle->GetOwningNode();
 				if (NodeHandle->IsValid())
 				{
-					const FMetasoundFrontendClassName& ClassName = NodeHandle->GetClassMetadata().GetClassName();
-					if (ControllerPrivate::IsRerouteClass(ClassName))
+					if (NodeHandle->GetClassMetadata().GetClassName() == FRerouteNodeTemplate::ClassName)
 					{
 						TArray<FConstInputHandle> Inputs = NodeHandle->GetConstInputs();
 						if (!Inputs.IsEmpty())
@@ -129,14 +109,14 @@ namespace Metasound
 
 		void FindReroutedInputs(FConstInputHandle InHandleToCheck, TArray<FConstInputHandle>& InOutInputHandles)
 		{
-			using namespace ControllerPrivate;
+			using namespace Frontend;
 
 			if (InHandleToCheck->IsValid())
 			{
 				FConstNodeHandle NodeHandle = InHandleToCheck->GetOwningNode();
 				if (NodeHandle->IsValid())
 				{
-					if (ControllerPrivate::IsRerouteClass(NodeHandle->GetClassMetadata().GetClassName()))
+					if (NodeHandle->GetClassMetadata().GetClassName() == FRerouteNodeTemplate::ClassName)
 					{
 						TArray<FConstOutputHandle> Outputs = NodeHandle->GetConstOutputs();
 						for (FConstOutputHandle& OutputHandle : Outputs)
@@ -159,12 +139,14 @@ namespace Metasound
 
 		void IterateReroutedInputs(FConstInputHandle InHandleToCheck, TFunctionRef<void(FConstInputHandle)> Func)
 		{
+			using namespace Frontend;
+
 			if (InHandleToCheck->IsValid())
 			{
 				FConstNodeHandle NodeHandle = InHandleToCheck->GetOwningNode();
 				if (NodeHandle->IsValid())
 				{
-					if (ControllerPrivate::IsRerouteClass(NodeHandle->GetClassMetadata().GetClassName()))
+					if (NodeHandle->GetClassMetadata().GetClassName() == FRerouteNodeTemplate::ClassName)
 					{
 						TArray<FConstOutputHandle> Outputs = NodeHandle->GetConstOutputs();
 						for (FConstOutputHandle& OutputHandle : Outputs)
