[Backout] - CL33084850

diff --git a/Engine/Plugins/Runtime/Metasound/Source/MetasoundFrontend/Private/MetasoundFrontendDocumentBuilder.cpp b/Engine/Plugins/Runtime/Metasound/Source/MetasoundFrontend/Private/MetasoundFrontendDocumentBuilder.cpp
--- a/Engine/Plugins/Runtime/Metasound/Source/MetasoundFrontend/Private/MetasoundFrontendDocumentBuilder.cpp
+++ b/Engine/Plugins/Runtime/Metasound/Source/MetasoundFrontend/Private/MetasoundFrontendDocumentBuilder.cpp
@@ -13,7 +13,6 @@
 #include "MetasoundAssetManager.h"
 #include "MetasoundDocumentInterface.h"
 #include "MetasoundFrontendController.h"
-#include "MetasoundFrontendDataTypeRegistry.h"
 #include "MetasoundFrontendDocumentCache.h"
 #include "MetasoundFrontendDocument.h"
 #include "MetasoundFrontendDocumentIdGenerator.h"
@@ -24,8 +23,6 @@
 #include "MetasoundFrontendSearchEngine.h"
 #include "MetasoundFrontendTransform.h"
 #include "MetasoundTrace.h"
-#include "MetasoundVariableNodes.h"
-#include "NodeTemplates/MetasoundFrontendNodeTemplateInput.h"
 
 #include UE_INLINE_GENERATED_CPP_BY_NAME(MetasoundFrontendDocumentBuilder)
 
@@ -34,60 +31,6 @@ namespace Metasound::Frontend
 {
 	namespace DocumentBuilderPrivate
 	{
-		bool FindInputRegistryClass(FName TypeName, EMetasoundFrontendVertexAccessType AccessType, FMetasoundFrontendClass& OutClass)
-		{
-			switch (AccessType)
-			{
-				case EMetasoundFrontendVertexAccessType::Value:
-				{
-					return IDataTypeRegistry::Get().GetFrontendConstructorInputClass(TypeName, OutClass);
-				}
-				break;
-
-				case EMetasoundFrontendVertexAccessType::Reference:
-				{
-					return IDataTypeRegistry::Get().GetFrontendInputClass(TypeName, OutClass);
-				}
-				break;
-
-				case EMetasoundFrontendVertexAccessType::Unset:
-				default:
-				{
-					checkNoEntry();
-				}
-				break;
-			}
-
-			return false;
-		}
-
-		bool FindOutputRegistryClass(FName TypeName, EMetasoundFrontendVertexAccessType AccessType, FMetasoundFrontendClass& OutClass)
-		{
-			switch (AccessType)
-			{
-				case EMetasoundFrontendVertexAccessType::Value:
-				{
-					return IDataTypeRegistry::Get().GetFrontendConstructorOutputClass(TypeName, OutClass);
-				}
-				break;
-
-				case EMetasoundFrontendVertexAccessType::Reference:
-				{
-					return IDataTypeRegistry::Get().GetFrontendOutputClass(TypeName, OutClass);
-				}
-				break;
-
-				case EMetasoundFrontendVertexAccessType::Unset:
-				default:
-				{
-					checkNoEntry();
-				}
-				break;
-			}
-
-			return false;
-		}
-
 		bool TryGetInterfaceBoundEdges(
 			const FGuid& InFromNodeID,
 			const TSet<FMetasoundFrontendVersion>& InFromNodeInterfaces,
@@ -455,7 +398,7 @@ namespace Metasound::Frontend
 					const FName NodeName = Node->Name;
 					if (AddedNames.Contains(NodeName))
 					{
-						Node->Style.Display.Locations.Add(FGuid::NewGuid(), NextLocation);
+						Node->Style.Display.Locations.Add(FGuid(), NextLocation);
 						NextLocation += DisplayStyle::NodeLayout::DefaultOffsetY;
 					}
 					else
@@ -512,7 +455,7 @@ namespace Metasound::Frontend
 				}
 
 #if WITH_EDITORONLY_DATA
-				if (bAddedVertices && Options.bSetDefaultNodeLocations && !IsRunningCookCommandlet())
+				if (bAddedVertices && Options.bSetDefaultNodeLocations)
 				{
 					TArray<FMetasoundFrontendNode>& Nodes = Document.RootGraph.Graph.Nodes;
 					// Sort/Place Inputs
@@ -645,11 +588,6 @@ UMetaSoundBuilderDocument& UMetaSoundBuilderDocument::Create(const IMetaSoundDoc
 	return *DocObject;
 }
 
-bool UMetaSoundBuilderDocument::ConformObjectToDocument()
-{
-	return false;
-}
-
 FTopLevelAssetPath UMetaSoundBuilderDocument::GetAssetPathChecked() const
 {
 	FTopLevelAssetPath Path;
@@ -684,12 +622,32 @@ void UMetaSoundBuilderDocument::OnFinishActiveBuilder()
 	// Nothing to do here. UMetaSoundBuilderDocuments are always being used by builders
 }
 
-FMetaSoundFrontendDocumentBuilder::FMetaSoundFrontendDocumentBuilder(TScriptInterface<IMetaSoundDocumentInterface> InDocumentInterface, TSharedPtr<Metasound::Frontend::FDocumentModifyDelegates> InDocumentDelegates, bool bPrimeCache)
+FMetaSoundFrontendDocumentBuilder::FMetaSoundFrontendDocumentBuilder()
+	: DocumentDelegates(MakeShared<Metasound::Frontend::FDocumentModifyDelegates>())
+{
+}
+
+FMetaSoundFrontendDocumentBuilder::FMetaSoundFrontendDocumentBuilder(TScriptInterface<IMetaSoundDocumentInterface> InDocumentInterface)
 	: DocumentInterface(InDocumentInterface)
+	, DocumentDelegates(MakeShared<Metasound::Frontend::FDocumentModifyDelegates>())
 {
-	BeginBuilding(InDocumentDelegates, bPrimeCache);
+	if (DocumentInterface)
+	{
+		BeginBuilding();
+		InitCacheInternal();
+	}
 }
 
+FMetaSoundFrontendDocumentBuilder::FMetaSoundFrontendDocumentBuilder(TScriptInterface<IMetaSoundDocumentInterface> InDocumentInterface, TSharedRef<Metasound::Frontend::FDocumentModifyDelegates> InDocumentDelegates)
+	: DocumentInterface(InDocumentInterface)
+	, DocumentDelegates(InDocumentDelegates)
+{
+	if (DocumentInterface)
+	{
+		BeginBuilding();
+		InitCacheInternal();
+	}
+}
 FMetaSoundFrontendDocumentBuilder::~FMetaSoundFrontendDocumentBuilder()
 {
 	FinishBuilding();
@@ -724,15 +682,7 @@ void FMetaSoundFrontendDocumentBuilder::AddEdge(FMetasoundFrontendEdge&& InNewEd
 {
 	using namespace Metasound::Frontend;
 
-#if DO_CHECK
-	{
-		const IDocumentGraphEdgeCache& EdgeCache = DocumentCache->GetEdgeCache();
-		checkf(!EdgeCache.IsNodeInputConnected(InNewEdge.ToNodeID, InNewEdge.ToVertexID), TEXT("Failed to add edge in MetaSound Builder: Destination input already connected"));
-
-		const EInvalidEdgeReason Reason = IsValidEdge(InNewEdge);
-		checkf(Reason == Metasound::Frontend::EInvalidEdgeReason::None, TEXT("Attempted call to AddEdge in MetaSound Builder where edge is invalid: %s."), *LexToString(Reason));
-	}
-#endif // DO_CHECK
+	checkf(CanAddEdge(InNewEdge), TEXT("Attempted call to AddEdge in MetaSound Builder where not valid."));
 
 	FMetasoundFrontendDocument& Document = GetDocument();
 	FMetasoundFrontendGraph& Graph = Document.RootGraph.Graph;
@@ -995,15 +945,42 @@ const FMetasoundFrontendNode* FMetaSoundFrontendDocumentBuilder::AddGraphInput(c
 		return nullptr;
 	}
 
+	auto FindRegistryClass = [&InClassInput](FMetasoundFrontendClass& OutClass) -> bool
+	{
+		switch (InClassInput.AccessType)
+		{
+		case EMetasoundFrontendVertexAccessType::Value:
+		{
+			return IDataTypeRegistry::Get().GetFrontendConstructorInputClass(InClassInput.TypeName, OutClass);
+		}
+		break;
+
+		case EMetasoundFrontendVertexAccessType::Reference:
+		{
+			return IDataTypeRegistry::Get().GetFrontendInputClass(InClassInput.TypeName, OutClass);
+		}
+		break;
+
+		case EMetasoundFrontendVertexAccessType::Unset:
+		default:
+		{
+			checkNoEntry();
+		}
+		break;
+		}
+
+		return false;
+	};
+
 	FMetasoundFrontendClass Class;
-	if (DocumentBuilderPrivate::FindInputRegistryClass(InClassInput.TypeName, InClassInput.AccessType, Class))
+	if (FindRegistryClass(Class))
 	{
 		if(!FindDependency(Class.Metadata))
 		{
 			AddDependency(Class);
 		}
 
-		auto FinalizeNode = [this, &InClassInput](FMetasoundFrontendNode& InOutNode, const Metasound::Frontend::FNodeRegistryKey&)
+		auto FinalizeNode = [&InClassInput](FMetasoundFrontendNode& InOutNode, const Metasound::Frontend::FNodeRegistryKey&)
 		{
 			// Sets the name of the node an vertices on the node to match the class vertex name
 			DocumentBuilderPrivate::SetNodeAndVertexNames(InOutNode, InClassInput);
@@ -1063,8 +1040,35 @@ const FMetasoundFrontendNode* FMetaSoundFrontendDocumentBuilder::AddGraphOutput(
 		return nullptr;
 	}
 
+	auto FindRegistryClass = [&InClassOutput](FMetasoundFrontendClass& OutClass) -> bool
+	{
+		switch (InClassOutput.AccessType)
+		{
+		case EMetasoundFrontendVertexAccessType::Value:
+		{
+			return IDataTypeRegistry::Get().GetFrontendConstructorOutputClass(InClassOutput.TypeName, OutClass);
+		}
+		break;
+
+		case EMetasoundFrontendVertexAccessType::Reference:
+		{
+			return IDataTypeRegistry::Get().GetFrontendOutputClass(InClassOutput.TypeName, OutClass);
+		}
+		break;
+
+		case EMetasoundFrontendVertexAccessType::Unset:
+		default:
+		{
+			checkNoEntry();
+		}
+		break;
+		}
+
+		return false;
+	};
+
 	FMetasoundFrontendClass Class;
-	if (DocumentBuilderPrivate::FindOutputRegistryClass(InClassOutput.TypeName, InClassOutput.AccessType, Class))
+	if (FindRegistryClass(Class))
 	{
 		if (!FindDependency(Class.Metadata))
 		{
@@ -1081,7 +1085,7 @@ const FMetasoundFrontendNode* FMetaSoundFrontendDocumentBuilder::AddGraphOutput(
 			FMetasoundFrontendClassOutput& NewOutput = RootGraph.Interface.Outputs.Add_GetRef(InClassOutput);
 			if (!NewOutput.VertexID.IsValid())
 			{
-				NewOutput.VertexID = FDocumentIDGenerator::Get().CreateVertexID(Document);
+				NewOutput.VertexID = FDocumentIDGenerator::Get().CreateVertexID(Document);;
 			}
 
 			DocumentDelegates->InterfaceDelegates.OnOutputAdded.Broadcast(NewIndex);
@@ -1199,33 +1203,7 @@ const FMetasoundFrontendNode* FMetaSoundFrontendDocumentBuilder::AddNodeByClassN
 	return nullptr;
 }
 
-const FMetasoundFrontendNode* FMetaSoundFrontendDocumentBuilder::AddNodeByTemplate(const Metasound::Frontend::INodeTemplate& InTemplate, Metasound::Frontend::FNodeTemplateGenerateInterfaceParams Params, FGuid InNodeID)
-{
-	using namespace Metasound;
-	using namespace Metasound::Frontend;
-
-	const FMetasoundFrontendClass& TemplateClass = InTemplate.GetFrontendClass();
-	checkf(TemplateClass.Metadata.GetType() == EMetasoundFrontendClassType::Template, TEXT("INodeTemplate ClassType must always be 'Template'"));
-	const FMetasoundFrontendClass* Dependency = FindDependency(TemplateClass.Metadata);
-	if (!Dependency)
-	{
-		Dependency = AddDependency(TemplateClass);
-	}
-	check(Dependency);
-
-	auto FinalizeNodeFunction = [](const FMetasoundFrontendNode& Node, const Metasound::Frontend::FNodeRegistryKey& ClassKey)
-	{
-		return Node.Name;
-	};
-
-	FMetasoundFrontendNode* NewNode = AddNodeInternal(Dependency->Metadata, FinalizeNodeFunction, InNodeID);
-	check(NewNode);
-	NewNode->Interface = InTemplate.GenerateNodeInterface(MoveTemp(Params));
-
-	return NewNode;
-}
-
-FMetasoundFrontendNode* FMetaSoundFrontendDocumentBuilder::AddNodeInternal(const FMetasoundFrontendClassMetadata& InClassMetadata, FFinalizeNodeFunctionRef FinalizeNode, FGuid InNodeID, int32* NewNodeIndex)
+FMetasoundFrontendNode* FMetaSoundFrontendDocumentBuilder::AddNodeInternal(const FMetasoundFrontendClassMetadata& InClassMetadata, FFinalizeNodeFunctionRef FinalizeNode, FGuid InNodeID)
 {
 	METASOUND_TRACE_CPUPROFILER_EVENT_SCOPE(FMetaSoundFrontendDocumentBuilder::AddNodeInternal);
 
@@ -1245,11 +1223,6 @@ FMetasoundFrontendNode* FMetaSoundFrontendDocumentBuilder::AddNodeInternal(const
 		const IDocumentGraphNodeCache& NodeCache = DocumentCache->GetNodeCache();
 		DocumentDelegates->NodeDelegates.OnNodeAdded.Broadcast(NewIndex);
 
-		if (NewNodeIndex)
-		{
-			*NewNodeIndex = NewIndex;
-		}
-
 #if WITH_EDITORONLY_DATA
 		Document.Metadata.ModifyContext.AddNodeIDModified(InNodeID);
 #endif // WITH_EDITORONLY_DATA
@@ -1284,137 +1257,25 @@ bool FMetaSoundFrontendDocumentBuilder::RemoveUnusedDependencies()
 		const FGuid& ClassID = Dependencies[i].ID;
 		if (!DocumentCache->GetNodeCache().ContainsNodesOfClassID(ClassID))
 		{
-			const bool bRemoved = RemoveDependency(ClassID);
-			bDidEdit |= ensureAlwaysMsgf(bRemoved, TEXT("Failed to remove dependency that was found on document and was not referenced by nodes"));
+			checkf(RemoveDependency(ClassID), TEXT("Failed to remove dependency that was found on document and was not referenced by nodes"));
+			bDidEdit = true;
 		}
 	}
 
 	return bDidEdit;
 }
 
-void FMetaSoundFrontendDocumentBuilder::ClearGraph(TSharedRef<Metasound::Frontend::FDocumentModifyDelegates> ModifyDelegates)
+void FMetaSoundFrontendDocumentBuilder::ClearGraph()
 {
-	FMetasoundFrontendDocument& Doc = GetDocument();
-	FMetasoundFrontendGraphClass& GraphClass = Doc.RootGraph;
+	FMetasoundFrontendGraphClass& GraphClass = GetDocument().RootGraph;
 	GraphClass.Graph.Nodes.Reset();
 	GraphClass.Graph.Edges.Reset();
 	GraphClass.Interface.Inputs.Reset();
 	GraphClass.Interface.Outputs.Reset();
 	GraphClass.PresetOptions.InputsInheritingDefault.Reset();
-	Doc.Interfaces.Reset();
-	Doc.Dependencies.Reset();
-
-#if WITH_EDITORONLY_DATA
-	Doc.Metadata.MemberMetadata.Reset();
-#endif // WITH_EDITORONLY_DATA
-
-	Reload(ModifyDelegates);
-}
-
-#if WITH_EDITORONLY_DATA
-bool FMetaSoundFrontendDocumentBuilder::ClearMemberMetadata(const FGuid& InMemberID)
-{
-	return GetDocument().Metadata.MemberMetadata.Remove(InMemberID) > 0;
-}
-#endif // WITH_EDITORONLY_DATA
-
-bool FMetaSoundFrontendDocumentBuilder::ConformGraphInputNodeToClass(const FMetasoundFrontendClassInput& GraphInput)
-{
-	using namespace Metasound::Frontend;
-
-	FMetasoundFrontendDocument& Document = GetDocument();
-	const IDocumentGraphNodeCache& NodeCache = DocumentCache->GetNodeCache();
-	if (const int32* NodeIndexPtr = NodeCache.FindNodeIndex(GraphInput.NodeID))
-	{
-		TArray<FMetasoundFrontendNode>& Nodes = Document.RootGraph.Graph.Nodes;
-		FMetasoundFrontendNode& Node = Nodes[*NodeIndexPtr];
-		FMetasoundFrontendClass Class;
-		const bool bClassFound = DocumentBuilderPrivate::FindInputRegistryClass(GraphInput.TypeName, GraphInput.AccessType, Class);
-		if (ensureAlways(bClassFound))
-		{
-			const FMetasoundFrontendClass* Dependency = FindDependency(Class.Metadata);
-			if (!Dependency)
-			{
-				Dependency = AddDependency(Class);
-			}
-
-			if (ensureAlways(Dependency))
-			{
-				const int32 RemovalIndex = *NodeIndexPtr; // Have to cache as next delegate broadcast invalidates index pointer
-				DocumentDelegates->NodeDelegates.OnRemoveSwappingNode.Broadcast(RemovalIndex, Nodes.Num() - 1);
-				FMetasoundFrontendNode NewNode = MoveTemp(Node);
-				Nodes.RemoveAtSwap(RemovalIndex, EAllowShrinking::No);
-				NewNode.ClassID = Dependency->ID;
-				NewNode.Interface.Inputs.Last().TypeName = GraphInput.TypeName;
-				NewNode.Interface.Outputs.Last().TypeName = GraphInput.TypeName;
-
-#if WITH_EDITORONLY_DATA
-				Document.Metadata.ModifyContext.AddNodeIDModified(NewNode.GetID());
-#endif // WITH_EDITORONLY_DATA
-
-				// Set the default literal on the nodes inputs so that it gets passed to the instantiated TInputNode on a live
-				// auditioned MetaSound.
-				DocumentBuilderPrivate::SetDefaultLiteralOnInputNode(NewNode, GraphInput);
-
-				FMetasoundFrontendNode& NewNodeRef = Nodes.Add_GetRef(MoveTemp(NewNode));
-				DocumentDelegates->NodeDelegates.OnNodeAdded.Broadcast(Nodes.Num() - 1);
-
-				// Remove the default literal on the node added during the "FinalizeNode" call. This matches how 
-				// nodes are serialized in editor. The default literals are only stored on the FMetasoundFrontendClassInputs.
-				NewNodeRef.InputLiterals.Reset();
-
-				RemoveUnusedDependencies();
-				return true;
-			}
-		}
-	}
-
-	return false;
-}
-
-bool FMetaSoundFrontendDocumentBuilder::ConformGraphOutputNodeToClass(const FMetasoundFrontendClassOutput& GraphOutput)
-{
-	using namespace Metasound::Frontend;
-
-	FMetasoundFrontendDocument& Document = GetDocument();
-	const IDocumentGraphNodeCache& NodeCache = DocumentCache->GetNodeCache();
-	if (const int32* NodeIndexPtr = NodeCache.FindNodeIndex(GraphOutput.NodeID))
-	{
-		TArray<FMetasoundFrontendNode>& Nodes = Document.RootGraph.Graph.Nodes;
-		FMetasoundFrontendNode& Node = Nodes[*NodeIndexPtr];
-		FMetasoundFrontendClass Class;
-		const bool bClassFound = DocumentBuilderPrivate::FindOutputRegistryClass(GraphOutput.TypeName, GraphOutput.AccessType, Class);
-		if (ensureAlways(bClassFound))
-		{
-			const FMetasoundFrontendClass* Dependency = FindDependency(Class.Metadata);
-			if (!Dependency)
-			{
-				Dependency = AddDependency(Class);
-			}
-
-			if (ensureAlways(Dependency))
-			{
-				const int32 RemovalIndex = *NodeIndexPtr; // Have to cache as next delegate broadcast invalidates index pointer
-				DocumentDelegates->NodeDelegates.OnRemoveSwappingNode.Broadcast(RemovalIndex, Nodes.Num() - 1);
-				FMetasoundFrontendNode NewNode = MoveTemp(Node);
-				Nodes.RemoveAtSwap(RemovalIndex, EAllowShrinking::No);
-				NewNode.ClassID = Dependency->ID;
-				NewNode.Interface.Inputs.Last().TypeName = GraphOutput.TypeName;
-				NewNode.Interface.Outputs.Last().TypeName = GraphOutput.TypeName;
-
-#if WITH_EDITORONLY_DATA
-				Document.Metadata.ModifyContext.AddNodeIDModified(NewNode.GetID());
-#endif // WITH_EDITORONLY_DATA
-				Nodes.Add(MoveTemp(NewNode));
-				DocumentDelegates->NodeDelegates.OnNodeAdded.Broadcast(Nodes.Num() - 1);
-
-				RemoveUnusedDependencies();
-				return true;
-			}
-		}
-	}
-
-	return false;
+	GetDocument().Interfaces.Reset();
+	RemoveUnusedDependencies();
+	InitCacheInternal();
 }
 
 bool FMetaSoundFrontendDocumentBuilder::ContainsDependencyOfType(EMetasoundFrontendClassType ClassType) const
@@ -1458,12 +1319,12 @@ bool FMetaSoundFrontendDocumentBuilder::ConvertFromPreset()
 	return false;
 }
 
-bool FMetaSoundFrontendDocumentBuilder::ConvertToPreset(const FMetasoundFrontendDocument& InReferencedDocument, TSharedRef<Metasound::Frontend::FDocumentModifyDelegates> ModifyDelegates)
+bool FMetaSoundFrontendDocumentBuilder::ConvertToPreset(const FMetasoundFrontendDocument& InReferencedDocument)
 {
 	using namespace Metasound;
 	using namespace Metasound::Frontend;
 
-	ClearGraph(ModifyDelegates);
+	ClearGraph();
 
 	FMetasoundFrontendGraphClass& PresetAssetRootGraph = GetDocument().RootGraph;
 	FMetasoundFrontendGraph& PresetAssetGraph = PresetAssetRootGraph.Graph;
@@ -1473,32 +1334,23 @@ bool FMetaSoundFrontendDocumentBuilder::ConvertToPreset(const FMetasoundFrontend
 #endif // WITH_EDITORONLY_DATA
 
 	// Mark all inputs as inherited by default
+	TArray<FName> InputsInheritingDefault;
+	Algo::Transform(PresetAssetRootGraph.Interface.Inputs, InputsInheritingDefault, [](const FMetasoundFrontendClassInput& Input)
 	{
-		TArray<FName> InputsInheritingDefault;
-		auto GetInputName = [](const FMetasoundFrontendClassInput& Input) { return Input.Name; };
-		Algo::Transform(PresetAssetRootGraph.Interface.Inputs, InputsInheritingDefault, GetInputName);
+		return Input.Name;
+	});
 
-		PresetAssetRootGraph.PresetOptions.bIsPreset = true;
-		PresetAssetRootGraph.PresetOptions.InputsInheritingDefault = TSet<FName>(InputsInheritingDefault);
-	}
+	PresetAssetRootGraph.PresetOptions.bIsPreset = true;
+	PresetAssetRootGraph.PresetOptions.InputsInheritingDefault = TSet<FName>(InputsInheritingDefault);
 
 	// Apply root graph transform 
 	FRebuildPresetRootGraph RebuildPresetRootGraph(InReferencedDocument);
 	if (RebuildPresetRootGraph.Transform(GetDocument()))
 	{
-		DocumentInterface->ConformObjectToDocument();
-
-		// TL/DR: Have to reload and assign delegates here due to the rebuild preset transform still being implemented via controllers.
-		// Onces its reimplemented with the builder API, this can be removed.
-		// 
-		// The invalidate cache call when accessing the mutable document handle from within the transform unfortunately doesn't reach this
-		// builder's cache indirectly as converting to preset can be called by transient builders that are not registered with the MetaSound
-		// builder subsystem.
-		Reload(ModifyDelegates);
+		InitCacheInternal();
 		return true;
 	}
 
-	// Hack until prior processes are all moved over to builder API
 	return false;
 }
 
@@ -1685,16 +1537,6 @@ const FMetasoundFrontendNode* FMetaSoundFrontendDocumentBuilder::FindGraphOutput
 	return nullptr;
 }
 
-#if WITH_EDITOR
-UMetaSoundFrontendMemberMetadata* FMetaSoundFrontendDocumentBuilder::FindMemberMetadata(const FGuid& InMemberID)
-{
-	FMetasoundFrontendDocument& Document = GetDocument();
-	TMap<FGuid, TObjectPtr<UMetaSoundFrontendMemberMetadata>>& LiteralMetadata = Document.Metadata.MemberMetadata;
-	TObjectPtr<UMetaSoundFrontendMemberMetadata> ToReturn = LiteralMetadata.FindRef(InMemberID);
-	return ToReturn;
-}
-#endif // WITH_EDITOR
-
 const FMetasoundFrontendNode* FMetaSoundFrontendDocumentBuilder::FindNode(const FGuid& InNodeID) const
 {
 	using namespace Metasound::Frontend;
@@ -1859,7 +1701,7 @@ FMetasoundFrontendDocument& FMetaSoundFrontendDocumentBuilder::GetDocument()
 
 const FMetasoundFrontendDocument& FMetaSoundFrontendDocumentBuilder::GetDocument() const
 {
-	return DocumentInterface->GetConstDocument();
+	return DocumentInterface->GetDocument();
 }
 
 const Metasound::Frontend::FDocumentModifyDelegates& FMetaSoundFrontendDocumentBuilder::GetDocumentDelegates() const
@@ -1874,28 +1716,6 @@ const IMetaSoundDocumentInterface& FMetaSoundFrontendDocumentBuilder::GetDocumen
 	return *Interface;
 }
 
-FMetasoundAssetBase& FMetaSoundFrontendDocumentBuilder::GetMetasoundAsset()
-{
-	using namespace Metasound::Frontend;
-
-	UObject* Object = DocumentInterface.GetObject();
-	check(Object);
-	FMetasoundAssetBase* Asset = IMetaSoundAssetManager::GetChecked().GetAsAsset(*Object);
-	check(Asset);
-	return *Asset;
-}
-
-const FMetasoundAssetBase& FMetaSoundFrontendDocumentBuilder::GetMetasoundAsset() const
-{
-	using namespace Metasound::Frontend;
-
-	UObject* Object = DocumentInterface.GetObject();
-	check(Object);
-	const FMetasoundAssetBase* Asset = IMetaSoundAssetManager::GetChecked().GetAsAsset(*Object);
-	check(Asset);
-	return *Asset;
-}
-
 EMetasoundFrontendVertexAccessType FMetaSoundFrontendDocumentBuilder::GetNodeInputAccessType(const FGuid& InNodeID, const FGuid& InVertexID) const
 {
 	using namespace Metasound::Frontend;
@@ -2077,48 +1897,6 @@ EMetasoundFrontendVertexAccessType FMetaSoundFrontendDocumentBuilder::GetNodeOut
 	return EMetasoundFrontendVertexAccessType::Unset;
 }
 
-void FMetaSoundFrontendDocumentBuilder::InitDocument(const FMetasoundFrontendDocument* InDocumentTemplate, const FMetasoundFrontendClassName* InNewClassName, bool bResetVersion)
-{
-	using namespace Metasound;
-	using namespace Metasound::Frontend;
-
-	METASOUND_TRACE_CPUPROFILER_EVENT_SCOPE(FMetaSoundFrontendDocumentBuilder::InitDocument);
-
-	FMetasoundFrontendDocument& Document = GetDocument();
-
-	// 1. Set default class Metadata.
-	if (InDocumentTemplate)
-	{
-		// 1a. If template provided, copy that.
-		Document = *InDocumentTemplate;
-		InitGraphClassMetadata(bResetVersion, InNewClassName);
-	}
-	else
-	{
-		// 1a. Initialize class using default data
-		FMetasoundFrontendClassMetadata& ClassMetadata = Document.RootGraph.Metadata;
-		InitGraphClassMetadata(Document.RootGraph.Metadata, bResetVersion, InNewClassName);
-
-		// 1b. Set default doc version Metadata
-		{
-			FMetasoundFrontendDocumentMetadata& DocMetadata = Document.Metadata;
-			DocMetadata.Version.Number = FMetasoundFrontendDocument::GetMaxVersion();
-		}
-
-		// 1c. Add default interfaces for given UClass
-		{
-			TArray<FMetasoundFrontendVersion> InitVersions = ISearchEngine::Get().FindUClassDefaultInterfaceVersions(GetBuilderClassPath());
-			FModifyInterfaceOptions Options({ }, InitVersions);
-			ModifyInterfaces(MoveTemp(Options));
-		}
-	}
-}
-
-bool FMetaSoundFrontendDocumentBuilder::IsValid() const
-{
-	return DocumentInterface.GetObject() != nullptr;
-}
-
 int32 FMetaSoundFrontendDocumentBuilder::GetTransactionCount() const
 {
 	using namespace Metasound::Frontend;
@@ -2150,9 +1928,34 @@ void FMetaSoundFrontendDocumentBuilder::InitGraphClassMetadata(FMetasoundFronten
 	InOutMetadata.SetType(EMetasoundFrontendClassType::Graph);
 }
 
-void FMetaSoundFrontendDocumentBuilder::InitGraphClassMetadata(bool bResetVersion, const FMetasoundFrontendClassName* NewClassName)
+void FMetaSoundFrontendDocumentBuilder::InitDocument()
 {
-	InitGraphClassMetadata(GetDocument().RootGraph.Metadata, bResetVersion, NewClassName);
+	using namespace Metasound;
+	using namespace Metasound::Frontend;
+
+	METASOUND_TRACE_CPUPROFILER_EVENT_SCOPE(FMetaSoundFrontendDocumentBuilder::InitDocument);
+
+	FMetasoundFrontendDocument& Document = GetDocument();
+
+	// 1. Set default class Metadata
+	{
+		constexpr bool bResetVersion = true;
+		FMetasoundFrontendClassMetadata& ClassMetadata = Document.RootGraph.Metadata;
+		InitGraphClassMetadata(ClassMetadata, bResetVersion);
+	}
+
+	// 2. Set default doc version Metadata
+	{
+		FMetasoundFrontendDocumentMetadata& DocMetadata = Document.Metadata;
+		DocMetadata.Version.Number = FMetasoundFrontendDocument::GetMaxVersion();
+	}
+
+	// 3. Add default interfaces for given UClass
+	{
+		TArray<FMetasoundFrontendVersion> InitVersions = ISearchEngine::Get().FindUClassDefaultInterfaceVersions(GetBuilderClassPath());
+		FModifyInterfaceOptions Options({ }, InitVersions);
+		ModifyInterfaces(MoveTemp(Options));
+	}
 }
 
 void FMetaSoundFrontendDocumentBuilder::InitNodeLocations()
@@ -2212,6 +2015,12 @@ void FMetaSoundFrontendDocumentBuilder::InitNodeLocations()
 #endif // WITH_EDITORONLY_DATA
 }
 
+void FMetaSoundFrontendDocumentBuilder::InvalidateCache()
+{
+	constexpr bool bPrimeCache = false;
+	InitCacheInternal(bPrimeCache);
+}
+
 bool FMetaSoundFrontendDocumentBuilder::IsDependencyReferenced(const FGuid& InClassID) const
 {
 	return DocumentCache->GetNodeCache().ContainsNodesOfClassID(InClassID);
@@ -2250,15 +2059,6 @@ bool FMetaSoundFrontendDocumentBuilder::IsPreset() const
 	return GetDocument().RootGraph.PresetOptions.bIsPreset;
 }
 
-bool FMetaSoundFrontendDocumentBuilder::IsVariableClass(EMetasoundFrontendClassType ClassType) const
-{
-	const bool bIsVariableNode = (ClassType == EMetasoundFrontendClassType::Variable)
-		|| (ClassType == EMetasoundFrontendClassType::VariableAccessor)
-		|| (ClassType == EMetasoundFrontendClassType::VariableDeferredAccessor)
-		|| (ClassType == EMetasoundFrontendClassType::VariableMutator);
-	return bIsVariableNode;
-}
-
 Metasound::Frontend::EInvalidEdgeReason FMetaSoundFrontendDocumentBuilder::IsValidEdge(const FMetasoundFrontendEdge& InEdge) const
 {
 	using namespace Metasound::Frontend;
@@ -2294,27 +2094,6 @@ Metasound::Frontend::EInvalidEdgeReason FMetaSoundFrontendDocumentBuilder::IsVal
 	return EInvalidEdgeReason::None;
 }
 
-void FMetaSoundFrontendDocumentBuilder::IterateNodesConnectedWithVertex(const FMetasoundFrontendVertexHandle& Vertex, TFunctionRef<void(const FMetasoundFrontendEdge&, FMetasoundFrontendNode&)> NodeIndexIterFunc)
-{
-	using namespace Metasound::Frontend;
-
-	FMetasoundFrontendDocument& Document = GetDocument();
-	TArray<FMetasoundFrontendEdge> EdgesToConnectedNodes; // Have to cache to avoid pointers becoming garbage in subsequent removal loop
-	const IDocumentGraphEdgeCache& EdgeCache = DocumentCache->GetEdgeCache();
-	const IDocumentGraphNodeCache& NodeCache = DocumentCache->GetNodeCache();
-	const TArray<const FMetasoundFrontendEdge*> Edges = EdgeCache.FindEdges(Vertex.NodeID, Vertex.VertexID);
-	Algo::Transform(Edges, EdgesToConnectedNodes, [](const FMetasoundFrontendEdge* Edge) { check(Edge); return *Edge; });
-	for (const FMetasoundFrontendEdge& Edge : EdgesToConnectedNodes)
-	{
-		const FGuid& ConnectedNodeID = Edge.ToNodeID == Vertex.NodeID ? Edge.FromNodeID : Edge.ToNodeID;
-		if (const int32* ConnectedNodeIndex = NodeCache.FindNodeIndex(ConnectedNodeID))
-		{
-			FMetasoundFrontendNode& Node = Document.RootGraph.Graph.Nodes[*ConnectedNodeIndex];
-			NodeIndexIterFunc(Edge, Node);
-		}
-	}
-}
-
 bool FMetaSoundFrontendDocumentBuilder::ModifyInterfaces(Metasound::Frontend::FModifyInterfaceOptions&& InOptions)
 {
 	using namespace Metasound::Frontend;
@@ -2377,7 +2156,7 @@ bool FMetaSoundFrontendDocumentBuilder::TransformTemplateNodes()
 		const FTemplateTransformParams& Params = Pair.Value;
 		if (Params.Template)
 		{
-			TUniquePtr<INodeTemplateTransform> NodeTransform = Params.Template->GenerateNodeTransform();
+			TUniquePtr<INodeTransform> NodeTransform = Params.Template->GenerateNodeTransform();
 			check(NodeTransform.IsValid());
 
 			for (const FGuid& NodeID : Params.NodeIDs)
@@ -2403,31 +2182,27 @@ bool FMetaSoundFrontendDocumentBuilder::TransformTemplateNodes()
 	return bModified;
 }
 
-void FMetaSoundFrontendDocumentBuilder::BeginBuilding(TSharedPtr<Metasound::Frontend::FDocumentModifyDelegates> Delegates, bool bPrimeCache)
+void FMetaSoundFrontendDocumentBuilder::InitCacheInternal(bool bPrimeCache)
 {
 	using namespace Metasound::Frontend;
+	DocumentCache = FDocumentCache::Create(GetDocument(), DocumentDelegates, bPrimeCache);
+}
 
-	DocumentDelegates = Delegates.IsValid() ? Delegates : MakeShared<FDocumentModifyDelegates>();
-
+void FMetaSoundFrontendDocumentBuilder::BeginBuilding()
+{
 	if (DocumentInterface)
 	{
 		DocumentInterface->OnBeginActiveBuilder();
-		DocumentCache = FDocumentCache::Create(GetDocument(), DocumentDelegates.ToSharedRef(), bPrimeCache);
 	}
 }
 
 void FMetaSoundFrontendDocumentBuilder::FinishBuilding()
 {
-	using namespace Metasound::Frontend;
-
 	if (DocumentInterface)
 	{
 		DocumentInterface->OnFinishActiveBuilder();
-		DocumentInterface = { };
 	}
-
-	DocumentDelegates.Reset();
-	DocumentCache.Reset();
+	DocumentInterface = TScriptInterface<IMetaSoundDocumentInterface>();
 }
 
 bool FMetaSoundFrontendDocumentBuilder::RemoveDependency(const FGuid& InClassID)
@@ -2564,33 +2339,15 @@ bool FMetaSoundFrontendDocumentBuilder::RemoveNamedEdges(const TSet<Metasound::F
 	return bSuccess;
 }
 
-void FMetaSoundFrontendDocumentBuilder::Reload(TSharedPtr<Metasound::Frontend::FDocumentModifyDelegates> Delegates, bool bPrimeCache)
+bool FMetaSoundFrontendDocumentBuilder::RemoveNodeInputDefault(const FGuid& InNodeID, const FGuid& InVertexID)
 {
 	using namespace Metasound::Frontend;
 
-	if (DocumentInterface)
+	const IDocumentGraphNodeCache& NodeCache = DocumentCache->GetNodeCache();
+	if (const int32* NodeIndex = NodeCache.FindNodeIndex(InNodeID))
 	{
-		DocumentInterface->OnFinishActiveBuilder();
-	}
-
-	DocumentDelegates = Delegates.IsValid() ? Delegates : MakeShared<FDocumentModifyDelegates>();
-
-	if (DocumentInterface)
-	{
-		DocumentCache = FDocumentCache::Create(GetDocument(), DocumentDelegates.ToSharedRef(), bPrimeCache);
-		DocumentInterface->OnBeginActiveBuilder();
-	}
-}
-
-bool FMetaSoundFrontendDocumentBuilder::RemoveNodeInputDefault(const FGuid& InNodeID, const FGuid& InVertexID)
-{
-	using namespace Metasound::Frontend;
-
-	const IDocumentGraphNodeCache& NodeCache = DocumentCache->GetNodeCache();
-	if (const int32* NodeIndex = NodeCache.FindNodeIndex(InNodeID))
-	{
-		FMetasoundFrontendGraph& Graph = GetDocument().RootGraph.Graph;
-		FMetasoundFrontendNode& Node = Graph.Nodes[*NodeIndex];
+		FMetasoundFrontendGraph& Graph = GetDocument().RootGraph.Graph;
+		FMetasoundFrontendNode& Node = Graph.Nodes[*NodeIndex];
 
 		auto IsVertex = [&InVertexID](const FMetasoundFrontendVertex& Vertex) { return Vertex.VertexID == InVertexID; };
 		const int32 VertexIndex = Node.Interface.Inputs.IndexOfByPredicate(IsVertex);
@@ -2683,10 +2440,6 @@ bool FMetaSoundFrontendDocumentBuilder::RemoveEdgeToNodeInput(const FGuid& InNod
 		DocumentDelegates->EdgeDelegates.OnRemoveSwappingEdge.Broadcast(Index, LastIndex);
 		Graph.Edges.RemoveAtSwap(Index, EAllowShrinking::No);
 
-#if WITH_EDITORONLY_DATA
-		Document.Metadata.ModifyContext.AddNodeIDModified(InNodeID);
-#endif // WITH_EDITORONLY_DATA
-
 		return true;
 	}
 
@@ -2696,128 +2449,51 @@ bool FMetaSoundFrontendDocumentBuilder::RemoveEdgeToNodeInput(const FGuid& InNod
 #if WITH_EDITOR
 bool FMetaSoundFrontendDocumentBuilder::RemoveGraphComment(const FGuid& InCommentID)
 {
-	FMetasoundFrontendDocument& Document = GetDocument();
-	TMap<FGuid, FMetaSoundFrontendGraphComment>& Comments = Document.RootGraph.Graph.Style.Comments;
-
-	if (Comments.Remove(InCommentID) > 0)
-	{
-#if WITH_EDITORONLY_DATA
-		Document.Metadata.ModifyContext.SetDocumentModified();
-#endif // WITH_EDITORONLY_DATA
-
-		return true;
-	}
-
-	return false;
+	TMap<FGuid, FMetaSoundFrontendGraphComment>& Comments = GetDocument().RootGraph.Graph.Style.Comments;
+	return Comments.Remove(InCommentID) > 0;
 }
 #endif // WITH_EDITOR
 
-TArray<const FMetasoundFrontendNode*> FMetaSoundFrontendDocumentBuilder::GetGraphInputTemplateNodes(FName InInputName)
+bool FMetaSoundFrontendDocumentBuilder::RemoveGraphInput(FName InInputName)
 {
-	using namespace Metasound::Frontend;
-
-	TArray<const FMetasoundFrontendNode*> TemplateNodes;
-
-	FMetasoundFrontendDocument& Document = GetDocument();
-	if (const int32* Index = DocumentCache->GetInterfaceCache().FindInputIndex(InInputName))
+	if (const FMetasoundFrontendNode* Node = FindGraphInputNode(InInputName))
 	{
-		const FMetasoundFrontendClassInput& InputClass = Document.RootGraph.Interface.Inputs[*Index];
-		const IDocumentGraphNodeCache& NodeCache = DocumentCache->GetNodeCache();
-		const IDocumentGraphEdgeCache& EdgeCache = DocumentCache->GetEdgeCache();
-
-		if (const FMetasoundFrontendNode* InputNode = NodeCache.FindNode(InputClass.NodeID))
+		const FGuid ClassID = Node->ClassID;
+		if (RemoveNode(Node->GetID()))
 		{
-			const FGuid OutputVertexID = InputNode->Interface.Outputs.Last().VertexID;
-			const TArray<const FMetasoundFrontendEdge*> ConnectedEdges = EdgeCache.FindEdges(InputClass.NodeID, OutputVertexID);
-			for (const FMetasoundFrontendEdge* Edge : ConnectedEdges)
+			TArray<FMetasoundFrontendClassInput>& Inputs = GetDocument().RootGraph.Interface.Inputs;
+			auto InputNameMatches = [InInputName](const FMetasoundFrontendClassInput& Input) { return Input.Name == InInputName; };
+			const int32 Index = Inputs.IndexOfByPredicate(InputNameMatches);
+			if (Index != INDEX_NONE)
 			{
-				check(Edge);
-				if (const int32* ConnectedNodeIndex = NodeCache.FindNodeIndex(Edge->ToNodeID))
+				DocumentDelegates->InterfaceDelegates.OnRemovingInput.Broadcast(Index);
+
+				const int32 LastIndex = Inputs.Num() - 1;
+				if (Index != LastIndex)
 				{
-					const FMetasoundFrontendNode& ConnectedNode = Document.RootGraph.Graph.Nodes[*ConnectedNodeIndex];
-					if (const FMetasoundFrontendClass* ConnectedNodeClass = FindDependency(ConnectedNode.ClassID))
+					DocumentDelegates->InterfaceDelegates.OnRemovingInput.Broadcast(LastIndex);
+				}
+				Inputs.RemoveAtSwap(Index, EAllowShrinking::No);
+				if (Index != LastIndex)
+				{
+					DocumentDelegates->InterfaceDelegates.OnInputAdded.Broadcast(Index);
+				}
+
+				if (IsDependencyReferenced(ClassID))
+				{
+					return true;
+				}
+				else
+				{
+					if (RemoveDependency(ClassID))
 					{
-						if (ConnectedNodeClass->Metadata.GetClassName() == FInputNodeTemplate::ClassName)
-						{
-							TemplateNodes.Add(&ConnectedNode);
-						}
+						return true;
 					}
 				}
 			}
 		}
 	}
 
-	return TemplateNodes;
-}
-
-bool FMetaSoundFrontendDocumentBuilder::RemoveGraphInput(FName InInputName)
-{
-	using namespace Metasound::Frontend;
-
-	FMetasoundFrontendDocument& Document = GetDocument();
-	if (const int32* IndexPtr = DocumentCache->GetInterfaceCache().FindInputIndex(InInputName))
-	{
-		TArray<FMetasoundFrontendClassInput>& Inputs = GetDocument().RootGraph.Interface.Inputs;
-		const FGuid NodeID = Inputs[*IndexPtr].NodeID;
-		FGuid ClassID;
-		TArray<FGuid> NodeIDsToRemove { NodeID };
-
-		if (const FMetasoundFrontendNode* Node = FindNode(NodeID))
-		{
-			ClassID = Node->ClassID;
-		}
-		else
-		{
-			return false;
-		}
-
-		const TArray<const FMetasoundFrontendNode*> TemplateNodes = GetGraphInputTemplateNodes(InInputName);
-		Algo::Transform(TemplateNodes, NodeIDsToRemove, [](const FMetasoundFrontendNode* Node) { return Node->GetID(); });
-
-		bool bNodesRemoved = true;
-		for (const FGuid& ToRemove : NodeIDsToRemove)
-		{
-			if (RemoveNode(ToRemove))
-			{
-#if WITH_EDITORONLY_DATA
-				Document.Metadata.ModifyContext.AddNodeIDModified(ToRemove);
-#endif // WITH_EDITORONLY_DATA
-			}
-			else
-			{
-				bNodesRemoved = false;
-			}
-		}
-
-		if (bNodesRemoved)
-		{
-			const int32 Index = *IndexPtr;
-			DocumentDelegates->InterfaceDelegates.OnRemovingInput.Broadcast(Index);
-
-			const int32 LastIndex = Inputs.Num() - 1;
-			if (Index != LastIndex)
-			{
-				DocumentDelegates->InterfaceDelegates.OnRemovingInput.Broadcast(LastIndex);
-			}
-			Inputs.RemoveAtSwap(Index, EAllowShrinking::No);
-			if (Index != LastIndex)
-			{
-				DocumentDelegates->InterfaceDelegates.OnInputAdded.Broadcast(Index);
-			}
-
-#if WITH_EDITORONLY_DATA
-			ClearMemberMetadata(NodeID);
-			Document.Metadata.ModifyContext.AddMemberIDModified(NodeID);
-#endif // WITH_EDITORONLY_DATA
-
-			const bool bDependencyReferenced = IsDependencyReferenced(ClassID);
-			if (bDependencyReferenced || RemoveDependency(ClassID))
-			{
-				return true;
-			}
-		}
-	}
-
 	return false;
 }
 
@@ -2826,16 +2502,9 @@ bool FMetaSoundFrontendDocumentBuilder::RemoveGraphOutput(FName InOutputName)
 	if (const FMetasoundFrontendNode* Node = FindGraphOutputNode(InOutputName))
 	{
 		const FGuid ClassID = Node->ClassID;
-		const FGuid NodeID = Node->GetID();
-		if (RemoveNode(NodeID))
+		if (RemoveNode(Node->GetID()))
 		{
-			FMetasoundFrontendDocument& Document = GetDocument();
-
-#if WITH_EDITORONLY_DATA
-			Document.Metadata.ModifyContext.AddNodeIDModified(NodeID);
-#endif // WITH_EDITORONLY_DATA
-
-			TArray<FMetasoundFrontendClassOutput>& Outputs = Document.RootGraph.Interface.Outputs;
+			TArray<FMetasoundFrontendClassOutput>& Outputs = GetDocument().RootGraph.Interface.Outputs;
 			auto OutputNameMatches = [InOutputName](const FMetasoundFrontendClassOutput& Output) { return Output.Name == InOutputName; };
 			const int32 Index = Outputs.IndexOfByPredicate(OutputNameMatches);
 			if (Index != INDEX_NONE)
@@ -2853,16 +2522,17 @@ bool FMetaSoundFrontendDocumentBuilder::RemoveGraphOutput(FName InOutputName)
 					DocumentDelegates->InterfaceDelegates.OnOutputAdded.Broadcast(Index);
 				}
 
-#if WITH_EDITORONLY_DATA
-				ClearMemberMetadata(NodeID);
-				Document.Metadata.ModifyContext.AddMemberIDModified(NodeID);
-#endif // WITH_EDITORONLY_DATA
-
-				const bool bDependencyReferenced = IsDependencyReferenced(ClassID);
-				if (bDependencyReferenced || RemoveDependency(ClassID))
+				if (IsDependencyReferenced(ClassID))
 				{
 					return true;
 				}
+				else
+				{
+					if (RemoveDependency(ClassID))
+					{
+						return true;
+					}
+				}
 			}
 		}
 	}
@@ -2911,7 +2581,6 @@ bool FMetaSoundFrontendDocumentBuilder::RemoveNode(const FGuid& InNodeID)
 	using namespace Metasound::Frontend;
 
 	const IDocumentGraphNodeCache& NodeCache = DocumentCache->GetNodeCache();
-	const IDocumentGraphEdgeCache& EdgeCache = DocumentCache->GetEdgeCache();
 
 	if (const int32* IndexPtr = NodeCache.FindNodeIndex(InNodeID))
 	{
@@ -2921,55 +2590,15 @@ bool FMetaSoundFrontendDocumentBuilder::RemoveNode(const FGuid& InNodeID)
 		FMetasoundFrontendGraph& Graph = Document.RootGraph.Graph;
 		TArray<FMetasoundFrontendNode>& Nodes = Graph.Nodes;
 		const FMetasoundFrontendNode& Node = Nodes[Index];
-		const FGuid& NodeID = Node.GetID();
-
-		const FMetasoundFrontendClass* NodeClass = DocumentCache->FindDependency(Node.ClassID);
-		check(NodeClass);
 
-		const bool bIsVariableNode = IsVariableClass(NodeClass->Metadata.GetType());
-		if (bIsVariableNode)
-		{
-			const bool bUnlinkedVariable = UnlinkVariableNode(Node.GetID());
-			ensureMsgf(bUnlinkedVariable, TEXT("Removed variable node that was not linked to variable"));
-		}
-
-		FMetasoundFrontendVertexHandle FromVertexHandle;
 		for (const FMetasoundFrontendVertex& Vertex : Node.Interface.Inputs)
 		{
-			if (bIsVariableNode)
-			{
-				if (const int32* InputEdgeIndex = EdgeCache.FindEdgeIndexToNodeInput(InNodeID, Vertex.VertexID))
-				{
-					FromVertexHandle = Graph.Edges[*InputEdgeIndex].GetFromVertexHandle();
-				}
-			}
-
 			RemoveEdgeToNodeInput(InNodeID, Vertex.VertexID);
 		}
 
-		TArray<FMetasoundFrontendVertexHandle> ToVertexHandles;
 		for (const FMetasoundFrontendVertex& Vertex : Node.Interface.Outputs)
 		{
-			if (bIsVariableNode)
-			{
-				ToVertexHandles.Reset();
-				const TArrayView<const int32> OutputEdgeIndices = EdgeCache.FindEdgeIndicesFromNodeOutput(InNodeID, Vertex.VertexID);
-				Algo::Transform(OutputEdgeIndices, ToVertexHandles, [&Graph](const int32& VertIndex) { return Graph.Edges[VertIndex].GetToVertexHandle(); });
-			}
 			RemoveEdgesFromNodeOutput(InNodeID, Vertex.VertexID);
-			if (bIsVariableNode)
-			{
-				for (const FMetasoundFrontendVertexHandle& ToHandle : ToVertexHandles)
-				{
-					AddEdge(FMetasoundFrontendEdge
-					{
-						FromVertexHandle.NodeID,
-						FromVertexHandle.VertexID,
-						ToHandle.NodeID,
-						ToHandle.VertexID
-					});
-				}
-			}
 		}
 
 		const int32 LastIndex = Nodes.Num() - 1;
@@ -3029,9 +2658,8 @@ bool FMetaSoundFrontendDocumentBuilder::RenameRootGraphClass(const FMetasoundFro
 
 void FMetaSoundFrontendDocumentBuilder::ReloadCache()
 {
-	using namespace Metasound::Frontend;
-
-	Reload(DocumentDelegates, true);
+	constexpr bool bPrimeCache = true;
+	InitCacheInternal(bPrimeCache);
 }
 
 #if WITH_EDITOR
@@ -3042,142 +2670,6 @@ void FMetaSoundFrontendDocumentBuilder::SetAuthor(const FString& InAuthor)
 }
 #endif // WITH_EDITOR
 
-bool FMetaSoundFrontendDocumentBuilder::SetGraphInputAccessType(FName InputName, EMetasoundFrontendVertexAccessType AccessType)
-{
-	using namespace Metasound::Frontend;
-
-	if (!ensureMsgf(AccessType != EMetasoundFrontendVertexAccessType::Unset, TEXT("Cannot set graph input access type to '%s'"), LexToString(AccessType)))
-	{
-		return false;
-	}
-
-	if (const int32* Index = DocumentCache->GetInterfaceCache().FindInputIndex(InputName))
-	{
-		FMetasoundFrontendDocument& Document = GetDocument();
-		FMetasoundFrontendClassInput& GraphInput = Document.RootGraph.Interface.Inputs[*Index];
-
-		if (GraphInput.AccessType != AccessType)
-		{
-			GraphInput.AccessType = AccessType;
-			if (AccessType == EMetasoundFrontendVertexAccessType::Reference)
-			{
-				const IDocumentGraphNodeCache& NodeCache = DocumentCache->GetNodeCache();
-				if (const int32* NodeIndex = NodeCache.FindNodeIndex(GraphInput.NodeID))
-				{
-					FMetasoundFrontendNode& Node = Document.RootGraph.Graph.Nodes[*NodeIndex];
-					const FMetasoundFrontendVertex& NodeOutput = Node.Interface.Outputs.Last();
-					IterateNodesConnectedWithVertex({ GraphInput.NodeID, NodeOutput.VertexID }, [this, &AccessType, &Document](const FMetasoundFrontendEdge& Edge, FMetasoundFrontendNode& ConnectedNode)
-					{
-						if (const FMetasoundFrontendClass* ConnectedNodeClass = FindDependency(ConnectedNode.ClassID))
-						{
-							// If connected to an input template node, disconnect the template node from other nodes as the data type is
-							// about to be mismatched.  Otherwise, direct connection to other nodes (i.e. at runtime when template
-							// nodes aren't injected) forcefully remove to avoid data type mismatch.
-							if (ConnectedNodeClass->Metadata.GetClassName() == FInputNodeTemplate::ClassName)
-							{
-								const FMetasoundFrontendVertex& ConnectedNodeOutput = ConnectedNode.Interface.Outputs.Last();
-								IterateNodesConnectedWithVertex({ Edge.ToNodeID, ConnectedNodeOutput.VertexID }, [this, &AccessType](const FMetasoundFrontendEdge& TempEdge, FMetasoundFrontendNode&)
-								{
-									const EMetasoundFrontendVertexAccessType ConnectedAccessType = GetNodeInputAccessType(TempEdge.ToNodeID, TempEdge.ToVertexID);
-									if (!FMetasoundFrontendClassVertex::CanConnectVertexAccessTypes(AccessType, ConnectedAccessType))
-									{
-										RemoveEdgeToNodeInput(TempEdge.ToNodeID, TempEdge.ToVertexID);
-									}
-								});
-							}
-							else
-							{
-								const EMetasoundFrontendVertexAccessType ConnectedAccessType = GetNodeInputAccessType(Edge.ToNodeID, Edge.ToVertexID);
-								if (!FMetasoundFrontendClassVertex::CanConnectVertexAccessTypes(AccessType, ConnectedAccessType))
-								{
-									RemoveEdgeToNodeInput(Edge.ToNodeID, Edge.ToVertexID);
-								}
-							}
-						}
-					});
-
-				}
-			}
-
-			const bool bNodeConformed = ConformGraphInputNodeToClass(GraphInput);
-			if (bNodeConformed)
-			{
-#if WITH_EDITORONLY_DATA
-				Document.Metadata.ModifyContext.AddMemberIDModified(GraphInput.NodeID);
-#endif // WITH_EDITORONLY_DATA
-
-				return true;
-			}
-		}
-	}
-
-	return false;
-}
-
-bool FMetaSoundFrontendDocumentBuilder::SetGraphInputDataType(FName InputName, FName DataType)
-{
-	using namespace Metasound::Frontend;
-
-	if (IDataTypeRegistry::Get().IsRegistered(DataType))
-	{
-		if (const int32* Index = DocumentCache->GetInterfaceCache().FindInputIndex(InputName))
-		{
-			FMetasoundFrontendDocument& Document = GetDocument();
-			FMetasoundFrontendClassInput& GraphInput = Document.RootGraph.Interface.Inputs[*Index];
-			if (GraphInput.TypeName != DataType)
-			{
-				FMetasoundFrontendLiteral DefaultLiteral;
-				DefaultLiteral.SetFromLiteral(IDataTypeRegistry::Get().CreateDefaultLiteral(DataType));
-				GraphInput.DefaultLiteral = DefaultLiteral;
-				GraphInput.TypeName = DataType;
-
-				const IDocumentGraphNodeCache& NodeCache = DocumentCache->GetNodeCache();
-				if (const int32* NodeIndex = NodeCache.FindNodeIndex(GraphInput.NodeID))
-				{
-					FMetasoundFrontendNode& Node = Document.RootGraph.Graph.Nodes[*NodeIndex];
-					FMetasoundFrontendVertex& NodeOutput = Node.Interface.Outputs.Last();
-					IterateNodesConnectedWithVertex({ GraphInput.NodeID, NodeOutput.VertexID }, [this, &DataType](const FMetasoundFrontendEdge& Edge, FMetasoundFrontendNode& ConnectedNode)
-					{
-						const FMetasoundFrontendClass* ConnectedNodeClass = FindDependency(ConnectedNode.ClassID);
-						if (ensure(ConnectedNodeClass))
-						{
-							// If connected to an input template node, disconnect the template node from other nodes as the data type is
-							// about to be mismatched.  Otherwise, direct connection to other nodes (i.e. at runtime when template
-							// nodes aren't injected) forcefully remove to avoid data type mismatch.
-							if (ConnectedNodeClass->Metadata.GetClassName() == FInputNodeTemplate::ClassName)
-							{
-								RemoveEdgesFromNodeOutput(Edge.ToNodeID, ConnectedNode.Interface.Outputs.Last().VertexID);
-								ConnectedNode.Interface.Inputs.Last().TypeName = DataType;
-								ConnectedNode.Interface.Outputs.Last().TypeName = DataType;
-							}
-							else
-							{
-								RemoveEdgeToNodeInput(Edge.ToNodeID, Edge.ToVertexID);
-							}
-						}
-					});
-				}
-
-				const bool bNodeConformed = ConformGraphInputNodeToClass(GraphInput);
-				if (bNodeConformed)
-				{
-					RemoveUnusedDependencies();
-
-#if WITH_EDITORONLY_DATA
-					ClearMemberMetadata(GraphInput.NodeID);
-					Document.Metadata.ModifyContext.AddMemberIDModified(GraphInput.NodeID);
-					Document.Metadata.ModifyContext.AddNodeIDModified(GraphInput.NodeID);
-#endif // WITH_EDITORONLY_DATA
-
-					return true;
-				}
-			}
-		}
-	}
-
-	return false;
-}
-
 bool FMetaSoundFrontendDocumentBuilder::SetGraphInputDefault(FName InputName, const FMetasoundFrontendLiteral& InDefaultLiteral)
 {
 	using namespace Metasound::Frontend;
@@ -3230,139 +2722,6 @@ bool FMetaSoundFrontendDocumentBuilder::SetGraphInputInheritsDefault(FName InNam
 	return false;
 }
 
-bool FMetaSoundFrontendDocumentBuilder::SetGraphOutputAccessType(FName OutputName, EMetasoundFrontendVertexAccessType AccessType)
-{
-	using namespace Metasound::Frontend;
-
-	if (!ensureMsgf(AccessType != EMetasoundFrontendVertexAccessType::Unset, TEXT("Cannot set graph output access type to '%s'"), LexToString(AccessType)))
-	{
-		return false;
-	}
-
-	if (const int32* Index = DocumentCache->GetInterfaceCache().FindOutputIndex(OutputName))
-	{
-		FMetasoundFrontendDocument& Document = GetDocument();
-		FMetasoundFrontendClassOutput& GraphOutput = Document.RootGraph.Interface.Outputs[*Index];
-		if (GraphOutput.AccessType != AccessType)
-		{
-			GraphOutput.AccessType = AccessType;
-			if (AccessType == EMetasoundFrontendVertexAccessType::Value)
-			{
-				const IDocumentGraphNodeCache& NodeCache = DocumentCache->GetNodeCache();
-				if (const int32* NodeIndex = NodeCache.FindNodeIndex(GraphOutput.NodeID))
-				{
-					FMetasoundFrontendNode& Node = Document.RootGraph.Graph.Nodes[*NodeIndex];
-					const FMetasoundFrontendVertex& NodeInput = Node.Interface.Inputs.Last();
-					IterateNodesConnectedWithVertex({ GraphOutput.NodeID, NodeInput.VertexID }, [this, &AccessType](const FMetasoundFrontendEdge& Edge, FMetasoundFrontendNode& ConnectedNode)
-					{
-						if (const FMetasoundFrontendClass* ConnectedNodeClass = FindDependency(ConnectedNode.ClassID))
-						{
-							const FMetasoundFrontendVertex& ConnectedNodeOutput = ConnectedNode.Interface.Outputs.Last();
-							const EMetasoundFrontendVertexAccessType ConnectedAccessType = GetNodeOutputAccessType(ConnectedNode.GetID(), ConnectedNodeOutput.VertexID);
-							if (!FMetasoundFrontendClassVertex::CanConnectVertexAccessTypes(ConnectedAccessType, AccessType))
-							{
-								RemoveEdgeToNodeInput(Edge.ToNodeID, Edge.ToVertexID);
-							}
-						}
-					});
-				}
-			}
-
-			const bool bNodeConformed = ConformGraphOutputNodeToClass(GraphOutput);
-			if (bNodeConformed)
-			{
-	#if WITH_EDITORONLY_DATA
-				Document.Metadata.ModifyContext.AddMemberIDModified(GraphOutput.NodeID);
-	#endif // WITH_EDITORONLY_DATA
-
-				return true;
-			}
-		}
-	}
-
-	return false;
-}
-
-bool FMetaSoundFrontendDocumentBuilder::SetGraphOutputDataType(FName OutputName, FName DataType)
-{
-	using namespace Metasound::Frontend;
-
-	if (IDataTypeRegistry::Get().IsRegistered(DataType))
-	{
-		if (const int32* Index = DocumentCache->GetInterfaceCache().FindOutputIndex(OutputName))
-		{
-			FMetasoundFrontendDocument& Document = GetDocument();
-			FMetasoundFrontendClassOutput& GraphOutput = Document.RootGraph.Interface.Outputs[*Index];
-			if (GraphOutput.TypeName != DataType)
-			{
-				GraphOutput.TypeName = DataType;
-
-				const IDocumentGraphNodeCache& NodeCache = DocumentCache->GetNodeCache();
-				if (const int32* NodeIndex = NodeCache.FindNodeIndex(GraphOutput.NodeID))
-				{
-					FMetasoundFrontendNode& Node = Document.RootGraph.Graph.Nodes[*NodeIndex];
-
-					FMetasoundFrontendLiteral DefaultLiteral;
-					DefaultLiteral.SetFromLiteral(IDataTypeRegistry::Get().CreateDefaultLiteral(DataType));
-					FMetasoundFrontendVertex& NodeInput = Node.Interface.Inputs.Last();
-					Node.InputLiterals = { FMetasoundFrontendVertexLiteral { NodeInput.VertexID, MoveTemp(DefaultLiteral) } };
-
-					RemoveEdgeToNodeInput(GraphOutput.NodeID, NodeInput.VertexID);
-					GraphOutput.TypeName = DataType;
-
-				}
-
-				const bool bNodeConformed = ConformGraphOutputNodeToClass(GraphOutput);
-				if (bNodeConformed)
-				{
-
-#if WITH_EDITORONLY_DATA
-					ClearMemberMetadata(GraphOutput.NodeID);
-					Document.Metadata.ModifyContext.AddMemberIDModified(GraphOutput.NodeID);
-#endif // WITH_EDITORONLY_DATA
-
-					return true;
-				}
-			}
-		}
-	}
-
-	return false;
-}
-
-#if WITH_EDITOR
-void FMetaSoundFrontendDocumentBuilder::SetMemberMetadata(UMetaSoundFrontendMemberMetadata& NewMetadata)
-{
-	check(NewMetadata.MemberID.IsValid());
-
-	TMap<FGuid, TObjectPtr<UMetaSoundFrontendMemberMetadata>>& LiteralMetadata = GetDocument().Metadata.MemberMetadata;
-	LiteralMetadata.Remove(NewMetadata.MemberID);
-	LiteralMetadata.Add(NewMetadata.MemberID, NewMetadata);
-}
-
-bool FMetaSoundFrontendDocumentBuilder::SetNodeComment(const FGuid& InNodeID, FString&& InNewComment)
-{
-	if (FMetasoundFrontendNode* Node = FindNodeInternal(InNodeID))
-	{
-		Node->Style.Display.Comment = MoveTemp(InNewComment);
-		return true;
-	}
-
-	return false;
-}
-
-bool FMetaSoundFrontendDocumentBuilder::SetNodeCommentVisible(const FGuid& InNodeID, bool bIsVisible)
-{
-	if (FMetasoundFrontendNode* Node = FindNodeInternal(InNodeID))
-	{
-		Node->Style.Display.bCommentVisible = bIsVisible;
-		return true;
-	}
-
-	return false;
-}
-#endif // WITH_EDITOR
-
 bool FMetaSoundFrontendDocumentBuilder::SetNodeInputDefault(const FGuid& InNodeID, const FGuid& InVertexID, const FMetasoundFrontendLiteral& InLiteral)
 {
 	using namespace Metasound::Frontend;
@@ -3403,6 +2762,28 @@ bool FMetaSoundFrontendDocumentBuilder::SetNodeInputDefault(const FGuid& InNodeI
 }
 
 #if WITH_EDITOR
+bool FMetaSoundFrontendDocumentBuilder::SetNodeComment(const FGuid& InNodeID, FString&& InNewComment)
+{
+	if (FMetasoundFrontendNode* Node = FindNodeInternal(InNodeID))
+	{
+		Node->Style.Display.Comment = MoveTemp(InNewComment);
+		return true;
+	}
+
+	return false;
+}
+
+bool FMetaSoundFrontendDocumentBuilder::SetNodeCommentVisible(const FGuid& InNodeID, bool bIsVisible)
+{
+	if (FMetasoundFrontendNode* Node = FindNodeInternal(InNodeID))
+	{
+		Node->Style.Display.bCommentVisible = bIsVisible;
+		return true;
+	}
+
+	return false;
+}
+
 bool FMetaSoundFrontendDocumentBuilder::SetNodeLocation(const FGuid& InNodeID, const FVector2D& InLocation, const FGuid* InLocationGuid)
 {
 	if (FMetasoundFrontendNode* Node = FindNodeInternal(InNodeID))
@@ -3410,16 +2791,9 @@ bool FMetaSoundFrontendDocumentBuilder::SetNodeLocation(const FGuid& InNodeID, c
 		FMetasoundFrontendNodeStyle& Style = Node->Style;
 		if (InLocationGuid)
 		{
-			if (InLocationGuid->IsValid())
-			{
-				Style.Display.Locations.FindOrAdd(*InLocationGuid) = InLocation;
-				return true;
-			}
-
-			UE_LOG(LogMetaSound, Display, TEXT("Invalid Location Guid no longer supported, reseting display location for node with ID '%s'"), *InNodeID.ToString());
+			Style.Display.Locations.FindOrAdd(*InLocationGuid) = InLocation;
 		}
-
-		if (Style.Display.Locations.IsEmpty())
+		else if (Style.Display.Locations.IsEmpty())
 		{
 			Style.Display.Locations = { { FGuid::NewGuid(), InLocation } };
 		}
@@ -3598,39 +2972,6 @@ bool FMetaSoundFrontendDocumentBuilder::SwapGraphOutput(const FMetasoundFrontend
 	return true;
 }
 
-bool FMetaSoundFrontendDocumentBuilder::UnlinkVariableNode(const FGuid& InNodeID)
-{
-	using namespace Metasound::VariableNames;
-
-	FMetasoundFrontendGraphClass& GraphClass = GetDocument().RootGraph;
-	for (FMetasoundFrontendVariable& Variable : GraphClass.Graph.Variables)
-	{
-		if (InNodeID == Variable.VariableNodeID)
-		{
-			Variable.VariableNodeID = FGuid();
-			return true;
-		}
-
-		if (InNodeID == Variable.MutatorNodeID)
-		{
-			Variable.MutatorNodeID = FGuid();
-			return true;
-		}
-
-		if (Variable.AccessorNodeIDs.Remove(InNodeID) > 0)
-		{
-			return true;
-		}
-
-		if (Variable.DeferredAccessorNodeIDs.Remove(InNodeID) > 0)
-		{
-			return true;
-		}
-	}
-
-	return false;
-}
-
 bool FMetaSoundFrontendDocumentBuilder::UpdateDependencyClassNames(const TMap<FMetasoundFrontendClassName, FMetasoundFrontendClassName>& OldToNewReferencedClassNames)
 {
 	for (FMetasoundFrontendClass& Dependency : GetDocument().Dependencies)
diff --git a/Engine/Plugins/Runtime/Metasound/Source/MetasoundFrontend/Private/MetasoundFrontendDocumentCache.cpp b/Engine/Plugins/Runtime/Metasound/Source/MetasoundFrontend/Private/MetasoundFrontendDocumentCache.cpp
--- a/Engine/Plugins/Runtime/Metasound/Source/MetasoundFrontend/Private/MetasoundFrontendDocumentCache.cpp
+++ b/Engine/Plugins/Runtime/Metasound/Source/MetasoundFrontend/Private/MetasoundFrontendDocumentCache.cpp
@@ -338,11 +338,6 @@ namespace Metasound::Frontend
 		return nullptr;
 	}
 
-	const int32* FDocumentGraphInterfaceCache::FindInputIndex(FName InputName) const
-	{
-		return InputNameToIndex.Find(InputName);
-	}
-
 	const FMetasoundFrontendClassOutput* FDocumentGraphInterfaceCache::FindOutput(FName OutputName) const
 	{
 		if (const int32* Index = OutputNameToIndex.Find(OutputName))
@@ -354,11 +349,6 @@ namespace Metasound::Frontend
 		return nullptr;
 	}
 
-	const int32* FDocumentGraphInterfaceCache::FindOutputIndex(FName OutputName) const
-	{
-		return OutputNameToIndex.Find(OutputName);
-	}
-
 	int32 FDocumentGraphInterfaceCache::GetTransactionCount() const
 	{
 		return TransactionCount;
diff --git a/Engine/Plugins/Runtime/Metasound/Source/MetasoundFrontend/Private/MetasoundFrontendDocumentCache.h b/Engine/Plugins/Runtime/Metasound/Source/MetasoundFrontend/Private/MetasoundFrontendDocumentCache.h
--- a/Engine/Plugins/Runtime/Metasound/Source/MetasoundFrontend/Private/MetasoundFrontendDocumentCache.h
+++ b/Engine/Plugins/Runtime/Metasound/Source/MetasoundFrontend/Private/MetasoundFrontendDocumentCache.h
@@ -126,9 +126,7 @@ namespace Metasound::Frontend
 
 		// IDocumentGraphInterfaceCache implementation
 		virtual const FMetasoundFrontendClassInput* FindInput(FName InputName) const override;
-		virtual const int32* FindInputIndex(FName InputName) const override;
 		virtual const FMetasoundFrontendClassOutput* FindOutput(FName OutputName) const override;
-		virtual const int32* FindOutputIndex(FName OutputName) const override;
 
 		int32 GetTransactionCount() const;
 
diff --git a/Engine/Plugins/Runtime/Metasound/Source/MetasoundFrontend/Private/MetasoundFrontendDocumentVersioning.h b/Engine/Plugins/Runtime/Metasound/Source/MetasoundFrontend/Private/MetasoundFrontendDocumentVersioning.h
--- a/Engine/Plugins/Runtime/Metasound/Source/MetasoundFrontend/Private/MetasoundFrontendDocumentVersioning.h
+++ b/Engine/Plugins/Runtime/Metasound/Source/MetasoundFrontend/Private/MetasoundFrontendDocumentVersioning.h
@@ -2,35 +2,31 @@
 #pragma once
 
 #include "Containers/UnrealString.h"
-#include "MetasoundDocumentInterface.h"
 #include "MetasoundFrontendController.h"
 #include "MetasoundFrontendTransform.h"
 #include "MetasoundFrontendDocument.h"
 #include "UObject/NameTypes.h"
-#include "UObject/ScriptInterface.h"
 
-class IMetaSoundDocumentInterface;
 
-
-namespace Metasound::Frontend
+namespace Metasound
 {
-	/** Base class for versioning a document. */
-	class METASOUNDFRONTEND_API FVersionDocument
+	namespace Frontend
 	{
-		const FName Name;
-		const FString& Path;
-
-	public:
-		static FMetasoundFrontendVersionNumber GetMaxVersion()
+		/** Base class for versioning a document. */
+		class METASOUNDFRONTEND_API FVersionDocument : public IDocumentTransform
 		{
-			return FMetasoundFrontendVersionNumber { 1, 12 };
-		}
+			const FName Name;
+			const FString& Path;
 
-		FVersionDocument(FName InName, const FString& InPath);
+		public:
+			static FMetasoundFrontendVersionNumber GetMaxVersion()
+			{
+				return FMetasoundFrontendVersionNumber { 1, 11 };
+			}
 
-		UE_DEPRECATED(5.4, "Use DocumentInterface overload instead")
-		bool Transform(FDocumentHandle InDocument) const;
+			FVersionDocument(FName InName, const FString& InPath);
 
-		bool Transform(TScriptInterface<IMetaSoundDocumentInterface> DocumentInterface) const;
-	};
-} // namespace Metasound::Frontend
+			bool Transform(FDocumentHandle InDocument) const override;
+		};
+	} // namespace Frontend
+} // namespace Metasound
diff --git a/Engine/Plugins/Runtime/Metasound/Source/MetasoundFrontend/Private/MetasoundFrontendGraphController.cpp b/Engine/Plugins/Runtime/Metasound/Source/MetasoundFrontend/Private/MetasoundFrontendGraphController.cpp
--- a/Engine/Plugins/Runtime/Metasound/Source/MetasoundFrontend/Private/MetasoundFrontendGraphController.cpp
+++ b/Engine/Plugins/Runtime/Metasound/Source/MetasoundFrontend/Private/MetasoundFrontendGraphController.cpp
@@ -750,13 +750,6 @@ namespace Metasound
 				GraphClass->PresetOptions.InputsInheritingDefault.Reset();
 				OwningDocument->ClearInterfaceVersions();
 				OwningDocument->RemoveUnreferencedDependencies();
-
-#if WITH_EDITORONLY_DATA
-				FDocumentAccessPtr DocumentPtr = OwningDocument->GetDocumentPtr();
-				FMetasoundFrontendDocument* Document = DocumentPtr.Get();
-				check(Document);
-				Document->Metadata.MemberMetadata.Reset();
-#endif // WITH_EDITORONLY_DATA
 			}
 		}
 
diff --git a/Engine/Plugins/Runtime/Metasound/Source/MetasoundFrontend/Private/MetasoundFrontendInputController.cpp b/Engine/Plugins/Runtime/Metasound/Source/MetasoundFrontend/Private/MetasoundFrontendInputController.cpp
--- a/Engine/Plugins/Runtime/Metasound/Source/MetasoundFrontend/Private/MetasoundFrontendInputController.cpp
+++ b/Engine/Plugins/Runtime/Metasound/Source/MetasoundFrontend/Private/MetasoundFrontendInputController.cpp
@@ -63,7 +63,6 @@ namespace Metasound
 			EMetasoundFrontendVertexAccessType AccessType = EMetasoundFrontendVertexAccessType::Unset;
 			bool bIsRerouted = false;
 
-			PRAGMA_DISABLE_DEPRECATION_WARNINGS
 			Frontend::IterateReroutedInputs(AsShared(), [this, &bIsRerouted, &AccessType](const FConstInputHandle& ReroutedInput)
 			{
 				bIsRerouted = true;
@@ -94,7 +93,6 @@ namespace Metasound
 					}
 				}
 			});
-			PRAGMA_ENABLE_DEPRECATION_WARNINGS
 
 			return bIsRerouted ? AccessType : EMetasoundFrontendVertexAccessType::Reference;
 		}
diff --git a/Engine/Plugins/Runtime/Metasound/Source/MetasoundFrontend/Private/MetasoundFrontendNodeTemplateRegistry.cpp b/Engine/Plugins/Runtime/Metasound/Source/MetasoundFrontend/Private/MetasoundFrontendNodeTemplateRegistry.cpp
--- a/Engine/Plugins/Runtime/Metasound/Source/MetasoundFrontend/Private/MetasoundFrontendNodeTemplateRegistry.cpp
+++ b/Engine/Plugins/Runtime/Metasound/Source/MetasoundFrontend/Private/MetasoundFrontendNodeTemplateRegistry.cpp
@@ -12,6 +12,11 @@
 
 namespace Metasound::Frontend
 {
+	TUniquePtr<INodeTransform> INodeTemplate::GenerateNodeTransform(FMetasoundFrontendDocument& InDocument) const
+	{
+		return nullptr;
+	}
+
 	class FNodeTemplateRegistry : public INodeTemplateRegistry
 	{
 	public:
@@ -53,9 +58,9 @@ namespace Metasound::Frontend
 
 	const INodeTemplate* FNodeTemplateRegistry::FindTemplate(const FNodeRegistryKey& InKey) const
 	{
-		if (const TUniquePtr<const INodeTemplate>* Template = Templates.Find(InKey))
+		if (const TUniquePtr<const INodeTemplate>* TemplatePtr = Templates.Find(InKey))
 		{
-			return Template->Get();
+			return TemplatePtr->Get();
 		}
 
 		return nullptr;
@@ -82,11 +87,6 @@ namespace Metasound::Frontend
 		return Registry;
 	}
 
-	TUniquePtr<INodeTransform> INodeTemplate::GenerateNodeTransform(FMetasoundFrontendDocument& InDocument) const
-	{
-		return nullptr;
-	}
-
 #if WITH_EDITOR
 	FText FNodeTemplateBase::GetNodeDisplayName(const IMetaSoundDocumentInterface& Interface, const FGuid& InNodeID) const
 	{
diff --git a/Engine/Plugins/Runtime/Metasound/Source/MetasoundFrontend/Private/MetasoundFrontendOutputController.cpp b/Engine/Plugins/Runtime/Metasound/Source/MetasoundFrontend/Private/MetasoundFrontendOutputController.cpp
--- a/Engine/Plugins/Runtime/Metasound/Source/MetasoundFrontend/Private/MetasoundFrontendOutputController.cpp
+++ b/Engine/Plugins/Runtime/Metasound/Source/MetasoundFrontend/Private/MetasoundFrontendOutputController.cpp
@@ -2,7 +2,6 @@
 
 #include "MetasoundFrontendOutputController.h"
 #include "MetasoundFrontendInvalidController.h"
-#include "MetasoundFrontendNodeTemplateRegistry.h"
 
 #define LOCTEXT_NAMESPACE "MetasoundFrontendOutputController"
 
@@ -52,10 +51,7 @@ namespace Metasound
 
 		EMetasoundFrontendVertexAccessType FBaseOutputController::GetVertexAccessType() const
 		{
-			PRAGMA_DISABLE_DEPRECATION_WARNINGS
 			FConstOutputHandle ReroutedOutput = Frontend::FindReroutedOutput(AsShared());
-			PRAGMA_ENABLE_DEPRECATION_WARNINGS
-
 			if (ReroutedOutput->IsValid())
 			{
 				if (ReroutedOutput->GetOwningNodeID() != GetOwningNodeID() || ReroutedOutput->GetID() != GetID())
@@ -149,17 +145,6 @@ namespace Metasound
 
 		bool FBaseOutputController::IsConnectionUserModifiable() const 
 		{
-			FConstNodeHandle Owner = GetOwningNode();
-			if (Owner->GetClassMetadata().GetType() == EMetasoundFrontendClassType::Template)
-			{
-				const FNodeRegistryKey Key(GetOwningNode()->GetClassMetadata());
-				const INodeTemplate* Template = INodeTemplateRegistry::Get().FindTemplate(Key);
-				if (ensure(Template))
-				{
-					return Template->IsOutputConnectionUserModifiable();
-				}
-			}
-
 			return true;
 		}
 
diff --git a/Engine/Plugins/Runtime/Metasound/Source/MetasoundFrontend/Private/MetasoundFrontendRegistryContainerImpl.cpp b/Engine/Plugins/Runtime/Metasound/Source/MetasoundFrontend/Private/MetasoundFrontendRegistryContainerImpl.cpp
--- a/Engine/Plugins/Runtime/Metasound/Source/MetasoundFrontend/Private/MetasoundFrontendRegistryContainerImpl.cpp
+++ b/Engine/Plugins/Runtime/Metasound/Source/MetasoundFrontend/Private/MetasoundFrontendRegistryContainerImpl.cpp
@@ -17,7 +17,6 @@
 #include "UObject/NameTypes.h"
 #include "UObject/ScriptInterface.h"
 
-#include "MetasoundAssetBase.h"
 #include "MetasoundAssetManager.h"
 #include "MetasoundDocumentInterface.h"
 #include "MetasoundFrontendDataTypeRegistry.h"
@@ -43,31 +42,41 @@ namespace Metasound::Frontend
 {
 	namespace RegistryPrivate
 	{
-		TScriptInterface<IMetaSoundDocumentInterface> BuildRegistryDocument(TScriptInterface<IMetaSoundDocumentInterface> DocumentInterface, bool bAsync)
+		TScriptInterface<IMetaSoundDocumentInterface> BuildRegistryDocument(TScriptInterface<IMetaSoundDocumentInterface> DocumentInterface, bool bForceCopy)
 		{
 			using namespace Metasound::Frontend;
 
 			METASOUND_TRACE_CPUPROFILER_EVENT_SCOPE(Metasound::Frontend::BuildRegistryDocument);
 
-			UObject* DocObject = DocumentInterface.GetObject();
-			check(DocObject);
-
 #if WITH_EDITOR
-			// Only assets require template node processing
-			if (DocObject->IsAsset())
+			// Node template transform is performed on copy of local document to avoid overwriting editable data
+			constexpr bool bTransformDocumentBeforeRegistering = true;
+#else // !WITH_EDITOR
+			// Node template transform is performed on local document only if cook determinism ID generation
+			// is enabled to avoid transforms potentially creating new edges with non-deterministic IDs.
+			const bool bTransformDocumentBeforeRegistering = MetaSoundEnableCookDeterministicIDGeneration == 0;
+#endif // !WITH_EDITOR
+
+			if (bTransformDocumentBeforeRegistering)
 			{
-				DocumentInterface = &UMetaSoundBuilderDocument::Create(*DocumentInterface.GetInterface());
-				FMetaSoundFrontendDocumentBuilder Builder(DocumentInterface);
-				Builder.TransformTemplateNodes();
-			}
+				// 1. Find template dependencies to build prior to making new document/builder as an optimization
+				// (no sense in creating new document/builder if no templates need processing)
+				FMetaSoundFrontendDocumentBuilder OriginalDocBuilder(DocumentInterface);
+				const bool bContainsTemplateDependency = OriginalDocBuilder.ContainsDependencyOfType(EMetasoundFrontendClassType::Template);
+				if (bContainsTemplateDependency)
+				{
+					UMetaSoundBuilderDocument& RegistryDocObject = UMetaSoundBuilderDocument::Create(*DocumentInterface.GetInterface());
+					FMetaSoundFrontendDocumentBuilder RegistryDocBuilder(&RegistryDocObject);
 
-			return DocumentInterface;
-#else
-	#if !NO_LOGGING
-			// Only assets require template node processing and support document attachment
-			if (DocObject->IsAsset()) 
+					RegistryDocBuilder.TransformTemplateNodes();
+
+					return &RegistryDocObject;
+				}
+			}
+#if !NO_LOGGING
+			else
 			{
-				const FMetaSoundFrontendDocumentBuilder& OriginalDocBuilder = IMetaSoundAssetManager::GetChecked().AttachDocumentBuilderChecked(*DocObject);
+				FMetaSoundFrontendDocumentBuilder OriginalDocBuilder(DocumentInterface);
 				const bool bContainsTemplateDependency = OriginalDocBuilder.ContainsDependencyOfType(EMetasoundFrontendClassType::Template);
 				if (bContainsTemplateDependency)
 				{
@@ -76,12 +85,8 @@ namespace Metasound::Frontend
 						*OriginalDocBuilder.GetDebugName());
 				}
 			}
-	#endif // !NO_LOGGING
+#endif // !NO_LOGGING
 
-			// Force a copy if async registration is enabled and we need to protect against race conditions from external modifications.
-			const FMetasoundAssetBase* AssetBase = IMetaSoundAssetManager::GetChecked().GetAsAsset(*DocObject);
-			check(AssetBase);
-			const bool bForceCopy = AssetBase->IsBuilderActive() && bAsync;
 			if (bForceCopy)
 			{
 				return &UMetaSoundBuilderDocument::Create(*DocumentInterface.GetInterface());
@@ -90,7 +95,6 @@ namespace Metasound::Frontend
 			{
 				return DocumentInterface;
 			}
-#endif // WITH_EDITOR
 		}
 
 		// FGraphNode is used to create unique INodes based off of a IGraph.
@@ -433,7 +437,7 @@ namespace Metasound::Frontend
 		return MakeUnique<FNodeRegistryTransactionStream>(TransactionBuffer);
 	}
 
-	FGraphRegistryKey FRegistryContainerImpl::RegisterGraph(const TScriptInterface<IMetaSoundDocumentInterface>& InDocumentInterface, bool bAsync)
+	FGraphRegistryKey FRegistryContainerImpl::RegisterGraph(const TScriptInterface<IMetaSoundDocumentInterface>& InDocumentInterface, bool bAsync, bool bForceCopy)
 	{
 		using namespace UE;
 
@@ -442,8 +446,15 @@ namespace Metasound::Frontend
 		check(InDocumentInterface);
 		check(IsInGameThread());
 
-		const FMetasoundFrontendDocument& Document = InDocumentInterface->GetConstDocument();
+		// Use the asset path of the provided document interface object for identification, *NOT* the
+		// built version as the build process may in fact create a new object with a transient path.
 		const FTopLevelAssetPath AssetPath = InDocumentInterface->GetAssetPathChecked();
+		const TScriptInterface<IMetaSoundDocumentInterface> RegistryDocInterface = RegistryPrivate::BuildRegistryDocument(InDocumentInterface, bForceCopy);
+
+		UObject* OwningObject = RegistryDocInterface.GetObject();
+		check(OwningObject);
+
+		const FMetasoundFrontendDocument& Document = RegistryDocInterface->GetConstDocument();
 		const FGraphRegistryKey RegistryKey { FNodeRegistryKey(Document.RootGraph), AssetPath };
 
 		if (!RegistryKey.IsValid())
@@ -455,19 +466,6 @@ namespace Metasound::Frontend
 
 		METASOUND_TRACE_CPUPROFILER_EVENT_SCOPE_TEXT(*FString::Printf(TEXT("FRegistryContainerImpl::RegisterGraph key:%s, asset %s"), *RegistryKey.ToString(), *AssetPath.ToString()));
 
-		// Wait for any async tasks that are in flight which correspond to the same graph prior to building as building, if force copy is false, may mutate the serialized document
-		if (bAsync)
-		{
-			WaitForAsyncGraphRegistration(RegistryKey);
-		}
-
-		// Use the asset path of the provided document interface object for identification, *NOT* the
-		// built version as the build process may in fact create a new object with a transient path.
-		const TScriptInterface<IMetaSoundDocumentInterface> RegistryDocInterface = RegistryPrivate::BuildRegistryDocument(InDocumentInterface, bAsync);
-
-		UObject* OwningObject = RegistryDocInterface.GetObject();
-		check(OwningObject);
-
 		FNodeClassInfo NodeClassInfo(Document.RootGraph, AssetPath);
 
 		// Proxies are created synchronously to avoid creating proxies in async tasks. Proxies
@@ -484,13 +482,20 @@ namespace Metasound::Frontend
 
 		if (bAsync)
 		{
+			// Wait for any async tasks that are in flight which correspond to the same graph
+			WaitForAsyncGraphRegistration(RegistryKey);
+
 			Tasks::FTask BuildAndRegisterTask = AsyncRegistrationPipe.Launch(
 				UE_SOURCE_LOCATION,
 				[RegistryKey, ClassInfo = MoveTemp(NodeClassInfo), RegistryDocInterface, ProxyDataCache = MoveTemp(ProxyDataCache)]() mutable
 				{
 					FRegistryContainerImpl& Registry = FRegistryContainerImpl::Get();
 					// Unregister the graph before re-registering
-					Registry.UnregisterGraphInternal(RegistryKey);
+					if (Registry.IsGraphRegistered(RegistryKey))
+					{
+						Registry.UnregisterGraphInternal(RegistryKey);
+					}
+
 					Registry.BuildAndRegisterGraphFromDocument(RegistryDocInterface, ProxyDataCache, MoveTemp(ClassInfo));
 					Registry.RemoveRegistrationTask(RegistryKey, FNodeRegistryTransaction::ETransactionType::NodeRegistration);
 					Registry.RemoveDocumentReference(RegistryDocInterface);
@@ -507,7 +512,10 @@ namespace Metasound::Frontend
 		}
 		else
 		{
-			UnregisterGraphInternal(RegistryKey);
+			if (IsGraphRegistered(RegistryKey))
+			{
+				UnregisterGraphInternal(RegistryKey);
+			}
 
 			// Build and register graph synchronously
 			BuildAndRegisterGraphFromDocument(RegistryDocInterface, ProxyDataCache, MoveTemp(NodeClassInfo));
@@ -612,11 +620,6 @@ namespace Metasound::Frontend
 
 		FScopeLock Lock(&RegistryMapsCriticalSection);
 		{
-			if (!RegisteredGraphs.Contains(InKey))
-			{
-				return false;
-			}
-
 			const int32 GraphUnregistered = RegisteredGraphs.Remove(InKey) > 0;
 			const bool bNodeUnregistered = UnregisterNodeInternal(InKey.NodeKey);
 
diff --git a/Engine/Plugins/Runtime/Metasound/Source/MetasoundFrontend/Private/NodeTemplates/MetasoundFrontendNodeTemplateReroute.cpp b/Engine/Plugins/Runtime/Metasound/Source/MetasoundFrontend/Private/NodeTemplates/MetasoundFrontendNodeTemplateReroute.cpp
--- a/Engine/Plugins/Runtime/Metasound/Source/MetasoundFrontend/Private/NodeTemplates/MetasoundFrontendNodeTemplateReroute.cpp
+++ b/Engine/Plugins/Runtime/Metasound/Source/MetasoundFrontend/Private/NodeTemplates/MetasoundFrontendNodeTemplateReroute.cpp
@@ -16,7 +16,7 @@ namespace Metasound::Frontend
 {
 	namespace ReroutePrivate
 	{
-		class FRerouteNodeTemplateTransform : public INodeTemplateTransform
+		class FRerouteNodeTemplateTransform : public INodeTransform
 		{
 		public:
 			FRerouteNodeTemplateTransform() = default;
@@ -109,8 +109,6 @@ namespace Metasound::Frontend
 							}
 						}
 					}
-
-					return nullptr;
 				}
 			}
 
@@ -204,13 +202,12 @@ namespace Metasound::Frontend
 		NewInterface.Outputs.Add(FMetasoundFrontendVertex { VertexName, DataType, FGuid::NewGuid() });
 
 		return NewInterface;
-
 	}
 
-	TUniquePtr<INodeTemplateTransform> FRerouteNodeTemplate::GenerateNodeTransform() const
+	TUniquePtr<INodeTransform> FRerouteNodeTemplate::GenerateNodeTransform() const
 	{
 		using namespace ReroutePrivate;
-		return TUniquePtr<INodeTemplateTransform>(new FRerouteNodeTemplateTransform());
+		return TUniquePtr<INodeTransform>(new FRerouteNodeTemplateTransform());
 	}
 
 	const FMetasoundFrontendClass& FRerouteNodeTemplate::GetFrontendClass() const
@@ -253,14 +250,14 @@ namespace Metasound::Frontend
 			const FMetasoundFrontendVertex& RerouteOutput = Node->Interface.Outputs.Last();
 
 			TArray<const FMetasoundFrontendNode*> ConnectedNodes;
-			TArray<const FMetasoundFrontendVertex*> ConnectedInputs = InBuilder.FindNodeInputsConnectedToNodeOutput(InNodeID, RerouteOutput.VertexID, &ConnectedNodes);
+			TArray<const FMetasoundFrontendVertex*> ConnectedInputs = InBuilder.FindNodeInputsConnectedToNodeOutput(InNodeID, InVertexID, &ConnectedNodes);
 			for (int32 Index = 0; Index < ConnectedNodes.Num(); ++Index)
 			{
 				const FMetasoundFrontendNode* ConnectedNode = ConnectedNodes[Index];
 				if (const FMetasoundFrontendClass* ConnectedNodeClass = InBuilder.FindDependency(ConnectedNode->ClassID))
 				{
 					const FMetasoundFrontendVertex* ConnectedInput = ConnectedInputs[Index];
-					if (ConnectedNodeClass->Metadata.GetClassName() == GetClassName())
+					if (ConnectedNodeClass->Metadata.GetClassName() == ClassName)
 					{
 						return this->GetNodeInputAccessType(InBuilder, ConnectedNode->GetID(), ConnectedInput->VertexID);
 					}
diff --git a/Engine/Plugins/Runtime/Metasound/Source/MetasoundFrontend/Public/MetasoundAssetBase.h b/Engine/Plugins/Runtime/Metasound/Source/MetasoundFrontend/Public/MetasoundAssetBase.h
--- a/Engine/Plugins/Runtime/Metasound/Source/MetasoundFrontend/Public/MetasoundAssetBase.h
+++ b/Engine/Plugins/Runtime/Metasound/Source/MetasoundFrontend/Public/MetasoundAssetBase.h
@@ -67,7 +67,6 @@ public:
 	virtual const UEdGraph* GetGraph() const = 0;
 	virtual UEdGraph& GetGraphChecked() = 0;
 	virtual const UEdGraph& GetGraphChecked() const = 0;
-	virtual void MigrateEditorGraph(FMetaSoundFrontendDocumentBuilder& OutBuilder) = 0;
 
 	// Sets the graph associated with this Metasound. Graph is required to be referenced on
 	// Metasound UObject for editor serialization purposes.
@@ -79,8 +78,10 @@ public:
 	virtual void SetRegistryAssetClassInfo(const Metasound::Frontend::FNodeClassInfo& InClassInfo) = 0;
 #endif // WITH_EDITORONLY_DATA
 
-	UE_DEPRECATED(5.5, "Moved to IMetaSoundDocumentInterface::ConformObjectToDocument")
-	virtual bool ConformObjectDataToInterfaces();
+	// Called when the interface is changed, presenting the opportunity for
+	// any reflected object data to be updated based on the new interface.
+	// Returns whether or not any edits were made.
+	virtual bool ConformObjectDataToInterfaces() = 0;
 
 	// Registers the root graph of the given asset with the MetaSound Frontend.
 	virtual void RegisterGraphWithFrontend(Metasound::Frontend::FMetaSoundAssetRegistrationOptions InRegistrationOptions = Metasound::Frontend::FMetaSoundAssetRegistrationOptions());
@@ -91,11 +92,19 @@ public:
 	// Cooks this MetaSound and recursively checks and cooks referenced MetaSounds if necessary. Cook includes autoupdating and resolving the document, which is then registered with the MetaSound Frontend.
 	void CookMetaSound();
 
+	// Sets/overwrites the root class metadata
+	UE_DEPRECATED(5.3, "Directly setting graph class Metadata is no longer be supported. Use the FMetaSoundFrontendDocumentBuilder to modify class data.")
+	virtual void SetMetadata(FMetasoundFrontendClassMetadata& InMetadata);
+
 #if WITH_EDITOR
 	// Rebuild dependent asset classes
 	void RebuildReferencedAssetClasses();
 #endif // WITH_EDITOR
 
+	// Returns the interface entries declared by the given asset's document from the InterfaceRegistry.
+	UE_DEPRECATED(5.3, "Use static FMetaSoundFrontendDocumentBuilder 'FindDeclaredInterfaces instead.")
+	bool GetDeclaredInterfaces(TArray<const Metasound::Frontend::IInterfaceRegistryEntry*>& OutInterfaces) const;
+
 	// Returns whether an interface with the given version is declared by the given asset's document.
 	bool IsInterfaceDeclared(const FMetasoundFrontendVersion& InVersion) const;
 
@@ -120,6 +129,9 @@ public:
 	bool AddingReferenceCausesLoop(const FSoftObjectPath& InReferencePath) const;
 	bool IsReferencedAsset(const FMetasoundAssetBase& InAssetToCheck) const;
 
+	UE_DEPRECATED(5.3, "ConvertFromPreset moved to FMetaSoundFrontendDocumentBuilder.")
+	void ConvertFromPreset();
+
 	bool IsRegistered() const;
 
 	// Imports data from a JSON string directly
@@ -136,16 +148,12 @@ public:
 	Metasound::Frontend::FGraphHandle GetRootGraphHandle();
 	Metasound::Frontend::FConstGraphHandle GetRootGraphHandle() const;
 
-	UE_DEPRECATED(5.5, "Direct mutation of the document is no longer supported via AssetBase.")
-	void SetDocument(FMetasoundFrontendDocument InDocument, bool bMarkDirty = true);
-
-	const FMetasoundFrontendDocument& GetConstDocumentChecked() const;
+	// Overwrites the existing document. If the document's interface is not supported,
+	// the FMetasoundAssetBase be while queried for a new one using `GetPreferredInterface`.
+	void SetDocument(const FMetasoundFrontendDocument& InDocument, bool bMarkDirty = true);
+	void SetDocument(FMetasoundFrontendDocument&& InDocument, bool bMarkDirty = true);
 
-	// Soft deprecated.  Document layer should not be directly mutated via asset base in anticipation
-	// of moving to all mutation via the Frontend Document & Builder Subsystem API.
 	FMetasoundFrontendDocument& GetDocumentChecked();
-
-	UE_DEPRECATED(5.5, "Use GetConstDocumentChecked instead.")
 	const FMetasoundFrontendDocument& GetDocumentChecked() const;
 
 	const Metasound::Frontend::FGraphRegistryKey& GetGraphRegistryKey() const;
@@ -153,6 +161,9 @@ public:
 	UE_DEPRECATED(5.4, "Use GetGraphRegistryKey instead.")
 	const Metasound::Frontend::FNodeRegistryKey& GetRegistryKey() const;
 
+	UE_DEPRECATED(5.3, "AddDefaultInterfaces is included in now applied via FMetaSoundFrontendDocumentBuilder::InitDocument and no longer directly supported via this function.")
+	void AddDefaultInterfaces();
+
 	bool VersionAsset();
 
 #if WITH_EDITOR
@@ -165,11 +176,7 @@ public:
 	void CacheRegistryMetadata();
 
 	FMetasoundFrontendDocumentModifyContext& GetModifyContext();
-
-	UE_DEPRECATED(5.5, "Use GetConstModifyContext")
 	const FMetasoundFrontendDocumentModifyContext& GetModifyContext() const;
-
-	const FMetasoundFrontendDocumentModifyContext& GetConstModifyContext() const;
 #endif // WITH_EDITOR
 
 	// Calls the outermost package and marks it dirty.
@@ -189,18 +196,6 @@ public:
 
 	FString GetOwningAssetName() const;
 
-#if WITH_EDITORONLY_DATA
-	void ClearVersionedOnLoad();
-	bool GetVersionedOnLoad() const;
-	void SetVersionedOnLoad();
-#endif // WITH_EDITORONLY_DATA
-
-	// Returns true if the IMetaSoundDocumentInterface is currently has an active builder.
-	//
-	// If true, calls to register the graph will be performed synchronously in order to avoid
-	// race conditions with an active builder.
-	virtual bool IsBuilderActive() const = 0;
-
 protected:
 	void OnNotifyBeginDestroy();
 
@@ -227,6 +222,8 @@ protected:
 	// When a graph is registered, the underlying IMetaSoundDocumentInterface may be accessed on an
 	// async tasks. If modifications need to be made to the IMetaSoundDocumentInterface, callers should
 	// wait for the inflight graph registration to complete by calling this method.
+
+protected:
 	
 	// Container for runtime data of MetaSound graph.
 	struct FRuntimeData_DEPRECATED
@@ -254,6 +251,10 @@ protected:
 	const FRuntimeData& GetRuntimeData() const;
 	PRAGMA_ENABLE_DEPRECATION_WARNINGS
 
+	// Returns all public class inputs.  This is a potentially expensive.
+	// Prefer accessing public class inputs using CacheRuntimeData.
+	TArray<FMetasoundFrontendClassInput> GetPublicClassInputs() const;
+
 	bool AutoUpdate(bool bInLogWarningsOnDroppedConnection);
 	void CookReferencedMetaSounds();
 
@@ -266,16 +267,13 @@ protected:
 private:
 #if WITH_EDITORONLY_DATA
 	void UpdateAssetRegistry();
-
-	bool bVersionedOnLoad = false;
 #endif
 
-	// Checks if version is up-to-date. If so, returns true. If false, updates the interfaces within the given asset's document to the most recent version.
-	bool TryUpdateInterfaceFromVersion(const FMetasoundFrontendVersion& Version);
-
-	// Returns new interface to be versioned to from the given version. If no interface versioning is
-	// required, returns invalid interface (interface with no name and invalid version number).
-	FMetasoundFrontendInterface GetInterfaceToVersion(const FMetasoundFrontendVersion& InterfaceVersion, Metasound::Frontend::FConstDocumentHandle InDocument) const;
+	// Returns true if the IMetaSoundDocumentInterface is currently has an active builder.
+	//
+	// If true, calls to register the graph will be performed synchronously in order to avoid
+	// race conditions with an active builder.
+	virtual bool IsBuilderActive() const = 0;
 
 	Metasound::Frontend::FGraphRegistryKey GraphRegistryKey;
 };
diff --git a/Engine/Plugins/Runtime/Metasound/Source/MetasoundFrontend/Public/MetasoundAssetManager.h b/Engine/Plugins/Runtime/Metasound/Source/MetasoundFrontend/Public/MetasoundAssetManager.h
--- a/Engine/Plugins/Runtime/Metasound/Source/MetasoundFrontend/Public/MetasoundAssetManager.h
+++ b/Engine/Plugins/Runtime/Metasound/Source/MetasoundFrontend/Public/MetasoundAssetManager.h
@@ -12,7 +12,6 @@ class FMetasoundAssetBase;
 class UEdGraph;
 
 struct FMetasoundFrontendClassName;
-struct FMetaSoundFrontendDocumentBuilder;
 
 namespace Metasound
 {
@@ -132,8 +131,6 @@ namespace Metasound
 			// Add or Update a MetaSound Asset's entry data
 			virtual Metasound::Frontend::FNodeRegistryKey AddOrUpdateAsset(const UObject& InObject) = 0;
 
-			virtual FMetaSoundFrontendDocumentBuilder& AttachDocumentBuilderChecked(UObject& InObject) const = 0;
-
 			// Whether or not the class is eligible for auto-update
 			virtual bool CanAutoUpdate(const FMetasoundFrontendClassName& InClassName) const = 0;
 
diff --git a/Engine/Plugins/Runtime/Metasound/Source/MetasoundFrontend/Public/MetasoundDocumentInterface.h b/Engine/Plugins/Runtime/Metasound/Source/MetasoundFrontend/Public/MetasoundDocumentInterface.h
--- a/Engine/Plugins/Runtime/Metasound/Source/MetasoundFrontend/Public/MetasoundDocumentInterface.h
+++ b/Engine/Plugins/Runtime/Metasound/Source/MetasoundFrontend/Public/MetasoundDocumentInterface.h
@@ -45,10 +45,6 @@ public:
 	// Returns the parent class registered with the MetaSound UObject registry.
 	virtual const UClass& GetBaseMetaSoundUClass() const = 0;
 
-	// Conforms UProperty data outside the Frontend Document Model to the document's data.
-	// Returns whether or not object data was modified.
-	virtual bool ConformObjectToDocument() = 0;
-
 private:
 	virtual FMetasoundFrontendDocument& GetDocument() = 0;
 
diff --git a/Engine/Plugins/Runtime/Metasound/Source/MetasoundFrontend/Public/MetasoundFrontendController.h b/Engine/Plugins/Runtime/Metasound/Source/MetasoundFrontend/Public/MetasoundFrontendController.h
--- a/Engine/Plugins/Runtime/Metasound/Source/MetasoundFrontend/Public/MetasoundFrontendController.h
+++ b/Engine/Plugins/Runtime/Metasound/Source/MetasoundFrontend/Public/MetasoundFrontendController.h
@@ -1134,13 +1134,11 @@ namespace Metasound
 			virtual FString ExportToJSON() const = 0;
 		};
 
-		UE_DEPRECATED(5.5, "Reroute output recursion is now privately implemented in the reroute template node class. Use template registry to query reroute node input/output attributes")
 		METASOUNDFRONTEND_API FConstOutputHandle FindReroutedOutput(FConstOutputHandle InOutputHandle);
 
-		UE_DEPRECATED(5.5, "Reroute input recursion is now privately implemented in the reroute template node class. Use template registry to query reroute node input/output attributes")
+		UE_DEPRECATED(5.5, "Reroute input recursion is now privately implemented in the reroute template node class.")
 		METASOUNDFRONTEND_API void FindReroutedInputs(FConstInputHandle InHandleToCheck, TArray<FConstInputHandle>& InOutInputHandles);
 
-		UE_DEPRECATED(5.5, "Reroute input recursion is now privately implemented in the reroute template node class. Use template registry to query reroute node input/output attributes")
 		METASOUNDFRONTEND_API void IterateReroutedInputs(FConstInputHandle InHandleToCheck, TFunctionRef<void(FConstInputHandle)> Func);
 	} // namespace Frontend
 } // namespace Metasound
diff --git a/Engine/Plugins/Runtime/Metasound/Source/MetasoundFrontend/Public/MetasoundFrontendDocument.h b/Engine/Plugins/Runtime/Metasound/Source/MetasoundFrontend/Public/MetasoundFrontendDocument.h
--- a/Engine/Plugins/Runtime/Metasound/Source/MetasoundFrontend/Public/MetasoundFrontendDocument.h
+++ b/Engine/Plugins/Runtime/Metasound/Source/MetasoundFrontend/Public/MetasoundFrontendDocument.h
@@ -134,10 +134,10 @@ enum class EMetasoundFrontendClassType : uint8
 	// The MetaSound class is an output from a graph in the containing document.
 	Output,
 
-	// The MetaSound class is an literal requiring a literal value to construct.
+	// The MetaSound class is an literal requiring an literal value to construct.
 	Literal,
 
-	// The MetaSound class is an variable requiring a literal value to construct.
+	// The MetaSound class is an variable requiring an literal value to construct.
 	Variable,
 
 	// The MetaSound class accesses variables.
@@ -1652,24 +1652,6 @@ struct METASOUNDFRONTEND_API FMetasoundFrontendGraphClass : public FMetasoundFro
 	FMetasoundFrontendGraphClassPresetOptions PresetOptions;
 };
 
-UCLASS()
-class METASOUNDFRONTEND_API UMetaSoundFrontendMemberMetadata : public UObject
-{
-	GENERATED_BODY()
-
-#if WITH_EDITOR
-public:
-	virtual void ForceRefresh() { }
-	virtual FMetasoundFrontendLiteral GetDefault() const { return FMetasoundFrontendLiteral(); }
-	virtual EMetasoundFrontendLiteralType GetLiteralType() const { return EMetasoundFrontendLiteralType::None; }
-	virtual void SetFromLiteral(const FMetasoundFrontendLiteral& InLiteral) { }
-#endif // WITH_EDITOR
-
-#if WITH_EDITORONLY_DATA
-	UPROPERTY()
-	FGuid MemberID;
-#endif // WITH_EDITORONLY_DATA
-};
 
 USTRUCT()
 struct METASOUNDFRONTEND_API FMetasoundFrontendDocumentMetadata
@@ -1681,12 +1663,6 @@ struct METASOUNDFRONTEND_API FMetasoundFrontendDocumentMetadata
 
 #if WITH_EDITORONLY_DATA
 	FMetasoundFrontendDocumentModifyContext ModifyContext;
-
-	// Map of MemberID to metadata used to constrain how literals can be manipulated
-	// with the editor context. This can be used to implement things like numeric ranges,
-	// hardware control parameters, etc.
-	UPROPERTY()
-	TMap<FGuid, TObjectPtr<UMetaSoundFrontendMemberMetadata>> MemberMetadata;
 #endif // WITH_EDITORONLY_DATA
 };
 
@@ -1723,41 +1699,36 @@ public:
 	}
 
 private:
-#if WITH_EDITORONLY_DATA
 	UPROPERTY(meta = (DeprecatedProperty, DeprecationMessage = "5.0 - ArchetypeVersion has been migrated to InterfaceVersions array."))
 	FMetasoundFrontendVersion ArchetypeVersion;
 
 	UPROPERTY(meta = (DeprecatedProperty, DeprecationMessage = "5.0 - InterfaceVersions has been migrated to Interfaces set."))
 	TArray<FMetasoundFrontendVersion> InterfaceVersions;
-#endif // WITH_EDITORONLY_DATA
 
 	// Used for generating deterministic IDs per document
 	mutable uint32 IdCounter = 1;
 
 public:
-#if WITH_EDITORONLY_DATA
-	bool RequiresInterfaceVersioning() const
-	{
-		return ArchetypeVersion.IsValid() || !InterfaceVersions.IsEmpty();
-	}
-
 	// Data migration for 5.0 Early Access data. ArchetypeVersion/InterfaceVersions properties can be removed post 5.0 release
 	// and this fix-up can be removed post 5.0 release.
-	void VersionInterfaces()
+	bool VersionInterfaces()
 	{
+		bool bDidEdit = false;
 		if (ArchetypeVersion.IsValid())
 		{
 			Interfaces.Add(ArchetypeVersion);
 			ArchetypeVersion = FMetasoundFrontendVersion::GetInvalid();
+			bDidEdit = true;
 		}
-
 		if (!InterfaceVersions.IsEmpty())
 		{
 			Interfaces.Append(InterfaceVersions);
 			InterfaceVersions.Reset();
+			bDidEdit = true;
 		}
+
+		return bDidEdit;
 	}
-#endif // WITH_EDITORONLY_DATA
 };
 
 METASOUNDFRONTEND_API const TCHAR* LexToString(EMetasoundFrontendClassType InClassType);
