[Backout] - CL33084850

diff --git a/Engine/Plugins/Runtime/Metasound/Source/MetasoundFrontend/Public/MetasoundFrontendDocumentBuilder.h b/Engine/Plugins/Runtime/Metasound/Source/MetasoundFrontend/Public/MetasoundFrontendDocumentBuilder.h
--- a/Engine/Plugins/Runtime/Metasound/Source/MetasoundFrontend/Public/MetasoundFrontendDocumentBuilder.h
+++ b/Engine/Plugins/Runtime/Metasound/Source/MetasoundFrontend/Public/MetasoundFrontendDocumentBuilder.h
@@ -101,8 +101,6 @@ public:
 	// Create and return a valid builder document which copies the provided interface's document & class
 	static UMetaSoundBuilderDocument& Create(const IMetaSoundDocumentInterface& InDocToCopy);
 
-	virtual bool ConformObjectToDocument() override;
-
 	// Returns the document
 	virtual const FMetasoundFrontendDocument& GetConstDocument() const override;
 
@@ -135,8 +133,11 @@ struct METASOUNDFRONTEND_API FMetaSoundFrontendDocumentBuilder
 	GENERATED_BODY()
 
 public:
-	// Default ctor should typically never be used directly as builder interface (and optionally delegates) should be specified on construction (Default exists only to make UObject reflection happy).
-	FMetaSoundFrontendDocumentBuilder(TScriptInterface<IMetaSoundDocumentInterface> InDocumentInterface = { }, TSharedPtr<Metasound::Frontend::FDocumentModifyDelegates> InDocumentDelegates = { }, bool bPrimeCache = false);
+	// Exists only to make UObject reflection happy.  Should never be
+	// used directly as builder interface (and optionally delegates) should be specified on construction.
+	FMetaSoundFrontendDocumentBuilder();
+	FMetaSoundFrontendDocumentBuilder(TScriptInterface<IMetaSoundDocumentInterface> InDocumentInterface);
+	FMetaSoundFrontendDocumentBuilder(TScriptInterface<IMetaSoundDocumentInterface> InDocumentInterface, TSharedRef<Metasound::Frontend::FDocumentModifyDelegates> InDocumentDelegates);
 	virtual ~FMetaSoundFrontendDocumentBuilder();
 
 	// Call when the builder will no longer modify the IMetaSoundDocumentInterface
@@ -149,39 +150,30 @@ public:
 	bool AddEdgesFromMatchingInterfaceNodeOutputsToGraphOutputs(const FGuid& InNodeID, TArray<const FMetasoundFrontendEdge*>& OutEdgesCreated, bool bReplaceExistingConnections = true);
 	bool AddEdgesFromMatchingInterfaceNodeInputsToGraphInputs(const FGuid& InNodeID, TArray<const FMetasoundFrontendEdge*>& OutEdgesCreated, bool bReplaceExistingConnections = true);
 	const FMetasoundFrontendNode* AddGraphInput(const FMetasoundFrontendClassInput& InClassInput);
-	const FMetasoundFrontendNode* AddGraphNode(const FMetasoundFrontendGraphClass& InClass, FGuid InNodeID = FGuid::NewGuid());
 	const FMetasoundFrontendNode* AddGraphOutput(const FMetasoundFrontendClassOutput& InClassOutput);
-
 	bool AddInterface(FName InterfaceName);
 
+	const FMetasoundFrontendNode* AddGraphNode(const FMetasoundFrontendGraphClass& InClass, FGuid InNodeID = FGuid::NewGuid());
 	const FMetasoundFrontendNode* AddNodeByClassName(const FMetasoundFrontendClassName& InClassName, int32 InMajorVersion = 1, FGuid InNodeID = FGuid::NewGuid());
-	const FMetasoundFrontendNode* AddNodeByTemplate(const Metasound::Frontend::INodeTemplate& InTemplate, Metasound::Frontend::FNodeTemplateGenerateInterfaceParams Params, FGuid InNodeID = FGuid::NewGuid());
 
 	// Returns whether or not the given edge can be added, which requires that its input
 	// is not already connected and the edge is valid (see function 'IsValidEdge').
 	bool CanAddEdge(const FMetasoundFrontendEdge& InEdge) const;
 
-	// Clears graph completely of all nodes, edges, dependencies, & member metadata.
-	// Reloads the builder state, so external delegates must be relinked if desired.
-	void ClearGraph(TSharedRef<Metasound::Frontend::FDocumentModifyDelegates> ModifyDelegates = { });
-
-#if WITH_EDITORONLY_DATA
-	bool ClearMemberMetadata(const FGuid& InMemberID);
-#endif // WITH_EDITORONLY_DATA
+	void ClearGraph();
 
 	bool ContainsDependencyOfType(EMetasoundFrontendClassType ClassType) const;
 	bool ContainsEdge(const FMetasoundFrontendEdge& InEdge) const;
 	bool ContainsNode(const FGuid& InNodeID) const;
 
 	bool ConvertFromPreset();
-	bool ConvertToPreset(const FMetasoundFrontendDocument& InReferencedDocument, TSharedRef<Metasound::Frontend::FDocumentModifyDelegates> ModifyDelegates = { });
+	bool ConvertToPreset(const FMetasoundFrontendDocument& InReferencedDocument);
 
-#if WITH_EDITORONLY_DATA
+#if WITH_EDITOR
 	const FMetaSoundFrontendGraphComment* FindGraphComment(const FGuid& InCommentID) const;
 	FMetaSoundFrontendGraphComment* FindGraphComment(const FGuid& InCommentID);
 	FMetaSoundFrontendGraphComment& FindOrAddGraphComment(const FGuid& InCommentID);
-	UMetaSoundFrontendMemberMetadata* FindMemberMetadata(const FGuid& InMemberID);
-#endif // WITH_EDITORONLY_DATA
+#endif // WITH_EDITOR
 
 	static bool FindDeclaredInterfaces(const FMetasoundFrontendDocument& InDocument, TArray<const Metasound::Frontend::IInterfaceRegistryEntry*>& OutInterfaces);
 	bool FindDeclaredInterfaces(TArray<const Metasound::Frontend::IInterfaceRegistryEntry*>& OutInterfaces) const;
@@ -190,14 +182,14 @@ public:
 	const FMetasoundFrontendClass* FindDependency(const FMetasoundFrontendClassMetadata& InMetadata) const;
 	TArray<const FMetasoundFrontendEdge*> FindEdges(const FGuid& InNodeID, const FGuid& InVertexID) const;
 
+	bool FindInterfaceInputNodes(FName InterfaceName, TArray<const FMetasoundFrontendNode*>& OutInputs) const;
+	bool FindInterfaceOutputNodes(FName InterfaceName, TArray<const FMetasoundFrontendNode*>& OutOutputs) const;
+
 	const FMetasoundFrontendClassInput* FindGraphInput(FName InputName) const;
 	const FMetasoundFrontendNode* FindGraphInputNode(FName InputName) const;
 	const FMetasoundFrontendClassOutput* FindGraphOutput(FName OutputName) const;
 	const FMetasoundFrontendNode* FindGraphOutputNode(FName OutputName) const;
 
-	bool FindInterfaceInputNodes(FName InterfaceName, TArray<const FMetasoundFrontendNode*>& OutInputs) const;
-	bool FindInterfaceOutputNodes(FName InterfaceName, TArray<const FMetasoundFrontendNode*>& OutOutputs) const;
-
 	const FMetasoundFrontendNode* FindNode(const FGuid& InNodeID) const;
 
 	const FMetasoundFrontendVertex* FindNodeInput(const FGuid& InNodeID, const FGuid& InVertexID) const;
@@ -229,28 +221,21 @@ public:
 
 	const Metasound::Frontend::FDocumentModifyDelegates& GetDocumentDelegates() const;
 	const IMetaSoundDocumentInterface& GetDocumentInterface() const;
-	FMetasoundAssetBase& GetMetasoundAsset();
-	const FMetasoundAssetBase& GetMetasoundAsset() const;
 
 	int32 GetTransactionCount() const;
 
-	TArray<const FMetasoundFrontendNode*> GetGraphInputTemplateNodes(FName InInputName);
 	EMetasoundFrontendVertexAccessType GetNodeInputAccessType(const FGuid& InNodeID, const FGuid& InVertexID) const;
 	const FMetasoundFrontendLiteral* GetNodeInputClassDefault(const FGuid& InNodeID, const FGuid& InVertexID) const;
 	const FMetasoundFrontendLiteral* GetNodeInputDefault(const FGuid& InNodeID, const FGuid& InVertexID) const;
 	EMetasoundFrontendVertexAccessType GetNodeOutputAccessType(const FGuid& InNodeID, const FGuid& InVertexID) const;
 
-	// Initializes the builder's document, using the (optional) provided document template, (optional) class name, and (optionally) whether or not to reset the existing class version.
-	void InitDocument(const FMetasoundFrontendDocument* InDocumentTemplate = nullptr, const FMetasoundFrontendClassName* InNewClassName = nullptr, bool bResetVersion = true);
-
 	// Initializes GraphClass Metadata, optionally resetting the version back to 1.0 and/or creating a unique class name if a name is not provided.
 	static void InitGraphClassMetadata(FMetasoundFrontendClassMetadata& InOutMetadata, bool bResetVersion = false, const FMetasoundFrontendClassName* NewClassName = nullptr);
-	void InitGraphClassMetadata(bool bResetVersion, const FMetasoundFrontendClassName* NewClassName);
-
+	void InitDocument();
 	void InitNodeLocations();
 
-	UE_DEPRECATED(5.5, "Use invalidate overload that is provided new version of modify delegates")
-	void InvalidateCache() { }
+	// Clears the cached data enabling fast document queries. Generally discouraged unless document has been unavoidably mutated outside of the builder.
+	void InvalidateCache();
 
 	bool IsDependencyReferenced(const FGuid& InClassID) const;
 	bool IsNodeInputConnected(const FGuid& InNodeID, const FGuid& InVertexID) const;
@@ -260,19 +245,14 @@ public:
 	bool IsInterfaceDeclared(const FMetasoundFrontendVersion& InInterfaceVersion) const;
 	bool IsPreset() const;
 
-	// Returns whether or not builder is attached to a DocumentInterface and is valid to build or act on a document.
-	bool IsValid() const;
-	bool IsVariableClass(EMetasoundFrontendClassType ClassType) const;
-
 	// Returns whether or not the given edge is valid (i.e. represents an input and output that equate in data and access types) or malformed.
 	// Note that this does not return whether or not the given edge exists, but rather if it could be legally applied to the given edge vertices.
 	Metasound::Frontend::EInvalidEdgeReason IsValidEdge(const FMetasoundFrontendEdge& InEdge) const;
 
 	bool ModifyInterfaces(Metasound::Frontend::FModifyInterfaceOptions&& InOptions);
 
-	UE_DEPRECATED(5.5,
-		"Cache invalidation may require new copy of delegates. In addition, re-priming is discouraged. "
-		"To enforce this, new recommended pattern is to construct a new builder instead")
+	// Reloads the cached data enabling fast document queries. Can be expensive and is generally discouraged in favor of using InvalidateCache
+	// only when document has been known to be mutated outside of the builder (also, invalidate doesn't prime the cache and only reloads piecemeal upon request).
 	void ReloadCache();
 
 	bool RemoveDependency(const FGuid& InClassID);
@@ -296,8 +276,6 @@ public:
 	int32 RemoveNodeLocation(const FGuid& InNodeID, const FGuid* InLocationGuid = nullptr);
 #endif // WITH_EDITOR
 
-	void Reload(TSharedPtr<Metasound::Frontend::FDocumentModifyDelegates> Delegates = {}, bool bPrimeCache = false);
-
 	bool RemoveNodeInputDefault(const FGuid& InNodeID, const FGuid& InVertexID);
 	bool RemoveUnusedDependencies();
 	bool RenameRootGraphClass(const FMetasoundFrontendClassName& InName);
@@ -305,18 +283,6 @@ public:
 #if WITH_EDITOR
 	void SetAuthor(const FString& InAuthor);
 
-#endif // WITH_EDITOR
-
-	bool SetGraphInputAccessType(FName InputName, EMetasoundFrontendVertexAccessType AccessType);
-	bool SetGraphInputDataType(FName InputName, FName DataType);
-	bool SetGraphInputDefault(FName InputName, const FMetasoundFrontendLiteral& InDefaultLiteral);
-
-	bool SetGraphOutputAccessType(FName OutputName, EMetasoundFrontendVertexAccessType AccessType);
-	bool SetGraphOutputDataType(FName OutputName, FName DataType);
-
-#if WITH_EDITOR
-	void SetMemberMetadata(UMetaSoundFrontendMemberMetadata& NewMetadata);
-
 	// Sets the editor-only comment to the provided value.
 	// Returns true if the node was found and the comment was updated, false if not.
 	bool SetNodeComment(const FGuid& InNodeID, FString&& InNewComment);
@@ -330,6 +296,7 @@ public:
 	bool SetNodeLocation(const FGuid& InNodeID, const FVector2D& InLocation, const FGuid* InLocationGuid = nullptr);
 #endif // WITH_EDITOR
 
+	bool SetGraphInputDefault(FName InputName, const FMetasoundFrontendLiteral& InDefaultLiteral);
 	bool SetNodeInputDefault(const FGuid& InNodeID, const FGuid& InVertexID, const FMetasoundFrontendLiteral& InLiteral);
 
 	// Sets the document's version number.  Should only be called by document versioning.
@@ -346,36 +313,24 @@ public:
 private:
 	using FFinalizeNodeFunctionRef = TFunctionRef<void(FMetasoundFrontendNode&, const Metasound::Frontend::FNodeRegistryKey&)>;
 
-	FMetasoundFrontendNode* AddNodeInternal(const FMetasoundFrontendClassMetadata& InClassMetadata, Metasound::Frontend::FFinalizeNodeFunctionRef FinalizeNode, FGuid InNodeID = FGuid::NewGuid(), int32* NewNodeIndex = nullptr);
-
-	// Conforms GraphOutput node's ClassID, Access & Data Type with the GraphOutput.
-	// creating and removing dependencies as necessary within the document dependency array. Does *NOT*
-	// modify edge data (i.e. if the DataType is changed on the given node and it has corresponding
-	// edges, edges may then be invalid due to access type/DataType incompatibility).
-	bool ConformGraphInputNodeToClass(const FMetasoundFrontendClassInput& GraphInput);
-
-	// Conforms GraphOutput node's ClassID, Access & Data Type with the GraphOutput.
-	// Creates and removes dependencies as necessary within the document dependency array. Does *NOT*
-	// modify edge data (i.e. if the DataType is changed on the given node and it has corresponding
-	// edges, edges may then be invalid due to access type/DataType incompatibility).
-	bool ConformGraphOutputNodeToClass(const FMetasoundFrontendClassOutput& GraphOutput);
-
-	bool FindNodeClassInterfaces(const FGuid& InNodeID, TSet<FMetasoundFrontendVersion>& OutInterfaces) const;
+	FMetasoundFrontendNode* AddNodeInternal(const FMetasoundFrontendClassMetadata& InClassMetadata, Metasound::Frontend::FFinalizeNodeFunctionRef FinalizeNode, FGuid InNodeID = FGuid::NewGuid());
 	FMetasoundFrontendNode* FindNodeInternal(const FGuid& InNodeID);
 
-	void IterateNodesConnectedWithVertex(const FMetasoundFrontendVertexHandle& Vertex, TFunctionRef<void(const FMetasoundFrontendEdge&, FMetasoundFrontendNode&)> NodeIndexIterFunc);
-
 	const FTopLevelAssetPath GetBuilderClassPath() const;
 	FMetasoundFrontendDocument& GetDocument();
 
+	bool FindNodeClassInterfaces(const FGuid& InNodeID, TSet<FMetasoundFrontendVersion>& OutInterfaces) const;
+
+	void InitCacheInternal(bool bPrimeCache = false);
+
 	bool SetGraphInputInheritsDefault(FName InName, bool bInputInheritsDefault);
-	bool UnlinkVariableNode(const FGuid& InNodeID);
 
-	void BeginBuilding(TSharedPtr<Metasound::Frontend::FDocumentModifyDelegates> Delegates = {}, bool bPrimeCache = false);
+	// Called to mark that the builder is beginning modifications to the IMetaSoundDocumentInterface
+	void BeginBuilding();
 
 	UPROPERTY(Transient)
 	TScriptInterface<IMetaSoundDocumentInterface> DocumentInterface;
 
 	TSharedPtr<Metasound::Frontend::IDocumentCache> DocumentCache;
-	TSharedPtr<Metasound::Frontend::FDocumentModifyDelegates> DocumentDelegates;
+	TSharedRef<Metasound::Frontend::FDocumentModifyDelegates> DocumentDelegates;
 };
diff --git a/Engine/Plugins/Runtime/Metasound/Source/MetasoundFrontend/Public/MetasoundFrontendDocumentCacheInterface.h b/Engine/Plugins/Runtime/Metasound/Source/MetasoundFrontend/Public/MetasoundFrontendDocumentCacheInterface.h
--- a/Engine/Plugins/Runtime/Metasound/Source/MetasoundFrontend/Public/MetasoundFrontendDocumentCacheInterface.h
+++ b/Engine/Plugins/Runtime/Metasound/Source/MetasoundFrontend/Public/MetasoundFrontendDocumentCacheInterface.h
@@ -88,9 +88,7 @@ namespace Metasound::Frontend
 		virtual ~IDocumentGraphInterfaceCache() = default;
 
 		virtual const FMetasoundFrontendClassInput* FindInput(FName InputName) const = 0;
-		virtual const int32* FindInputIndex(FName InputName) const = 0;
 		virtual const FMetasoundFrontendClassOutput* FindOutput(FName OutputName) const = 0;
-		virtual const int32* FindOutputIndex(FName OutputName) const = 0;
 	};
 
 	/** Interface for querying cached document dependencies. */
diff --git a/Engine/Plugins/Runtime/Metasound/Source/MetasoundFrontend/Public/MetasoundFrontendNodeTemplateRegistry.h b/Engine/Plugins/Runtime/Metasound/Source/MetasoundFrontend/Public/MetasoundFrontendNodeTemplateRegistry.h
--- a/Engine/Plugins/Runtime/Metasound/Source/MetasoundFrontend/Public/MetasoundFrontendNodeTemplateRegistry.h
+++ b/Engine/Plugins/Runtime/Metasound/Source/MetasoundFrontend/Public/MetasoundFrontendNodeTemplateRegistry.h
@@ -29,22 +29,12 @@ namespace Metasound::Frontend
 
 	struct METASOUNDFRONTEND_API FNodeTemplateGenerateInterfaceParams
 	{
-		// TODO: Currently just DataType FName. Subsequent change will replace this with a VertexHandle
-		// and will add builder reference to once builder supports template nodes and controllers
-		// are no longer used to add template nodes from editor code.
+		// Currently just DataType FName. Subsequent change will replace this with a VertexHandle
+		// and will add builder reference to once builder supports template nodes.
 		TArray<FName> InputsToConnect;
 		TArray<FName> OutputsToConnect;
 	};
 
-	class METASOUNDFRONTEND_API INodeTemplateTransform
-	{
-	public:
-		virtual ~INodeTemplateTransform() = default;
-
-		/** Return true if the builder was modified, false otherwise. */
-		virtual bool Transform(const FGuid& InNodeID, FMetaSoundFrontendDocumentBuilder& OutBuilder) const = 0;
-	};
-
 	/**
 		* Base interface for a node template, which acts in place of frontend node class and respective instance(s).
 		* Instances are preprocessed, allowing for custom graph manipulation prior to generating a respective runtime
@@ -70,7 +60,7 @@ namespace Metasound::Frontend
 #endif // WITH_EDITOR
 
 		// Generates node transform that is used to preprocess nodes.
-		virtual TUniquePtr<INodeTemplateTransform> GenerateNodeTransform() const = 0;
+		virtual TUniquePtr<INodeTransform> GenerateNodeTransform() const = 0;
 
 		// Returns the class definition for the given node class template.
 		virtual const FMetasoundFrontendClass& GetFrontendClass() const = 0;
diff --git a/Engine/Plugins/Runtime/Metasound/Source/MetasoundFrontend/Public/MetasoundFrontendTransform.h b/Engine/Plugins/Runtime/Metasound/Source/MetasoundFrontend/Public/MetasoundFrontendTransform.h
--- a/Engine/Plugins/Runtime/Metasound/Source/MetasoundFrontend/Public/MetasoundFrontendTransform.h
+++ b/Engine/Plugins/Runtime/Metasound/Source/MetasoundFrontend/Public/MetasoundFrontendTransform.h
@@ -131,15 +131,24 @@ namespace Metasound::Frontend
 
 	};
 
+	/** Updates document's given interface to the most recent version. */
 	class METASOUNDFRONTEND_API FUpdateRootGraphInterface : public IDocumentTransform
 	{
 	public:
 		FUpdateRootGraphInterface(const FMetasoundFrontendVersion& InInterfaceVersion, const FString& InOwningAssetName=FString(TEXT("Unknown")))
+			: InterfaceVersion(InInterfaceVersion)
+			, OwningAssetName(InOwningAssetName)
 		{
 		}
 
-		UE_DEPRECATED(5.5, "RootGraph update is now handled privately by internal MetaSound asset management")
-		virtual bool Transform(FDocumentHandle InDocument) const override { return false; }
+		bool Transform(FDocumentHandle InDocument) const override;
+
+	private:
+		void GetUpdatePathForDocument(const FMetasoundFrontendVersion& InCurrentVersion, const FMetasoundFrontendVersion& InTargetVersion, TArray<const IInterfaceRegistryEntry*>& OutUpgradePath) const;
+		bool UpdateDocumentInterface(const TArray<const IInterfaceRegistryEntry*>& InUpgradePath, FDocumentHandle InDocument) const;
+
+		FMetasoundFrontendVersion InterfaceVersion;
+		FString OwningAssetName;
 	};
 
 	/** Completely rebuilds the graph connecting a preset's inputs to the reference
diff --git a/Engine/Plugins/Runtime/Metasound/Source/MetasoundFrontend/Public/NodeTemplates/MetasoundFrontendNodeTemplateReroute.h b/Engine/Plugins/Runtime/Metasound/Source/MetasoundFrontend/Public/NodeTemplates/MetasoundFrontendNodeTemplateReroute.h
--- a/Engine/Plugins/Runtime/Metasound/Source/MetasoundFrontend/Public/NodeTemplates/MetasoundFrontendNodeTemplateReroute.h
+++ b/Engine/Plugins/Runtime/Metasound/Source/MetasoundFrontend/Public/NodeTemplates/MetasoundFrontendNodeTemplateReroute.h
@@ -36,7 +36,7 @@ namespace Metasound::Frontend
 		virtual TUniquePtr<INodeTransform> GenerateNodeTransform(FMetasoundFrontendDocument& InPreprocessedDocument) const override { return { }; }
 
 		virtual FMetasoundFrontendNodeInterface GenerateNodeInterface(FNodeTemplateGenerateInterfaceParams InParams) const override;
-		virtual TUniquePtr<INodeTemplateTransform> GenerateNodeTransform() const override;
+		virtual TUniquePtr<INodeTransform> GenerateNodeTransform() const override;
 		virtual const FMetasoundFrontendClass& GetFrontendClass() const override;
 		virtual EMetasoundFrontendVertexAccessType GetNodeInputAccessType(const FMetaSoundFrontendDocumentBuilder& InBuilder, const FGuid& InNodeID, const FGuid& InVertexID) const override;
 		virtual EMetasoundFrontendVertexAccessType GetNodeOutputAccessType(const FMetaSoundFrontendDocumentBuilder& InBuilder, const FGuid& InNodeID, const FGuid& InVertexID) const override;
diff --git a/Engine/Source/Runtime/Engine/Classes/EdGraph/EdGraph.h b/Engine/Source/Runtime/Engine/Classes/EdGraph/EdGraph.h
--- a/Engine/Source/Runtime/Engine/Classes/EdGraph/EdGraph.h
+++ b/Engine/Source/Runtime/Engine/Classes/EdGraph/EdGraph.h
@@ -163,7 +163,7 @@ public:
 	* 
 	* @return True if the node has been removed from the graph
 	*/
-	ENGINE_API bool RemoveNode( UEdGraphNode* NodeToRemove, bool bBreakAllLinks = true, bool bMarkDirty = true );
+	ENGINE_API bool RemoveNode( UEdGraphNode* NodeToRemove, bool bBreakAllLinks = true );
 
 	/** Signal to listeners that the graph has changed - prefer to use NotifyNodeChanged when updating a single node */
 	ENGINE_API virtual void NotifyGraphChanged();
diff --git a/Engine/Source/Runtime/Engine/Classes/EdGraph/EdGraphPin.h b/Engine/Source/Runtime/Engine/Classes/EdGraph/EdGraphPin.h
--- a/Engine/Source/Runtime/Engine/Classes/EdGraph/EdGraphPin.h
+++ b/Engine/Source/Runtime/Engine/Classes/EdGraph/EdGraphPin.h
@@ -398,13 +398,13 @@ public:
 	ENGINE_API static UEdGraphPin* CreatePin(UEdGraphNode* InOwningNode);
 
 	/** Create a link. Note, this does not check that schema allows it, and will not break any existing connections */
-	ENGINE_API void MakeLinkTo(UEdGraphPin* ToPin, bool bAlwaysMarkDirty = true);
+	ENGINE_API void MakeLinkTo(UEdGraphPin* ToPin);
 
 	/** Break a link to the specified pin (if present) */
-	ENGINE_API void BreakLinkTo(UEdGraphPin* ToPin, bool bAlwaysMarkDirty = true);
+	ENGINE_API void BreakLinkTo(UEdGraphPin* ToPin);
 
 	/** Break all links from this pin */
-	ENGINE_API void BreakAllPinLinks(bool bNotifyNodes = false, bool bAlwaysMarkDirty = true);
+	ENGINE_API void BreakAllPinLinks(bool bNotifyNodes = false);
 
 	/**
 	* Moves the persistent data (across a node refresh) from the SourcePin.
@@ -577,7 +577,7 @@ private:
 	bool Serialize(FArchive& Ar);
 
 	// Helper functions
-	static void ConvertConnectedGhostNodesToRealNodes(UEdGraphNode* InNode, bool bAlwaysMarkDirty = true);
+	static void ConvertConnectedGhostNodesToRealNodes(UEdGraphNode* InNode);
 	static void ResolveReferencesToPin(UEdGraphPin* Pin, bool bStrictValidation = true);
 	static void SerializePinArray(FArchive& Ar, TArray<UEdGraphPin*>& ArrayRef, UEdGraphPin* RequestingPin, EPinResolveType ResolveType);
 	static bool SerializePin(FArchive& Ar, UEdGraphPin*& PinRef, int32 ArrayIdx, UEdGraphPin* RequestingPin, EPinResolveType ResolveType, TArray<UEdGraphPin*>& OldPins);
diff --git a/Engine/Source/Runtime/Engine/Private/EdGraph/EdGraphPin.cpp b/Engine/Source/Runtime/Engine/Private/EdGraph/EdGraphPin.cpp
--- a/Engine/Source/Runtime/Engine/Private/EdGraph/EdGraphPin.cpp
+++ b/Engine/Source/Runtime/Engine/Private/EdGraph/EdGraphPin.cpp
@@ -509,7 +509,7 @@ UEdGraphPin* UEdGraphPin::CreatePin(UEdGraphNode* InOwningNode)
 	return NewPin;
 }
 
-void UEdGraphPin::MakeLinkTo(UEdGraphPin* ToPin, bool bAlwaysMarkDirty)
+void UEdGraphPin::MakeLinkTo(UEdGraphPin* ToPin)
 {
 	if (ToPin)
 	{
@@ -525,32 +525,32 @@ void UEdGraphPin::MakeLinkTo(UEdGraphPin* ToPin, bool bAlwaysMarkDirty)
 			ensureMsgf(MyNode->GetOuter() == ToPin->GetOwningNode()->GetOuter(), TEXT("%s"), *GetLinkInfoString( LOCTEXT("MakeLinkTo", "MakeLinkTo").ToString(), LOCTEXT("OuterMismatch", "has a different outer than pin").ToString(), ToPin)); // Ensure both pins belong to the same graph
 
 			// Notify owning nodes about upcoming change
-			Modify(bAlwaysMarkDirty);
-			ToPin->Modify(bAlwaysMarkDirty);
+			Modify();
+			ToPin->Modify();
 
 			// Add to both lists
 			LinkedTo.Add(ToPin);
 			ToPin->LinkedTo.Add(this);
 
 			// If either node was a pre-placed ghost, turn it into a real thing
-			UEdGraphPin::ConvertConnectedGhostNodesToRealNodes(MyNode, bAlwaysMarkDirty);
-			UEdGraphPin::ConvertConnectedGhostNodesToRealNodes(ToPin->GetOwningNode(), bAlwaysMarkDirty);
+			UEdGraphPin::ConvertConnectedGhostNodesToRealNodes(MyNode);
+			UEdGraphPin::ConvertConnectedGhostNodesToRealNodes(ToPin->GetOwningNode());
 		}
 	}
 }
 
-void UEdGraphPin::BreakLinkTo(UEdGraphPin* ToPin, bool bAlwaysMarkDirty)
+void UEdGraphPin::BreakLinkTo(UEdGraphPin* ToPin)
 {
 	if (ToPin)
 	{
 		// If we do indeed link to the passed in pin...
 		if (LinkedTo.Contains(ToPin))
 		{
-			Modify(bAlwaysMarkDirty);
+			Modify();
 
 			if (ToPin->LinkedTo.Contains(this))
 			{
-				ToPin->Modify(bAlwaysMarkDirty);
+				ToPin->Modify();
 				ToPin->LinkedTo.Remove(this);
 			}
 			else if (OwningNode && !OwningNode->HasAnyFlags(RF_BeginDestroyed))
@@ -569,13 +569,13 @@ void UEdGraphPin::BreakLinkTo(UEdGraphPin* ToPin, bool bAlwaysMarkDirty)
 	}
 }
 
-void UEdGraphPin::BreakAllPinLinks(bool bNotifyNodes, bool bAlwaysMarkDirty)
+void UEdGraphPin::BreakAllPinLinks(const bool bNotifyNodes)
 {
 	TArray<UEdGraphPin*> LinkedToCopy = LinkedTo;
 
 	for (UEdGraphPin* LinkedToPin : LinkedToCopy)
 	{
-		BreakLinkTo(LinkedToPin, bAlwaysMarkDirty);
+		BreakLinkTo(LinkedToPin);
 #if WITH_EDITOR
 		if (bNotifyNodes)
 		{
@@ -1799,12 +1799,12 @@ void UEdGraphPin::DeclarePinCustomVersions(FArchive& Ar)
 }
 #endif
 
-void UEdGraphPin::ConvertConnectedGhostNodesToRealNodes(UEdGraphNode* InNode, bool bAlwaysMarkDirty)
+void UEdGraphPin::ConvertConnectedGhostNodesToRealNodes(UEdGraphNode* InNode)
 {
 	if (InNode && InNode->IsAutomaticallyPlacedGhostNode())
 	{
 		// Enable the node and clear the comment
-		InNode->Modify(bAlwaysMarkDirty);
+		InNode->Modify();
 		InNode->SetEnabledState(ENodeEnabledState::Enabled, /*bUserAction=*/ false);
 		InNode->NodeComment.Empty();
 #if WITH_EDITORONLY_DATA
@@ -1816,7 +1816,7 @@ void UEdGraphPin::ConvertConnectedGhostNodesToRealNodes(UEdGraphNode* InNode, bo
 		{
 			for (UEdGraphPin* OtherPin : Pin->LinkedTo)
 			{
-				ConvertConnectedGhostNodesToRealNodes(OtherPin->GetOwningNode(), bAlwaysMarkDirty);
+				ConvertConnectedGhostNodesToRealNodes(OtherPin->GetOwningNode());
 			}
 		}
 	}
