Fixed unused and potentially uninitialized FQuats and FVector2Ds.

diff --git a/Engine/Plugins/Runtime/Metasound/Source/MetasoundEngine/Private/Interfaces/MetasoundDeprecatedInterfaces.cpp b/Engine/Plugins/Runtime/Metasound/Source/MetasoundEngine/Private/Interfaces/MetasoundDeprecatedInterfaces.cpp
--- a/Engine/Plugins/Runtime/Metasound/Source/MetasoundEngine/Private/Interfaces/MetasoundDeprecatedInterfaces.cpp
+++ b/Engine/Plugins/Runtime/Metasound/Source/MetasoundEngine/Private/Interfaces/MetasoundDeprecatedInterfaces.cpp
@@ -315,7 +315,9 @@ namespace Metasound::Engine
 				InDocument->AddInterfaceVersion(MetasoundOutputFormatMonoV1_1::GetInterface().Version);
 
 				FNodeHandle MonoFormatOutput = Graph->GetOutputNodeWithName(MetasoundOutputFormatMonoV1_0::GetAudioOutputName());
+#if WITH_EDITOR
 				FVector2D MonoFormatLocation;
+#endif // WITH_EDITOR
 
 				FOutputHandle OutputToReconnect = IOutputController::GetInvalidHandle();
 				if (MonoFormatOutput->IsValid())
@@ -487,7 +489,9 @@ namespace Metasound::Engine
 				FOutputHandle LeftOutputToReconnect = IOutputController::GetInvalidHandle();
 				FOutputHandle RightOutputToReconnect = IOutputController::GetInvalidHandle();
 
+#if WITH_EDITOR
 				FVector2D StereoFormatLocation;
+#endif // WITH_EDITOR
 
 				if (StereoFormatOutput->IsValid())
 				{
diff --git a/Engine/Source/Runtime/Core/Tests/Math/UnrealMathTest.cpp b/Engine/Source/Runtime/Core/Tests/Math/UnrealMathTest.cpp
--- a/Engine/Source/Runtime/Core/Tests/Math/UnrealMathTest.cpp
+++ b/Engine/Source/Runtime/Core/Tests/Math/UnrealMathTest.cpp
@@ -2200,7 +2200,7 @@ bool RunDoubleVectorTest()
 
 	// Quat / Rotator conversion to vectors, matrices
 	{
-		FQuat4d Q0, Q1, Q2, Q3;
+		FQuat4d Q0, Q1;
 
 		FRotator3d Rotator0;
 		Rotator0 = FRotator3d(30.0f, -45.0f, 90.0f);
diff --git a/Engine/Source/Runtime/Engine/Private/Particles/ParticleSystemRender.cpp b/Engine/Source/Runtime/Engine/Private/Particles/ParticleSystemRender.cpp
--- a/Engine/Source/Runtime/Engine/Private/Particles/ParticleSystemRender.cpp
+++ b/Engine/Source/Runtime/Engine/Private/Particles/ParticleSystemRender.cpp
@@ -1737,7 +1737,7 @@ void FDynamicMeshEmitterData::CalculateParticleTransform(
 		}
 	}
 
-	FQuat PointToLockedAxis;
+	FQuat PointToLockedAxis = FQuat::Identity;
 	if (bUseMeshLockedAxis == true)
 	{
 		// facing axis is taken to be the local x axis.	
