Stopgap fix for cook failure/optimization to ignore MetaSound editor data migration in version 1.12 if cooking

diff --git a/Engine/Plugins/Runtime/Metasound/Source/MetasoundFrontend/Private/MetasoundFrontendDocumentVersioning.cpp b/Engine/Plugins/Runtime/Metasound/Source/MetasoundFrontend/Private/MetasoundFrontendDocumentVersioning.cpp
--- a/Engine/Plugins/Runtime/Metasound/Source/MetasoundFrontend/Private/MetasoundFrontendDocumentVersioning.cpp
+++ b/Engine/Plugins/Runtime/Metasound/Source/MetasoundFrontend/Private/MetasoundFrontendDocumentVersioning.cpp
@@ -572,11 +572,18 @@ namespace Metasound::Frontend
 			using namespace Metasound::Frontend;
 
 #if WITH_EDITORONLY_DATA
-			FMetasoundAssetBase& Asset = OutBuilder.GetMetasoundAsset();
-			Asset.MigrateEditorGraph(OutBuilder);
-			UE_LOG(LogMetaSound, Display, TEXT("Resave recommended: Asset '%s' at '%s' successfully migrated editor data in target document version '%s'."), *Name.ToString(), *Path->ToString(), *GetTargetVersion().ToString());
+			if (IsRunningCookCommandlet())
+			{
+				UE_LOG(LogMetaSound, Display, TEXT("Resave recommended: Asset '%s' at '%s' skipped migrated editor data/creation of input template nodes during cook to target document version '%s'."), *Name.ToString(), *Path->ToString(), *GetTargetVersion().ToString());
+			}
+			else
+			{
+				FMetasoundAssetBase& Asset = OutBuilder.GetMetasoundAsset();
+				Asset.MigrateEditorGraph(OutBuilder);
+				UE_LOG(LogMetaSound, Display, TEXT("Resave recommended: Asset '%s' at '%s' successfully migrated editor data in target document version '%s'."), *Name.ToString(), *Path->ToString(), *GetTargetVersion().ToString());
+			}
 #else // !WITH_EDITORONLY_DATA
-			UE_LOG(LogMetaSound, Error, TEXT("Asset '%s' at '%s' must be saved with editor data enabled in order to version document to target version '%s'."), *Name.ToString(), *Path->ToString(), *GetTargetVersion().ToString());
+			UE_LOG(LogMetaSound, Warning, TEXT("Asset '%s' at '%s' must be re-saved with editor data enabled in order to migrate document to target version '%s'."), *Name.ToString(), *Path->ToString(), *GetTargetVersion().ToString());
 #endif // !WITH_EDITORONLY_DATA
 		}
 	};
