Fix another place where could potentially be loading engine subsystem too early if serialized assets are being versioned before MetaSound Engine Subsystems are initialized
#rnx

diff --git a/Engine/Plugins/Runtime/Metasound/Source/MetasoundFrontend/Private/MetasoundFrontendDocumentVersioning.cpp b/Engine/Plugins/Runtime/Metasound/Source/MetasoundFrontend/Private/MetasoundFrontendDocumentVersioning.cpp
--- a/Engine/Plugins/Runtime/Metasound/Source/MetasoundFrontend/Private/MetasoundFrontendDocumentVersioning.cpp
+++ b/Engine/Plugins/Runtime/Metasound/Source/MetasoundFrontend/Private/MetasoundFrontendDocumentVersioning.cpp
@@ -552,12 +552,14 @@ namespace Metasound::Frontend
 	/** Versions document from 1.11 to 1.12. */
 	class FVersionDocument_1_12 : public FVersionDocumentTransform
 	{
+		FMetasoundAssetBase* AssetBase = nullptr;
 		const FName Name;
 		const FSoftObjectPath* Path = nullptr;
 
 	public:
-		FVersionDocument_1_12(FName InName, const FSoftObjectPath& InAssetPath)
-			: Name(InName)
+		FVersionDocument_1_12(FMetasoundAssetBase& InAssetBase, FName InName, const FSoftObjectPath& InAssetPath)
+			: AssetBase(&InAssetBase)
+			, Name(InName)
 			, Path(&InAssetPath)
 		{
 		}
@@ -578,8 +580,8 @@ namespace Metasound::Frontend
 			}
 			else
 			{
-				FMetasoundAssetBase& Asset = OutBuilder.GetMetasoundAsset();
-				Asset.MigrateEditorGraph(OutBuilder);
+				check(AssetBase);
+				AssetBase->MigrateEditorGraph(OutBuilder);
 				UE_LOG(LogMetaSound, Display, TEXT("Resave recommended: Asset '%s' at '%s' successfully migrated editor data in target document version '%s'."), *Name.ToString(), *Path->ToString(), *GetTargetVersion().ToString());
 			}
 #else // !WITH_EDITORONLY_DATA
@@ -626,7 +628,7 @@ namespace Metasound::Frontend
 			// 1. Builders were implemented post document version 1.11, so earlier versions are not supported.
 			// 2. Controller mutations are not tracked by analogous builder, which can cause internal cache corruption.
 			FMetaSoundFrontendDocumentBuilder Builder(DocumentInterface);
-			bWasUpdated |= FVersionDocument_1_12(Name, Path).Transform(Builder);
+			bWasUpdated |= FVersionDocument_1_12(InAssetBase, Name, Path).Transform(Builder);
 		}
 
 		if (bWasUpdated)
