- Avoid creating MetaSound builder at all if already up to date (transaction stack when undo/redo is inacted can attempt to call versioning code via serialize call even though its up-to-date spawning unnecessary errors regarding multiple builders being active)
- Avoid creating additional builder when swapping interfaces
#tests verified listed jira repros resolved locally in editor
#rnx

diff --git a/Engine/Plugins/Runtime/Metasound/Source/MetasoundEditor/Private/MetasoundDetailCustomization.cpp b/Engine/Plugins/Runtime/Metasound/Source/MetasoundEditor/Private/MetasoundDetailCustomization.cpp
--- a/Engine/Plugins/Runtime/Metasound/Source/MetasoundEditor/Private/MetasoundDetailCustomization.cpp
+++ b/Engine/Plugins/Runtime/Metasound/Source/MetasoundEditor/Private/MetasoundDetailCustomization.cpp
@@ -382,7 +382,7 @@ namespace Metasound
 							MetaSoundAsset->GetGraphChecked().Modify();
 
 							TScriptInterface<IMetaSoundDocumentInterface> MetaSoundPatchDocInterface = MetaSound.Get();
-							FMetaSoundFrontendDocumentBuilder Builder(MetaSoundPatchDocInterface);
+							FMetaSoundFrontendDocumentBuilder& Builder = IMetaSoundAssetManager::GetChecked().AttachDocumentBuilderChecked(*MetaSound.Get());
 							FModifyInterfaceOptions Options({ }, { InterfaceToAdd });
 							Options.bSetDefaultNodeLocations = false; // Don't automatically add nodes to ed graph
 							Builder.ModifyInterfaces(MoveTemp(Options));
diff --git a/Engine/Plugins/Runtime/Metasound/Source/MetasoundFrontend/Private/MetasoundFrontendDocumentVersioning.cpp b/Engine/Plugins/Runtime/Metasound/Source/MetasoundFrontend/Private/MetasoundFrontendDocumentVersioning.cpp
--- a/Engine/Plugins/Runtime/Metasound/Source/MetasoundFrontend/Private/MetasoundFrontendDocumentVersioning.cpp
+++ b/Engine/Plugins/Runtime/Metasound/Source/MetasoundFrontend/Private/MetasoundFrontendDocumentVersioning.cpp
@@ -5,7 +5,6 @@
 #include "Interfaces/MetasoundFrontendInterface.h"
 #include "Interfaces/MetasoundFrontendInterfaceRegistry.h"
 #include "MetasoundAccessPtr.h"
-#include "MetasoundAssetManager.h"
 #include "MetasoundDocumentInterface.h"
 #include "MetasoundFrontendDocumentBuilder.h"
 #include "MetasoundFrontendDocumentController.h"
@@ -597,44 +596,47 @@ namespace Metasound::Frontend
 		UObject* OwningAsset = InAssetBase.GetOwningAsset();
 		TScriptInterface<IMetaSoundDocumentInterface> DocumentInterface(OwningAsset);
 		const FMetasoundFrontendDocument& Document = DocumentInterface->GetConstDocument();
-		FName Name = FName(OwningAsset->GetName());
-		FString Path = OwningAsset->GetPathName();
 
 		// Copied as value will be mutated with each applicable transform below
 		const FMetasoundFrontendVersionNumber InitVersionNumber = Document.Metadata.Version.Number;
-
-		// Controller (Soft Deprecated) Transforms
-		if (InitVersionNumber.Major == 1 && InitVersionNumber.Minor < 12)
+		if (InitVersionNumber < GetMaxDocumentVersion())
 		{
-			UObject* DocObject = DocumentInterface.GetObject();
-			FDocumentHandle DocHandle = InAssetBase.GetDocumentHandle();
+			FName Name = FName(OwningAsset->GetName());
+			FString Path = OwningAsset->GetPathName();
 
-			bWasUpdated |= FVersionDocument_1_1(Name, Path).Transform(DocHandle);
-			bWasUpdated |= FVersionDocument_1_2(Name, Path).Transform(DocHandle);
-			bWasUpdated |= FVersionDocument_1_3().Transform(DocHandle);
-			bWasUpdated |= FVersionDocument_1_4().Transform(DocHandle);
-			bWasUpdated |= FVersionDocument_1_5(Name, Path).Transform(DocHandle);
-			bWasUpdated |= FVersionDocument_1_6().Transform(DocHandle);
-			bWasUpdated |= FVersionDocument_1_7(Name, Path).Transform(DocHandle);
-			bWasUpdated |= FVersionDocument_1_8(Name, Path).Transform(DocHandle);
-			bWasUpdated |= FVersionDocument_1_9(Name, Path).Transform(DocHandle);
-			bWasUpdated |= FVersionDocument_1_10().Transform(DocHandle);
-			bWasUpdated |= FVersionDocument_1_11().Transform(DocHandle);
-		}
+			// Controller (Soft Deprecated) Transforms
+			if (InitVersionNumber.Major == 1 && InitVersionNumber.Minor < 12)
+			{
+				UObject* DocObject = DocumentInterface.GetObject();
+				FDocumentHandle DocHandle = InAssetBase.GetDocumentHandle();
+
+				bWasUpdated |= FVersionDocument_1_1(Name, Path).Transform(DocHandle);
+				bWasUpdated |= FVersionDocument_1_2(Name, Path).Transform(DocHandle);
+				bWasUpdated |= FVersionDocument_1_3().Transform(DocHandle);
+				bWasUpdated |= FVersionDocument_1_4().Transform(DocHandle);
+				bWasUpdated |= FVersionDocument_1_5(Name, Path).Transform(DocHandle);
+				bWasUpdated |= FVersionDocument_1_6().Transform(DocHandle);
+				bWasUpdated |= FVersionDocument_1_7(Name, Path).Transform(DocHandle);
+				bWasUpdated |= FVersionDocument_1_8(Name, Path).Transform(DocHandle);
+				bWasUpdated |= FVersionDocument_1_9(Name, Path).Transform(DocHandle);
+				bWasUpdated |= FVersionDocument_1_10().Transform(DocHandle);
+				bWasUpdated |= FVersionDocument_1_11().Transform(DocHandle);
+			}
 
-		// Builder Transforms
-		{
-			// Mutation of document via both controller & builder systems simultaneously is forbidden as:
-			// 1. Builders were implemented post document version 1.11, so earlier versions are not supported.
-			// 2. Controller mutations are not tracked by analogous builder, which can cause internal cache corruption.
-			FMetaSoundFrontendDocumentBuilder Builder(DocumentInterface);
-			bWasUpdated |= FVersionDocument_1_12(InAssetBase, Name, Path).Transform(Builder);
-		}
+			// Builder Transforms
+			{
+				// Mutation of document via both controller & builder systems simultaneously is forbidden as:
+				// 1. Builders were implemented post document version 1.11, so earlier versions are not supported.
+				// 2. Controller mutations are not tracked by analogous builder, which can cause internal cache corruption.
+				FMetaSoundFrontendDocumentBuilder Builder(DocumentInterface);
+				bWasUpdated |= FVersionDocument_1_12(InAssetBase, Name, Path).Transform(Builder);
+			}
 
-		if (bWasUpdated)
-		{
-			const FMetasoundFrontendVersionNumber& NewVersionNumber = Document.Metadata.Version.Number;
-			UE_LOG(LogMetaSound, Verbose, TEXT("MetaSound at '%s' Document Versioned: '%s' --> '%s'"), *Path, *InitVersionNumber.ToString(), *NewVersionNumber.ToString());
+			if (bWasUpdated)
+			{
+				const FMetasoundFrontendVersionNumber& NewVersionNumber = Document.Metadata.Version.Number;
+				UE_LOG(LogMetaSound, Verbose, TEXT("MetaSound at '%s' Document Versioned: '%s' --> '%s'"), *Path, *InitVersionNumber.ToString(), *NewVersionNumber.ToString());
+			}
 		}
 
 		return bWasUpdated;
