Fix for pin audio inspector oscilloscope data dropouts.

diff --git a/Engine/Plugins/Runtime/Metasound/Source/MetasoundEditor/Private/MetasoundEditorGraphConnectionManager.cpp b/Engine/Plugins/Runtime/Metasound/Source/MetasoundEditor/Private/MetasoundEditorGraphConnectionManager.cpp
--- a/Engine/Plugins/Runtime/Metasound/Source/MetasoundEditor/Private/MetasoundEditorGraphConnectionManager.cpp
+++ b/Engine/Plugins/Runtime/Metasound/Source/MetasoundEditor/Private/MetasoundEditorGraphConnectionManager.cpp
@@ -1,7 +1,7 @@
 // Copyright Epic Games, Inc. All Rights Reserved.
 #include "MetasoundEditorGraphConnectionManager.h"
 
-#include "Analysis/MetasoundFrontendVertexAnalyzerAudioBuffer.h"
+#include "Analysis/MetasoundFrontendVertexAnalyzerAudioBusWriter.h"
 #include "Analysis/MetasoundFrontendVertexAnalyzerEnvelopeFollower.h"
 #include "Analysis/MetasoundFrontendVertexAnalyzerForwardValue.h"
 #include "Analysis/MetasoundFrontendVertexAnalyzerTriggerDensity.h"
@@ -25,7 +25,6 @@ namespace Metasound
 			GraphAnalyzerView->AddAnalyzerForAllSupportedOutputs(FVertexAnalyzerForwardInt::GetAnalyzerName());
 			GraphAnalyzerView->AddAnalyzerForAllSupportedOutputs(FVertexAnalyzerForwardString::GetAnalyzerName());
 
-			GraphAnalyzerView->AddAnalyzerForAllSupportedOutputs(FVertexAnalyzerAudioBuffer::GetAnalyzerName());
 			GraphAnalyzerView->AddAnalyzerForAllSupportedOutputs(FVertexAnalyzerEnvelopeFollower::GetAnalyzerName());
 			GraphAnalyzerView->AddAnalyzerForAllSupportedOutputs(FVertexAnalyzerTriggerDensity::GetAnalyzerName());
 		}
@@ -76,14 +75,28 @@ namespace Metasound
 			Window.SetNumZeroed(InWindowSize);
 		}
 
-		void FGraphConnectionManager::TrackAudioPin(const FGuid& InNodeID, FVertexName InOutputName, FName InAnalyzerName, Audio::FPatchInput& InPatchInput)
+		FGuid FGraphConnectionManager::AddAudioBusWriter(const FGuid& InNodeID, FVertexName InOutputName, Audio::FDeviceId InDeviceID, const UAudioBus* InAudioBus)
 		{
-			TrackedAudioPins.Add({ InNodeID, InOutputName, InAnalyzerName }, MakeShared<Audio::FPatchInput>(InPatchInput));
+			FGuid AnalyzerInstanceID;
+
+			if (GraphAnalyzerView.IsValid())
+			{
+				const uint32 AudioBusID = InAudioBus->GetUniqueID();
+				const FName AnalyzerName = Frontend::FVertexAnalyzerAudioBusWriter::GetAnalyzerName();
+				const FName AnalyzerMemberName = Frontend::FVertexAnalyzerAudioBusWriter::GetAnalyzerMemberName(InDeviceID, AudioBusID);
+				AnalyzerInstanceID = GraphAnalyzerView->AddAnalyzerForSpecifiedOutput(InNodeID, InOutputName, AnalyzerName, AnalyzerMemberName);
+			}
+
+			return AnalyzerInstanceID;
 		}
 
-		void FGraphConnectionManager::UntrackAudioPin(const FGuid& InNodeID, FVertexName InOutputName, FName InAnalyzerName)
+		void FGraphConnectionManager::RemoveAudioBusWriter(const FGuid& InAnalyzerInstanceID)
 		{
-			TrackedAudioPins.Remove({ InNodeID, InOutputName, InAnalyzerName });
+			if (GraphAnalyzerView.IsValid())
+			{
+				const FName AnalyzerName = Frontend::FVertexAnalyzerAudioBusWriter::GetAnalyzerName();
+				GraphAnalyzerView->RemoveAnalyzerInstance(AnalyzerName, InAnalyzerInstanceID);
+			}
 		}
 
 		void FGraphConnectionManager::Update(float InDeltaTime)
@@ -131,23 +144,6 @@ namespace Metasound
 				{
 					Pair.Value.Update(InDeltaTime);
 				}
-
-				// Audio Pins update
-				if (const UMetasoundEditorSettings* MetaSoundSettings = GetDefault<UMetasoundEditorSettings>();
-				    MetaSoundSettings && MetaSoundSettings->bShowOscilloscopeOnAudioPinMouseOver)
-				{
-					auto UpdateAudio = [this](const FAnalyzerAddress& InAnalyzerAddress, const FAudioBuffer& InAudioBuffer)
-					{
-						TSharedRef<Audio::FPatchInput>* PatchInput = TrackedAudioPins.Find({ InAnalyzerAddress.NodeID, InAnalyzerAddress.OutputName, InAnalyzerAddress.AnalyzerName });
-
-						if (PatchInput)
-						{
-							(*PatchInput)->PushAudio(InAudioBuffer.GetData(), InAudioBuffer.Num());
-						}
-					};
-				
-					UpdateConnections<FVertexAnalyzerAudioBuffer, FAudioBuffer>(InDeltaTime, UpdateAudio);
-				}
 			}
 		}
 
diff --git a/Engine/Plugins/Runtime/Metasound/Source/MetasoundEditor/Private/MetasoundEditorGraphConnectionManager.h b/Engine/Plugins/Runtime/Metasound/Source/MetasoundEditor/Private/MetasoundEditorGraphConnectionManager.h
--- a/Engine/Plugins/Runtime/Metasound/Source/MetasoundEditor/Private/MetasoundEditorGraphConnectionManager.h
+++ b/Engine/Plugins/Runtime/Metasound/Source/MetasoundEditor/Private/MetasoundEditorGraphConnectionManager.h
@@ -99,8 +99,11 @@ namespace Metasound
 			// Retains existing values provided from update up to given size.
 			void TrackValue(const FGuid& InNodeID, FVertexName InOutputName, FName InAnalyzerName, int32 InWindowSize);
 
-			void TrackAudioPin(const FGuid& InNodeID, FVertexName InOutputName, FName InAnalyzerName, Audio::FPatchInput& InPatchInput);
-			void UntrackAudioPin(const FGuid& InNodeID, FVertexName InOutputName, FName InAnalyzerName);
+			// Send the audio from the given node output to the given bus. If successful, an AnalyzerInstanceID will be returned.
+			FGuid AddAudioBusWriter(const FGuid& InNodeID, FVertexName InOutputName, Audio::FDeviceId InDeviceID, const UAudioBus* InAudioBus);
+
+			// Remove the specified analyzer. InAnalyzerInstanceID is an FGuid that was previously returned from a call to AddAudioBusWriter(...).
+			void RemoveAudioBusWriter(const FGuid& InAnalyzerInstanceID);
 
 			void Update(float InDeltaTime);
 
@@ -220,7 +223,6 @@ namespace Metasound
 			struct FAudioBufferKey : FWindowValueKey{};
 
 			TMap<FWindowValueKey, FFloatMovingWindow> WindowedValues;
-			TMap<FAudioBufferKey, TSharedRef<Audio::FPatchInput>> TrackedAudioPins;
 
 			TWeakObjectPtr<const UAudioComponent> AudioComponent;
 			TMap<FString, Audio::FVolumeFader> ConnectionFaders;
diff --git a/Engine/Plugins/Runtime/Metasound/Source/MetasoundEditor/Private/MetasoundPinAudioInspector.cpp b/Engine/Plugins/Runtime/Metasound/Source/MetasoundEditor/Private/MetasoundPinAudioInspector.cpp
--- a/Engine/Plugins/Runtime/Metasound/Source/MetasoundEditor/Private/MetasoundPinAudioInspector.cpp
+++ b/Engine/Plugins/Runtime/Metasound/Source/MetasoundEditor/Private/MetasoundPinAudioInspector.cpp
@@ -1,7 +1,6 @@
 // Copyright Epic Games, Inc. All Rights Reserved.
 #include "MetasoundPinAudioInspector.h"
 
-#include "Analysis/MetasoundFrontendVertexAnalyzerAudioBuffer.h"
 #include "AudioBusSubsystem.h"
 #include "AudioDefines.h"
 #include "AudioDeviceManager.h"
@@ -57,20 +56,7 @@ namespace Metasound
 
 			Oscilloscope->StartProcessing();
 
-			// Set PatchInput from Oscilloscope AudioBus
-			const FAudioDeviceManager* AudioDeviceManager = FAudioDeviceManager::Get();
-			check(AudioDeviceManager);
-
-			const FMixerDevice* MixerDevice = static_cast<const FMixerDevice*>(AudioDeviceManager->GetAudioDeviceRaw(AudioDeviceId));
-			check(MixerDevice);
-
-			UAudioBusSubsystem* AudioBusSubsystem = MixerDevice->GetSubsystem<UAudioBusSubsystem>();
-			check(AudioBusSubsystem);
-
-			const uint32 AudioBusId = Oscilloscope->GetAudioBus()->GetUniqueID();
-			PatchInput = AudioBusSubsystem->AddPatchInputForAudioBus(FAudioBusKey(AudioBusId), MixerDevice->GetNumOutputFrames(), NumChannels);
-
-			// Track Audio Pin
+			// Analyze Audio Pin
 			if (GraphPinObj && !GraphPinObj->LinkedTo.IsEmpty() && GraphPinObj->PinType.PinCategory == FGraphBuilder::PinCategoryAudio)
 			{
 				if (FGraphConnectionManager* ConnectionManager = GetConnectionManager())
@@ -78,24 +64,21 @@ namespace Metasound
 					const Frontend::FConstOutputHandle OutputHandle = FGraphBuilder::FindReroutedOutputHandleFromPin(GraphPinObj);
 					const FGuid NodeID = OutputHandle->GetOwningNodeID();
 					const FName OutputName = OutputHandle->GetName();
-					const FName AnalyzerName = Frontend::FVertexAnalyzerAudioBuffer::GetAnalyzerName();
-
-					ConnectionManager->TrackAudioPin(NodeID, OutputName, AnalyzerName, PatchInput);
+					AnalyzerInstanceID = ConnectionManager->AddAudioBusWriter(NodeID, OutputName, AudioDeviceId, Oscilloscope->GetAudioBus());
 				}
 			}
 		}
 
 		FMetasoundPinAudioInspector::~FMetasoundPinAudioInspector()
 		{
-			// Untrack audio pin
-			if (FGraphConnectionManager* ConnectionManager = GetConnectionManager())
+			// Remove audio pin analyzer
+			if (AnalyzerInstanceID.IsValid())
 			{
-				const Frontend::FConstOutputHandle OutputHandle = FGraphBuilder::FindReroutedOutputHandleFromPin(GraphPinObj);
-				const FGuid NodeID = OutputHandle->GetOwningNodeID();
-				const FName OutputName = OutputHandle->GetName();
-				const FName AnalyzerName = Frontend::FVertexAnalyzerAudioBuffer::GetAnalyzerName();
-
-				ConnectionManager->UntrackAudioPin(NodeID, OutputName, AnalyzerName);
+				if (FGraphConnectionManager* ConnectionManager = GetConnectionManager())
+				{
+					ConnectionManager->RemoveAudioBusWriter(AnalyzerInstanceID);
+					AnalyzerInstanceID.Invalidate();
+				}
 			}
 
 			// Stop AudioBus
diff --git a/Engine/Plugins/Runtime/Metasound/Source/MetasoundEditor/Private/MetasoundPinAudioInspector.h b/Engine/Plugins/Runtime/Metasound/Source/MetasoundEditor/Private/MetasoundPinAudioInspector.h
--- a/Engine/Plugins/Runtime/Metasound/Source/MetasoundEditor/Private/MetasoundPinAudioInspector.h
+++ b/Engine/Plugins/Runtime/Metasound/Source/MetasoundEditor/Private/MetasoundPinAudioInspector.h
@@ -33,9 +33,8 @@ namespace Metasound
 			UEdGraphPin* GraphPinObj = nullptr;
 
 			TSharedPtr<AudioWidgets::FAudioOscilloscope> Oscilloscope = nullptr;
-			Audio::FPatchInput PatchInput;
-			
 			TSharedPtr<SMetasoundPinAudioInspector> PinAudioInspectorWidget = nullptr;
+			FGuid AnalyzerInstanceID;
 		};
 	} // namespace Editor
 } // namespace Metasound
diff --git a/Engine/Plugins/Runtime/Metasound/Source/MetasoundEngine/Private/MetasoundEngineModule.cpp b/Engine/Plugins/Runtime/Metasound/Source/MetasoundEngine/Private/MetasoundEngineModule.cpp
--- a/Engine/Plugins/Runtime/Metasound/Source/MetasoundEngine/Private/MetasoundEngineModule.cpp
+++ b/Engine/Plugins/Runtime/Metasound/Source/MetasoundEngine/Private/MetasoundEngineModule.cpp
@@ -14,6 +14,7 @@
 #include "MetasoundWaveTable.h"
 #include "Analysis/MetasoundFrontendAnalyzerRegistry.h"
 #include "Analysis/MetasoundFrontendVertexAnalyzerAudioBuffer.h"
+#include "Analysis/MetasoundFrontendVertexAnalyzerAudioBusWriter.h"
 #include "Analysis/MetasoundFrontendVertexAnalyzerEnvelopeFollower.h"
 #include "Analysis/MetasoundFrontendVertexAnalyzerForwardValue.h"
 #include "Analysis/MetasoundFrontendVertexAnalyzerTriggerDensity.h"
@@ -118,6 +119,7 @@ public:
 
 		// Register Analyzers
 		METASOUND_REGISTER_VERTEX_ANALYZER_FACTORY(Frontend::FVertexAnalyzerAudioBuffer)
+		METASOUND_REGISTER_VERTEX_ANALYZER_FACTORY(Frontend::FVertexAnalyzerAudioBusWriter)
 		METASOUND_REGISTER_VERTEX_ANALYZER_FACTORY(Frontend::FVertexAnalyzerEnvelopeFollower)
 		METASOUND_REGISTER_VERTEX_ANALYZER_FACTORY(Frontend::FVertexAnalyzerForwardBool)
 		METASOUND_REGISTER_VERTEX_ANALYZER_FACTORY(Frontend::FVertexAnalyzerForwardFloat)
diff --git a/Engine/Plugins/Runtime/Metasound/Source/MetasoundFrontend/MetasoundFrontend.Build.cs b/Engine/Plugins/Runtime/Metasound/Source/MetasoundFrontend/MetasoundFrontend.Build.cs
--- a/Engine/Plugins/Runtime/Metasound/Source/MetasoundFrontend/MetasoundFrontend.Build.cs
+++ b/Engine/Plugins/Runtime/Metasound/Source/MetasoundFrontend/MetasoundFrontend.Build.cs
@@ -16,8 +16,10 @@ namespace UnrealBuildTool.Rules
 				new string[]
 				{
 					"AudioExtensions",
+					"AudioMixer",
 					"Core",
 					"CoreUObject",
+					"Engine",
 					"Serialization",
 					"SignalProcessing",
 					"MetasoundGraphCore"
diff --git a/Engine/Plugins/Runtime/Metasound/Source/MetasoundFrontend/Private/Analysis/MetasoundFrontendGraphAnalyzerView.cpp b/Engine/Plugins/Runtime/Metasound/Source/MetasoundFrontend/Private/Analysis/MetasoundFrontendGraphAnalyzerView.cpp
--- a/Engine/Plugins/Runtime/Metasound/Source/MetasoundFrontend/Private/Analysis/MetasoundFrontendGraphAnalyzerView.cpp
+++ b/Engine/Plugins/Runtime/Metasound/Source/MetasoundFrontend/Private/Analysis/MetasoundFrontendGraphAnalyzerView.cpp
@@ -110,9 +110,43 @@ namespace Metasound
 				}
 			});
 
-			TArray<FString> ActiveAnalyzerStrings;
-			Algo::Transform(ActiveAnalyzers, ActiveAnalyzerStrings, [](const FAnalyzerAddress& Address) { return Address.ToString(); });
-			ActiveAnalyzerSender->PushLiteral(MoveTemp(ActiveAnalyzerStrings));
+			SendActiveAnalyzers();
+		}
+
+		FGuid FMetasoundGraphAnalyzerView::AddAnalyzerForSpecifiedOutput(const FGuid& InNodeID, FVertexName InOutputName, FName InAnalyzerName, FName InAnalyzerMemberName)
+		{
+			const IVertexAnalyzerFactory* Factory = IVertexAnalyzerRegistry::Get().FindAnalyzerFactory(InAnalyzerName);
+			if (!Factory)
+			{
+				return {};
+			}
+
+			FAnalyzerAddress AnalyzerAddress;
+			AnalyzerAddress.DataType = Factory->GetDataType();
+			AnalyzerAddress.InstanceID = InstanceID;
+			AnalyzerAddress.NodeID = InNodeID;
+			AnalyzerAddress.OutputName = InOutputName;
+			AnalyzerAddress.AnalyzerName = InAnalyzerName;
+			AnalyzerAddress.AnalyzerInstanceID = FGuid::NewGuid();
+			AnalyzerAddress.AnalyzerMemberName = InAnalyzerMemberName;
+			ActiveAnalyzers.Add(AnalyzerAddress);
+
+			const TArray<FAnalyzerOutput>& AnalyzerOutputs = Factory->GetAnalyzerOutputs();
+			for (const FAnalyzerOutput& AnalyzerOutput : AnalyzerOutputs)
+			{
+				FAnalyzerAddress OutputReceiverAddress = AnalyzerAddress;
+				OutputReceiverAddress.AnalyzerMemberName = AnalyzerOutput.Name;
+				OutputReceiverAddress.DataType = AnalyzerOutput.DataType;
+
+				FMetasoundGraphAnalyzerOutputKey OutputKey{ AnalyzerAddress.NodeID, AnalyzerAddress.OutputName };
+				FMetasoundAnalyzerView NewView(MoveTemp(OutputReceiverAddress));
+				NewView.BindToAllOutputs(OperatorSettings);
+				AnalyzerViews.FindOrAdd(OutputKey).Add(MoveTemp(NewView));
+			}
+
+			SendActiveAnalyzers();
+
+			return AnalyzerAddress.AnalyzerInstanceID;
 		}
 
 		TArray<FMetasoundAnalyzerView*> FMetasoundGraphAnalyzerView::GetAnalyzerViews(FName InAnalyzerName)
@@ -221,9 +255,42 @@ namespace Metasound
 				}
 			}
 
-			TArray<FString> ActiveAnalyzerStrings;
-			Algo::Transform(ActiveAnalyzers, ActiveAnalyzerStrings, [] (const FAnalyzerAddress& Address) { return Address.ToString(); });
-			ActiveAnalyzerSender->PushLiteral(MoveTemp(ActiveAnalyzerStrings));
+			SendActiveAnalyzers();
+		}
+
+		void FMetasoundGraphAnalyzerView::RemoveAnalyzerInstance(FName InAnalyzerName, const FGuid& InAnalyzerInstanceID)
+		{
+			TArray<FMetasoundGraphAnalyzerOutputKey> OutputsToRemove;
+			for (TPair<FMetasoundGraphAnalyzerOutputKey, TArray<FMetasoundAnalyzerView>>& AnalyzerPair : AnalyzerViews)
+			{
+				TArray<FMetasoundAnalyzerView>& Views = AnalyzerPair.Value;
+				Views.RemoveAll([InAnalyzerName, InAnalyzerInstanceID](const FMetasoundAnalyzerView& View)
+					{
+						return View.AnalyzerAddress.AnalyzerName == InAnalyzerName && View.AnalyzerAddress.AnalyzerInstanceID == InAnalyzerInstanceID;
+					});
+
+				if (Views.IsEmpty())
+				{
+					OutputsToRemove.Add(AnalyzerPair.Key);
+				}
+			}
+
+			for (const FMetasoundGraphAnalyzerOutputKey& Key : OutputsToRemove)
+			{
+				AnalyzerViews.Remove(Key);
+			}
+
+			TArray<FAnalyzerAddress> AnalyzerAddresses = ActiveAnalyzers.Array();
+			for (const FAnalyzerAddress& AnalyzerAddress : AnalyzerAddresses)
+			{
+				if (AnalyzerAddress.AnalyzerName == InAnalyzerName && AnalyzerAddress.AnalyzerInstanceID == InAnalyzerInstanceID)
+				{
+					ActiveAnalyzers.Remove(AnalyzerAddress);
+					break;
+				}
+			}
+
+			SendActiveAnalyzers();
 		}
 
 		const FMetasoundAssetBase& FMetasoundGraphAnalyzerView::GetMetaSoundAssetChecked() const
@@ -231,5 +298,12 @@ namespace Metasound
 			check(MetaSoundAsset);
 			return *MetaSoundAsset;
 		}
+
+		void FMetasoundGraphAnalyzerView::SendActiveAnalyzers()
+		{
+			TArray<FString> ActiveAnalyzerStrings;
+			Algo::Transform(ActiveAnalyzers, ActiveAnalyzerStrings, [](const FAnalyzerAddress& Address) { return Address.ToString(); });
+			ActiveAnalyzerSender->PushLiteral(MoveTemp(ActiveAnalyzerStrings));
+		}
 	} // namespace Frontend
 } // namespace Metasound
diff --git a/Engine/Plugins/Runtime/Metasound/Source/MetasoundFrontend/Private/Analysis/MetasoundFrontendVertexAnalyzerAudioBusWriter.cpp b/Engine/Plugins/Runtime/Metasound/Source/MetasoundFrontend/Private/Analysis/MetasoundFrontendVertexAnalyzerAudioBusWriter.cpp
new file mode 100644
--- /dev/null
+++ b/Engine/Plugins/Runtime/Metasound/Source/MetasoundFrontend/Private/Analysis/MetasoundFrontendVertexAnalyzerAudioBusWriter.cpp
@@ -0,0 +1,84 @@
+// Copyright Epic Games, Inc. All Rights Reserved.
+
+#include "Analysis/MetasoundFrontendVertexAnalyzerAudioBusWriter.h"
+
+#include "AudioBusSubsystem.h"
+#include "AudioDevice.h"
+#include "AudioDeviceManager.h"
+#include "MetasoundAudioBuffer.h"
+
+namespace Metasound
+{
+	namespace Frontend
+	{
+		const FName& FVertexAnalyzerAudioBusWriter::GetAnalyzerName()
+		{
+			static const FName AnalyzerName = "UE.Audio.AudioBusWriter";
+			return AnalyzerName;
+		}
+
+		const FName& FVertexAnalyzerAudioBusWriter::GetDataType()
+		{
+			return GetMetasoundDataTypeName<FAudioBuffer>();
+		}
+
+		FName FVertexAnalyzerAudioBusWriter::GetAnalyzerMemberName(const Audio::FDeviceId InDeviceID, const uint32 InAudioBusID)
+		{
+			const FBusAddress BusAddress
+			{
+				.DeviceID = InDeviceID,
+				.AudioBusID = InAudioBusID,
+			};
+			return FName(BusAddress.ToString());
+		}
+
+		FVertexAnalyzerAudioBusWriter::FVertexAnalyzerAudioBusWriter(const FCreateAnalyzerParams& InParams)
+			: FVertexAnalyzerBase(InParams.AnalyzerAddress, InParams.VertexDataReference)
+		{
+			// AnalyzerMemberName is used as an input to modify analyzer settings:
+			const FBusAddress BusAddress = FBusAddress::FromString(InParams.AnalyzerAddress.AnalyzerMemberName.ToString());
+
+			if (FAudioDeviceManager* ADM = FAudioDeviceManager::Get())
+			{
+				if (const FAudioDevice* AudioDevice = ADM->GetAudioDeviceRaw(BusAddress.DeviceID))
+				{
+					UAudioBusSubsystem* AudioBusSubsystem = AudioDevice->GetSubsystem<UAudioBusSubsystem>();
+
+					const int32 Frames = InParams.OperatorSettings.GetNumFramesPerBlock();
+					constexpr int32 Channels = 1;
+					AudioBusPatchInput = AudioBusSubsystem->AddPatchInputForAudioBus(BusAddress.AudioBusID, Frames, Channels);
+				}
+			}
+		}
+
+		void FVertexAnalyzerAudioBusWriter::Execute()
+		{
+			const FAudioBuffer& AudioBuffer = GetVertexData<FAudioBuffer>();
+			AudioBusPatchInput.PushAudio(AudioBuffer.GetData(), AudioBuffer.Num());
+		}
+
+		FString FVertexAnalyzerAudioBusWriter::FBusAddress::ToString() const
+		{
+			return FString::Join(TArray<FString>
+			{
+				FString::FromInt(DeviceID),
+				FString::FromInt(AudioBusID),
+			}
+			, TEXT(","));
+		}
+
+		FVertexAnalyzerAudioBusWriter::FBusAddress FVertexAnalyzerAudioBusWriter::FBusAddress::FromString(const FString& InAnalyzerMemberName)
+		{
+			FBusAddress BusAddress;
+
+			TArray<FString> Tokens;
+			if (ensure(InAnalyzerMemberName.ParseIntoArray(Tokens, TEXT(",")) == 2))
+			{
+				BusAddress.DeviceID = FCString::Atoi(*Tokens[0]);
+				BusAddress.AudioBusID = FCString::Atoi(*Tokens[1]);
+			}
+
+			return BusAddress;
+		}
+	} // namespace Frontend
+} // namespace Metasound
diff --git a/Engine/Plugins/Runtime/Metasound/Source/MetasoundFrontend/Public/Analysis/MetasoundFrontendGraphAnalyzerView.h b/Engine/Plugins/Runtime/Metasound/Source/MetasoundFrontend/Public/Analysis/MetasoundFrontendGraphAnalyzerView.h
--- a/Engine/Plugins/Runtime/Metasound/Source/MetasoundFrontend/Public/Analysis/MetasoundFrontendGraphAnalyzerView.h
+++ b/Engine/Plugins/Runtime/Metasound/Source/MetasoundFrontend/Public/Analysis/MetasoundFrontendGraphAnalyzerView.h
@@ -27,6 +27,7 @@ namespace Metasound
 			TSet<FAnalyzerAddress> ActiveAnalyzers;
 
 			const FMetasoundAssetBase& GetMetaSoundAssetChecked() const;
+			void SendActiveAnalyzers();
 
 			uint64 InstanceID = INDEX_NONE;
 			const FMetasoundAssetBase* MetaSoundAsset = nullptr;
@@ -48,7 +49,9 @@ namespace Metasound
 			~FMetasoundGraphAnalyzerView();
 
 			void AddAnalyzerForAllSupportedOutputs(FName InAnalyzerName, bool bInRequiresConnection = true);
+			FGuid AddAnalyzerForSpecifiedOutput(const FGuid& InNodeID, FVertexName InOutputName, FName InAnalyzerName, FName InAnalyzerMemberName);
 			void RemoveAnalyzerForAllSupportedOutputs(FName InAnalyzerName);
+			void RemoveAnalyzerInstance(FName InAnalyzerName, const FGuid& InAnalyzerInstanceID);
 			TArray<FMetasoundAnalyzerView*> GetAnalyzerViews(FName InAnalyzerName);
 			TArray<const FMetasoundAnalyzerView*> GetAnalyzerViews(FName InAnalyzerName) const;
 			TArray<FMetasoundAnalyzerView*> GetAnalyzerViewsForOutput(const FGuid& InNodeID, FName InOutputName, FName InAnalyzerName);
diff --git a/Engine/Plugins/Runtime/Metasound/Source/MetasoundFrontend/Public/Analysis/MetasoundFrontendVertexAnalyzerAudioBusWriter.h b/Engine/Plugins/Runtime/Metasound/Source/MetasoundFrontend/Public/Analysis/MetasoundFrontendVertexAnalyzerAudioBusWriter.h
new file mode 100644
--- /dev/null
+++ b/Engine/Plugins/Runtime/Metasound/Source/MetasoundFrontend/Public/Analysis/MetasoundFrontendVertexAnalyzerAudioBusWriter.h
@@ -0,0 +1,48 @@
+// Copyright Epic Games, Inc. All Rights Reserved.
+#pragma once
+
+#include "Analysis/MetasoundFrontendAnalyzerFactory.h"
+#include "Analysis/MetasoundFrontendVertexAnalyzer.h"
+#include "AudioDefines.h"
+#include "Containers/Array.h"
+
+namespace Metasound
+{
+	namespace Frontend
+	{
+		class METASOUNDFRONTEND_API FVertexAnalyzerAudioBusWriter : public FVertexAnalyzerBase
+		{
+		public:
+			static const FName& GetAnalyzerName();
+			static const FName& GetDataType();
+			static FName GetAnalyzerMemberName(const Audio::FDeviceId InDeviceID, const uint32 InAudioBusID);
+
+			class METASOUNDFRONTEND_API FFactory : public TVertexAnalyzerFactory<FVertexAnalyzerAudioBusWriter>
+			{
+			public:
+				virtual const TArray<FAnalyzerOutput>& GetAnalyzerOutputs() const override
+				{
+					static const TArray<FAnalyzerOutput> Outputs;
+					return Outputs;
+				}
+			};
+
+			FVertexAnalyzerAudioBusWriter(const FCreateAnalyzerParams& InParams);
+			virtual ~FVertexAnalyzerAudioBusWriter() = default;
+
+			virtual void Execute() override;
+
+		private:
+			struct FBusAddress
+			{
+				Audio::FDeviceId DeviceID = 0;
+				uint32 AudioBusID = 0;
+
+				FString ToString() const;
+				static FBusAddress FromString(const FString& InAnalyzerMemberName);
+			};
+
+			Audio::FPatchInput AudioBusPatchInput;
+		};
+	} // namespace Frontend
+} // namespace Metasound
