UE-171771 - MetaSound need to click canvas first and then spacebar to play graph
#rnx

diff --git a/Engine/Plugins/Runtime/Metasound/Source/MetasoundEditor/Private/MetasoundEditor.cpp b/Engine/Plugins/Runtime/Metasound/Source/MetasoundEditor/Private/MetasoundEditor.cpp
--- a/Engine/Plugins/Runtime/Metasound/Source/MetasoundEditor/Private/MetasoundEditor.cpp
+++ b/Engine/Plugins/Runtime/Metasound/Source/MetasoundEditor/Private/MetasoundEditor.cpp
@@ -964,6 +964,8 @@ namespace Metasound
 			{
 				Play();
 			}
+
+			FSlateApplication::Get().SetUserFocus(0, MetasoundGraphEditor);
 		}
 
 		UObject* FEditor::GetMetasoundObject() const
