Metasound Setting for are we using the old widgets, or the new AudioMaterialWidgets when possible in MetasoundEditor
#rnx

diff --git a/Engine/Plugins/Runtime/Metasound/Source/MetasoundEditor/Private/MetasoundEditor.cpp b/Engine/Plugins/Runtime/Metasound/Source/MetasoundEditor/Private/MetasoundEditor.cpp
--- a/Engine/Plugins/Runtime/Metasound/Source/MetasoundEditor/Private/MetasoundEditor.cpp
+++ b/Engine/Plugins/Runtime/Metasound/Source/MetasoundEditor/Private/MetasoundEditor.cpp
@@ -95,17 +95,6 @@ namespace Metasound
 {
 	namespace Editor
 	{
-		namespace EditorPrivate
-		{
-			int32 UseAudioMaterialWidgets = 0;
-			FAutoConsoleVariableRef CVarUseAudioMaterialWidgets(
-				TEXT("au.MetaSound.Editor.EnableAudioMaterialAnalyzers"),
-				UseAudioMaterialWidgets,
-				TEXT("Are new AudioMaterialWidgets used for visualization for the Metasound Analyzers, if implemented.\n")
-				TEXT("0: Disabled (default), !0: Enabled"),
-				ECVF_Default);
-		}//Metasound::Editor::EditorPrivate
-
 		static const TArray<FText> NodeSectionNames
 		{
 			LOCTEXT("NodeSectionName_Invalid", "INVALID"),
@@ -1267,8 +1256,12 @@ namespace Metasound
 
 					if (!OutputMeter.IsValid())
 					{
+						bool bUseAudioMaterialWidgets = false;
+						if (const UMetasoundEditorSettings* EditorSettings = GetDefault<UMetasoundEditorSettings>())
+						{
+							bUseAudioMaterialWidgets = EditorSettings->bUseAudioMaterialWidgets;
+						}
 
-						bool bUseAudioMaterialWidgets = EditorPrivate::UseAudioMaterialWidgets == 0 ? false : true;
 						OutputMeter = MakeShared<AudioWidgets::FAudioMeter>(MetaSoundSource->NumChannels, AudioDeviceId, nullptr, bUseAudioMaterialWidgets);
 					}
 					else if (OutputMeter->GetAudioBus()->GetNumChannels() != MetaSoundSource->NumChannels)
diff --git a/Engine/Plugins/Runtime/Metasound/Source/MetasoundEditor/Private/MetasoundEditorSettings.h b/Engine/Plugins/Runtime/Metasound/Source/MetasoundEditor/Private/MetasoundEditorSettings.h
--- a/Engine/Plugins/Runtime/Metasound/Source/MetasoundEditor/Private/MetasoundEditorSettings.h
+++ b/Engine/Plugins/Runtime/Metasound/Source/MetasoundEditor/Private/MetasoundEditorSettings.h
@@ -175,8 +175,13 @@ public:
 	/** Determines which details view to show in Metasounds Editor */
 	UPROPERTY(Transient)
 	EMetasoundActiveDetailView DetailView = EMetasoundActiveDetailView::General;
+
+	/** Whether the AudioMaterialWidgets are used when possible in Metasound Editor*/
+	UPROPERTY(EditAnywhere, config, DisplayName = "Use Audio Material Widgets", Category = MetasoundStyling)
+	bool bUseAudioMaterialWidgets = false;
 	
-	UPROPERTY(EditAnywhere, config, Category = MetasoundStyling, meta = (AllowedClasses = "/Script/SlateCore.SlateWidgetStyleAsset", DisplayName = "Knob Style"))
+	/**Override the Knob Style used in the Metasound Editor.*/
+	UPROPERTY(EditAnywhere, config, Category = MetasoundStyling, meta = (AllowedClasses = "/Script/SlateCore.SlateWidgetStyleAsset", EditCondition = "bUseAudioMaterialWidgets", DisplayName = "Knob Style"))
 	FSoftObjectPath KnobStyleOverride;
 	
 	/** Get the AudioMaterialKnob Style. if KnobStyleOverride is not set, returns default style.*/
diff --git a/Engine/Plugins/Runtime/Metasound/Source/MetasoundEditor/Private/SMetasoundGraphNode.cpp b/Engine/Plugins/Runtime/Metasound/Source/MetasoundEditor/Private/SMetasoundGraphNode.cpp
--- a/Engine/Plugins/Runtime/Metasound/Source/MetasoundEditor/Private/SMetasoundGraphNode.cpp
+++ b/Engine/Plugins/Runtime/Metasound/Source/MetasoundEditor/Private/SMetasoundGraphNode.cpp
@@ -64,17 +64,6 @@ namespace Metasound
 {
 	namespace Editor
 	{
-		namespace GraphNodePrivate
-		{
-			int32 UseAudioMaterialWidgets = 0;
-			FAutoConsoleVariableRef CVarUseAudioMaterialWidgets(
-				TEXT("au.MetaSound.Editor.UseAudioMaterialSlates"),
-				UseAudioMaterialWidgets,
-				TEXT("Are new AudioMaterialWidgets used for visualization in the Metasound editor, if implemented.\n")
-				TEXT("0: Disabled (default), !0: Enabled"),
-				ECVF_Default);
-		}// Metasound::Editor::GraphNodePrivate
-
 		SMetaSoundGraphNode::~SMetaSoundGraphNode()
 		{
 			// Clean up input widgets
@@ -707,6 +696,12 @@ namespace Metasound
 			TSharedPtr<SHorizontalBox> ContentBox = SNew(SHorizontalBox);
 			TSharedPtr<SWidget> OuterContentBox; // currently only used for input float nodes to accommodate the input widget
 
+			bool bUseAudioMaterialWidgets = false;
+			if (const UMetasoundEditorSettings* EditorSettings = GetDefault<UMetasoundEditorSettings>())
+			{
+				bUseAudioMaterialWidgets = EditorSettings->bUseAudioMaterialWidgets;
+			}
+
 			// If editable float input node and not constructor input, check if custom widget required
 			bool bShowContentWidget = false;
 			if (UMetasoundEditorGraphInput* GraphMember = Cast<UMetasoundEditorGraphInput>(GetMetaSoundMember()))
@@ -772,7 +767,7 @@ namespace Metasound
 
 						if (DefaultFloat->WidgetType == EMetasoundMemberDefaultWidget::Slider)
 						{
-							if (Metasound::Editor::GraphNodePrivate::UseAudioMaterialWidgets)
+							if (bUseAudioMaterialWidgets)
 							{
 								SAssignNew(InputWidget, SAudioMaterialLabeledSlider)
 									.Owner(GraphMember->GetOwningGraph())
@@ -861,7 +856,7 @@ namespace Metasound
 									InputWidget->SetDesiredSizeOverride(FVector2D(SliderDesiredSizeVertical.Y, SliderDesiredSizeVertical.X));
 								}
 
-								if (Metasound::Editor::GraphNodePrivate::UseAudioMaterialWidgets)
+								if (bUseAudioMaterialWidgets)
 								{
 									// safe downcast because the ptr was just assigned above 
 									StaticCastSharedPtr<SAudioMaterialLabeledSlider>(InputWidget)->SetOrientation(DefaultFloat->WidgetOrientation);
@@ -896,7 +891,7 @@ namespace Metasound
 								}
 							};
 
-							if (Metasound::Editor::GraphNodePrivate::UseAudioMaterialWidgets)
+							if (bUseAudioMaterialWidgets)
 							{
 								const UMetasoundEditorSettings* MetasoundSettings = GetDefault<UMetasoundEditorSettings>();
 								check(MetasoundSettings)
@@ -999,54 +994,57 @@ namespace Metasound
 							}
 						});
 					}
-					else if (UMetasoundEditorGraphMemberDefaultBool* DefaultBool = Cast<UMetasoundEditorGraphMemberDefaultBool>(GraphMember->GetLiteral()); Metasound::Editor::GraphNodePrivate::UseAudioMaterialWidgets)
+					else if (UMetasoundEditorGraphMemberDefaultBool* DefaultBool = Cast<UMetasoundEditorGraphMemberDefaultBool>(GraphMember->GetLiteral()))
 					{
-						bool bIsNotTriggerNode = GraphMember->GetDataType() != Metasound::GetMetasoundDataTypeName<Metasound::FTrigger>();
-
-						if ((bIsNotTriggerNode && IsValid(DefaultBool)) && DefaultBool->WidgetType != EMetasoundBoolMemberDefaultWidget::None)
+						if (bUseAudioMaterialWidgets)
 						{
-							constexpr float WidgetPadding = 3.0f;
-							static const FVector2D ButtonDesiredSize = FVector2D(56.0f, 87.0f);
+							bool bIsNotTriggerNode = GraphMember->GetDataType() != Metasound::GetMetasoundDataTypeName<Metasound::FTrigger>();
 
-							auto OnboolValueChangedLambda = [DefaultBool, GraphMember, this](bool Value)
-								{
-									if (MaterialButtonWidget.IsValid())
+							if ((bIsNotTriggerNode && IsValid(DefaultBool)) && DefaultBool->WidgetType != EMetasoundBoolMemberDefaultWidget::None)
+							{
+								constexpr float WidgetPadding = 3.0f;
+								static const FVector2D ButtonDesiredSize = FVector2D(56.0f, 87.0f);
+
+								auto OnboolValueChangedLambda = [DefaultBool, GraphMember, this](bool Value)
 									{
-										if (!bIsInputWidgetTransacting)
+										if (MaterialButtonWidget.IsValid())
 										{
-											GEditor->BeginTransaction(LOCTEXT("MetasoundGraphNode_MetasoundSetInputDefault", "Set MetaSound Input Default"));
-											bIsInputWidgetTransacting = true;
+											if (!bIsInputWidgetTransacting)
+											{
+												GEditor->BeginTransaction(LOCTEXT("MetasoundGraphNode_MetasoundSetInputDefault", "Set MetaSound Input Default"));
+												bIsInputWidgetTransacting = true;
+											}
+											GraphMember->GetOwningGraph()->GetMetasound()->Modify();
+											DefaultBool->Modify();
+
+											constexpr bool bPostTransaction = true;
+											DefaultBool->SetDefault(Value);
+											GraphMember->UpdateFrontendDefaultLiteral(bPostTransaction);
 										}
-										GraphMember->GetOwningGraph()->GetMetasound()->Modify();
-										DefaultBool->Modify();
+									};
 
-										constexpr bool bPostTransaction = true;
-										DefaultBool->SetDefault(Value);
-										GraphMember->UpdateFrontendDefaultLiteral(bPostTransaction);
-									}
-								};
-
-							bShowContentWidget = true;
-							SAssignNew(OuterContentBox, SVerticalBox)
-								+ SVerticalBox::Slot()
-								.HAlign(HAlign_Right)
-								.VAlign(VAlign_Center)
-								.AutoHeight()
-								[
-									ContentBox.ToSharedRef()
-								]
-								+ SVerticalBox::Slot()
-								.HAlign(HAlign_Fill)
-								.VAlign(VAlign_Top)
-								.Padding(WidgetPadding, 0.0f, WidgetPadding, WidgetPadding)
-								.AutoHeight()
-								[
-									SAssignNew(MaterialButtonWidget, SAudioMaterialButton)
-										.OnBooleanValueChanged_Lambda(OnboolValueChangedLambda)
-								];
+								bShowContentWidget = true;
+								SAssignNew(OuterContentBox, SVerticalBox)
+									+ SVerticalBox::Slot()
+									.HAlign(HAlign_Right)
+									.VAlign(VAlign_Center)
+									.AutoHeight()
+									[
+										ContentBox.ToSharedRef()
+									]
+									+ SVerticalBox::Slot()
+									.HAlign(HAlign_Fill)
+									.VAlign(VAlign_Top)
+									.Padding(WidgetPadding, 0.0f, WidgetPadding, WidgetPadding)
+									.AutoHeight()
+									[
+										SAssignNew(MaterialButtonWidget, SAudioMaterialButton)
+											.OnBooleanValueChanged_Lambda(OnboolValueChangedLambda)
+									];
 
-							MaterialButtonWidget->SetPressedState(DefaultBool->GetDefault());
-							MaterialButtonWidget->SetDesiredSizeOverride(ButtonDesiredSize);
+								MaterialButtonWidget->SetPressedState(DefaultBool->GetDefault());
+								MaterialButtonWidget->SetDesiredSizeOverride(ButtonDesiredSize);
+							}
 						}
 					}
 				}
