UE-186195 - Add console command to mute/solo MetaSounds
#rnx

diff --git a/Engine/Plugins/Runtime/Metasound/Source/MetasoundEditor/Private/MetasoundWaveConsole.cpp b/Engine/Plugins/Runtime/Metasound/Source/MetasoundEditor/Private/MetasoundWaveConsole.cpp
new file mode 100644
--- /dev/null
+++ b/Engine/Plugins/Runtime/Metasound/Source/MetasoundEditor/Private/MetasoundWaveConsole.cpp
@@ -0,0 +1,123 @@
+// Copyright Epic Games, Inc. All Rights Reserved.
+
+#include "AudioDevice.h"
+#include "AudioDeviceManager.h"
+#include "Audio/AudioDebug.h"
+#include "MetasoundSource.h"
+
+namespace Metasound
+{
+	namespace Console
+	{
+		static void HandleSoloMetaSoundWave(const TArray<FString>& InArgs, UWorld* World)
+		{
+			check(GEngine);
+			FAudioDeviceManager* DeviceManager = GEngine->GetAudioDeviceManager();
+			if (!DeviceManager)
+			{
+				return;
+			}
+
+			FAudioDeviceHandle AudioDevice = DeviceManager->GetActiveAudioDevice();
+			if (!AudioDevice)
+			{
+				return;
+			}
+
+			TArray<FString> MetaSoundWaves;
+
+			const int32 WorldID = World->GetUniqueID();
+			const TArray<FActiveSound*>& ActiveSounds = AudioDevice->GetActiveSounds();
+			for (const FActiveSound* const& ActiveSound : ActiveSounds)
+			{
+				if (WorldID == ActiveSound->GetWorldID())
+				{
+					// Make sure that sound asset is MetaSound based
+					if (Cast<UMetaSoundSource>(ActiveSound->GetSound()))
+					{
+						for (const TPair<UPTRINT, FWaveInstance*>& WaveInstancePair : ActiveSound->GetWaveInstances())
+						{
+							const FWaveInstance* WaveInstance = WaveInstancePair.Value;
+							MetaSoundWaves.Add(WaveInstance->GetName());
+						}
+					}
+				}
+			}
+
+			if (InArgs.Num() == 1)
+			{
+				for (const FString& Wave : MetaSoundWaves)
+				{
+					// we can't use FAudioDebugger::SetSoloSoundWave(...) because we don't want to mute non MetaSounds
+					const bool bMute = !InArgs[0].Equals(Wave);
+					DeviceManager->GetDebugger().SetMuteSoundWave(*Wave, bMute);
+				}
+			}
+			else
+			{
+				UE_LOG(LogMetaSound, Error, TEXT("You can solo ONLY ONE MetaSound wave!"));
+			}
+		}
+
+		static void HandleMuteMetaSoundWave(const TArray<FString>& InArgs, UWorld* World)
+		{
+			check(GEngine);
+			FAudioDeviceManager* DeviceManager = GEngine->GetAudioDeviceManager();
+			if (!DeviceManager)
+			{
+				return;
+			}
+
+			FAudioDeviceHandle AudioDevice = DeviceManager->GetActiveAudioDevice();
+			if (!AudioDevice)
+			{
+				return;
+			}
+
+			TArray<FString> MetaSoundWaves;
+
+			const int32 WorldID = World->GetUniqueID();
+			const TArray<FActiveSound*>& ActiveSounds = AudioDevice->GetActiveSounds();
+			for (const FActiveSound* const& ActiveSound : ActiveSounds)
+			{
+				if (WorldID == ActiveSound->GetWorldID())
+				{
+					// Make sure that sound asset is MetaSound based
+					if (Cast<UMetaSoundSource>(ActiveSound->GetSound()))
+					{
+						for (const TPair<UPTRINT, FWaveInstance*>& WaveInstancePair : ActiveSound->GetWaveInstances())
+						{
+							const FWaveInstance* WaveInstance = WaveInstancePair.Value;
+							MetaSoundWaves.Add(WaveInstance->GetName());
+						}
+					}
+				}
+			}
+
+			for (const FString& Arg : InArgs)
+			{
+				for (const FString& Wave : MetaSoundWaves)
+				{
+					if (Arg.Equals(Wave))
+					{
+						DeviceManager->GetDebugger().SetMuteSoundWave(*Wave, true);
+					}
+				}
+			}
+		}
+	} // namespace Console
+} // namespace Metasound
+
+static FAutoConsoleCommandWithWorldAndArgs SoloMetaSoundWave
+(
+	TEXT("au.MetaSound.SoloMetaSound"),
+	TEXT("Mutes all other MetaSound waves. Only the first argument is accepted."),
+	FConsoleCommandWithWorldAndArgsDelegate::CreateStatic(&Metasound::Console::HandleSoloMetaSoundWave)
+);
+
+static FAutoConsoleCommandWithWorldAndArgs MuteMetaSoundWave
+(
+	TEXT("au.MetaSound.MuteMetaSoundWave"),
+	TEXT("Mutes all given MetaSound waves."),
+	FConsoleCommandWithWorldAndArgsDelegate::CreateStatic(&Metasound::Console::HandleMuteMetaSoundWave)
+);
