#rnc

diff --git a/Engine/Plugins/Runtime/Metasound/Source/MetasoundEditor/Private/MetasoundWaveConsole.cpp b/Engine/Plugins/Runtime/Metasound/Source/MetasoundEditor/Private/MetasoundWaveConsole.cpp
--- a/Engine/Plugins/Runtime/Metasound/Source/MetasoundEditor/Private/MetasoundWaveConsole.cpp
+++ b/Engine/Plugins/Runtime/Metasound/Source/MetasoundEditor/Private/MetasoundWaveConsole.cpp
@@ -5,6 +5,8 @@
 #include "Audio/AudioDebug.h"
 #include "MetasoundSource.h"
 
+#if ENABLE_AUDIO_DEBUG
+
 namespace Metasound
 {
 	namespace Console
@@ -121,3 +123,5 @@ static FAutoConsoleCommandWithWorldAndArgs MuteMetaSoundWave
 	TEXT("Mutes all given MetaSound waves."),
 	FConsoleCommandWithWorldAndArgsDelegate::CreateStatic(&Metasound::Console::HandleMuteMetaSoundWave)
 );
+
+#endif // ENABLE_AUDIO_DEBUG
