Fix for Something is preventing Metasound editor to save properly when AudioMaterialBool button is used.
#rnx

diff --git a/Engine/Plugins/Runtime/AudioWidgets/Source/AudioWidgets/Private/AudioMaterialSlate/SAudioMaterialButton.cpp b/Engine/Plugins/Runtime/AudioWidgets/Source/AudioWidgets/Private/AudioMaterialSlate/SAudioMaterialButton.cpp
--- a/Engine/Plugins/Runtime/AudioWidgets/Source/AudioWidgets/Private/AudioMaterialSlate/SAudioMaterialButton.cpp
+++ b/Engine/Plugins/Runtime/AudioWidgets/Source/AudioWidgets/Private/AudioMaterialSlate/SAudioMaterialButton.cpp
@@ -12,7 +12,8 @@ void SAudioMaterialButton::Construct(const FArguments& InArgs)
 	Owner = InArgs._Owner;
 	AudioMaterialButtonStyle = InArgs._AudioMaterialButtonStyle;
 	bIsPressedAttribute = InArgs._bIsPressedAttribute;
-	OnIsPressedStateChanged = InArgs._OnBooleanValueChanged;
+	OnBooleanValueChanged = InArgs._OnBooleanValueChanged;
+	OnMouseCaptureEnd = InArgs._OnMouseCaptureEnd;
 
 	ApplyNewMaterial();
 
@@ -115,6 +116,7 @@ FReply SAudioMaterialButton::OnMouseButtonUp(const FGeometry& MyGeometry, const
 {
 	if ((MouseEvent.GetEffectingButton() == EKeys::LeftMouseButton) && this->HasMouseCapture())
 	{
+		OnMouseCaptureEnd.ExecuteIfBound();
 		return FReply::Handled().ReleaseMouseCapture();
 	}
 
@@ -123,6 +125,13 @@ FReply SAudioMaterialButton::OnMouseButtonUp(const FGeometry& MyGeometry, const
 
 void SAudioMaterialButton::CommitNewState(bool InPressedState)
 {
-	bIsPressedAttribute.Set(InPressedState);
-	OnIsPressedStateChanged.ExecuteIfBound(InPressedState);
+	if (bIsPressedAttribute.Get() != InPressedState)
+	{
+		if (!bIsPressedAttribute.IsBound())
+		{
+			bIsPressedAttribute.Set(InPressedState);
+		}
+		Invalidate(EInvalidateWidgetReason::Paint);
+		OnBooleanValueChanged.ExecuteIfBound(InPressedState);
+	}
 }
diff --git a/Engine/Plugins/Runtime/AudioWidgets/Source/AudioWidgets/Private/AudioMaterialSlate/SAudioMaterialKnob.cpp b/Engine/Plugins/Runtime/AudioWidgets/Source/AudioWidgets/Private/AudioMaterialSlate/SAudioMaterialKnob.cpp
--- a/Engine/Plugins/Runtime/AudioWidgets/Source/AudioWidgets/Private/AudioMaterialSlate/SAudioMaterialKnob.cpp
+++ b/Engine/Plugins/Runtime/AudioWidgets/Source/AudioWidgets/Private/AudioMaterialSlate/SAudioMaterialKnob.cpp
@@ -3,7 +3,6 @@
 
 #include "AudioMaterialSlate/SAudioMaterialKnob.h"
 #include "AudioMaterialSlate/AudioMaterialSlateTypes.h"
-#include "AudioWidgetsStyle.h"
 #include "Framework/Application/SlateApplication.h"
 #include "SlateOptMacros.h"
 #include "Components/Widget.h"
@@ -26,12 +25,6 @@ void SAudioMaterialKnob::Construct(const FArguments& InArgs)
 	OnMouseCaptureBegin = InArgs._OnMouseCaptureBegin;
 	OnMouseCaptureEnd = InArgs._OnMouseCaptureEnd;
 
-	//For now check if owner is not a widget -> use the default style.
-	if (!Cast<UWidget>(Owner))
-	{
-		AudioMaterialKnobStyle = &FAudioWidgetsStyle::Get().GetWidgetStyle<FAudioMaterialKnobStyle>("AudioMaterialKnob.Style");
-	}
-
 	ApplyNewMaterial();
 
 	if (InArgs._Value.IsSet())
@@ -41,9 +34,9 @@ void SAudioMaterialKnob::Construct(const FArguments& InArgs)
 }
 END_SLATE_FUNCTION_BUILD_OPTIMIZATION
 
-void SAudioMaterialKnob::SetValue(float InValue)
+void SAudioMaterialKnob::SetValue(const TAttribute<float>& InValueAttribute)
 {
-	CommitValue(InValue);
+	SetAttribute(ValueAttribute, InValueAttribute, EInvalidateWidgetReason::Paint);
 }
 
 void SAudioMaterialKnob::SetTuneSpeed(const float InTurnSpeed)
diff --git a/Engine/Plugins/Runtime/AudioWidgets/Source/AudioWidgets/Public/AudioMaterialSlate/SAudioMaterialButton.h b/Engine/Plugins/Runtime/AudioWidgets/Source/AudioWidgets/Public/AudioMaterialSlate/SAudioMaterialButton.h
--- a/Engine/Plugins/Runtime/AudioWidgets/Source/AudioWidgets/Public/AudioMaterialSlate/SAudioMaterialButton.h
+++ b/Engine/Plugins/Runtime/AudioWidgets/Source/AudioWidgets/Public/AudioMaterialSlate/SAudioMaterialButton.h
@@ -35,6 +35,9 @@ public:
 	/** Called when the button's state changes. */
 	SLATE_EVENT(FOnBooleanValueChanged, OnBooleanValueChanged)
 
+	/** Invoked when the mouse is released and a capture ends. */
+	SLATE_EVENT(FSimpleDelegate, OnMouseCaptureEnd)
+
 	SLATE_END_ARGS()
 
 	/** Constructs this widget with InArgs */
@@ -51,7 +54,10 @@ public:
 
 public:
 
-	FOnBooleanValueChanged OnIsPressedStateChanged;
+	FOnBooleanValueChanged OnBooleanValueChanged;
+
+	// Holds a delegate that is executed when the mouse is let up and a capture ends.
+	FSimpleDelegate OnMouseCaptureEnd;
 
 protected:
 
diff --git a/Engine/Plugins/Runtime/AudioWidgets/Source/AudioWidgets/Public/AudioMaterialSlate/SAudioMaterialKnob.h b/Engine/Plugins/Runtime/AudioWidgets/Source/AudioWidgets/Public/AudioMaterialSlate/SAudioMaterialKnob.h
--- a/Engine/Plugins/Runtime/AudioWidgets/Source/AudioWidgets/Public/AudioMaterialSlate/SAudioMaterialKnob.h
+++ b/Engine/Plugins/Runtime/AudioWidgets/Source/AudioWidgets/Public/AudioMaterialSlate/SAudioMaterialKnob.h
@@ -2,7 +2,8 @@
 
 #pragma once
 
-#include "AudioMaterialSlate/AudioMaterialSlateTypes.h"
+#include "AudioMaterialSlateTypes.h"
+#include "AudioWidgetsStyle.h"
 #include "Framework/SlateDelegates.h"
 #include "Styling/ISlateStyle.h"
 #include "Styling/SlateWidgetStyleAsset.h"
@@ -24,6 +25,7 @@ public:
 	, _Locked(false)
 	, _MouseUsesStep(false)
 	, _StepSize(0.01f)
+	, _AudioMaterialKnobStyle(&FAudioWidgetsStyle::Get().GetWidgetStyle<FAudioMaterialKnobStyle>("AudioMaterialKnob.Style"))
 	{}
 
 	/** The owner object*/
@@ -68,7 +70,7 @@ public:
 	void Construct(const FArguments& InArgs);
 
 	/** Set the Value attribute */
-	void SetValue(float InValueAttribute);
+	void SetValue(const TAttribute<float>& InValueAttribute);
 
 	/** Set the TuneSpeed attribute */
 	void SetTuneSpeed(const float InMouseSpeed);
diff --git a/Engine/Plugins/Runtime/Metasound/Source/MetasoundEditor/Private/MetasoundEditorGraphMemberDefaults.cpp b/Engine/Plugins/Runtime/Metasound/Source/MetasoundEditor/Private/MetasoundEditorGraphMemberDefaults.cpp
--- a/Engine/Plugins/Runtime/Metasound/Source/MetasoundEditor/Private/MetasoundEditorGraphMemberDefaults.cpp
+++ b/Engine/Plugins/Runtime/Metasound/Source/MetasoundEditor/Private/MetasoundEditorGraphMemberDefaults.cpp
@@ -72,7 +72,11 @@ void UMetasoundEditorGraphMemberDefaultBool::UpdatePreviewInstance(const Metasou
 
 void UMetasoundEditorGraphMemberDefaultBool::SetDefault(const bool InDefault)
 {
-	Default.Value = InDefault;
+	if (Default.Value != InDefault)
+	{
+		Default.Value = InDefault;
+		OnDefaultStateChanged.Broadcast(InDefault);
+	}
 }
 
 const bool UMetasoundEditorGraphMemberDefaultBool::GetDefault()
diff --git a/Engine/Plugins/Runtime/Metasound/Source/MetasoundEditor/Private/SMetasoundGraphNode.cpp b/Engine/Plugins/Runtime/Metasound/Source/MetasoundEditor/Private/SMetasoundGraphNode.cpp
--- a/Engine/Plugins/Runtime/Metasound/Source/MetasoundEditor/Private/SMetasoundGraphNode.cpp
+++ b/Engine/Plugins/Runtime/Metasound/Source/MetasoundEditor/Private/SMetasoundGraphNode.cpp
@@ -72,22 +72,28 @@ namespace Metasound
 			{
 				if (UMetasoundEditorGraphMember* GraphMember = MemberNode->GetMember())
 				{
-					if (UMetasoundEditorGraphMemberDefaultFloat* DefaultFloat = Cast<UMetasoundEditorGraphMemberDefaultFloat>(GraphMember->GetLiteral()))
+					// This may hit if the asset editor is closed while interacting with a widget 
+					// (ex. Ctrl-W is pressed mid drag before the value is committed)
+					if (bIsInputWidgetTransacting)
 					{
-						// This may hit if the asset editor is closed while interacting with a widget 
-						// (ex. Ctrl-W is pressed mid drag before the value is committed)
-						if (bIsInputWidgetTransacting)
+						GEditor->EndTransaction();
+						if (UMetasoundEditorGraph* Graph = GraphMember->GetOwningGraph())
 						{
-							GEditor->EndTransaction();
-							if (UMetasoundEditorGraph* Graph = GraphMember->GetOwningGraph())
-							{
-								constexpr bool bPostTransaction = false;
-								GraphMember->UpdateFrontendDefaultLiteral(bPostTransaction);
-								FGraphBuilder::GetOutermostMetaSoundChecked(*Graph).GetModifyContext().AddMemberIDsModified({ GraphMember->GetMemberID() });
-							}
+							constexpr bool bPostTransaction = false;
+							GraphMember->UpdateFrontendDefaultLiteral(bPostTransaction);
+							FGraphBuilder::GetOutermostMetaSoundChecked(*Graph).GetModifyContext().AddMemberIDsModified({ GraphMember->GetMemberID() });
 						}
+					}
+
+					if (UMetasoundEditorGraphMemberDefaultFloat* DefaultFloat = Cast<UMetasoundEditorGraphMemberDefaultFloat>(GraphMember->GetLiteral()))
+					{
 						DefaultFloat->OnDefaultValueChanged.Remove(InputSliderOnValueChangedDelegateHandle);
 						DefaultFloat->OnRangeChanged.Remove(InputSliderOnRangeChangedDelegateHandle);
+
+					}
+					else if (UMetasoundEditorGraphMemberDefaultBool* DefaultBool = Cast<UMetasoundEditorGraphMemberDefaultBool>(GraphMember->GetLiteral()))
+					{
+						DefaultBool->OnDefaultStateChanged.Remove(InputButtonOnStateChangedDelegateHandle);
 					}
 				}
 			}
@@ -1002,6 +1008,7 @@ namespace Metasound
 
 							if ((bIsNotTriggerNode && IsValid(DefaultBool)) && DefaultBool->WidgetType != EMetasoundBoolMemberDefaultWidget::None)
 							{
+								bShowContentWidget = true;
 								constexpr float WidgetPadding = 3.0f;
 								static const FVector2D ButtonDesiredSize = FVector2D(56.0f, 87.0f);
 
@@ -1023,7 +1030,19 @@ namespace Metasound
 										}
 									};
 
-								bShowContentWidget = true;
+								auto OnboolMouseCaptureEndLambda = [this]()
+									{
+										if (bIsInputWidgetTransacting)
+										{
+											GEditor->EndTransaction();
+											bIsInputWidgetTransacting = false;
+										}
+										else
+										{
+											UE_LOG(LogMetaSound, Warning, TEXT("Unmatched MetaSound editor widget transaction."));
+										}
+									};
+
 								SAssignNew(OuterContentBox, SVerticalBox)
 									+ SVerticalBox::Slot()
 									.HAlign(HAlign_Right)
@@ -1040,10 +1059,27 @@ namespace Metasound
 									[
 										SAssignNew(MaterialButtonWidget, SAudioMaterialButton)
 											.OnBooleanValueChanged_Lambda(OnboolValueChangedLambda)
+											.OnMouseCaptureEnd_Lambda(OnboolMouseCaptureEndLambda)
+											.bIsPressedAttribute(DefaultBool->GetDefault())
 									];
 
-								MaterialButtonWidget->SetPressedState(DefaultBool->GetDefault());
 								MaterialButtonWidget->SetDesiredSizeOverride(ButtonDesiredSize);
+
+
+								// Setup & clear delegate if necessary (ex. if was just saved)
+								if (InputButtonOnStateChangedDelegateHandle.IsValid())
+								{
+									DefaultBool->OnDefaultStateChanged.Remove(InputButtonOnStateChangedDelegateHandle);
+									InputButtonOnStateChangedDelegateHandle.Reset();
+								}
+
+								InputButtonOnStateChangedDelegateHandle = DefaultBool->OnDefaultStateChanged.AddLambda([Widget = MaterialButtonWidget](bool Value)
+									{
+										if (Widget.IsValid())
+										{
+											Widget->SetPressedState(Value);
+										}
+									});
 							}
 						}
 					}
diff --git a/Engine/Plugins/Runtime/Metasound/Source/MetasoundEditor/Private/SMetasoundGraphNode.h b/Engine/Plugins/Runtime/Metasound/Source/MetasoundEditor/Private/SMetasoundGraphNode.h
--- a/Engine/Plugins/Runtime/Metasound/Source/MetasoundEditor/Private/SMetasoundGraphNode.h
+++ b/Engine/Plugins/Runtime/Metasound/Source/MetasoundEditor/Private/SMetasoundGraphNode.h
@@ -82,7 +82,10 @@ namespace Metasound::Editor
 		TSharedPtr<SAudioInputWidget> InputWidget;
 
 		//Button Widget for bool input.
-		TSharedPtr <SAudioMaterialButton> MaterialButtonWidget;
+		TSharedPtr<SAudioMaterialButton> MaterialButtonWidget;
+
+		// Handle for on state changed delegate for Button 
+		FDelegateHandle InputButtonOnStateChangedDelegateHandle;
 
 		// Handle for on value changed delegate for input slider 
 		FDelegateHandle InputSliderOnValueChangedDelegateHandle;
diff --git a/Engine/Plugins/Runtime/Metasound/Source/MetasoundEditor/Public/MetasoundEditorGraphMemberDefaults.h b/Engine/Plugins/Runtime/Metasound/Source/MetasoundEditor/Public/MetasoundEditorGraphMemberDefaults.h
--- a/Engine/Plugins/Runtime/Metasound/Source/MetasoundEditor/Public/MetasoundEditorGraphMemberDefaults.h
+++ b/Engine/Plugins/Runtime/Metasound/Source/MetasoundEditor/Public/MetasoundEditorGraphMemberDefaults.h
@@ -39,6 +39,9 @@ struct FMetasoundEditorGraphMemberDefaultBoolRef
 	bool Value = false;
 };
 
+// For bool input widget
+DECLARE_MULTICAST_DELEGATE_OneParam(FOnMetasoundBoolStateChangedEvent, bool);
+
 UCLASS(MinimalAPI)
 class UMetasoundEditorGraphMemberDefaultBool : public UMetasoundEditorGraphMemberDefaultLiteral
 {
@@ -53,6 +56,8 @@ public:
 	UPROPERTY(EditAnywhere, Category = Widget, meta = (DisplayName = "Widget"))
 	EMetasoundBoolMemberDefaultWidget WidgetType = EMetasoundBoolMemberDefaultWidget::None;
 
+	FOnMetasoundBoolStateChangedEvent OnDefaultStateChanged;
+
 	virtual FMetasoundFrontendLiteral GetDefault() const override;
 	virtual EMetasoundFrontendLiteralType GetLiteralType() const override;
 	virtual void SetFromLiteral(const FMetasoundFrontendLiteral& InLiteral) override;
