Remove dependency on MetaSound engine subsystems of any kind being required to register MetaSounds with frontend
#rnx

diff --git a/Engine/Plugins/Runtime/Metasound/Source/MetasoundEngine/Private/Metasound.cpp b/Engine/Plugins/Runtime/Metasound/Source/MetasoundEngine/Private/Metasound.cpp
--- a/Engine/Plugins/Runtime/Metasound/Source/MetasoundEngine/Private/Metasound.cpp
+++ b/Engine/Plugins/Runtime/Metasound/Source/MetasoundEngine/Private/Metasound.cpp
@@ -214,7 +214,7 @@ void UMetaSoundPatch::OnAsyncReferencedAssetsLoaded(const TArray<FMetasoundAsset
 	Metasound::FMetaSoundEngineAssetHelper::OnAsyncReferencedAssetsLoaded(*this, InAsyncReferences);
 }
 
-bool UMetaSoundPatch::IsBuilderActive() const
+bool UMetaSoundPatch::IsActivelyBuilding() const
 {
 	return bIsBuilderActive;
 }
diff --git a/Engine/Plugins/Runtime/Metasound/Source/MetasoundEngine/Private/MetasoundSource.cpp b/Engine/Plugins/Runtime/Metasound/Source/MetasoundEngine/Private/MetasoundSource.cpp
--- a/Engine/Plugins/Runtime/Metasound/Source/MetasoundEngine/Private/MetasoundSource.cpp
+++ b/Engine/Plugins/Runtime/Metasound/Source/MetasoundEngine/Private/MetasoundSource.cpp
@@ -1578,7 +1578,7 @@ Metasound::SourcePrivate::FParameterRouter& UMetaSoundSource::GetParameterRouter
 	return Router;
 }
 
-bool UMetaSoundSource::IsBuilderActive() const
+bool UMetaSoundSource::IsActivelyBuilding() const
 {
 	return bIsBuilderActive;
 }
diff --git a/Engine/Plugins/Runtime/Metasound/Source/MetasoundEngine/Public/Metasound.h b/Engine/Plugins/Runtime/Metasound/Source/MetasoundEngine/Public/Metasound.h
--- a/Engine/Plugins/Runtime/Metasound/Source/MetasoundEngine/Public/Metasound.h
+++ b/Engine/Plugins/Runtime/Metasound/Source/MetasoundEngine/Public/Metasound.h
@@ -178,7 +178,7 @@ public:
 		return this;
 	}
 
-	virtual bool IsBuilderActive() const override;
+	virtual bool IsActivelyBuilding() const override;
 
 protected:
 #if WITH_EDITOR
diff --git a/Engine/Plugins/Runtime/Metasound/Source/MetasoundEngine/Public/MetasoundSource.h b/Engine/Plugins/Runtime/Metasound/Source/MetasoundEngine/Public/MetasoundSource.h
--- a/Engine/Plugins/Runtime/Metasound/Source/MetasoundEngine/Public/MetasoundSource.h
+++ b/Engine/Plugins/Runtime/Metasound/Source/MetasoundEngine/Public/MetasoundSource.h
@@ -273,7 +273,7 @@ public:
 	Metasound::FOperatorSettings GetOperatorSettings(Metasound::FSampleRate InDeviceSampleRate) const;
 
 	virtual const FMetasoundFrontendDocument& GetConstDocument() const override;
-	virtual bool IsBuilderActive() const override;
+	virtual bool IsActivelyBuilding() const override;
 
 protected:
 	Metasound::Frontend::FDocumentAccessPtr GetDocumentAccessPtr() override;
diff --git a/Engine/Plugins/Runtime/Metasound/Source/MetasoundFrontend/Private/MetasoundAssetBase.cpp b/Engine/Plugins/Runtime/Metasound/Source/MetasoundFrontend/Private/MetasoundAssetBase.cpp
--- a/Engine/Plugins/Runtime/Metasound/Source/MetasoundFrontend/Private/MetasoundAssetBase.cpp
+++ b/Engine/Plugins/Runtime/Metasound/Source/MetasoundFrontend/Private/MetasoundAssetBase.cpp
@@ -189,17 +189,6 @@ namespace Metasound
 			}
 		} // namespace AssetBasePrivate
 
-		namespace ConsoleVariables
-		{
-			static bool bDisableAsyncGraphRegistration = false;
-		}
-
-		FAutoConsoleVariableRef CVarMetaSoundDisableAsyncGraphRegistration(
-			TEXT("au.MetaSound.DisableAsyncGraphRegistration"),
-			Metasound::Frontend::ConsoleVariables::bDisableAsyncGraphRegistration,
-			TEXT("Disables async registration of MetaSound graphs\n")
-			TEXT("Default: false"),
-			ECVF_Default);
 		FConsoleVariableMulticastDelegate CVarMetaSoundBlockRateChanged;
 
 		FAutoConsoleVariableRef CVarMetaSoundBlockRate(
@@ -336,10 +325,7 @@ void FMetasoundAssetBase::RegisterGraphWithFrontend(Metasound::Frontend::FMetaSo
 	check(Owner);
 	const FString AssetName = Owner->GetName();
 
-	// Register graphs async by default;
-	const bool bAsync = !ConsoleVariables::bDisableAsyncGraphRegistration;
-
-	GraphRegistryKey = FRegistryContainerImpl::Get().RegisterGraph(Owner, bAsync);
+	GraphRegistryKey = FRegistryContainerImpl::Get().RegisterGraph(Owner);
 	if (GraphRegistryKey.IsValid())
 	{
 #if WITH_EDITORONLY_DATA
@@ -462,13 +448,7 @@ void FMetasoundAssetBase::UnregisterGraphWithFrontend()
 		UObject* OwningAsset = GetOwningAsset();
 		if (ensureAlways(OwningAsset))
 		{
-			// Async registration is only available if:
-			// 1. The IMetaSoundDocumentInterface is not actively modified by a builder
-			//    (built graph must be released synchronously to avoid a race condition on
-			//    reading/writing the IMetaSoundDocumentInterface on the Game Thread)
-			// 2. Async registration is globally disabled via console variable.
-			const bool bAsync = !(IsBuilderActive() || ConsoleVariables::bDisableAsyncGraphRegistration);
-			const bool bSuccess = FRegistryContainerImpl::Get().UnregisterGraph(GraphRegistryKey, OwningAsset, bAsync);
+			const bool bSuccess = FRegistryContainerImpl::Get().UnregisterGraph(GraphRegistryKey, OwningAsset);
 			if (!bSuccess)
 			{
 				UE_LOG(LogMetaSound, Verbose, TEXT("Failed to unregister node with key %s for asset %s. No registry entry exists with that key."), *GraphRegistryKey.ToString(), *GetOwningAssetName());
diff --git a/Engine/Plugins/Runtime/Metasound/Source/MetasoundFrontend/Private/MetasoundFrontendDocumentBuilder.cpp b/Engine/Plugins/Runtime/Metasound/Source/MetasoundFrontend/Private/MetasoundFrontendDocumentBuilder.cpp
--- a/Engine/Plugins/Runtime/Metasound/Source/MetasoundFrontend/Private/MetasoundFrontendDocumentBuilder.cpp
+++ b/Engine/Plugins/Runtime/Metasound/Source/MetasoundFrontend/Private/MetasoundFrontendDocumentBuilder.cpp
@@ -670,6 +670,11 @@ const UClass& UMetaSoundBuilderDocument::GetBaseMetaSoundUClass() const
 	return *MetaSoundUClass;
 }
 
+bool UMetaSoundBuilderDocument::IsActivelyBuilding() const
+{
+	return true;
+}
+
 FMetasoundFrontendDocument& UMetaSoundBuilderDocument::GetDocument()
 {
 	return Document;
diff --git a/Engine/Plugins/Runtime/Metasound/Source/MetasoundFrontend/Private/MetasoundFrontendRegistryContainerImpl.cpp b/Engine/Plugins/Runtime/Metasound/Source/MetasoundFrontend/Private/MetasoundFrontendRegistryContainerImpl.cpp
--- a/Engine/Plugins/Runtime/Metasound/Source/MetasoundFrontend/Private/MetasoundFrontendRegistryContainerImpl.cpp
+++ b/Engine/Plugins/Runtime/Metasound/Source/MetasoundFrontend/Private/MetasoundFrontendRegistryContainerImpl.cpp
@@ -41,6 +41,17 @@ static FAutoConsoleVariableRef CVarAsyncRegistrationTasksBusyWait(
 
 namespace Metasound::Frontend
 {
+	namespace ConsoleVariables
+	{
+		static bool bDisableAsyncGraphRegistration = false;
+		FAutoConsoleVariableRef CVarMetaSoundDisableAsyncGraphRegistration(
+			TEXT("au.MetaSound.DisableAsyncGraphRegistration"),
+			Metasound::Frontend::ConsoleVariables::bDisableAsyncGraphRegistration,
+			TEXT("Disables async registration of MetaSound graphs\n")
+			TEXT("Default: false"),
+			ECVF_Default);
+	} // namespace ConsoleVariables
+
 	namespace RegistryPrivate
 	{
 		TScriptInterface<IMetaSoundDocumentInterface> BuildRegistryDocument(TScriptInterface<IMetaSoundDocumentInterface> DocumentInterface, bool bAsync)
@@ -67,21 +78,22 @@ namespace Metasound::Frontend
 			// Only assets require template node processing and support document attachment
 			if (DocObject->IsAsset()) 
 			{
-				const FMetaSoundFrontendDocumentBuilder& OriginalDocBuilder = IMetaSoundAssetManager::GetChecked().AttachDocumentBuilderChecked(*DocObject);
-				const bool bContainsTemplateDependency = OriginalDocBuilder.ContainsDependencyOfType(EMetasoundFrontendClassType::Template);
-				if (bContainsTemplateDependency)
+				if (IMetaSoundAssetManager* AssetManager = IMetaSoundAssetManager::Get())
 				{
-					UE_LOG(LogMetaSound, Error,
-						TEXT("Template node processing disabled but provided asset class at '%s' to register contains template nodes. Runtime graph will fail to build."),
-						*OriginalDocBuilder.GetDebugName());
+					const FMetaSoundFrontendDocumentBuilder& OriginalDocBuilder = AssetManager->AttachDocumentBuilderChecked(*DocObject);
+					const bool bContainsTemplateDependency = OriginalDocBuilder.ContainsDependencyOfType(EMetasoundFrontendClassType::Template);
+					if (bContainsTemplateDependency)
+					{
+						UE_LOG(LogMetaSound, Error,
+							TEXT("Template node processing disabled but provided asset class at '%s' to register contains template nodes. Runtime graph will fail to build."),
+							*OriginalDocBuilder.GetDebugName());
+					}
 				}
 			}
 	#endif // !NO_LOGGING
 
 			// Force a copy if async registration is enabled and we need to protect against race conditions from external modifications.
-			const FMetasoundAssetBase* AssetBase = IMetaSoundAssetManager::GetChecked().GetAsAsset(*DocObject);
-			check(AssetBase);
-			const bool bForceCopy = AssetBase->IsBuilderActive() && bAsync;
+			const bool bForceCopy = DocumentInterface->IsActivelyBuilding() && bAsync;
 			if (bForceCopy)
 			{
 				return &UMetaSoundBuilderDocument::Create(*DocumentInterface.GetInterface());
@@ -433,7 +445,7 @@ namespace Metasound::Frontend
 		return MakeUnique<FNodeRegistryTransactionStream>(TransactionBuffer);
 	}
 
-	FGraphRegistryKey FRegistryContainerImpl::RegisterGraph(const TScriptInterface<IMetaSoundDocumentInterface>& InDocumentInterface, bool bAsync)
+	FGraphRegistryKey FRegistryContainerImpl::RegisterGraph(const TScriptInterface<IMetaSoundDocumentInterface>& InDocumentInterface)
 	{
 		using namespace UE;
 
@@ -456,6 +468,7 @@ namespace Metasound::Frontend
 		METASOUND_TRACE_CPUPROFILER_EVENT_SCOPE_TEXT(*FString::Printf(TEXT("FRegistryContainerImpl::RegisterGraph key:%s, asset %s"), *RegistryKey.ToString(), *AssetPath.ToString()));
 
 		// Wait for any async tasks that are in flight which correspond to the same graph prior to building as building, if force copy is false, may mutate the serialized document
+		const bool bAsync = !ConsoleVariables::bDisableAsyncGraphRegistration;
 		if (bAsync)
 		{
 			WaitForAsyncGraphRegistration(RegistryKey);
@@ -642,7 +655,7 @@ namespace Metasound::Frontend
 		return false;
 	}
 
-	bool FRegistryContainerImpl::UnregisterGraph(const FGraphRegistryKey& InRegistryKey, const TScriptInterface<IMetaSoundDocumentInterface>& InDocumentInterface, bool bAsync)
+	bool FRegistryContainerImpl::UnregisterGraph(const FGraphRegistryKey& InRegistryKey, const TScriptInterface<IMetaSoundDocumentInterface>& InDocumentInterface)
 	{
 		using namespace UE;
 		using namespace RegistryPrivate;
@@ -675,6 +688,12 @@ namespace Metasound::Frontend
 			TransactionBuffer->AddTransaction(FNodeRegistryTransaction(FNodeRegistryTransaction::ETransactionType::NodeUnregistration, NodeClassInfo, Timestamp));
 		}
 
+		// Async registration is only available if:
+		// 1. The IMetaSoundDocumentInterface is not actively modified by a builder
+		//    (built graph must be released synchronously to avoid a race condition on
+		//    reading/writing the IMetaSoundDocumentInterface on the Game Thread)
+		// 2. Async registration is globally disabled via console variable.
+		const bool bAsync = !(InDocumentInterface->IsActivelyBuilding() || ConsoleVariables::bDisableAsyncGraphRegistration);
 		if (bAsync)
 		{
 			// Wait for any async tasks that are in flight which correspond to the same graph
diff --git a/Engine/Plugins/Runtime/Metasound/Source/MetasoundFrontend/Private/MetasoundFrontendRegistryContainerImpl.h b/Engine/Plugins/Runtime/Metasound/Source/MetasoundFrontend/Private/MetasoundFrontendRegistryContainerImpl.h
--- a/Engine/Plugins/Runtime/Metasound/Source/MetasoundFrontend/Private/MetasoundFrontendRegistryContainerImpl.h
+++ b/Engine/Plugins/Runtime/Metasound/Source/MetasoundFrontend/Private/MetasoundFrontendRegistryContainerImpl.h
@@ -128,10 +128,10 @@ namespace Metasound
 			virtual TArray<FConverterNodeInfo> GetPossibleConverterNodes(const FName& FromDataType, const FName& ToDataType) override;
 
 			// Register a graph from an IMetaSoundDocumentInterface
-			FGraphRegistryKey RegisterGraph(const TScriptInterface<IMetaSoundDocumentInterface>& InDocumentInterface, bool bAsync = true);
+			FGraphRegistryKey RegisterGraph(const TScriptInterface<IMetaSoundDocumentInterface>& InDocumentInterface);
 
 			// Unregister a graph
-			bool UnregisterGraph(const FGraphRegistryKey& InRegistryKey, const TScriptInterface<IMetaSoundDocumentInterface>& InDocumentInterface, bool bAsync = true);
+			bool UnregisterGraph(const FGraphRegistryKey& InRegistryKey, const TScriptInterface<IMetaSoundDocumentInterface>& InDocumentInterface);
 
 			// Private implementation until hardened and used for template nodes other than reroutes.
 			FNodeRegistryKey RegisterNodeTemplate(TUniquePtr<Metasound::Frontend::INodeRegistryTemplateEntry>&& InEntry);
diff --git a/Engine/Plugins/Runtime/Metasound/Source/MetasoundFrontend/Public/MetasoundAssetBase.h b/Engine/Plugins/Runtime/Metasound/Source/MetasoundFrontend/Public/MetasoundAssetBase.h
--- a/Engine/Plugins/Runtime/Metasound/Source/MetasoundFrontend/Public/MetasoundAssetBase.h
+++ b/Engine/Plugins/Runtime/Metasound/Source/MetasoundFrontend/Public/MetasoundAssetBase.h
@@ -195,11 +195,8 @@ public:
 	void SetVersionedOnLoad();
 #endif // WITH_EDITORONLY_DATA
 
-	// Returns true if the IMetaSoundDocumentInterface is currently has an active builder.
-	//
-	// If true, calls to register the graph will be performed synchronously in order to avoid
-	// race conditions with an active builder.
-	virtual bool IsBuilderActive() const = 0;
+	UE_DEPRECATED(5.5, "Use IMetaSoundDocumentInterface 'IsActivelyBuilding' instead")
+	virtual bool IsBuilderActive() const { checkNoEntry(); return false; }
 
 protected:
 	void OnNotifyBeginDestroy();
diff --git a/Engine/Plugins/Runtime/Metasound/Source/MetasoundFrontend/Public/MetasoundDocumentInterface.h b/Engine/Plugins/Runtime/Metasound/Source/MetasoundFrontend/Public/MetasoundDocumentInterface.h
--- a/Engine/Plugins/Runtime/Metasound/Source/MetasoundFrontend/Public/MetasoundDocumentInterface.h
+++ b/Engine/Plugins/Runtime/Metasound/Source/MetasoundFrontend/Public/MetasoundDocumentInterface.h
@@ -49,6 +49,9 @@ public:
 	// Returns whether or not object data was modified.
 	virtual bool ConformObjectToDocument() = 0;
 
+	// Returns whether or not a document builder is currently active and can mutate the given interface's document
+	virtual bool IsActivelyBuilding() const = 0;
+
 private:
 	virtual FMetasoundFrontendDocument& GetDocument() = 0;
 
diff --git a/Engine/Plugins/Runtime/Metasound/Source/MetasoundFrontend/Public/MetasoundFrontendDocumentBuilder.h b/Engine/Plugins/Runtime/Metasound/Source/MetasoundFrontend/Public/MetasoundFrontendDocumentBuilder.h
--- a/Engine/Plugins/Runtime/Metasound/Source/MetasoundFrontend/Public/MetasoundFrontendDocumentBuilder.h
+++ b/Engine/Plugins/Runtime/Metasound/Source/MetasoundFrontend/Public/MetasoundFrontendDocumentBuilder.h
@@ -112,6 +112,9 @@ public:
 	// Base MetaSoundClass that document is published to.
 	virtual const UClass& GetBaseMetaSoundUClass() const final override;
 
+	// Returns if the document is being actively built (always true as builder documents are always being actively built)
+	virtual bool IsActivelyBuilding() const final override;
+
 private:
 	virtual FMetasoundFrontendDocument& GetDocument() override;
 
