Temporary hack to ensure builder exists prior to call in register while I refactor greater lifetime issue
#rnx

diff --git a/Engine/Plugins/Runtime/Metasound/Source/MetasoundFrontend/Private/MetasoundFrontendRegistryContainerImpl.cpp b/Engine/Plugins/Runtime/Metasound/Source/MetasoundFrontend/Private/MetasoundFrontendRegistryContainerImpl.cpp
--- a/Engine/Plugins/Runtime/Metasound/Source/MetasoundFrontend/Private/MetasoundFrontendRegistryContainerImpl.cpp
+++ b/Engine/Plugins/Runtime/Metasound/Source/MetasoundFrontend/Private/MetasoundFrontendRegistryContainerImpl.cpp
@@ -80,13 +80,20 @@ namespace Metasound::Frontend
 			{
 				if (IMetaSoundAssetManager* AssetManager = IMetaSoundAssetManager::Get())
 				{
-					const FMetaSoundFrontendDocumentBuilder& OriginalDocBuilder = AssetManager->AttachDocumentBuilderChecked(*DocObject);
-					const bool bContainsTemplateDependency = OriginalDocBuilder.ContainsDependencyOfType(EMetasoundFrontendClassType::Template);
-					if (bContainsTemplateDependency)
+					// Hack to ensure the builder is available prior to testing validity of cooked build
+					// in order to report if cooked asset is properly processed. This will be replaced with
+					// actual accessor directly from registry and AssetManager::AttachDocumentBuilderChecked
+					// will be deprecated.
+					if (IDocumentBuilderRegistry* BuilderRegistry = IDocumentBuilderRegistry::Get())
 					{
-						UE_LOG(LogMetaSound, Error,
-							TEXT("Template node processing disabled but provided asset class at '%s' to register contains template nodes. Runtime graph will fail to build."),
-							*OriginalDocBuilder.GetDebugName());
+						const FMetaSoundFrontendDocumentBuilder& OriginalDocBuilder = AssetManager->AttachDocumentBuilderChecked(*DocObject);
+						const bool bContainsTemplateDependency = OriginalDocBuilder.ContainsDependencyOfType(EMetasoundFrontendClassType::Template);
+						if (bContainsTemplateDependency)
+						{
+							UE_LOG(LogMetaSound, Error,
+								TEXT("Template node processing disabled but provided asset class at '%s' to register contains template nodes. Runtime graph will fail to build."),
+								*OriginalDocBuilder.GetDebugName());
+						}
 					}
 				}
 			}
