AudioMaterialSlate Style overrides to Editor Preferences
#rnx

diff --git a/Engine/Plugins/Runtime/AudioWidgets/Source/AudioWidgets/Private/AudioMaterialSlate/SAudioMaterialLabeledSlider.cpp b/Engine/Plugins/Runtime/AudioWidgets/Source/AudioWidgets/Private/AudioMaterialSlate/SAudioMaterialLabeledSlider.cpp
--- a/Engine/Plugins/Runtime/AudioWidgets/Source/AudioWidgets/Private/AudioMaterialSlate/SAudioMaterialLabeledSlider.cpp
+++ b/Engine/Plugins/Runtime/AudioWidgets/Source/AudioWidgets/Private/AudioMaterialSlate/SAudioMaterialLabeledSlider.cpp
@@ -40,6 +40,7 @@ void SAudioMaterialLabeledSlider::Construct(const SAudioMaterialLabeledSlider::F
 	SAssignNew(Slider, SAudioMaterialSlider)
 		.ValueAttribute(SliderValueAttribute.Get())
 		.Owner(InArgs._Owner)
+		.AudioMaterialSliderStyle(Style)
 		.Orientation(Orientation.Get())
 		.OnValueChanged_Lambda([this](float Value)
 		{
diff --git a/Engine/Plugins/Runtime/AudioWidgets/Source/AudioWidgets/Private/AudioMaterialSlate/SAudioMaterialMeter.cpp b/Engine/Plugins/Runtime/AudioWidgets/Source/AudioWidgets/Private/AudioMaterialSlate/SAudioMaterialMeter.cpp
--- a/Engine/Plugins/Runtime/AudioWidgets/Source/AudioWidgets/Private/AudioMaterialSlate/SAudioMaterialMeter.cpp
+++ b/Engine/Plugins/Runtime/AudioWidgets/Source/AudioWidgets/Private/AudioMaterialSlate/SAudioMaterialMeter.cpp
@@ -3,7 +3,6 @@
 
 #include "AudioMaterialSlate/SAudioMaterialMeter.h"
 #include "AudioMaterialSlate/AudioMaterialMeter.h"
-#include "AudioWidgetsStyle.h"
 #include "Components/AudioComponent.h"
 #include "Fonts/FontMeasure.h"
 #include "Framework/Application/SlateApplication.h"
@@ -22,11 +21,6 @@ void SAudioMaterialMeter::Construct(const FArguments& InArgs)
 	Style = InArgs._AudioMaterialMeterStyle;
 	MeterChannelInfoAttribute = InArgs._MeterChannelInfo;
 
-	if (!Cast<UWidget>(Owner))
-	{
-		Style = &FAudioWidgetsStyle::Get().GetWidgetStyle<FAudioMaterialMeterStyle>("AudioMaterialMeter.Style");
-	}
-
 	ApplyNewMaterial();
 }
 
diff --git a/Engine/Plugins/Runtime/AudioWidgets/Source/AudioWidgets/Private/AudioMaterialSlate/SAudioMaterialSlider.cpp b/Engine/Plugins/Runtime/AudioWidgets/Source/AudioWidgets/Private/AudioMaterialSlate/SAudioMaterialSlider.cpp
--- a/Engine/Plugins/Runtime/AudioWidgets/Source/AudioWidgets/Private/AudioMaterialSlate/SAudioMaterialSlider.cpp
+++ b/Engine/Plugins/Runtime/AudioWidgets/Source/AudioWidgets/Private/AudioMaterialSlate/SAudioMaterialSlider.cpp
@@ -27,12 +27,6 @@ void SAudioMaterialSlider::Construct(const FArguments& InArgs)
 	OnValueChanged = InArgs._OnValueChanged;
 	OnValueCommitted = InArgs._OnValueCommitted;
 
-	//For now check if owner is not a widget -> use the default style.
-	if (!Cast<UWidget>(Owner))
-	{
-		AudioMaterialSliderStyle = &FAudioWidgetsStyle::Get().GetWidgetStyle<FAudioMaterialSliderStyle>("AudioMaterialSlider.Style");
-	}
-
 	ApplyNewMaterial();
 
 	if (InArgs._ValueAttribute.IsSet())
diff --git a/Engine/Plugins/Runtime/AudioWidgets/Source/AudioWidgets/Private/AudioMeter.cpp b/Engine/Plugins/Runtime/AudioWidgets/Source/AudioWidgets/Private/AudioMeter.cpp
--- a/Engine/Plugins/Runtime/AudioWidgets/Source/AudioWidgets/Private/AudioMeter.cpp
+++ b/Engine/Plugins/Runtime/AudioWidgets/Source/AudioWidgets/Private/AudioMeter.cpp
@@ -5,6 +5,7 @@
 #include "AudioBusSubsystem.h"
 #include "AudioMaterialSlate/SAudioMaterialMeter.h"
 #include "AudioMixerDevice.h"
+#include "AudioWidgetsStyle.h"
 #include "SAudioMeter.h"
 #include "UObject/UObjectGlobals.h"
 
@@ -167,26 +168,27 @@ namespace AudioWidgets
 		
 	}
 
-	FAudioMeter::FAudioMeter(const int32 InNumChannels, const Audio::FDeviceId InAudioDeviceId, const TObjectPtr<UAudioBus> InExternalAudioBus, bool InbUseAudioMaterialWidget)
+	FAudioMeter::FAudioMeter(const int32 InNumChannels, const Audio::FDeviceId InAudioDeviceId, const TObjectPtr<UAudioBus> InExternalAudioBus)
 	{
-		if (InbUseAudioMaterialWidget)
-		{
-			Widget = (SNew(SAudioMaterialMeter));
-		}
-		else
-		{
-			Widget = SNew(SAudioMeter)
-				.Orientation(EOrientation::Orient_Vertical)
-				.BackgroundColor(FLinearColor::Transparent)
-
-				// TODO: Move to editor style
-				.MeterBackgroundColor(FLinearColor(0.031f, 0.031f, 0.031f, 1.0f))
-				.MeterValueColor(FLinearColor(0.025719f, 0.208333f, 0.069907f, 1.0f))
-				.MeterPeakColor(FLinearColor(0.24349f, 0.708333f, 0.357002f, 1.0f))
-				.MeterClippingColor(FLinearColor(1.0f, 0.0f, 0.112334f, 1.0f))
-				.MeterScaleColor(FLinearColor(0.017642f, 0.017642f, 0.017642f, 1.0f))
-				.MeterScaleLabelColor(FLinearColor(0.442708f, 0.442708f, 0.442708f, 1.0f));
-		}
+		Widget = SNew(SAudioMeter)
+			.Orientation(EOrientation::Orient_Vertical)
+			.BackgroundColor(FLinearColor::Transparent)
+
+			// TODO: Move to editor style
+			.MeterBackgroundColor(FLinearColor(0.031f, 0.031f, 0.031f, 1.0f))
+			.MeterValueColor(FLinearColor(0.025719f, 0.208333f, 0.069907f, 1.0f))
+			.MeterPeakColor(FLinearColor(0.24349f, 0.708333f, 0.357002f, 1.0f))
+			.MeterClippingColor(FLinearColor(1.0f, 0.0f, 0.112334f, 1.0f))
+			.MeterScaleColor(FLinearColor(0.017642f, 0.017642f, 0.017642f, 1.0f))
+			.MeterScaleLabelColor(FLinearColor(0.442708f, 0.442708f, 0.442708f, 1.0f));	
+
+		Init(InNumChannels, InAudioDeviceId, InExternalAudioBus);
+	}
+
+	FAudioMeter::FAudioMeter(const int32 InNumChannels, const Audio::FDeviceId InAudioDeviceId, const FAudioMaterialMeterStyle* AudioMaterialMeterStyle, TObjectPtr<UAudioBus> InExternalAudioBus)
+	{
+		Widget = SNew(SAudioMaterialMeter)	
+			.AudioMaterialMeterStyle(AudioMaterialMeterStyle);
 
 		Init(InNumChannels, InAudioDeviceId, InExternalAudioBus);
 	}
diff --git a/Engine/Plugins/Runtime/AudioWidgets/Source/AudioWidgets/Public/AudioMaterialSlate/SAudioMaterialMeter.h b/Engine/Plugins/Runtime/AudioWidgets/Source/AudioWidgets/Public/AudioMaterialSlate/SAudioMaterialMeter.h
--- a/Engine/Plugins/Runtime/AudioWidgets/Source/AudioWidgets/Public/AudioMaterialSlate/SAudioMaterialMeter.h
+++ b/Engine/Plugins/Runtime/AudioWidgets/Source/AudioWidgets/Public/AudioMaterialSlate/SAudioMaterialMeter.h
@@ -2,8 +2,9 @@
 
 #pragma once
 
-#include "AudioMaterialSlate/AudioMaterialSlateTypes.h"
+#include "AudioMaterialSlateTypes.h"
 #include "AudioMeterTypes.h"
+#include "AudioWidgetsStyle.h"
 #include "Framework/SlateDelegates.h"
 #include "SAudioMeter.h"
 #include "Styling/ISlateStyle.h"
@@ -21,6 +22,7 @@ class AUDIOWIDGETS_API SAudioMaterialMeter : public SAudioMeterBase
 public:
 	SLATE_BEGIN_ARGS(SAudioMaterialMeter)
 	: _Orientation(EOrientation::Orient_Vertical)
+	, _AudioMaterialMeterStyle(&FAudioWidgetsStyle::Get().GetWidgetStyle<FAudioMaterialMeterStyle>("AudioMaterialMeter.Style"))
 	{}
 
 	/** The meter's orientation. */
diff --git a/Engine/Plugins/Runtime/AudioWidgets/Source/AudioWidgets/Public/AudioMeter.h b/Engine/Plugins/Runtime/AudioWidgets/Source/AudioWidgets/Public/AudioMeter.h
--- a/Engine/Plugins/Runtime/AudioWidgets/Source/AudioWidgets/Public/AudioMeter.h
+++ b/Engine/Plugins/Runtime/AudioWidgets/Source/AudioWidgets/Public/AudioMeter.h
@@ -25,6 +25,7 @@ class SAudioMaterialMeter;
 class SAudioMeter;
 class UWorld;
 
+struct FAudioMaterialMeterStyle;
 
 /**
  * An audio meter widget.
@@ -160,8 +161,12 @@ namespace AudioWidgets
 		UE_DEPRECATED(5.4, "Use the FAudioMeter constructor that uses Audio::FDeviceId.")
 		FAudioMeter(int32 InNumChannels, UWorld& InWorld, TObjectPtr<UAudioBus> InExternalAudioBus = nullptr); 
 		
-		//** OPTIONAL PARAM InExternalAudioBus: An audio meter can be constructed from this audio bus. , InbUseAudioMaterialWidget: Is the AudioMaterialWidgets used for analyzer visualization.*/
-		FAudioMeter(const int32 InNumChannels, const Audio::FDeviceId InAudioDeviceId, TObjectPtr<UAudioBus> InExternalAudioBus = nullptr, bool InbUseAudioMaterialWidget = false); 
+		//** OPTIONAL PARAM InExternalAudioBus: An audio meter can be constructed from this audio bus.*/
+		FAudioMeter(const int32 InNumChannels, const Audio::FDeviceId InAudioDeviceId, TObjectPtr<UAudioBus> InExternalAudioBus = nullptr); 
+
+		//** Constructs the Meter using AudioMaterialMeter with the given style. OPTIONAL PARAM InExternalAudioBus: An audio meter can be constructed from this audio bus.
+		FAudioMeter(const int32 InNumChannels, const Audio::FDeviceId InAudioDeviceId, const FAudioMaterialMeterStyle* AudioMaterialMeterStyle, TObjectPtr<UAudioBus> InExternalAudioBus = nullptr);
+
 		~FAudioMeter();
 
 		UAudioBus* GetAudioBus() const;
diff --git a/Engine/Plugins/Runtime/Metasound/Source/MetasoundEditor/Private/MetasoundEditor.cpp b/Engine/Plugins/Runtime/Metasound/Source/MetasoundEditor/Private/MetasoundEditor.cpp
--- a/Engine/Plugins/Runtime/Metasound/Source/MetasoundEditor/Private/MetasoundEditor.cpp
+++ b/Engine/Plugins/Runtime/Metasound/Source/MetasoundEditor/Private/MetasoundEditor.cpp
@@ -1256,13 +1256,18 @@ namespace Metasound
 
 					if (!OutputMeter.IsValid())
 					{
-						bool bUseAudioMaterialWidgets = false;
-						if (const UMetasoundEditorSettings* EditorSettings = GetDefault<UMetasoundEditorSettings>())
+						const UMetasoundEditorSettings* EditorSettings = GetDefault<UMetasoundEditorSettings>();
+						check(EditorSettings)
+						const bool bUseAudioMaterialWidgets = EditorSettings->bUseAudioMaterialWidgets;
+
+						if (bUseAudioMaterialWidgets)
+						{										
+							OutputMeter = MakeShared<AudioWidgets::FAudioMeter>(MetaSoundSource->NumChannels, AudioDeviceId, EditorSettings->GetMeterStyle(), nullptr);
+						}						
+						else
 						{
-							bUseAudioMaterialWidgets = EditorSettings->bUseAudioMaterialWidgets;
+							OutputMeter = MakeShared<AudioWidgets::FAudioMeter>(MetaSoundSource->NumChannels, AudioDeviceId, nullptr);
 						}
-
-						OutputMeter = MakeShared<AudioWidgets::FAudioMeter>(MetaSoundSource->NumChannels, AudioDeviceId, nullptr, bUseAudioMaterialWidgets);
 					}
 					else if (OutputMeter->GetAudioBus()->GetNumChannels() != MetaSoundSource->NumChannels)
 					{
diff --git a/Engine/Plugins/Runtime/Metasound/Source/MetasoundEditor/Private/MetasoundEditorSettings.cpp b/Engine/Plugins/Runtime/Metasound/Source/MetasoundEditor/Private/MetasoundEditorSettings.cpp
--- a/Engine/Plugins/Runtime/Metasound/Source/MetasoundEditor/Private/MetasoundEditorSettings.cpp
+++ b/Engine/Plugins/Runtime/Metasound/Source/MetasoundEditor/Private/MetasoundEditorSettings.cpp
@@ -45,3 +45,42 @@ const FAudioMaterialKnobStyle* UMetasoundEditorSettings::GetKnobStyle() const
 
 	return &FAudioWidgetsStyle::Get().GetWidgetStyle<FAudioMaterialKnobStyle>("AudioMaterialKnob.Style");
 }
+
+const FAudioMaterialSliderStyle* UMetasoundEditorSettings::GetSliderStyle() const
+{
+	if (const UObject* Style = SliderStyleOverride.TryLoad())
+	{
+		if (const USlateWidgetStyleAsset* SlateWidgetStyleAsset = CastChecked<USlateWidgetStyleAsset>(Style))
+		{
+			return SlateWidgetStyleAsset->GetStyle<FAudioMaterialSliderStyle>();
+		}
+	}
+
+	return &FAudioWidgetsStyle::Get().GetWidgetStyle<FAudioMaterialSliderStyle>("AudioMaterialSlider.Style");
+}
+
+const FAudioMaterialButtonStyle* UMetasoundEditorSettings::GetButtonStyle() const
+{
+	if (const UObject* Style = ButtonStyleOverride.TryLoad())
+	{
+		if (const USlateWidgetStyleAsset* SlateWidgetStyleAsset = CastChecked<USlateWidgetStyleAsset>(Style))
+		{
+			return SlateWidgetStyleAsset->GetStyle<FAudioMaterialButtonStyle>();
+		}
+	}
+
+	return &FAudioWidgetsStyle::Get().GetWidgetStyle<FAudioMaterialButtonStyle>("AudioMaterialButton.Style");
+}
+
+const FAudioMaterialMeterStyle* UMetasoundEditorSettings::GetMeterStyle() const
+{
+	if (const UObject* Style = MeterStyleOverride.TryLoad())
+	{
+		if (const USlateWidgetStyleAsset* SlateWidgetStyleAsset = CastChecked<USlateWidgetStyleAsset>(Style))
+		{
+			return SlateWidgetStyleAsset->GetStyle<FAudioMaterialMeterStyle>();
+		}
+	}
+
+	return &FAudioWidgetsStyle::Get().GetWidgetStyle<FAudioMaterialMeterStyle>("AudioMaterialMeter.Style");
+}
diff --git a/Engine/Plugins/Runtime/Metasound/Source/MetasoundEditor/Private/MetasoundEditorSettings.h b/Engine/Plugins/Runtime/Metasound/Source/MetasoundEditor/Private/MetasoundEditorSettings.h
--- a/Engine/Plugins/Runtime/Metasound/Source/MetasoundEditor/Private/MetasoundEditorSettings.h
+++ b/Engine/Plugins/Runtime/Metasound/Source/MetasoundEditor/Private/MetasoundEditorSettings.h
@@ -11,6 +11,9 @@
 
 class USlateWidgetStyleAsset;
 struct FAudioMaterialKnobStyle;
+struct FAudioMaterialButtonStyle;
+struct FAudioMaterialSliderStyle;
+struct FAudioMaterialMeterStyle;
 
 UENUM()
 enum class EMetasoundActiveAnalyzerEnvelopeDirection : uint8
@@ -182,8 +185,30 @@ public:
 	
 	/**Override the Knob Style used in the Metasound Editor.*/
 	UPROPERTY(EditAnywhere, config, Category = MetasoundStyling, meta = (AllowedClasses = "/Script/SlateCore.SlateWidgetStyleAsset", EditCondition = "bUseAudioMaterialWidgets", DisplayName = "Knob Style"))
-	FSoftObjectPath KnobStyleOverride;
+	FSoftObjectPath KnobStyleOverride;	
 	
-	/** Get the AudioMaterialKnob Style. if KnobStyleOverride is not set, returns default style.*/
-	const FAudioMaterialKnobStyle* GetKnobStyle() const;
+	/**Override the Slider Style used in the Metasound Editor.*/
+	UPROPERTY(EditAnywhere, config, Category = MetasoundStyling, meta = (AllowedClasses = "/Script/SlateCore.SlateWidgetStyleAsset", EditCondition = "bUseAudioMaterialWidgets", DisplayName = "Slider Style"))
+	FSoftObjectPath SliderStyleOverride;
+
+	/**Override the Button Style used in the Metasound Editor.*/
+	UPROPERTY(EditAnywhere, config, Category = MetasoundStyling, meta = (AllowedClasses = "/Script/SlateCore.SlateWidgetStyleAsset", EditCondition = "bUseAudioMaterialWidgets", DisplayName = "Button Style"))
+	FSoftObjectPath ButtonStyleOverride;	
+	
+	/**Override the Meter Style used in the Metasound Editor.*/
+	UPROPERTY(EditAnywhere, config, Category = MetasoundStyling, meta = (AllowedClasses = "/Script/SlateCore.SlateWidgetStyleAsset", EditCondition = "bUseAudioMaterialWidgets", DisplayName = "Meter Style"))
+	FSoftObjectPath MeterStyleOverride;	
+		
+	/** Get the AudioMaterialKnob Style. If KnobStyleOverride is not set, returns default style.*/
+	const FAudioMaterialKnobStyle* GetKnobStyle() const;	
+
+	/** Get the AudioMaterialSlider Style. If SliderStyleOverride is not set, returns default style.*/
+	const FAudioMaterialSliderStyle* GetSliderStyle() const;	
+	
+	/** Get the AudioMaterialButton Style. If ButtonStyleOverride is not set, returns default style.*/
+	const FAudioMaterialButtonStyle* GetButtonStyle() const;	
+	
+	/** Get the AudioMaterialMeter Style. If MeterStyleOverride is not set, returns default style.*/
+	const FAudioMaterialMeterStyle* GetMeterStyle() const;
+
 };
diff --git a/Engine/Plugins/Runtime/Metasound/Source/MetasoundEditor/Private/SMetasoundGraphNode.cpp b/Engine/Plugins/Runtime/Metasound/Source/MetasoundEditor/Private/SMetasoundGraphNode.cpp
--- a/Engine/Plugins/Runtime/Metasound/Source/MetasoundEditor/Private/SMetasoundGraphNode.cpp
+++ b/Engine/Plugins/Runtime/Metasound/Source/MetasoundEditor/Private/SMetasoundGraphNode.cpp
@@ -775,8 +775,12 @@ namespace Metasound
 						{
 							if (bUseAudioMaterialWidgets)
 							{
+								const UMetasoundEditorSettings* EditorSettings = GetDefault<UMetasoundEditorSettings>();
+								check(EditorSettings)
+
 								SAssignNew(InputWidget, SAudioMaterialLabeledSlider)
 									.Owner(GraphMember->GetOwningGraph())
+									.Style(EditorSettings->GetSliderStyle())
 									.AudioUnitsValueType(DefaultFloat->WidgetUnitValueType)
 									.bUseLinearOutput(DefaultFloat->VolumeWidgetUseLinearOutput)
 									.OnValueChanged_Lambda(OnValueChangedLambda)
@@ -899,12 +903,12 @@ namespace Metasound
 
 							if (bUseAudioMaterialWidgets)
 							{
-								const UMetasoundEditorSettings* MetasoundSettings = GetDefault<UMetasoundEditorSettings>();
-								check(MetasoundSettings)
+								const UMetasoundEditorSettings* EditorSettings = GetDefault<UMetasoundEditorSettings>();
+								check(EditorSettings)
 
 								SAssignNew(InputWidget, SAudioMaterialLabeledKnob)
 									.Owner(GraphMember->GetOwningGraph())
-									.Style(MetasoundSettings->GetKnobStyle())
+									.Style(EditorSettings->GetKnobStyle())
 									.OnValueChanged_Lambda(OnValueChangedLambda)
 									.AudioUnitsValueType(DefaultFloat->WidgetUnitValueType)
 									.bUseLinearOutput(DefaultFloat->VolumeWidgetUseLinearOutput)
@@ -1043,6 +1047,15 @@ namespace Metasound
 										}
 									};
 
+								const UMetasoundEditorSettings* EditorSettings = GetDefault<UMetasoundEditorSettings>();
+								check(EditorSettings)
+
+									SAssignNew(MaterialButtonWidget, SAudioMaterialButton)
+									.AudioMaterialButtonStyle(EditorSettings->GetButtonStyle())
+									.OnBooleanValueChanged_Lambda(OnboolValueChangedLambda)
+									.OnMouseCaptureEnd_Lambda(OnboolMouseCaptureEndLambda)
+									.bIsPressedAttribute(DefaultBool->GetDefault());
+
 								SAssignNew(OuterContentBox, SVerticalBox)
 									+ SVerticalBox::Slot()
 									.HAlign(HAlign_Right)
@@ -1057,10 +1070,7 @@ namespace Metasound
 									.Padding(WidgetPadding, 0.0f, WidgetPadding, WidgetPadding)
 									.AutoHeight()
 									[
-										SAssignNew(MaterialButtonWidget, SAudioMaterialButton)
-											.OnBooleanValueChanged_Lambda(OnboolValueChangedLambda)
-											.OnMouseCaptureEnd_Lambda(OnboolMouseCaptureEndLambda)
-											.bIsPressedAttribute(DefaultBool->GetDefault())
+										MaterialButtonWidget.ToSharedRef()										
 									];
 
 								MaterialButtonWidget->SetDesiredSizeOverride(ButtonDesiredSize);
