Adding StartTime interface as an optional MetaSound Source interface.

StartTime used in BP API and Sequencer will now be forwarded to the MetaSound through this interface. The MetaSound can decide what to do with that information -- i.e. sound design what 'Start Time' means.

diff --git a/Engine/Plugins/Runtime/Metasound/Source/MetasoundEngine/Private/Interfaces/MetasoundInterface.cpp b/Engine/Plugins/Runtime/Metasound/Source/MetasoundEngine/Private/Interfaces/MetasoundInterface.cpp
--- a/Engine/Plugins/Runtime/Metasound/Source/MetasoundEngine/Private/Interfaces/MetasoundInterface.cpp
+++ b/Engine/Plugins/Runtime/Metasound/Source/MetasoundEngine/Private/Interfaces/MetasoundInterface.cpp
@@ -132,6 +132,7 @@ namespace Metasound::Engine
 		{
 			Reg.RegisterInterface(MakeUnique<FInterfaceRegistryEntry>(SourceInterface::CreateInterface(*UMetaSoundSource::StaticClass()), MakeUnique<SourceInterface::FUpdateInterface>()));
 			Reg.RegisterInterface(MakeUnique<FInterfaceRegistryEntry>(SourceOneShotInterface::CreateInterface(*UMetaSoundSource::StaticClass())));
+			Reg.RegisterInterface(MakeUnique<FInterfaceRegistryEntry>(SourceStartTimeInterface::CreateInterface(*UMetaSoundSource::StaticClass())));
 		}
 
 		RegisterAudioFormatInterfaces();
diff --git a/Engine/Plugins/Runtime/Metasound/Source/MetasoundEngine/Private/MetasoundSource.cpp b/Engine/Plugins/Runtime/Metasound/Source/MetasoundEngine/Private/MetasoundSource.cpp
--- a/Engine/Plugins/Runtime/Metasound/Source/MetasoundEngine/Private/MetasoundSource.cpp
+++ b/Engine/Plugins/Runtime/Metasound/Source/MetasoundEngine/Private/MetasoundSource.cpp
@@ -943,6 +943,19 @@ ISoundGeneratorPtr UMetaSoundSource::CreateSoundGenerator(const FSoundGeneratorI
 	FParameterRouter& Router = UMetaSoundSource::GetParameterRouter();
 	TSharedPtr<TSpscQueue<FMetaSoundParameterTransmitter::FParameter>> DataChannel = Router.FindOrCreateDataChannelForReader(InParams.AudioDeviceID, InParams.InstanceID);
 
+	if (InParams.StartTime > 0.0f)
+	{
+		if (IsInterfaceDeclared(SourceStartTimeInterface::GetVersion()))
+		{
+			InDefaultParameters.Add({ SourceStartTimeInterface::Inputs::StartTime, InParams.StartTime });
+		}
+		else
+		{
+			const FString AssetName = GetName();
+			UE_LOG(LogMetaSound, Warning, TEXT("A non-zero StartTime (%.2f) was used in asset '%s' that doesn't implement the Start Time interface. StartTime will be ignored."), InParams.StartTime, *AssetName);
+		}
+	}
+
 	FOperatorBuilderSettings BuilderSettings = FOperatorBuilderSettings::GetDefaultSettings();
 	// Graph analyzer currently only enabled for preview sounds (but can theoretically be supported for all sounds)
 	BuilderSettings.bPopulateInternalDataReferences = InParams.bIsPreviewSound;
diff --git a/Engine/Plugins/Runtime/Metasound/Source/MetasoundFrontend/Private/Interfaces/MetasoundFrontendSourceInterface.cpp b/Engine/Plugins/Runtime/Metasound/Source/MetasoundFrontend/Private/Interfaces/MetasoundFrontendSourceInterface.cpp
--- a/Engine/Plugins/Runtime/Metasound/Source/MetasoundFrontend/Private/Interfaces/MetasoundFrontendSourceInterface.cpp
+++ b/Engine/Plugins/Runtime/Metasound/Source/MetasoundFrontend/Private/Interfaces/MetasoundFrontendSourceInterface.cpp
@@ -67,6 +67,53 @@ namespace Metasound::Frontend
 
 #undef AUDIO_PARAMETER_INTERFACE_NAMESPACE
 
+#define AUDIO_PARAMETER_INTERFACE_NAMESPACE "UE.Source.StartTime"
+	namespace SourceStartTimeInterface
+	{
+		const FMetasoundFrontendVersion& GetVersion()
+		{
+			static const FMetasoundFrontendVersion Version = { AUDIO_PARAMETER_INTERFACE_NAMESPACE, { 1, 0 } };
+			return Version;
+		}
+
+		namespace Inputs
+		{
+			const FName StartTime = AUDIO_PARAMETER_INTERFACE_MEMBER_DEFINE("StartTime");
+		}
+
+		Audio::FParameterInterfacePtr CreateInterface(const UClass& InClass)
+		{
+			struct FInterface : public Audio::FParameterInterface
+			{
+				FInterface(const FTopLevelAssetPath& InClassPath)
+					: FParameterInterface(SourceStartTimeInterface::GetVersion().Name, SourceStartTimeInterface::GetVersion().Number.ToInterfaceVersion())
+				{
+					constexpr bool bIsModifiable = true;
+					constexpr bool bIsDefault = false;
+					UClassOptions = TArray<FClassOptions>
+					{
+						{ InClassPath, bIsModifiable, bIsDefault }
+					};
+
+					Inputs =
+					{
+						{
+							LOCTEXT("StartTime", "Start Time"),
+							LOCTEXT("StartTime Description", "The StartTime of the source passed into the MetaSound from the Gameplay API."),
+							GetMetasoundDataTypeName<FTime>(),
+							Inputs::StartTime
+						}
+					};
+				}
+			};
+
+			return MakeShared<FInterface>(InClass.GetClassPathName());
+		}
+	} // namespace SourceStartTimeInterface
+
+#undef AUDIO_PARAMETER_INTERFACE_NAMESPACE
+
+
 #define AUDIO_PARAMETER_INTERFACE_NAMESPACE "UE.Source"
 	namespace SourceInterfaceV1_0
 	{
diff --git a/Engine/Plugins/Runtime/Metasound/Source/MetasoundFrontend/Public/Interfaces/MetasoundFrontendSourceInterface.h b/Engine/Plugins/Runtime/Metasound/Source/MetasoundFrontend/Public/Interfaces/MetasoundFrontendSourceInterface.h
--- a/Engine/Plugins/Runtime/Metasound/Source/MetasoundFrontend/Public/Interfaces/MetasoundFrontendSourceInterface.h
+++ b/Engine/Plugins/Runtime/Metasound/Source/MetasoundFrontend/Public/Interfaces/MetasoundFrontendSourceInterface.h
@@ -26,6 +26,18 @@ namespace Metasound
 			METASOUNDFRONTEND_API Audio::FParameterInterfacePtr CreateInterface(const UClass& InClass);
 		}
 
+		namespace SourceStartTimeInterface
+		{
+			namespace Inputs
+			{
+				METASOUNDFRONTEND_API const extern FName StartTime;
+			}
+
+			METASOUNDFRONTEND_API const FMetasoundFrontendVersion& GetVersion();
+
+			METASOUNDFRONTEND_API Audio::FParameterInterfacePtr CreateInterface(const UClass& InClass);
+		}
+
 		namespace SourceInterfaceV1_0
 		{
 			namespace Inputs
diff --git a/Engine/Source/Runtime/AudioMixer/Private/AudioMixerSource.cpp b/Engine/Source/Runtime/AudioMixer/Private/AudioMixerSource.cpp
--- a/Engine/Source/Runtime/AudioMixer/Private/AudioMixerSource.cpp
+++ b/Engine/Source/Runtime/AudioMixer/Private/AudioMixerSource.cpp
@@ -1191,6 +1191,7 @@ namespace Audio
 		BufferInitArgs.LoopingMode = InWaveInstance->LoopingMode;
 		BufferInitArgs.bIsSeeking = bIsSeeking;
 		BufferInitArgs.bIsPreviewSound = bActiveSoundIsPreviewSound;
+		BufferInitArgs.StartTime = InWaveInstance->StartTime;
 
 		MixerSourceBuffer = FMixerSourceBuffer::Create(BufferInitArgs, MoveTemp(DefaultParameters));
 		
diff --git a/Engine/Source/Runtime/AudioMixer/Private/AudioMixerSourceBuffer.cpp b/Engine/Source/Runtime/AudioMixer/Private/AudioMixerSourceBuffer.cpp
--- a/Engine/Source/Runtime/AudioMixer/Private/AudioMixerSourceBuffer.cpp
+++ b/Engine/Source/Runtime/AudioMixer/Private/AudioMixerSourceBuffer.cpp
@@ -107,6 +107,7 @@ namespace Audio
 			InitParams.NumFramesPerCallback = MONO_PCM_BUFFER_SAMPLES;
 			InitParams.InstanceID = InArgs.InstanceID;
 			InitParams.bIsPreviewSound = InArgs.bIsPreviewSound;
+			InitParams.StartTime = InArgs.StartTime;
 
 			SoundGenerator = InArgs.SoundWave->CreateSoundGenerator(InitParams, MoveTemp(InDefaultParams));
 
diff --git a/Engine/Source/Runtime/AudioMixer/Private/AudioMixerSourceBuffer.h b/Engine/Source/Runtime/AudioMixer/Private/AudioMixerSourceBuffer.h
--- a/Engine/Source/Runtime/AudioMixer/Private/AudioMixerSourceBuffer.h
+++ b/Engine/Source/Runtime/AudioMixer/Private/AudioMixerSourceBuffer.h
@@ -56,6 +56,7 @@ namespace Audio
 		FMixerBuffer* Buffer = nullptr;
 		USoundWave* SoundWave = nullptr;
 		ELoopingMode LoopingMode = ELoopingMode::LOOP_Never;
+		float StartTime = 0.0f;
 		bool bIsSeeking = false;
 		bool bForceSyncDecode = false;
 		bool bIsPreviewSound = false;
diff --git a/Engine/Source/Runtime/AudioMixer/Private/SoundWaveDecoder.cpp b/Engine/Source/Runtime/AudioMixer/Private/SoundWaveDecoder.cpp
--- a/Engine/Source/Runtime/AudioMixer/Private/SoundWaveDecoder.cpp
+++ b/Engine/Source/Runtime/AudioMixer/Private/SoundWaveDecoder.cpp
@@ -72,6 +72,7 @@ namespace Audio
 			Args.LoopingMode = LoopingMode;
 			Args.bIsSeeking = bIsSeeking;
 			Args.bForceSyncDecode = bForceSyncDecode;
+			Args.StartTime = SeekTime;
 			MixerSourceBuffer = FMixerSourceBuffer::Create(Args);
 
 			bIsValid = MixerSourceBuffer.IsValid();
diff --git a/Engine/Source/Runtime/Engine/Classes/Sound/SoundGenerator.h b/Engine/Source/Runtime/Engine/Classes/Sound/SoundGenerator.h
--- a/Engine/Source/Runtime/Engine/Classes/Sound/SoundGenerator.h
+++ b/Engine/Source/Runtime/Engine/Classes/Sound/SoundGenerator.h
@@ -24,6 +24,7 @@ struct FSoundGeneratorInitParams
 	uint64 InstanceID = 0;
 	bool bIsPreviewSound = false;
 	FString GraphName;
+	float StartTime = 0.0f;
 };
 
 class ISoundGenerator
