Fixing single header compiles.

diff --git a/Engine/Plugins/Experimental/Water/Source/Runtime/Public/WaterEditorServices.h b/Engine/Plugins/Experimental/Water/Source/Runtime/Public/WaterEditorServices.h
--- a/Engine/Plugins/Experimental/Water/Source/Runtime/Public/WaterEditorServices.h
+++ b/Engine/Plugins/Experimental/Water/Source/Runtime/Public/WaterEditorServices.h
@@ -1,10 +1,12 @@
 // Copyright Epic Games, Inc. All Rights Reserved.
 
 #pragma once
+
 #include "CoreTypes.h"
 
+template <typename FuncType> class TFunctionRef;
 class UClass;
-
+class UPackage;
 class UTexture2D;
 
 #if WITH_EDITORONLY_DATA
diff --git a/Engine/Plugins/Runtime/Metasound/Source/MetasoundEditor/Private/MetasoundEditor.cpp b/Engine/Plugins/Runtime/Metasound/Source/MetasoundEditor/Private/MetasoundEditor.cpp
--- a/Engine/Plugins/Runtime/Metasound/Source/MetasoundEditor/Private/MetasoundEditor.cpp
+++ b/Engine/Plugins/Runtime/Metasound/Source/MetasoundEditor/Private/MetasoundEditor.cpp
@@ -645,6 +645,11 @@ namespace Metasound
 
 		const FName FEditor::EditorName = "MetaSoundEditor";
 
+		FEditor::FEditor()
+		: GraphConnectionManager(MakeUnique<FGraphConnectionManager>())
+		{
+		}
+
 		void FEditor::RegisterTabSpawners(const TSharedRef<FTabManager>& InTabManager)
 		{
 			WorkspaceMenuCategory = InTabManager->AddLocalWorkspaceMenuCategory(LOCTEXT("WorkspaceMenu_MetasoundEditor", "MetaSound Editor"));
diff --git a/Engine/Plugins/Runtime/Metasound/Source/MetasoundEditor/Private/MetasoundEditor.h b/Engine/Plugins/Runtime/Metasound/Source/MetasoundEditor/Private/MetasoundEditor.h
--- a/Engine/Plugins/Runtime/Metasound/Source/MetasoundEditor/Private/MetasoundEditor.h
+++ b/Engine/Plugins/Runtime/Metasound/Source/MetasoundEditor/Private/MetasoundEditor.h
@@ -103,11 +103,7 @@ namespace Metasound
 		public:
 			static const FName EditorName;
 
-			FEditor()
-				: GraphConnectionManager(MakeUnique<FGraphConnectionManager>())
-			{
-			}
-
+			FEditor();
 			virtual ~FEditor();
 
 			virtual void RegisterTabSpawners(const TSharedRef<FTabManager>& TabManager) override;
diff --git a/Engine/Plugins/VirtualProduction/DMX/DMXGDTF/Source/DMXGDTF/Internal/GDTF/Geometries/DMXGDTFAddress.h b/Engine/Plugins/VirtualProduction/DMX/DMXGDTF/Source/DMXGDTF/Internal/GDTF/Geometries/DMXGDTFAddress.h
--- a/Engine/Plugins/VirtualProduction/DMX/DMXGDTF/Source/DMXGDTF/Internal/GDTF/Geometries/DMXGDTFAddress.h
+++ b/Engine/Plugins/VirtualProduction/DMX/DMXGDTF/Source/DMXGDTF/Internal/GDTF/Geometries/DMXGDTFAddress.h
@@ -2,7 +2,7 @@
 
 #pragma once
 
-#include "GenericPlatform/GenericPlatform.h"
+#include "HAL/Platform.h"
 
 namespace UE::DMX::GDTF
 {
diff --git a/Engine/Source/Developer/MeshMergeUtilities/Public/MeshMergeDataTracker.h b/Engine/Source/Developer/MeshMergeUtilities/Public/MeshMergeDataTracker.h
--- a/Engine/Source/Developer/MeshMergeUtilities/Public/MeshMergeDataTracker.h
+++ b/Engine/Source/Developer/MeshMergeUtilities/Public/MeshMergeDataTracker.h
@@ -4,6 +4,7 @@
 
 #include "StaticMeshResources.h"
 #include "Materials/MaterialInstance.h"
+#include "MaterialUtilities.h"
 #include "StaticMeshAttributes.h"
 
 typedef TPair<uint32, uint32> MeshLODPair;
diff --git a/Engine/Source/Runtime/Engine/Classes/Sound/QuartzCommandQueue.h b/Engine/Source/Runtime/Engine/Classes/Sound/QuartzCommandQueue.h
--- a/Engine/Source/Runtime/Engine/Classes/Sound/QuartzCommandQueue.h
+++ b/Engine/Source/Runtime/Engine/Classes/Sound/QuartzCommandQueue.h
@@ -4,6 +4,7 @@
 
 #include <type_traits>
 #include "Containers/ConsumeAllMpmcQueue.h"
+#include "Templates/UniquePtr.h"
 #include "QuartzCompileTimeVisitor.h"
 
 
diff --git a/Engine/Source/Runtime/Engine/Classes/Sound/QuartzCompileTimeVisitor.h b/Engine/Source/Runtime/Engine/Classes/Sound/QuartzCompileTimeVisitor.h
--- a/Engine/Source/Runtime/Engine/Classes/Sound/QuartzCompileTimeVisitor.h
+++ b/Engine/Source/Runtime/Engine/Classes/Sound/QuartzCompileTimeVisitor.h
@@ -2,6 +2,7 @@
 
 #pragma once
 
+#include "Templates/Function.h"
 #include <type_traits>
 
 // forwards
diff --git a/Engine/Source/Runtime/SessionServices/Public/ITraceController.h b/Engine/Source/Runtime/SessionServices/Public/ITraceController.h
--- a/Engine/Source/Runtime/SessionServices/Public/ITraceController.h
+++ b/Engine/Source/Runtime/SessionServices/Public/ITraceController.h
@@ -2,6 +2,10 @@
 
 #pragma once
 
+#include "Containers/Map.h"
+#include "Containers/UnrealString.h"
+#include "Delegates/Delegate.h"
+#include "Misc/Guid.h"
 
 struct FTraceStatus
 {
