Fix for AudioUnitTest build break caused by added Engine & AudioMixer dependencies in MetaSound Frontend
#rnx

diff --git a/Engine/Plugins/Runtime/Metasound/Source/MetasoundEditor/Private/MetasoundEditorGraphConnectionManager.cpp b/Engine/Plugins/Runtime/Metasound/Source/MetasoundEditor/Private/MetasoundEditorGraphConnectionManager.cpp
--- a/Engine/Plugins/Runtime/Metasound/Source/MetasoundEditor/Private/MetasoundEditorGraphConnectionManager.cpp
+++ b/Engine/Plugins/Runtime/Metasound/Source/MetasoundEditor/Private/MetasoundEditorGraphConnectionManager.cpp
@@ -1,10 +1,10 @@
 // Copyright Epic Games, Inc. All Rights Reserved.
 #include "MetasoundEditorGraphConnectionManager.h"
 
-#include "Analysis/MetasoundFrontendVertexAnalyzerAudioBusWriter.h"
 #include "Analysis/MetasoundFrontendVertexAnalyzerEnvelopeFollower.h"
 #include "Analysis/MetasoundFrontendVertexAnalyzerForwardValue.h"
 #include "Analysis/MetasoundFrontendVertexAnalyzerTriggerDensity.h"
+#include "Analysis/MetasoundVertexAnalyzerAudioBusWriter.h"
 #include "MetasoundEditorGraphBuilder.h"
 #include "MetasoundEditorSettings.h"
 #include "MetasoundVertex.h"
@@ -82,8 +82,8 @@ namespace Metasound
 			if (GraphAnalyzerView.IsValid())
 			{
 				const uint32 AudioBusID = InAudioBus->GetUniqueID();
-				const FName AnalyzerName = Frontend::FVertexAnalyzerAudioBusWriter::GetAnalyzerName();
-				const FName AnalyzerMemberName = Frontend::FVertexAnalyzerAudioBusWriter::GetAnalyzerMemberName(InDeviceID, AudioBusID);
+				const FName AnalyzerName = Engine::FVertexAnalyzerAudioBusWriter::GetAnalyzerName();
+				const FName AnalyzerMemberName = Engine::FVertexAnalyzerAudioBusWriter::GetAnalyzerMemberName(InDeviceID, AudioBusID);
 				AnalyzerInstanceID = GraphAnalyzerView->AddAnalyzerForSpecifiedOutput(InNodeID, InOutputName, AnalyzerName, AnalyzerMemberName);
 			}
 
@@ -94,7 +94,7 @@ namespace Metasound
 		{
 			if (GraphAnalyzerView.IsValid())
 			{
-				const FName AnalyzerName = Frontend::FVertexAnalyzerAudioBusWriter::GetAnalyzerName();
+				const FName AnalyzerName = Engine::FVertexAnalyzerAudioBusWriter::GetAnalyzerName();
 				GraphAnalyzerView->RemoveAnalyzerInstance(AnalyzerName, InAnalyzerInstanceID);
 			}
 		}
diff --git a/Engine/Plugins/Runtime/Metasound/Source/MetasoundEngine/Private/Analysis/MetasoundVertexAnalyzerAudioBusWriter.cpp b/Engine/Plugins/Runtime/Metasound/Source/MetasoundEngine/Private/Analysis/MetasoundVertexAnalyzerAudioBusWriter.cpp
new file mode 100644
--- /dev/null
+++ b/Engine/Plugins/Runtime/Metasound/Source/MetasoundEngine/Private/Analysis/MetasoundVertexAnalyzerAudioBusWriter.cpp
@@ -0,0 +1,82 @@
+// Copyright Epic Games, Inc. All Rights Reserved.
+
+#include "Analysis/MetasoundVertexAnalyzerAudioBusWriter.h"
+
+#include "AudioBusSubsystem.h"
+#include "AudioDevice.h"
+#include "AudioDeviceManager.h"
+#include "MetasoundAudioBuffer.h"
+
+
+namespace Metasound::Engine
+{
+	const FName& FVertexAnalyzerAudioBusWriter::GetAnalyzerName()
+	{
+		static const FName AnalyzerName = "UE.Audio.AudioBusWriter";
+		return AnalyzerName;
+	}
+
+	const FName& FVertexAnalyzerAudioBusWriter::GetDataType()
+	{
+		return GetMetasoundDataTypeName<FAudioBuffer>();
+	}
+
+	FName FVertexAnalyzerAudioBusWriter::GetAnalyzerMemberName(const Audio::FDeviceId InDeviceID, const uint32 InAudioBusID)
+	{
+		const FBusAddress BusAddress
+		{
+			.DeviceID = InDeviceID,
+			.AudioBusID = InAudioBusID,
+		};
+		return FName(BusAddress.ToString());
+	}
+
+	FVertexAnalyzerAudioBusWriter::FVertexAnalyzerAudioBusWriter(const Frontend::FCreateAnalyzerParams& InParams)
+		: FVertexAnalyzerBase(InParams.AnalyzerAddress, InParams.VertexDataReference)
+	{
+		// AnalyzerMemberName is used as an input to modify analyzer settings:
+		const FBusAddress BusAddress = FBusAddress::FromString(InParams.AnalyzerAddress.AnalyzerMemberName.ToString());
+
+		if (FAudioDeviceManager* ADM = FAudioDeviceManager::Get())
+		{
+			if (const FAudioDevice* AudioDevice = ADM->GetAudioDeviceRaw(BusAddress.DeviceID))
+			{
+				UAudioBusSubsystem* AudioBusSubsystem = AudioDevice->GetSubsystem<UAudioBusSubsystem>();
+
+				const int32 Frames = InParams.OperatorSettings.GetNumFramesPerBlock();
+				constexpr int32 Channels = 1;
+				AudioBusPatchInput = AudioBusSubsystem->AddPatchInputForAudioBus(BusAddress.AudioBusID, Frames, Channels);
+			}
+		}
+	}
+
+	void FVertexAnalyzerAudioBusWriter::Execute()
+	{
+		const FAudioBuffer& AudioBuffer = GetVertexData<FAudioBuffer>();
+		AudioBusPatchInput.PushAudio(AudioBuffer.GetData(), AudioBuffer.Num());
+	}
+
+	FString FVertexAnalyzerAudioBusWriter::FBusAddress::ToString() const
+	{
+		return FString::Join(TArray<FString>
+		{
+			FString::FromInt(DeviceID),
+			FString::FromInt(AudioBusID),
+		}
+		, TEXT(","));
+	}
+
+	FVertexAnalyzerAudioBusWriter::FBusAddress FVertexAnalyzerAudioBusWriter::FBusAddress::FromString(const FString& InAnalyzerMemberName)
+	{
+		FBusAddress BusAddress;
+
+		TArray<FString> Tokens;
+		if (ensure(InAnalyzerMemberName.ParseIntoArray(Tokens, TEXT(",")) == 2))
+		{
+			BusAddress.DeviceID = FCString::Atoi(*Tokens[0]);
+			BusAddress.AudioBusID = FCString::Atoi(*Tokens[1]);
+		}
+
+		return BusAddress;
+	}
+} // namespace Metasound::Engine
diff --git a/Engine/Plugins/Runtime/Metasound/Source/MetasoundEngine/Private/MetasoundEngineModule.cpp b/Engine/Plugins/Runtime/Metasound/Source/MetasoundEngine/Private/MetasoundEngineModule.cpp
--- a/Engine/Plugins/Runtime/Metasound/Source/MetasoundEngine/Private/MetasoundEngineModule.cpp
+++ b/Engine/Plugins/Runtime/Metasound/Source/MetasoundEngine/Private/MetasoundEngineModule.cpp
@@ -14,11 +14,11 @@
 #include "MetasoundWaveTable.h"
 #include "Analysis/MetasoundFrontendAnalyzerRegistry.h"
 #include "Analysis/MetasoundFrontendVertexAnalyzerAudioBuffer.h"
-#include "Analysis/MetasoundFrontendVertexAnalyzerAudioBusWriter.h"
 #include "Analysis/MetasoundFrontendVertexAnalyzerEnvelopeFollower.h"
 #include "Analysis/MetasoundFrontendVertexAnalyzerForwardValue.h"
 #include "Analysis/MetasoundFrontendVertexAnalyzerTriggerDensity.h"
 #include "Analysis/MetasoundFrontendVertexAnalyzerTriggerToTime.h"
+#include "Analysis/MetasoundVertexAnalyzerAudioBusWriter.h"
 #include "Interfaces/MetasoundDeprecatedInterfaces.h"
 #include "Interfaces/MetasoundInterface.h"
 #include "Interfaces/MetasoundInterfaceBindingsPrivate.h"
@@ -118,8 +118,9 @@ public:
 		FMetasoundFrontendRegistryContainer::Get()->RegisterPendingNodes();
 
 		// Register Analyzers
+
+		// TODO: Determine if we can move this registration to Frontend where it likely belongs
 		METASOUND_REGISTER_VERTEX_ANALYZER_FACTORY(Frontend::FVertexAnalyzerAudioBuffer)
-		METASOUND_REGISTER_VERTEX_ANALYZER_FACTORY(Frontend::FVertexAnalyzerAudioBusWriter)
 		METASOUND_REGISTER_VERTEX_ANALYZER_FACTORY(Frontend::FVertexAnalyzerEnvelopeFollower)
 		METASOUND_REGISTER_VERTEX_ANALYZER_FACTORY(Frontend::FVertexAnalyzerForwardBool)
 		METASOUND_REGISTER_VERTEX_ANALYZER_FACTORY(Frontend::FVertexAnalyzerForwardFloat)
@@ -129,6 +130,8 @@ public:
 		METASOUND_REGISTER_VERTEX_ANALYZER_FACTORY(Frontend::FVertexAnalyzerTriggerDensity)
 		METASOUND_REGISTER_VERTEX_ANALYZER_FACTORY(Frontend::FVertexAnalyzerTriggerToTime)
 
+		METASOUND_REGISTER_VERTEX_ANALYZER_FACTORY(Engine::FVertexAnalyzerAudioBusWriter)
+
 		// Register passthrough output analyzers
 		UMetasoundGeneratorHandle::RegisterPassthroughAnalyzerForType(
 			GetMetasoundDataTypeName<float>(),
diff --git a/Engine/Plugins/Runtime/Metasound/Source/MetasoundEngine/Public/Analysis/MetasoundVertexAnalyzerAudioBusWriter.h b/Engine/Plugins/Runtime/Metasound/Source/MetasoundEngine/Public/Analysis/MetasoundVertexAnalyzerAudioBusWriter.h
new file mode 100644
--- /dev/null
+++ b/Engine/Plugins/Runtime/Metasound/Source/MetasoundEngine/Public/Analysis/MetasoundVertexAnalyzerAudioBusWriter.h
@@ -0,0 +1,48 @@
+// Copyright Epic Games, Inc. All Rights Reserved.
+#pragma once
+
+#include "Analysis/MetasoundFrontendAnalyzerFactory.h"
+#include "Analysis/MetasoundFrontendVertexAnalyzer.h"
+#include "AudioDefines.h"
+#include "Containers/Array.h"
+
+namespace Metasound::Engine
+{
+	class METASOUNDENGINE_API FVertexAnalyzerAudioBusWriter : public Frontend::FVertexAnalyzerBase
+	{
+	public:
+		using FAnalyzerOutput = Frontend::FAnalyzerOutput;
+		using FCreateAnalyzerParams = Frontend::FCreateAnalyzerParams;
+
+		static const FName& GetAnalyzerName();
+		static const FName& GetDataType();
+		static FName GetAnalyzerMemberName(const Audio::FDeviceId InDeviceID, const uint32 InAudioBusID);
+
+		class METASOUNDENGINE_API FFactory : public Frontend::TVertexAnalyzerFactory<FVertexAnalyzerAudioBusWriter>
+		{
+		public:
+			virtual const TArray<FAnalyzerOutput>& GetAnalyzerOutputs() const override
+			{
+				static const TArray<FAnalyzerOutput> Outputs;
+				return Outputs;
+			}
+		};
+
+		FVertexAnalyzerAudioBusWriter(const FCreateAnalyzerParams& InParams);
+		virtual ~FVertexAnalyzerAudioBusWriter() = default;
+
+		virtual void Execute() override;
+
+	private:
+		struct FBusAddress
+		{
+			Audio::FDeviceId DeviceID = 0;
+			uint32 AudioBusID = 0;
+
+			FString ToString() const;
+			static FBusAddress FromString(const FString& InAnalyzerMemberName);
+		};
+
+		Audio::FPatchInput AudioBusPatchInput;
+	};
+} // namespace Metasound::Engine
diff --git a/Engine/Plugins/Runtime/Metasound/Source/MetasoundFrontend/MetasoundFrontend.Build.cs b/Engine/Plugins/Runtime/Metasound/Source/MetasoundFrontend/MetasoundFrontend.Build.cs
--- a/Engine/Plugins/Runtime/Metasound/Source/MetasoundFrontend/MetasoundFrontend.Build.cs
+++ b/Engine/Plugins/Runtime/Metasound/Source/MetasoundFrontend/MetasoundFrontend.Build.cs
@@ -16,21 +16,11 @@ namespace UnrealBuildTool.Rules
 				new string[]
 				{
 					"AudioExtensions",
-					"AudioMixer",
 					"Core",
 					"CoreUObject",
-					"Engine",
 					"Serialization",
 					"SignalProcessing",
 					"MetasoundGraphCore"
-
-				}
-			);
-
-			PrivateDependencyModuleNames.AddRange
-			(
-				new string[]
-				{
 				}
 			);
 
diff --git a/Engine/Plugins/Runtime/Metasound/Source/MetasoundFrontend/Private/Analysis/MetasoundFrontendVertexAnalyzerAudioBusWriter.cpp b/Engine/Plugins/Runtime/Metasound/Source/MetasoundFrontend/Private/Analysis/MetasoundFrontendVertexAnalyzerAudioBusWriter.cpp
deleted file mode 100644
diff --git a/Engine/Plugins/Runtime/Metasound/Source/MetasoundFrontend/Public/Analysis/MetasoundFrontendVertexAnalyzerAudioBusWriter.h b/Engine/Plugins/Runtime/Metasound/Source/MetasoundFrontend/Public/Analysis/MetasoundFrontendVertexAnalyzerAudioBusWriter.h
deleted file mode 100644
