Cut down on metasound graph profiling cost.
Uses cached specIds on FProfilingOperator to avoid event name map lookups on every scope.
#tests PIE, compiled shipping locally

diff --git a/Engine/Plugins/Runtime/Metasound/Source/MetasoundGraphCore/Private/MetasoundProfilingOperator.cpp b/Engine/Plugins/Runtime/Metasound/Source/MetasoundGraphCore/Private/MetasoundProfilingOperator.cpp
--- a/Engine/Plugins/Runtime/Metasound/Source/MetasoundGraphCore/Private/MetasoundProfilingOperator.cpp
+++ b/Engine/Plugins/Runtime/Metasound/Source/MetasoundGraphCore/Private/MetasoundProfilingOperator.cpp
@@ -87,7 +87,7 @@ namespace Metasound
 	{
 		FProfilingOperator* DerivedOperator = static_cast<FProfilingOperator*>(InOperator);
 		check(nullptr != DerivedOperator && DerivedOperator->ResetFunction);
-		METASOUND_TRACE_CPUPROFILER_EVENT_SCOPE_TEXT(*DerivedOperator->InsightsResetEventName);
+		METASOUND_TRACE_CPUPROFILER_EVENT_SCOPE_FAST(DerivedOperator->InsightsResetEventSpecId, *DerivedOperator->InsightsResetEventName);
 		DerivedOperator->ResetFunction(DerivedOperator->Operator.Get(), InParams);
 	}
 
@@ -95,7 +95,7 @@ namespace Metasound
 	{
 		FProfilingOperator* DerivedOperator = static_cast<FProfilingOperator*>(InOperator);
 		check(nullptr != DerivedOperator && DerivedOperator->ExecuteFunction);
-		METASOUND_TRACE_CPUPROFILER_EVENT_SCOPE_TEXT(*DerivedOperator->InsightsExecuteEventName);
+		METASOUND_TRACE_CPUPROFILER_EVENT_SCOPE_FAST(DerivedOperator->InsightsExecuteEventSpecId, *DerivedOperator->InsightsExecuteEventName);
 		DerivedOperator->ExecuteFunction(DerivedOperator->Operator.Get());
 	}
 
@@ -103,7 +103,7 @@ namespace Metasound
 	{
 		FProfilingOperator* DerivedOperator = static_cast<FProfilingOperator*>(InOperator);
 		check(nullptr != DerivedOperator && DerivedOperator->PostExecuteFunction);
-		METASOUND_TRACE_CPUPROFILER_EVENT_SCOPE_TEXT(*DerivedOperator->InsightsPostExecuteEventName);
+		METASOUND_TRACE_CPUPROFILER_EVENT_SCOPE_FAST(DerivedOperator->InsightsPostExecuteEventSpecId, *DerivedOperator->InsightsPostExecuteEventName);
 		DerivedOperator->PostExecuteFunction(DerivedOperator->Operator.Get());
 	}
 
diff --git a/Engine/Plugins/Runtime/Metasound/Source/MetasoundGraphCore/Private/MetasoundProfilingOperator.h b/Engine/Plugins/Runtime/Metasound/Source/MetasoundGraphCore/Private/MetasoundProfilingOperator.h
--- a/Engine/Plugins/Runtime/Metasound/Source/MetasoundGraphCore/Private/MetasoundProfilingOperator.h
+++ b/Engine/Plugins/Runtime/Metasound/Source/MetasoundGraphCore/Private/MetasoundProfilingOperator.h
@@ -24,17 +24,25 @@ namespace Metasound
 			, ResetFunction(Operator->GetResetFunction())
 			, ExecuteFunction(Operator->GetExecuteFunction())
 			, PostExecuteFunction(Operator->GetPostExecuteFunction())
+			, InsightsResetEventSpecId(0)
+			, InsightsExecuteEventSpecId(0)
+			, InsightsPostExecuteEventSpecId(0)
 		{
 			check(Operator);
 			const FNodeClassMetadata& NodeMetadata = Node->GetMetadata();
-			FString BaseEventName = NodeMetadata.ClassName.GetName().ToString();
-			if (NodeMetadata.ClassName.GetName().IsNone())
+			FName BaseEventName = NodeMetadata.ClassName.GetName();
+			if (BaseEventName.IsNone())
 			{
-				BaseEventName = Node->GetInstanceName().ToString();
+				BaseEventName = Node->GetInstanceName();
 			}
-			InsightsResetEventName = FString::Printf(TEXT("%s_RESET"), *BaseEventName);
-			InsightsExecuteEventName = FString::Printf(TEXT("%s_EXECUTE"), *BaseEventName);
-			InsightsPostExecuteEventName = FString::Printf(TEXT("%s_POSTEXECUTE"), *BaseEventName);
+			
+			BaseEventName.AppendString(InsightsResetEventName);
+			InsightsExecuteEventName = InsightsResetEventName;
+			InsightsPostExecuteEventName = InsightsResetEventName;
+
+			InsightsResetEventName.Append(TEXT("_RESET"));
+			InsightsExecuteEventName.Append(TEXT("_EXECUTE"));
+			InsightsPostExecuteEventName.Append(TEXT("_POSTEXECUTE"));
 		}
 
 		virtual ~FProfilingOperator() = default;
@@ -86,5 +94,8 @@ namespace Metasound
 		FString               InsightsResetEventName;
 		FString               InsightsExecuteEventName;
 		FString               InsightsPostExecuteEventName;
+		uint32                InsightsResetEventSpecId;
+		uint32                InsightsExecuteEventSpecId;
+		uint32                InsightsPostExecuteEventSpecId;
 	};
 }
\ No newline at end of file
diff --git a/Engine/Plugins/Runtime/Metasound/Source/MetasoundGraphCore/Public/MetasoundTrace.h b/Engine/Plugins/Runtime/Metasound/Source/MetasoundGraphCore/Public/MetasoundTrace.h
--- a/Engine/Plugins/Runtime/Metasound/Source/MetasoundGraphCore/Public/MetasoundTrace.h
+++ b/Engine/Plugins/Runtime/Metasound/Source/MetasoundGraphCore/Public/MetasoundTrace.h
@@ -24,11 +24,18 @@
 
 #define METASOUND_TRACE_CPUPROFILER_EVENT_SCOPE_TEXT(Name) \
     SCOPED_NAMED_EVENT_TCHAR(Name, FColor::Green)
+
+// Uses cached Insights SpecId to avoid string lookup
+#define METASOUND_TRACE_CPUPROFILER_EVENT_SCOPE_FAST(TraceSpecId, Name) \
+	FScopedNamedEvent ANONYMOUS_VARIABLE(NamedEvent_)(FColor::Green, Name);\
+	TRACE_CPUPROFILER_EVENT_SCOPE_USE(TraceSpecId, Name, PREPROCESSOR_JOIN(MetaSound_NamedEventScope_, __LINE__), true)
+
 #else
 // Metasound CPU profiler trace *not* enabled
 
 #define METASOUND_TRACE_CPUPROFILER_EVENT_SCOPE(Name)
 #define METASOUND_TRACE_CPUPROFILER_EVENT_SCOPE_TEXT(Name)
+#define METASOUND_TRACE_CPUPROFILER_EVENT_SCOPE_FAST(TraceSpecId, Name)
 
 #endif
 
