Fix for missing default MetaSound OutputMeter in asset editor
#rnx

#tests Open Editor & toggle material audio widgets in MetaSound editor settings

diff --git a/Engine/Plugins/Runtime/AudioWidgets/Source/AudioWidgets/Private/AudioMeter.cpp b/Engine/Plugins/Runtime/AudioWidgets/Source/AudioWidgets/Private/AudioMeter.cpp
--- a/Engine/Plugins/Runtime/AudioWidgets/Source/AudioWidgets/Private/AudioMeter.cpp
+++ b/Engine/Plugins/Runtime/AudioWidgets/Source/AudioWidgets/Private/AudioMeter.cpp
@@ -165,7 +165,6 @@ namespace AudioWidgets
 	FAudioMeter::FAudioMeter(int32 InNumChannels, UWorld& InWorld, TObjectPtr<UAudioBus> InExternalAudioBus)
 		: FAudioMeter(InNumChannels, InWorld.GetAudioDevice().GetDeviceID(), InExternalAudioBus)
 	{
-		
 	}
 
 	FAudioMeter::FAudioMeter(const int32 InNumChannels, const Audio::FDeviceId InAudioDeviceId, const TObjectPtr<UAudioBus> InExternalAudioBus)
@@ -180,15 +179,15 @@ namespace AudioWidgets
 			.MeterPeakColor(FLinearColor(0.24349f, 0.708333f, 0.357002f, 1.0f))
 			.MeterClippingColor(FLinearColor(1.0f, 0.0f, 0.112334f, 1.0f))
 			.MeterScaleColor(FLinearColor(0.017642f, 0.017642f, 0.017642f, 1.0f))
-			.MeterScaleLabelColor(FLinearColor(0.442708f, 0.442708f, 0.442708f, 1.0f));	
+			.MeterScaleLabelColor(FLinearColor(0.442708f, 0.442708f, 0.442708f, 1.0f));
 
 		Init(InNumChannels, InAudioDeviceId, InExternalAudioBus);
 	}
 
-	FAudioMeter::FAudioMeter(const int32 InNumChannels, const Audio::FDeviceId InAudioDeviceId, const FAudioMaterialMeterStyle* AudioMaterialMeterStyle, TObjectPtr<UAudioBus> InExternalAudioBus)
+	FAudioMeter::FAudioMeter(const int32 InNumChannels, const Audio::FDeviceId InAudioDeviceId, const FAudioMaterialMeterStyle& AudioMaterialMeterStyle, TObjectPtr<UAudioBus> InExternalAudioBus)
 	{
-		Widget = SNew(SAudioMaterialMeter)	
-			.AudioMaterialMeterStyle(AudioMaterialMeterStyle);
+		Widget = SNew(SAudioMaterialMeter)
+			.AudioMaterialMeterStyle(&AudioMaterialMeterStyle);
 
 		Init(InNumChannels, InAudioDeviceId, InExternalAudioBus);
 	}
diff --git a/Engine/Plugins/Runtime/AudioWidgets/Source/AudioWidgets/Public/AudioMeter.h b/Engine/Plugins/Runtime/AudioWidgets/Source/AudioWidgets/Public/AudioMeter.h
--- a/Engine/Plugins/Runtime/AudioWidgets/Source/AudioWidgets/Public/AudioMeter.h
+++ b/Engine/Plugins/Runtime/AudioWidgets/Source/AudioWidgets/Public/AudioMeter.h
@@ -165,7 +165,7 @@ namespace AudioWidgets
 		FAudioMeter(const int32 InNumChannels, const Audio::FDeviceId InAudioDeviceId, TObjectPtr<UAudioBus> InExternalAudioBus = nullptr); 
 
 		//** Constructs the Meter using AudioMaterialMeter with the given style. OPTIONAL PARAM InExternalAudioBus: An audio meter can be constructed from this audio bus.
-		FAudioMeter(const int32 InNumChannels, const Audio::FDeviceId InAudioDeviceId, const FAudioMaterialMeterStyle* AudioMaterialMeterStyle, TObjectPtr<UAudioBus> InExternalAudioBus = nullptr);
+		FAudioMeter(const int32 InNumChannels, const Audio::FDeviceId InAudioDeviceId, const FAudioMaterialMeterStyle& AudioMaterialMeterStyle, TObjectPtr<UAudioBus> InExternalAudioBus = nullptr);
 
 		~FAudioMeter();
 
diff --git a/Engine/Plugins/Runtime/Metasound/Source/MetasoundEditor/Private/MetasoundEditor.cpp b/Engine/Plugins/Runtime/Metasound/Source/MetasoundEditor/Private/MetasoundEditor.cpp
--- a/Engine/Plugins/Runtime/Metasound/Source/MetasoundEditor/Private/MetasoundEditor.cpp
+++ b/Engine/Plugins/Runtime/Metasound/Source/MetasoundEditor/Private/MetasoundEditor.cpp
@@ -1258,25 +1258,30 @@ namespace Metasound
 				if (ensure(GEditor))
 				{
 					const Audio::FDeviceId AudioDeviceId = GEditor->GetMainAudioDeviceID();
+					constexpr UAudioBus* DefaultBus = nullptr;
 
 					if (!OutputMeter.IsValid())
 					{
 						const UMetasoundEditorSettings* EditorSettings = GetDefault<UMetasoundEditorSettings>();
 						check(EditorSettings)
 						const bool bUseAudioMaterialWidgets = EditorSettings->bUseAudioMaterialWidgets;
-
 						if (bUseAudioMaterialWidgets)
-						{										
-							OutputMeter = MakeShared<AudioWidgets::FAudioMeter>(MetaSoundSource->NumChannels, AudioDeviceId, EditorSettings->GetMeterStyle(), nullptr);
-						}						
-						else
 						{
-							OutputMeter = MakeShared<AudioWidgets::FAudioMeter>(MetaSoundSource->NumChannels, AudioDeviceId, nullptr);
+							const FAudioMaterialMeterStyle* MeterStyle = EditorSettings->GetMeterStyle();
+							if (ensureMsgf(MeterStyle, TEXT("Failed to find MaterialMeterStyle when attempting to build MetaSound Editor output meter. Falling back to default non-material meter.")))
+							{
+								OutputMeter = MakeShared<AudioWidgets::FAudioMeter>(MetaSoundSource->NumChannels, AudioDeviceId, *MeterStyle, DefaultBus);
+							}
+						}
+
+						if (!OutputMeter.IsValid())
+						{
+							OutputMeter = MakeShared<AudioWidgets::FAudioMeter>(MetaSoundSource->NumChannels, AudioDeviceId, DefaultBus);
 						}
 					}
 					else if (OutputMeter->GetAudioBus()->GetNumChannels() != MetaSoundSource->NumChannels)
 					{
-						OutputMeter->Init(MetaSoundSource->NumChannels, AudioDeviceId, nullptr);
+						OutputMeter->Init(MetaSoundSource->NumChannels, AudioDeviceId, DefaultBus);
 					}
 
 					const uint32 MetaSoundNumChannels = static_cast<uint32>(MetaSoundSource->NumChannels);
