[MetaSounds] Fix for AudioBusReaderNode resizing an array to a large value (unsigned overflow)
#rnx

diff --git a/Engine/Plugins/Runtime/Metasound/Source/MetasoundEngine/Private/MetasoundAudioBusReaderNode.cpp b/Engine/Plugins/Runtime/Metasound/Source/MetasoundEngine/Private/MetasoundAudioBusReaderNode.cpp
--- a/Engine/Plugins/Runtime/Metasound/Source/MetasoundEngine/Private/MetasoundAudioBusReaderNode.cpp
+++ b/Engine/Plugins/Runtime/Metasound/Source/MetasoundEngine/Private/MetasoundAudioBusReaderNode.cpp
@@ -154,6 +154,12 @@ namespace Metasound
 			const int32 NumSamplesToPop = BlockSizeFrames * AudioBusChannels;
 
 			const FAudioBusProxyPtr& BusProxy = AudioBusAsset->GetAudioBusProxy();
+			if (BusProxy->NumChannels <= 0)
+			{
+				// the audio bus is invalid / uninitialized
+				return;
+			}
+
 			if (BusProxy.IsValid() && BusProxy->AudioBusId != AudioBusId)
 			{
 				AudioBusPatchOutput.Reset();
@@ -230,6 +236,11 @@ namespace Metasound
 			const FAudioBusProxyPtr& AudioBusProxy = AudioBusAsset->GetAudioBusProxy();
 			if (AudioBusProxy.IsValid())
 			{
+				if (AudioBusProxy->NumChannels <= 0)
+				{
+					UE_LOG(LogMetaSound, Warning, TEXT("AudioBusProxy is invalid (NumChannels = %i)."), AudioBusProxy->NumChannels);
+					return;
+				}
 				if (FAudioDeviceManager* ADM = FAudioDeviceManager::Get())
 				{
 					if (FAudioDevice* AudioDevice = ADM->GetAudioDeviceRaw(AudioDeviceId))
