UE-213206 Bool Type has "button" widget option but no widget appears on the graph
#rnx

diff --git a/Engine/Plugins/Runtime/Metasound/Source/MetasoundEditor/Private/MetasoundNodeDetailCustomization.cpp b/Engine/Plugins/Runtime/Metasound/Source/MetasoundEditor/Private/MetasoundNodeDetailCustomization.cpp
--- a/Engine/Plugins/Runtime/Metasound/Source/MetasoundEditor/Private/MetasoundNodeDetailCustomization.cpp
+++ b/Engine/Plugins/Runtime/Metasound/Source/MetasoundEditor/Private/MetasoundNodeDetailCustomization.cpp
@@ -19,6 +19,7 @@
 #include "MetasoundFrontendController.h"
 #include "MetasoundFrontendDataTypeRegistry.h"
 #include "MetasoundFrontendRegistries.h"
+#include "MetasoundEditorSettings.h"
 #include "Modules/ModuleManager.h"
 #include "PropertyCustomizationHelpers.h"
 #include "PropertyEditorDelegates.h"
@@ -409,12 +410,16 @@ namespace Metasound
 				Frontend::FDataTypeRegistryInfo DataTypeInfo;
 				MemberCustomizationPrivate::GetDataTypeFromElementPropertyHandle(DefaultValueHandle, DataTypeInfo);
 
-				//TODO: Check if using AudioMaterialSlates, only then add.
-				//Only display widgetType option if this is not trigger
-				if (MemberCustomizationPrivate::GetPrimitiveTypeName(DataTypeInfo) != Metasound::GetMetasoundDataTypeName<Metasound::FTrigger>())
+				const UMetasoundEditorSettings* EditorSettings = GetDefault<UMetasoundEditorSettings>();
+				check(EditorSettings)
+
+				if (EditorSettings->bUseAudioMaterialWidgets)
 				{
-					IDetailCategoryBuilder& WidgetCategoryBuilder = InDetailLayout.EditCategory("EditorOptions");
-					DefaultRows.Add(WidgetCategoryBuilder.AddExternalObjectProperty(TArray<UObject*>({ DefaultBool }), GET_MEMBER_NAME_CHECKED(UMetasoundEditorGraphMemberDefaultBool, WidgetType)));
+					if (MemberCustomizationPrivate::GetPrimitiveTypeName(DataTypeInfo) != Metasound::GetMetasoundDataTypeName<Metasound::FTrigger>())
+					{
+						IDetailCategoryBuilder& WidgetCategoryBuilder = InDetailLayout.EditCategory("EditorOptions");
+						DefaultRows.Add(WidgetCategoryBuilder.AddExternalObjectProperty(TArray<UObject*>({ DefaultBool }), GET_MEMBER_NAME_CHECKED(UMetasoundEditorGraphMemberDefaultBool, WidgetType)));
+					}
 				}
 			}
 
