[MetasoundEngine] TAudioBusReaderOperator crash fix - BusProxy.Object was nullptr

diff --git a/Engine/Plugins/Runtime/Metasound/Source/MetasoundEngine/Private/MetasoundAudioBusReaderNode.cpp b/Engine/Plugins/Runtime/Metasound/Source/MetasoundEngine/Private/MetasoundAudioBusReaderNode.cpp
--- a/Engine/Plugins/Runtime/Metasound/Source/MetasoundEngine/Private/MetasoundAudioBusReaderNode.cpp
+++ b/Engine/Plugins/Runtime/Metasound/Source/MetasoundEngine/Private/MetasoundAudioBusReaderNode.cpp
@@ -154,13 +154,13 @@ namespace Metasound
 			const int32 NumSamplesToPop = BlockSizeFrames * AudioBusChannels;
 
 			const FAudioBusProxyPtr& BusProxy = AudioBusAsset->GetAudioBusProxy();
-			if (BusProxy->NumChannels <= 0)
+			if (!BusProxy.IsValid() || BusProxy->NumChannels <= 0)
 			{
 				// the audio bus is invalid / uninitialized
 				return;
 			}
 
-			if (BusProxy.IsValid() && BusProxy->AudioBusId != AudioBusId)
+			if (BusProxy->AudioBusId != AudioBusId)
 			{
 				AudioBusPatchOutput.Reset();
 			}
