Fix presets that require vertex updates crashing when attempting to migrate editor data

diff --git a/Engine/Plugins/Runtime/Metasound/Source/MetasoundEditor/Private/MetasoundEditorGraph.cpp b/Engine/Plugins/Runtime/Metasound/Source/MetasoundEditor/Private/MetasoundEditorGraph.cpp
--- a/Engine/Plugins/Runtime/Metasound/Source/MetasoundEditor/Private/MetasoundEditorGraph.cpp
+++ b/Engine/Plugins/Runtime/Metasound/Source/MetasoundEditor/Private/MetasoundEditorGraph.cpp
@@ -1853,6 +1853,13 @@ void UMetasoundEditorGraph::MigrateEditorDocumentData(FMetaSoundFrontendDocument
 		}
 #endif // WITH_EDITOR
 
+		// Presets get rebuilt anyway and may have invalid connections (newly referenced vertices)
+		// that need to be rebuilt later in asset load.
+		if (OutBuilder.IsPreset())
+		{
+			return;
+		}
+
 		// Cache data to be used in edge swapping below, being careful to not reference the input node pointer
 		// or vertex reference within the lower inner loop as the underlying node array may be reallocated by
 		// template nodes being added.
@@ -1904,20 +1911,30 @@ void UMetasoundEditorGraph::MigrateEditorDocumentData(FMetaSoundFrontendDocument
 						if (Linked)
 						{
 							const UMetasoundEditorGraphNode* ConnectedNode = CastChecked<const UMetasoundEditorGraphNode>(Linked->GetOwningNode());
-							const FMetasoundFrontendVertex* ConnectedInput = OutBuilder.FindNodeInput(ConnectedNode->GetNodeID(), Linked->GetFName());
-
-							// Swap connection from input node to connected node to now be from template node to connected node
-							EdgeToRemove.ToNodeID = ConnectedNode->GetNodeID(),
-							EdgeToRemove.ToVertexID = ConnectedInput->VertexID;
-							OutBuilder.RemoveEdge(EdgeToRemove);
-
-							OutBuilder.AddEdge(FMetasoundFrontendEdge
+							const FGuid ConnectedNodeID = ConnectedNode->GetNodeID();
+							if (const FMetasoundFrontendVertex* ConnectedInput = OutBuilder.FindNodeInput(ConnectedNodeID, Linked->GetFName()))
+							{
+								// Swap connection from input node to connected node to now be from template node to connected node
+								EdgeToRemove.ToNodeID = ConnectedNodeID,
+								EdgeToRemove.ToVertexID = ConnectedInput->VertexID;
+								OutBuilder.RemoveEdge(EdgeToRemove);
+
+								OutBuilder.AddEdge(FMetasoundFrontendEdge
+								{
+									TemplateNodeID,
+									TemplateOutputID,
+									ConnectedNodeID,
+									ConnectedInput->VertexID
+								});
+							}
+							else
 							{
-								TemplateNodeID,
-								TemplateOutputID,
-								ConnectedNode->GetNodeID(),
-								ConnectedInput->VertexID
-							});
+								UE_LOG(LogMetaSound, Warning, TEXT("Editor graph '%s' migration failed to find node '%s' class output '%s': Ignoring connection upgrade from input '%s'"),
+									*OutBuilder.GetDebugName(),
+									*ConnectedNode->GetDisplayName().ToString(),
+									*Linked->GetName(),
+									*InputName.ToString());
+							}
 						}
 					}
 				}
@@ -1926,12 +1943,16 @@ void UMetasoundEditorGraph::MigrateEditorDocumentData(FMetaSoundFrontendDocument
 	});
 
 	// 4. Add comment nodes as builder graph comments to frontend document
-	TArray<UEdGraphNode_Comment*> CommentNodes;
-	GetNodesOfClass(CommentNodes);
-	for (UEdGraphNode_Comment* Node : CommentNodes)
+	// (No need to propagate comments for presets)
+	if (!OutBuilder.IsPreset())
 	{
-		FMetaSoundFrontendGraphComment& NewComment = OutBuilder.FindOrAddGraphComment(FGuid::NewGuid());
-		UMetasoundEditorGraphCommentNode::ConvertToFrontendComment(*Node, NewComment);
+		TArray<UEdGraphNode_Comment*> CommentNodes;
+		GetNodesOfClass(CommentNodes);
+		for (UEdGraphNode_Comment* Node : CommentNodes)
+		{
+			FMetaSoundFrontendGraphComment& NewComment = OutBuilder.FindOrAddGraphComment(FGuid::NewGuid());
+			UMetasoundEditorGraphCommentNode::ConvertToFrontendComment(*Node, NewComment);
+		}
 	}
 
 	// 5. Remove input locations and ensure that all other nodes only have at most one
