Fix being able to delete uneditable MetaSound nodes when deleted in a group with editable nodes

diff --git a/Engine/Plugins/Runtime/Metasound/Source/MetasoundEditor/Private/MetasoundEditor.cpp b/Engine/Plugins/Runtime/Metasound/Source/MetasoundEditor/Private/MetasoundEditor.cpp
--- a/Engine/Plugins/Runtime/Metasound/Source/MetasoundEditor/Private/MetasoundEditor.cpp
+++ b/Engine/Plugins/Runtime/Metasound/Source/MetasoundEditor/Private/MetasoundEditor.cpp
@@ -2329,7 +2329,10 @@ namespace Metasound
 			{
 				if (UMetasoundEditorGraphNode* Node = Cast<UMetasoundEditorGraphNode>(NodeObj))
 				{
-					Node->RemoveFromDocument();
+					if (Node->CanUserDeleteNode())
+					{
+						Node->RemoveFromDocument();
+					}
 				}
 				else if (UMetasoundEditorGraphCommentNode* CommentNode = Cast<UMetasoundEditorGraphCommentNode>(NodeObj))
 				{
