UE-171014
Shortcut way to safely 'clear all unused inputs' from MetaSounds

diff --git a/Engine/Plugins/Runtime/Metasound/Source/MetasoundEditor/Private/MetasoundEditor.cpp b/Engine/Plugins/Runtime/Metasound/Source/MetasoundEditor/Private/MetasoundEditor.cpp
--- a/Engine/Plugins/Runtime/Metasound/Source/MetasoundEditor/Private/MetasoundEditor.cpp
+++ b/Engine/Plugins/Runtime/Metasound/Source/MetasoundEditor/Private/MetasoundEditor.cpp
@@ -3085,6 +3085,83 @@ namespace Metasound
 			}
 		}
 
+		void FEditor::DeleteAllUnusedInSection()
+		{
+			TArray<TSharedPtr<FMetasoundGraphMemberSchemaAction>> ActionsToDelete;
+			TArray<TSharedPtr<FEdGraphSchemaAction>> Actions;		
+			GraphMembersMenu->GetSelectedCategorySubActions(Actions);
+
+			for (TSharedPtr<FEdGraphSchemaAction> Action : Actions)
+			{				
+				const TSharedPtr<FMetasoundGraphMemberSchemaAction> MetasoundAction = StaticCastSharedPtr<FMetasoundGraphMemberSchemaAction>(Action);
+				if (MetasoundAction.IsValid())
+				{
+					if (const UMetasoundEditorGraphMember* GraphMember = MetasoundAction->GetGraphMember())
+					{
+						const TArray<UMetasoundEditorGraphMemberNode*> Nodes = GraphMember->GetNodes();
+
+						if (Nodes.Num() > 0)
+						{
+							continue;
+						}
+						else
+						{
+							const FMetasoundFrontendVersion* InterfaceVersion = nullptr;
+							if (const UMetasoundEditorGraphVertex* Vertex = Cast<UMetasoundEditorGraphVertex>(GraphMember))
+							{
+								InterfaceVersion = &Vertex->GetInterfaceVersion();
+							}
+
+							// Interface members cannot be deleted
+							const bool bIsInterfaceMember = InterfaceVersion && InterfaceVersion->IsValid();
+							if (!bIsInterfaceMember)
+							{
+								ActionsToDelete.Add(MetasoundAction);														
+							}
+						}
+					}					
+				}
+			}
+
+			for (TSharedPtr<FMetasoundGraphMemberSchemaAction> Action : ActionsToDelete)
+			{
+				DeleteInterfaceItem(Action);
+			}
+		}
+
+		bool FEditor::CanDeleteUnusedMembers() const
+		{
+			if (!IsGraphEditable())
+			{
+				return false;
+			}
+
+			if (!GraphMembersMenu.IsValid())
+			{
+				return false;
+			}
+
+			TArray<TSharedPtr<FEdGraphSchemaAction>> SelectedActions;
+			GraphMembersMenu->GetSelectedActions(SelectedActions);
+
+			if (SelectedActions.IsEmpty())
+			{
+				return true;
+			}
+
+			//Option to delete unused member becomes pressable only if right-click happened on non valid Action.
+			for (const TSharedPtr<FEdGraphSchemaAction>& SelectedAction : SelectedActions)
+			{
+				TSharedPtr<FMetasoundGraphMemberSchemaAction> MetasoundAction = StaticCastSharedPtr<FMetasoundGraphMemberSchemaAction>(SelectedAction);
+				if (MetasoundAction.IsValid())
+				{					
+					return false;				
+				}
+			}
+			
+			return true;
+		}
+
 		FActionMenuContent FEditor::OnCreateGraphActionMenu(UEdGraph* InGraph, const FVector2D& InNodePosition, const TArray<UEdGraphPin*>& InDraggedPins, bool bAutoExpand, SGraphEditor::FActionMenuClosed InOnMenuClosed)
 		{
 			TSharedRef<SMetasoundActionMenu> ActionMenu = SNew(SMetasoundActionMenu)
@@ -3164,6 +3241,14 @@ namespace Metasound
 				FUIAction(
 					FExecuteAction::CreateSP(this, &FEditor::JumpToNodesForSelectedInterfaceItem), 
 					FCanExecuteAction::CreateSP(this, &FEditor::CanJumpToNodesForSelectedInterfaceItem)));
+					
+			MenuBuilder.AddMenuEntry(
+				LOCTEXT("DeleteAllUnusedInSection", "Delete Unused Members"),
+				LOCTEXT("DeleteAllUnusedInSectionTooltip", "Delete all Unused Members under this Header"),
+				FSlateIcon(),
+				FUIAction(
+					FExecuteAction::CreateSP(this, &FEditor::DeleteAllUnusedInSection),
+					FCanExecuteAction::CreateSP(this, &FEditor::CanDeleteUnusedMembers)));
 
 			return MenuBuilder.MakeWidget();
 		}
diff --git a/Engine/Plugins/Runtime/Metasound/Source/MetasoundEditor/Private/MetasoundEditor.h b/Engine/Plugins/Runtime/Metasound/Source/MetasoundEditor/Private/MetasoundEditor.h
--- a/Engine/Plugins/Runtime/Metasound/Source/MetasoundEditor/Private/MetasoundEditor.h
+++ b/Engine/Plugins/Runtime/Metasound/Source/MetasoundEditor/Private/MetasoundEditor.h
@@ -321,6 +321,12 @@ namespace Metasound
 
 			/** Jumps to the nodes corresponding to the first valid currently selected interface item. */
 			void JumpToNodesForSelectedInterfaceItem();
+			
+			/** Delete all unused members from the selected section*/
+			void DeleteAllUnusedInSection();
+
+			/** Whether the selection is not to a valid member*/
+			bool CanDeleteUnusedMembers() const;
 
 			/** Called to undo the last action */
 			void UndoGraphAction();
