Refactor MetaSound Array Random Get Node shared state
- Add input for Shared State Behavior to share state between other array random nodes that are the same node, same node in composition, or same data
- Simplify behavior for reinitializing and updating state
- Add hidden friend GetHashType hash functions to MetaSound types
- Add output for chosen index
- Tooltip updates and debug logging

diff --git a/Engine/Plugins/Runtime/Metasound/Source/MetasoundEngine/Private/MetasoundSource.cpp b/Engine/Plugins/Runtime/Metasound/Source/MetasoundEngine/Private/MetasoundSource.cpp
--- a/Engine/Plugins/Runtime/Metasound/Source/MetasoundEngine/Private/MetasoundSource.cpp
+++ b/Engine/Plugins/Runtime/Metasound/Source/MetasoundEngine/Private/MetasoundSource.cpp
@@ -1027,6 +1027,9 @@ ISoundGeneratorPtr UMetaSoundSource::CreateSoundGenerator(const FSoundGeneratorI
 				TArray<FAudioParameter> MergedParameters;
 				MergePresetOverridesAndSuppliedDefaults(InDefaultParameters, MergedParameters);
 
+				// Update Graph Hierarchy with the asset unique ID (the hierarchy was collapsed by only using the base graph)
+				Environment.SetValue<TArray<FGuid>>(OperatorBuilder::Environment::GraphHierarchy, TArray<FGuid>({ AssetClassID }));
+
 				// Create generator.
 				FMetasoundGeneratorInitParams InitParams
 				{
diff --git a/Engine/Plugins/Runtime/Metasound/Source/MetasoundEngineTest/Private/EngineTestMetaSoundAutomatedNodeTest.cpp b/Engine/Plugins/Runtime/Metasound/Source/MetasoundEngineTest/Private/EngineTestMetaSoundAutomatedNodeTest.cpp
--- a/Engine/Plugins/Runtime/Metasound/Source/MetasoundEngineTest/Private/EngineTestMetaSoundAutomatedNodeTest.cpp
+++ b/Engine/Plugins/Runtime/Metasound/Source/MetasoundEngineTest/Private/EngineTestMetaSoundAutomatedNodeTest.cpp
@@ -16,6 +16,7 @@
 #include "MetasoundFrontendRegistries.h"
 #include "MetasoundFrontendSearchEngine.h"
 #include "MetasoundLog.h"
+#include "MetasoundOperatorBuilder.h"
 #include "MetasoundPrimitives.h"
 #include "MetasoundVariableNodes.h"
 #include "MetasoundVertex.h"
@@ -53,6 +54,7 @@ namespace Metasound::EngineTest{
 		Environment.SetValue<bool>(SourceInterface::Environment::IsPreview, false);
 		Environment.SetValue<uint64>(SourceInterface::Environment::TransmitterID, 0);
 		Environment.SetValue<FString>(SourceInterface::Environment::GraphName, TEXT("ENGINE_TEST_REGISTERED_NODES"));
+		Environment.SetValue<TArray<FGuid>>(OperatorBuilder::Environment::GraphHierarchy, TArray<FGuid>({ FGuid() }));
 
 		if (Audio::FMixerDevice* MixerDevice = GetAudioMixerDevice())
 		{
diff --git a/Engine/Plugins/Runtime/Metasound/Source/MetasoundFrontend/Private/MetasoundArrayNodes.cpp b/Engine/Plugins/Runtime/Metasound/Source/MetasoundFrontend/Private/MetasoundArrayNodes.cpp
--- a/Engine/Plugins/Runtime/Metasound/Source/MetasoundFrontend/Private/MetasoundArrayNodes.cpp
+++ b/Engine/Plugins/Runtime/Metasound/Source/MetasoundFrontend/Private/MetasoundArrayNodes.cpp
@@ -8,13 +8,13 @@ namespace Metasound
 {
 	namespace MetasoundArrayNodesPrivate
 	{
-		FNodeClassMetadata CreateArrayNodeClassMetadata(const FName& InDataTypeName, const FName& InOperatorName, const FText& InDisplayName, const FText& InDescription, const FVertexInterface& InDefaultInterface)
+		FNodeClassMetadata CreateArrayNodeClassMetadata(const FName& InDataTypeName, const FName& InOperatorName, const FText& InDisplayName, const FText& InDescription, const FVertexInterface& InDefaultInterface, int32 MajorVersion, int32 MinorVersion, bool bIsDeprecated)
 		{
 			FNodeClassMetadata Metadata
 			{
 				FNodeClassName{FName("Array"), InOperatorName, InDataTypeName},
-				1, // Major Version
-				0, // Minor Version
+				MajorVersion, 
+				MinorVersion,
 				InDisplayName, 
 				InDescription,
 				PluginAuthor,
@@ -22,7 +22,8 @@ namespace Metasound
 				InDefaultInterface,
 				{ METASOUND_LOCTEXT("ArrayCategory", "Array") },
 				{ METASOUND_LOCTEXT("MetasoundArrayKeyword", "Array") },
-				FNodeDisplayStyle{}
+				FNodeDisplayStyle{}, 
+				bIsDeprecated
 			};
 
 			return Metadata;
diff --git a/Engine/Plugins/Runtime/Metasound/Source/MetasoundFrontend/Private/MetasoundArrayRandomNode.cpp b/Engine/Plugins/Runtime/Metasound/Source/MetasoundFrontend/Private/MetasoundArrayRandomNode.cpp
--- a/Engine/Plugins/Runtime/Metasound/Source/MetasoundFrontend/Private/MetasoundArrayRandomNode.cpp
+++ b/Engine/Plugins/Runtime/Metasound/Source/MetasoundFrontend/Private/MetasoundArrayRandomNode.cpp
@@ -1,7 +1,7 @@
 // Copyright Epic Games, Inc. All Rights Reserved.
 
 #include "MetasoundArrayRandomNode.h"
-#include "CoreMinimal.h"
+#include "Containers/CircularQueue.h"
 #include "Internationalization/Text.h"
 #include "MetasoundBuilderInterface.h"
 #include "MetasoundDataFactory.h"
@@ -13,40 +13,63 @@
 #include "MetasoundOperatorInterface.h"
 #include "MetasoundPrimitives.h"
 #include "MetasoundTrigger.h"
-#include "Containers/CircularQueue.h"
 #include "MetasoundArrayNodes.h"
+#include "MetasoundEnumRegistrationMacro.h"
 
 #include <type_traits>
 
+#define LOCTEXT_NAMESPACE "MetasoundFrontend"
 
 namespace Metasound
 {
+	DEFINE_METASOUND_ENUM_BEGIN(ESharedStateBehaviorType, FEnumSharedStateBehaviorType, "SharedStateBehaviorType")
+		DEFINE_METASOUND_ENUM_ENTRY(ESharedStateBehaviorType::SameNode, \
+			"SameNodeDescription", "Same Node", "SameNodeTT", \
+			"State is shared with other instances of this individual node regardless of the MetaSound it is in."),
+		DEFINE_METASOUND_ENUM_ENTRY(ESharedStateBehaviorType::SameNodeInComposition, \
+			"SameNodeInCompositionDescription", "Same Node in Composition", "SSameNodeInCompositionTT", \
+			"State is shared with other instances of this node with the same parent MetaSound graph(s). Useful for differentiating shared state between nodes used in different presets or multiple composed graphs."),
+		DEFINE_METASOUND_ENUM_ENTRY(ESharedStateBehaviorType::SameData, \
+			"SameDataDescription", "Same Data", "SameDataTT", \
+			"State is shared with other nodes with the same input array data (by value) regardless of where the node is located. Useful for sharing state regardless of graph composition or between multiple nodes within a single MetaSound using the same input data. Input array type must implement a hash function.")
+	DEFINE_METASOUND_ENUM_END()
+
 	FArrayRandomGet::FArrayRandomGet(int32 InSeed, int32 InMaxIndex, const TArray<float>& InWeights, int32 InNoRepeatOrder)
 	{
-		Init(InSeed, InMaxIndex, InWeights, InNoRepeatOrder);
+		UpdateState(InSeed, InMaxIndex, InWeights, InNoRepeatOrder);
 	}
 
 	void FArrayRandomGet::Init(int32 InSeed, int32 InMaxIndex, const TArray<float>& InWeights, int32 InNoRepeatOrder)
+	{
+		UpdateState(InSeed, InMaxIndex, InWeights, InNoRepeatOrder);
+	}
+
+	void FArrayRandomGet::UpdateState(int32 InSeed, int32 InMaxIndex, const TArray<float>& InWeights, int32 InNoRepeatOrder)
 	{
 		SetSeed(InSeed);
 		MaxIndex = InMaxIndex;
 		SetNoRepeatOrder(InNoRepeatOrder);
 		check(!InNoRepeatOrder || !PreviousIndicesQueue->IsFull());
-		RandomWeights = InWeights;
+		SetRandomWeights(InWeights);
 	}
 
 	void FArrayRandomGet::SetSeed(int32 InSeed)
 	{
-		if (InSeed == INDEX_NONE)
+		if (!bRandomStreamInitialized || InSeed != Seed)
 		{
-			RandomStream.Initialize(FPlatformTime::Cycles());
-		}
-		else
-		{
-			RandomStream.Initialize(InSeed);
-		}
+			Seed = InSeed;
+			if (InSeed == INDEX_NONE)
+			{
+				RandomStream.Initialize(FPlatformTime::Cycles());
+			}
+			else
+			{
+				RandomStream.Initialize(InSeed);
+			}
 
-		ResetSeed();
+			ResetSeed();
+			bRandomStreamInitialized = true;
+		}
 	}
 
 	void FArrayRandomGet::SetNoRepeatOrder(int32 InNoRepeatOrder)
@@ -64,7 +87,10 @@ namespace Metasound
 
 	void FArrayRandomGet::SetRandomWeights(const TArray<float>& InRandomWeights)
 	{
-		RandomWeights = InRandomWeights;
+		if (InRandomWeights != RandomWeights)
+		{
+			RandomWeights = InRandomWeights;
+		}
 	}
 
 	void FArrayRandomGet::ResetSeed()
@@ -219,6 +245,14 @@ namespace Metasound
 		return (*RG)->NextValue();
 	}
 
+	int32 FSharedStateRandomGetManager::NextValue(const FGuid& InSharedStateId, InitSharedStateArgs& InStateArgs)
+	{
+		FScopeLock Lock(&CritSect);
+		InitOrUpdate(InStateArgs);
+		TUniquePtr<FArrayRandomGet>* RG = RandomGets.Find(InSharedStateId);
+		return (*RG)->NextValue();
+	}
+
 	void FSharedStateRandomGetManager::SetSeed(const FGuid& InSharedStateId, int32 InSeed)
 	{
 		FScopeLock Lock(&CritSect);
@@ -246,5 +280,24 @@ namespace Metasound
 		TUniquePtr<FArrayRandomGet>* RG = RandomGets.Find(InSharedStateId);
 		(*RG)->ResetSeed();
 	}
+	
+	void FSharedStateRandomGetManager::ResetSeed(const FGuid& InSharedStateId, InitSharedStateArgs& InStateArgs)
+	{
+		FScopeLock Lock(&CritSect);
+		InitOrUpdate(InStateArgs);
+		TUniquePtr<FArrayRandomGet>* RG = RandomGets.Find(InSharedStateId);
+		(*RG)->ResetSeed();
+	}
 
+	void FSharedStateRandomGetManager::InitOrUpdate(InitSharedStateArgs& InStateArgs)
+	{
+		if (TUniquePtr<FArrayRandomGet>* FoundExisting = RandomGets.Find(InStateArgs.SharedStateId))
+		{
+			(*FoundExisting)->UpdateState(InStateArgs.Seed, InStateArgs.NumElements, InStateArgs.Weights, InStateArgs.NoRepeatOrder);
+		}
+		else
+		{
+			RandomGets.Add(InStateArgs.SharedStateId, MakeUnique<FArrayRandomGet>(InStateArgs.Seed, InStateArgs.NumElements, InStateArgs.Weights, InStateArgs.NoRepeatOrder));
+		}
+	}
 }
diff --git a/Engine/Plugins/Runtime/Metasound/Source/MetasoundFrontend/Public/MetasoundArrayNodes.h b/Engine/Plugins/Runtime/Metasound/Source/MetasoundFrontend/Public/MetasoundArrayNodes.h
--- a/Engine/Plugins/Runtime/Metasound/Source/MetasoundFrontend/Public/MetasoundArrayNodes.h
+++ b/Engine/Plugins/Runtime/Metasound/Source/MetasoundFrontend/Public/MetasoundArrayNodes.h
@@ -26,7 +26,7 @@ namespace Metasound
 	namespace MetasoundArrayNodesPrivate
 	{
 		// Convenience function for make FNodeClassMetadata of array nodes.
-		METASOUNDFRONTEND_API FNodeClassMetadata CreateArrayNodeClassMetadata(const FName& InDataTypeName, const FName& InOperatorName, const FText& InDisplayName, const FText& InDescription, const FVertexInterface& InDefaultInterface);
+		METASOUNDFRONTEND_API FNodeClassMetadata CreateArrayNodeClassMetadata(const FName& InDataTypeName, const FName& InOperatorName, const FText& InDisplayName, const FText& InDescription, const FVertexInterface& InDefaultInterface, int32 MajorVersion=1, int32 MinorVersion=0, bool bIsDeprecated=false);
 
 		// Retrieve the ElementType from an ArrayType
 		template<typename ArrayType>
diff --git a/Engine/Plugins/Runtime/Metasound/Source/MetasoundFrontend/Public/MetasoundArrayRandomNode.h b/Engine/Plugins/Runtime/Metasound/Source/MetasoundFrontend/Public/MetasoundArrayRandomNode.h
--- a/Engine/Plugins/Runtime/Metasound/Source/MetasoundFrontend/Public/MetasoundArrayRandomNode.h
+++ b/Engine/Plugins/Runtime/Metasound/Source/MetasoundFrontend/Public/MetasoundArrayRandomNode.h
@@ -1,33 +1,87 @@
 // Copyright Epic Games, Inc. All Rights Reserved.
 #pragma once
 
+#include "Concepts/GetTypeHashable.h"
 #include "Containers/CircularQueue.h"
 #include "Interfaces/MetasoundFrontendSourceInterface.h"
+#include "Templates/Models.h"
 #include "MetasoundArrayNodes.h"
+#include "MetasoundEnum.h"
 #include "MetasoundExecutableOperator.h"
 #include "MetasoundFacade.h"
 #include "MetasoundNodeInterface.h"
+#include "MetasoundOperatorBuilder.h"
 #include "MetasoundParamHelper.h"
 #include "MetasoundTrigger.h"
 
 #define LOCTEXT_NAMESPACE "MetasoundFrontend"
 
+namespace MetasoundArrayHashPrivate
+{
+	template<typename ElementType>
+	FGuid GetArrayContentHashGuid(const TArray<ElementType>& InArray)
+	{
+		if constexpr (TModels_V<CGetTypeHashable, ElementType>)
+		{
+			uint32 A = GetTypeHash(Metasound::GetMetasoundDataTypeName<ElementType>());
+			uint32 B = A; uint32 C = A; uint32 D = A;
+			for (int32 i = 0; i < InArray.Num(); i++)
+			{
+				const int32 Pos = i % 4;
+				switch (Pos)
+				{
+				case 0: 
+					A = HashCombineFast(A, GetTypeHash(InArray[i]));
+					break;
+				case 1:
+					B = HashCombineFast(B, GetTypeHash(InArray[i]));
+					break;
+				case 2:
+					C = HashCombineFast(C, GetTypeHash(InArray[i]));
+					break;
+				case 3:
+					D = HashCombineFast(D, GetTypeHash(InArray[i]));
+					break;
+				}
+			}
+			return FGuid(A, B, C, D);
+		}
+		else
+		{
+			UE_LOG(LogMetaSound, Warning, TEXT("Array Random Get: Please implement \"uint32 GetTypeHash(const T&)\" for type %s to use Same Data for Shared State Behavior."), *Metasound::GetMetasoundDataTypeString<ElementType>());
+			return FGuid();
+		}
+	}
+}
+
 namespace Metasound
 {
 	namespace ArrayNodeRandomGetVertexNames
 	{
 		METASOUND_PARAM(InputTriggerNextValue, "Next", "Trigger to get the next value in the randomized array.")
 		METASOUND_PARAM(InputTriggerResetSeed, "Reset", "Trigger to reset the seed for the randomized array.")
-		METASOUND_PARAM(InputRandomArray, "In Array", "Input array to randomized.")
+		METASOUND_PARAM(InputRandomArray, "In Array", "Input array to randomize.")
 		METASOUND_PARAM(InputWeights, "Weights", "Input array of weights to use for random selection. Will repeat if this array is shorter than the input array to select from.")
-		METASOUND_PARAM(InputSeed, "Seed", "Seed to use for the random shuffle.")
+		METASOUND_PARAM(InputSeed, "Seed", "Seed to use for the random stream. Set to -1 to use a random seed.")
 		METASOUND_PARAM(InputNoRepeatOrder, "No Repeats", "The number of elements to track to avoid repeating in a row.")
-		METASOUND_PARAM(InputEnableSharedState, "Enable Shared State", "Set to enabled to share state across instances of this MetaSound.")
+		METASOUND_PARAM(InputEnableSharedState, "Enable Shared State", "Set to enabled to share state with other Random Get (Array) nodes. Does not apply when previewing in the MetaSound editor; use PIE or game.")
+		METASOUND_PARAM(InputSharedStateBehavior, "Shared State Behavior", "The behavior for how state is shared with other Random Get (Array) nodes. Only applied when Enable Shared State is true.")
 		METASOUND_PARAM(OutputTriggerOnNext, "On Next", "Triggers when the \"Next\" input is triggered.")
-		METASOUND_PARAM(OutputTriggerOnReset, "On Reset", "Triggers when the \"Shuffle\" input is triggered or if the array is auto-shuffled.")
-		METASOUND_PARAM(ShuffleOutputValue, "Value", "Value of the current shuffled element.")
+		METASOUND_PARAM(OutputTriggerOnReset, "On Reset", "Triggers when the \"Reset\" input is triggered.")
+		METASOUND_PARAM(ShuffleOutputValue, "Value", "Value of the current random element.")
+		METASOUND_PARAM(OutputIndex, "Index", "Array index of the current random element.")
 	}
 
+	enum class ESharedStateBehaviorType : int32
+	{
+		SameNode,
+		SameNodeInComposition, 
+		SameData
+	};
+
+	DECLARE_METASOUND_ENUM(ESharedStateBehaviorType, ESharedStateBehaviorType::SameNodeInComposition, METASOUNDFRONTEND_API,
+		FEnumSharedStateBehaviorType, FEnumSharedStateBehaviorTypeInfo, FEnumSharedStateBehaviorTypeReadRef, FSharedStateBehaviorTypeWriteRef);
+
 	class METASOUNDFRONTEND_API FArrayRandomGet
 	{
 	public:
@@ -35,7 +89,9 @@ namespace Metasound
 		FArrayRandomGet(int32 InSeed, int32 InMaxIndex, const TArray<float>& InWeights, int32 InNoRepeatOrder);
 		~FArrayRandomGet() = default;
 
+		UE_DEPRECATED(5.5, "Use UpdateState instead")
 		void Init(int32 InSeed, int32 InMaxIndex, const TArray<float>& InWeights, int32 InNoRepeatOrder);
+		void UpdateState(int32 InSeed, int32 InMaxIndex, const TArray<float>& InWeights, int32 InNoRepeatOrder);
 		void SetSeed(int32 InSeed);
 		void SetNoRepeatOrder(int32 InNoRepeatOrder);
 		void SetRandomWeights(const TArray<float>& InRandomWeights);
@@ -59,6 +115,8 @@ namespace Metasound
 
 		// Random stream to use to randomize the shuffling
 		FRandomStream RandomStream;
+		int32 Seed = INDEX_NONE;
+		bool bRandomStreamInitialized = false;
 	};
 
 	struct InitSharedStateArgs
@@ -77,12 +135,23 @@ namespace Metasound
 		static FSharedStateRandomGetManager& Get();
 
 		void InitSharedState(InitSharedStateArgs& InArgs);
+		// Initialize or update state for a given shared state id. No lock, so call this function within one if needed
+		void InitOrUpdate(InitSharedStateArgs& InStateArgs);
+
+		// Get the next array index 
+		// Init or update state with the given args, then return next value (within a single lock operation)
+		int32 NextValue(const FGuid& InSharedStateId, InitSharedStateArgs& InStateArgs);
 		int32 NextValue(const FGuid& InSharedStateId);
+
 		void SetSeed(const FGuid& InSharedStateId, int32 InSeed);
 		void SetNoRepeatOrder(const FGuid& InSharedStateId, int32 InNoRepeatOrder);
 		void SetRandomWeights(const FGuid& InSharedStateId, const TArray<float>& InRandomWeights);
+
+		// Init or update state with the given args, then reset seed (within a single lock operation)
+		void ResetSeed(const FGuid& InSharedStateId, InitSharedStateArgs& InStateArgs);
 		void ResetSeed(const FGuid& InSharedStateId);
 
+
 	private:
 		FSharedStateRandomGetManager() = default;
 		~FSharedStateRandomGetManager() = default;
@@ -114,12 +183,14 @@ namespace Metasound
 					TInputDataVertex<WeightsArrayType>(METASOUND_GET_PARAM_NAME_AND_METADATA(InputWeights)),
 					TInputDataVertex<int32>(METASOUND_GET_PARAM_NAME_AND_METADATA(InputSeed), -1),
 					TInputDataVertex<int32>(METASOUND_GET_PARAM_NAME_AND_METADATA(InputNoRepeatOrder), 1),
-					TInputDataVertex<bool>(METASOUND_GET_PARAM_NAME_AND_METADATA(InputEnableSharedState), false)
+					TInputDataVertex<bool>(METASOUND_GET_PARAM_NAME_AND_METADATA(InputEnableSharedState), false),
+					TInputConstructorVertex<FEnumSharedStateBehaviorType>(METASOUND_GET_PARAM_NAME_AND_METADATA(InputSharedStateBehavior), (int32)ESharedStateBehaviorType::SameNodeInComposition)
 				),
 				FOutputVertexInterface(
 					TOutputDataVertex<FTrigger>(METASOUND_GET_PARAM_NAME_AND_METADATA(OutputTriggerOnNext)),
 					TOutputDataVertex<FTrigger>(METASOUND_GET_PARAM_NAME_AND_METADATA(OutputTriggerOnReset)),
-					TOutputDataVertex<ElementType>(METASOUND_GET_PARAM_NAME_AND_METADATA(ShuffleOutputValue))
+					TOutputDataVertex<ElementType>(METASOUND_GET_PARAM_NAME_AND_METADATA(ShuffleOutputValue)),
+					TOutputDataVertex<int32>(METASOUND_GET_PARAM_NAME_AND_METADATA(OutputIndex))
 				)
 			);
 
@@ -136,7 +207,7 @@ namespace Metasound
 				FText NodeDescription = METASOUND_LOCTEXT("RandomArrayGetNode_Description", "Randomly retrieve data from input array using the supplied weights.");
 				FVertexInterface NodeInterface = GetDefaultInterface();
 
-				return MetasoundArrayNodesPrivate::CreateArrayNodeClassMetadata(DataTypeName, OperatorName, NodeDisplayName, NodeDescription, NodeInterface);
+				return MetasoundArrayNodesPrivate::CreateArrayNodeClassMetadata(DataTypeName, OperatorName, NodeDisplayName, NodeDescription, NodeInterface, /*MajorVersion=*/1, /*MinorVersion=*/1);
 			};
 
 			static const FNodeClassMetadata Metadata = CreateNodeClassMetadata();
@@ -158,8 +229,9 @@ namespace Metasound
 			FInt32ReadRef InSeedValue = InputData.GetOrCreateDefaultDataReadReference<int32>(METASOUND_GET_PARAM_NAME(InputSeed), InParams.OperatorSettings);
 			FInt32ReadRef InNoRepeatOrder = InputData.GetOrCreateDefaultDataReadReference<int32>(METASOUND_GET_PARAM_NAME(InputNoRepeatOrder), InParams.OperatorSettings);
 			FBoolReadRef bInEnableSharedState = InputData.GetOrCreateDefaultDataReadReference<bool>(METASOUND_GET_PARAM_NAME(InputEnableSharedState), InParams.OperatorSettings);
+			FEnumSharedStateBehaviorType InSharedStateBehavior = InputData.GetOrCreateDefaultValue<FEnumSharedStateBehaviorType>(METASOUND_GET_PARAM_NAME(InputSharedStateBehavior), InParams.OperatorSettings);
 
-			return MakeUnique<TArrayRandomGetOperator<ArrayType>>(InParams, InTriggerNext, InTriggerReset, InInputArray, InInputWeightsArray, InSeedValue, InNoRepeatOrder, bInEnableSharedState);
+			return MakeUnique<TArrayRandomGetOperator<ArrayType>>(InParams, InTriggerNext, InTriggerReset, InInputArray, InInputWeightsArray, InSeedValue, InNoRepeatOrder, bInEnableSharedState, InSharedStateBehavior);
 		}
 
 		TArrayRandomGetOperator(
@@ -170,22 +242,22 @@ namespace Metasound
 			const TDataReadReference<WeightsArrayType>& InInputWeightsArray,
 			const FInt32ReadRef& InSeedValue,
 			const FInt32ReadRef& InNoRepeatOrder,
-			const FBoolReadRef& bInEnableSharedState)
+			const FBoolReadRef& bInEnableSharedState,
+			const FEnumSharedStateBehaviorType InSharedStateBehavior)
 			: TriggerNext(InTriggerNext)
 			, TriggerReset(InTriggerReset)
 			, InputArray(InInputArray)
 			, InputWeightsArray(InInputWeightsArray)
 			, SeedValue(InSeedValue)
 			, NoRepeatOrder(InNoRepeatOrder)
+			, bEnableSharedState(bInEnableSharedState)
+			, SharedStateBehavior(InSharedStateBehavior)
 			, TriggerOnNext(FTriggerWriteRef::CreateNew(InParams.OperatorSettings))
 			, TriggerOnReset(FTriggerWriteRef::CreateNew(InParams.OperatorSettings))
 			, OutValue(TDataWriteReferenceFactory<ElementType>::CreateAny(InParams.OperatorSettings))
-			, bEnableSharedState(bInEnableSharedState)
+			, OutIndex(FInt32WriteRef::CreateNew(INDEX_NONE))			
 		{
-			using namespace Frontend;
-
-			SharedStateUniqueId = InParams.Node.GetInstanceID();
-
+			NodeId = InParams.Node.GetInstanceID();
 			Reset(InParams);
 		}
 
@@ -201,6 +273,7 @@ namespace Metasound
 			InOutVertexData.BindReadVertex(METASOUND_GET_PARAM_NAME(InputSeed), SeedValue);
 			InOutVertexData.BindReadVertex(METASOUND_GET_PARAM_NAME(InputNoRepeatOrder), NoRepeatOrder);
 			InOutVertexData.BindReadVertex(METASOUND_GET_PARAM_NAME(InputEnableSharedState), bEnableSharedState);
+			InOutVertexData.SetValue(METASOUND_GET_PARAM_NAME(InputSharedStateBehavior), SharedStateBehavior);
 		}
 
 		virtual void BindOutputs(FOutputVertexInterfaceData& InOutVertexData) override
@@ -209,6 +282,7 @@ namespace Metasound
 			InOutVertexData.BindReadVertex(METASOUND_GET_PARAM_NAME(OutputTriggerOnNext), TriggerOnNext);
 			InOutVertexData.BindReadVertex(METASOUND_GET_PARAM_NAME(OutputTriggerOnReset), TriggerOnReset);
 			InOutVertexData.BindReadVertex(METASOUND_GET_PARAM_NAME(ShuffleOutputValue), OutValue);
+			InOutVertexData.BindReadVertex(METASOUND_GET_PARAM_NAME(OutputIndex), OutIndex);
 		}
 
 		virtual FDataReferenceCollection GetInputs() const override
@@ -235,27 +309,43 @@ namespace Metasound
 			{
 				GraphName = *InParams.Environment.GetValue<FString>(SourceInterface::Environment::GraphName);
 			}
+
+			TOptional<FName> EnumName = FEnumSharedStateBehaviorType::ToName(SharedStateBehavior);
+			if (EnumName.IsSet())
+			{
+				DebugSharedStateBehaviorString = EnumName.GetValue().ToString();
+			}
 #endif // WITH_METASOUND_DEBUG_ENVIRONMENT
 			bIsPreviewSound = InParams.Environment.GetValue<bool>(SourceInterface::Environment::IsPreview);
-
-			// Check to see if this is a global shuffler or a local one. 
-			// Global shuffler will use a namespace to opt into it.
-			PrevSeedValue = *SeedValue;
 			
-			WeightsArray = *InputWeightsArray;
-
-			const ArrayType& InputArrayRef = *InputArray;
-			PrevArraySize = InputArrayRef.Num();
-			PrevNoRepeatOrder = FMath::Clamp(*NoRepeatOrder, 0, PrevArraySize - 1);
-
 			*OutValue = TDataTypeFactory<ElementType>::CreateAny(InParams.OperatorSettings);
+			*OutIndex = INDEX_NONE;
 			TriggerOnNext->Reset();
 			TriggerOnReset->Reset();
-		}
 
-		bool UseSharedState() const
-		{
-			return *bEnableSharedState && bSharedStateInitialized;
+			// Cache shared state id for shared state behavior types that cannot be changed after node init
+			if (SharedStateBehavior == ESharedStateBehaviorType::SameNode)
+			{
+				SharedStateId = NodeId;
+			}
+			else if (SharedStateBehavior == ESharedStateBehaviorType::SameNodeInComposition)
+			{
+				if (InParams.Environment.Contains<TArray<FGuid>>(OperatorBuilder::Environment::GraphHierarchy))
+				{
+					const TArray<FGuid>& GraphHierarchy = InParams.Environment.GetValue<TArray<FGuid>>(OperatorBuilder::Environment::GraphHierarchy);
+					SharedStateId = GetSameNodeInCompositionId(NodeId, GraphHierarchy);
+				}
+				else
+				{
+#if WITH_METASOUND_DEBUG_ENVIRONMENT
+					if (!bHasLoggedMissingGraphHierarchyWarning)
+					{
+						UE_LOG(LogMetaSound, Warning, TEXT("Array Random Get: Graph Hierarchy environment variable needed for Same Node in Composition shared state id not found (Graph '%s')"), *GraphName);
+						bHasLoggedMissingGraphHierarchyWarning = true;
+					}
+#endif // WITH_METASOUND_DEBUG_ENVIRONMENT
+				}
+			}
 		}
 
 		void Execute()
@@ -275,29 +365,14 @@ namespace Metasound
 #endif // WITH_METASOUND_DEBUG_ENVIRONMENT
 				return;
 			}
-			
+
  			TriggerReset->ExecuteBlock(
 				[&](int32 StartFrame, int32 EndFrame)
 				{
 				},
 				[this](int32 StartFrame, int32 EndFrame)
 				{
-					if (IsStateInitializationNeeded())
-					{
-						InitializeState(PrevArraySize);
-					}
-
-					if (UseSharedState())
-					{
-						FSharedStateRandomGetManager& RGM = FSharedStateRandomGetManager::Get();
-						RGM.ResetSeed(SharedStateUniqueId);
-					}
-					else
-					{
-						check(ArrayRandomGet.IsValid());
-						ArrayRandomGet->ResetSeed();
-					}
-					TriggerOnReset->TriggerFrame(StartFrame);
+					ExecuteTriggerReset(StartFrame);
 				}
 			);
  
@@ -313,153 +388,111 @@ namespace Metasound
 		}
 
 	private:
-		void ExecuteTriggerNext(int32 StartFrame)
+		void CreateSharedStateArgs(InitSharedStateArgs& InOutStateArgs)
 		{
-			const ArrayType& InputArrayRef = *InputArray;
-			int32 OutRandomIndex = INDEX_NONE;
-
-			const bool bIsStateReinitializationNeeded = IsStateInitializationNeeded();
-			const bool bIsArraySizeChanged = PrevArraySize != InputArrayRef.Num(); 
-			const bool bSeedValueChanged = PrevSeedValue != *SeedValue;
-			const bool bNoRepeatOrderChanged = PrevNoRepeatOrder != *NoRepeatOrder;
-			const bool bWeightsArrayChanged = WeightsArray != *InputWeightsArray;
+			InOutStateArgs.SharedStateId = SharedStateId;
+			InOutStateArgs.Seed = *SeedValue;
+			InOutStateArgs.NumElements = (*InputArray).Num();
+			InOutStateArgs.NoRepeatOrder = *NoRepeatOrder;
+			InOutStateArgs.bIsPreviewSound = bIsPreviewSound;
+			InOutStateArgs.Weights = *InputWeightsArray;
+		}
 
-			// Update cached values if changed
-			if (bIsArraySizeChanged)
-			{
-				PrevArraySize = InputArrayRef.Num();
-			}
-			if (bSeedValueChanged)
-			{
-				PrevSeedValue = *SeedValue;
-			}
-			if (bNoRepeatOrderChanged)
-			{
-				PrevNoRepeatOrder = *NoRepeatOrder;
-			}
-			if (bWeightsArrayChanged)
+		void ExecuteTriggerReset(int32 StartFrame)
+		{
+			const ArrayType& InputArrayRef = *InputArray;
+			if (*bEnableSharedState && !bIsPreviewSound)
 			{
-				WeightsArray = *InputWeightsArray;
-			}
+				// Update shared state id for array content hash 
+				if (SharedStateBehavior == ESharedStateBehaviorType::SameData)
+				{
+					SharedStateId = MetasoundArrayHashPrivate::GetArrayContentHashGuid(*InputArray);
+				}
 
-			// Reinitialize state (with new values) if needed 
-			if (bIsStateReinitializationNeeded)
-			{
-				InitializeState(PrevArraySize);
-			}
+				FSharedStateRandomGetManager& RGM = FSharedStateRandomGetManager::Get();
+				InitSharedStateArgs StateArgs;
+				CreateSharedStateArgs(StateArgs);
 
-			// Update other state (which was not necessarily changed by new state initialization)
-			if (bSeedValueChanged)
+				// Update and reset seed as one operation
+				RGM.ResetSeed(SharedStateId, StateArgs);
+			}
+			else // No shared state
 			{
-				if (UseSharedState())
+				if (!ArrayRandomGet.IsValid())
 				{
-					FSharedStateRandomGetManager& RGM = FSharedStateRandomGetManager::Get();
-					RGM.SetSeed(SharedStateUniqueId, PrevSeedValue);
+					ArrayRandomGet = MakeUnique<FArrayRandomGet>(*SeedValue, InputArrayRef.Num(), *InputWeightsArray, *NoRepeatOrder);
 				}
 				else
 				{
-					check(ArrayRandomGet.IsValid());
-					ArrayRandomGet->SetSeed(PrevSeedValue);
+					ArrayRandomGet->UpdateState(*SeedValue, InputArrayRef.Num(), *InputWeightsArray, *NoRepeatOrder);
 				}
+				ArrayRandomGet->ResetSeed();
 			}
-			if (bNoRepeatOrderChanged)
+			TriggerOnReset->TriggerFrame(StartFrame);
+		}
+		
+		void ExecuteTriggerNext(int32 StartFrame)
+		{
+			const ArrayType& InputArrayRef = *InputArray;
+			if (*bEnableSharedState && !bIsPreviewSound)
 			{
-				if (UseSharedState())
+				// Update shared state id for array content hash 
+				if (SharedStateBehavior == ESharedStateBehaviorType::SameData)
 				{
-					FSharedStateRandomGetManager& RGM = FSharedStateRandomGetManager::Get();
-					RGM.SetNoRepeatOrder(SharedStateUniqueId, PrevNoRepeatOrder);
-				}
-				else
-				{
-					check(ArrayRandomGet.IsValid());
-					ArrayRandomGet->SetNoRepeatOrder(PrevNoRepeatOrder);
+					SharedStateId = MetasoundArrayHashPrivate::GetArrayContentHashGuid(*InputArray);
 				}
+
+				FSharedStateRandomGetManager& RGM = FSharedStateRandomGetManager::Get();
+				InitSharedStateArgs StateArgs;
+				CreateSharedStateArgs(StateArgs);
+
+				// Update and get next value as one operation
+				*OutIndex = RGM.NextValue(SharedStateId, StateArgs);
 			}
-			if (bWeightsArrayChanged)
+			else // No shared state
 			{
-				if (UseSharedState())
+				// Initialize or update state
+				if (!ArrayRandomGet.IsValid())
 				{
-					FSharedStateRandomGetManager& RGM = FSharedStateRandomGetManager::Get();
-					RGM.SetRandomWeights(SharedStateUniqueId, WeightsArray);
+					ArrayRandomGet = MakeUnique<FArrayRandomGet>(*SeedValue, InputArrayRef.Num(), *InputWeightsArray, *NoRepeatOrder);
 				}
 				else
 				{
-					check(ArrayRandomGet.IsValid());
-					ArrayRandomGet->SetRandomWeights(WeightsArray);
+					ArrayRandomGet->UpdateState(*SeedValue, InputArrayRef.Num(), *InputWeightsArray, *NoRepeatOrder);
 				}
+				// Get next value
+				*OutIndex = ArrayRandomGet->NextValue();
 			}
 
-			// Get next value 
-			if (UseSharedState())
-			{
-				FSharedStateRandomGetManager& RGM = FSharedStateRandomGetManager::Get();
-				OutRandomIndex = RGM.NextValue(SharedStateUniqueId);
-			}
-			else
-			{
-				check(ArrayRandomGet.IsValid());
-				OutRandomIndex = ArrayRandomGet->NextValue();
-			}
-
-			check(OutRandomIndex != INDEX_NONE);
+			check(*OutIndex != INDEX_NONE);
 #if WITH_METASOUND_DEBUG_ENVIRONMENT
-			UE_LOG(LogMetaSound, Verbose, TEXT("Array Random Get: Index chosen: '%u'"), OutRandomIndex);
+			UE_LOG(LogMetaSound, VeryVerbose, TEXT("Array Random Get Execute Next: \
+				Index chosen: %u, Graph: '%s', NumRepeats: %d, Array Size: %u, Seed: %d, Type: %s\
+				 Node Id: %s, Shared State Enabled: %u, Shared State Behavior: %s, Shared State Id: %s"), \
+				*OutIndex, *GraphName, *NoRepeatOrder, InputArrayRef.Num(), *SeedValue, *Metasound::GetMetasoundDataTypeString<ElementType>(), \
+				*NodeId.ToString(), *bEnableSharedState, *DebugSharedStateBehaviorString, *SharedStateId.ToString());
 #endif // WITH_METASOUND_DEBUG_ENVIRONMENT
 
 			// The input array size may have changed, so make sure it's wrapped into range of the input array
-			*OutValue = InputArrayRef[OutRandomIndex % InputArrayRef.Num()];
+			*OutValue = InputArrayRef[*OutIndex % InputArrayRef.Num()];
 
 			TriggerOnNext->TriggerFrame(StartFrame);
 		}
 
-		bool IsStateInitializationNeeded()
+		// Hash combine the current node id with the graph hierarchy ids
+		FGuid GetSameNodeInCompositionId(const FGuid& InNodeId, const TArray<FGuid>& InGraphHierarchy) const
 		{
-			const ArrayType& InputArrayRef = *InputArray;
-			const bool bIsArrayNonEmpty = InputArrayRef.Num() != 0; // Skip reinit if the array is empty because it represents an invalid state for this node.
-			const bool bIsArraySizeChanged = PrevArraySize != InputArrayRef.Num(); // Need to reinit for array size changes. 
-			const bool bIsSharedStateEnablementInconsistent = (*bEnableSharedState != bSharedStateInitialized); // Need to reinit if the shared state enablement has been updated.
-			const bool bIsNonSharedStateInitializationNeeded = !*bEnableSharedState && !ArrayRandomGet.IsValid(); // For the first initialization of the non shared state random get (bIsSharedStateEnablementInconsistent will take care of that for shared state)
-
-			return (bIsArrayNonEmpty && (bIsArraySizeChanged || bIsSharedStateEnablementInconsistent || bIsNonSharedStateInitializationNeeded));
-		}
-
-		void InitializeState(int32 InArraySize)
-		{
-			bSharedStateInitialized = false;
-			if (InArraySize > 0)
-			{
-				if (*bEnableSharedState)
-				{
-					// Get the environment variable for the unique ID of the sound
-					check(SharedStateUniqueId.IsValid());
-					FSharedStateRandomGetManager& RGM = FSharedStateRandomGetManager::Get();
-
-					InitSharedStateArgs Args;
-					Args.SharedStateId = SharedStateUniqueId;
-					Args.Seed = PrevSeedValue;
-					Args.NumElements = PrevArraySize;
-					Args.NoRepeatOrder = PrevNoRepeatOrder;
-					Args.bIsPreviewSound = bIsPreviewSound;
-					Args.Weights = WeightsArray;
-
-					RGM.InitSharedState(Args);
-
-					bSharedStateInitialized = true;
-				}
-				else
-				{
-					ArrayRandomGet = MakeUnique<FArrayRandomGet>(PrevSeedValue, PrevArraySize, WeightsArray, PrevNoRepeatOrder);
-				}
-			}
-			else
+			uint32 A = InNodeId.A; uint32 B = InNodeId.B; uint32 C = InNodeId.C; uint32 D = InNodeId.D;
+			for (int i = 0; i < InGraphHierarchy.Num(); ++i)
 			{
-				ArrayRandomGet = MakeUnique<FArrayRandomGet>();
-#if WITH_METASOUND_DEBUG_ENVIRONMENT
- 				UE_LOG(LogMetaSound, Verbose, TEXT("Array Random Get: Can't retrieve random elements from an empty array in graph '%s'"), *GraphName);
-#endif // WITH_METASOUND_DEBUG_ENVIRONMENT
+				A = HashCombineFast(A, InGraphHierarchy[i].A);
+				B = HashCombineFast(B, InGraphHierarchy[i].B);
+				C = HashCombineFast(C, InGraphHierarchy[i].C);
+				D = HashCombineFast(D, InGraphHierarchy[i].D);
 			}
+			return FGuid(A, B, C, D);
 		}
-
+		
 		// Inputs
 		FTriggerReadRef TriggerNext;
 		FTriggerReadRef TriggerReset;
@@ -467,27 +500,27 @@ namespace Metasound
 		TDataReadReference<WeightsArrayType> InputWeightsArray;
 		FInt32ReadRef SeedValue;
 		FInt32ReadRef NoRepeatOrder;
+		FBoolReadRef bEnableSharedState;
+		FEnumSharedStateBehaviorType SharedStateBehavior;
 
 		// Outputs
 		FTriggerWriteRef TriggerOnNext;
 		FTriggerWriteRef TriggerOnReset;
 		TDataWriteReference<ElementType> OutValue;
+		FInt32WriteRef OutIndex;
 
 #if WITH_METASOUND_DEBUG_ENVIRONMENT
 		FString GraphName;
 		bool bHasLoggedEmptyArrayWarning = false;
+		bool bHasLoggedMissingGraphHierarchyWarning = false;
+		FString DebugSharedStateBehaviorString;
 #endif // WITH_METASOUND_DEBUG_ENVIRONMENT
 
 		// Data
 		TUniquePtr<FArrayRandomGet> ArrayRandomGet;
-		TArray<float> WeightsArray;
-		int32 PrevSeedValue = INDEX_NONE;
-		int32 PrevNoRepeatOrder = INDEX_NONE;
-		FGuid SharedStateUniqueId;
-		int32 PrevArraySize = 0;
+		FGuid NodeId;
+		FGuid SharedStateId;
 		bool bIsPreviewSound = false;
-		FBoolReadRef bEnableSharedState;
-		bool bSharedStateInitialized = false;
 	};
 
 	template<typename ArrayType>
diff --git a/Engine/Plugins/Runtime/Metasound/Source/MetasoundFrontend/Public/MetasoundDataTypeRegistrationMacro.h b/Engine/Plugins/Runtime/Metasound/Source/MetasoundFrontend/Public/MetasoundDataTypeRegistrationMacro.h
--- a/Engine/Plugins/Runtime/Metasound/Source/MetasoundFrontend/Public/MetasoundDataTypeRegistrationMacro.h
+++ b/Engine/Plugins/Runtime/Metasound/Source/MetasoundFrontend/Public/MetasoundDataTypeRegistrationMacro.h
@@ -741,7 +741,7 @@ namespace Metasound
 			using namespace MetasoundDataTypeRegistrationPrivate;
 			using TArrayType = TArray<TDataType>;
 
-			if (TEnableAutoArrayTypeRegistration<TDataType>::Value)
+			if constexpr (TEnableAutoArrayTypeRegistration<TDataType>::Value)
 			{
 				constexpr bool bIsArrayType = true;
 				bool bSuccess = RegisterDataTypeWithFrontendInternal<TArrayType, TLiteralArrayEnum<PreferredArgType>::Value>();
@@ -749,8 +749,10 @@ namespace Metasound
 				bSuccess = bSuccess && RegisterDataTypeWithFrontendInternal<TVariable<TArrayType>>();
 				return bSuccess;
 			}
-
-			return true;
+			else
+			{
+				return true;
+			}
 		}
 	}
 	
diff --git a/Engine/Plugins/Runtime/Metasound/Source/MetasoundFrontend/Public/MetasoundEnum.h b/Engine/Plugins/Runtime/Metasound/Source/MetasoundFrontend/Public/MetasoundEnum.h
--- a/Engine/Plugins/Runtime/Metasound/Source/MetasoundFrontend/Public/MetasoundEnum.h
+++ b/Engine/Plugins/Runtime/Metasound/Source/MetasoundFrontend/Public/MetasoundEnum.h
@@ -280,3 +280,35 @@ namespace Metasound
 		static constexpr T DefaultValue = D;
 	};
 }
+
+// Helper macros: for defining/declaring a new Enum to be used in Metasounds.
+// * Declares/defines a typedef FEnum<YourEnum> wrapper around the enum that you pass it
+// * Declares/defines a specialization of TEnumStringHelper which is used to validate and convert between values/names.
+
+// These helps generate a couple of boiler plate functions:
+//  - GetNamespace() which will return the name of the Enum "EMyEnum" etc. 
+//  - GetAllEntries() which returns all posible entries in the enum. 
+
+// Definition macros are in MetasoundEnumRegistrationMacro.h
+
+/** DECLARE_METASOUND_ENUM
+ * @param ENUMNAME - The typename of your EnumType you want to use for Metasounds. e.g. MyEnum
+ * @param DEFAULT - A fully qualified default Enum value. e.g. EMyEnum::One
+ * @param API - The module API this is declared inside e.g. METASOUNDSTANDARDNODES_API
+ * @param ENUMTYPEDEF - The name of the TEnum<YourType> wrapper type
+ * @param TYPEINFO - The name of the TypeInfo type you want to define e.g. FMyEnumTypeInfo
+ * @param READREF - The name of the Read Reference type you want to define. e.g FMyEnumReadRef
+ * @param WRITEREF -The name of the Write Reference type you want to define e.g. FMyEnumWriteRef
+ */
+#define DECLARE_METASOUND_ENUM(ENUMNAME, DEFAULT, API, ENUMTYPEDEF, TYPEINFO, READREF, WRITEREF)\
+	using ENUMTYPEDEF = Metasound::TEnum<ENUMNAME, DEFAULT>; \
+	DECLARE_METASOUND_DATA_REFERENCE_TYPES(ENUMTYPEDEF, API, TYPEINFO, READREF, WRITEREF);\
+	template<> struct API Metasound::TEnumStringHelper<ENUMNAME> : Metasound::TEnumStringHelperBase<Metasound::TEnumStringHelper<ENUMNAME>, ENUMNAME>\
+	{\
+		static FName GetNamespace()\
+		{\
+			static const FName ThisName { TEXT(#ENUMNAME) };\
+			return ThisName;\
+		}\
+		static TArrayView<const Metasound::TEnumEntry<ENUMNAME>> GetAllEntries();\
+	};
diff --git a/Engine/Plugins/Runtime/Metasound/Source/MetasoundFrontend/Public/MetasoundEnumRegistrationMacro.h b/Engine/Plugins/Runtime/Metasound/Source/MetasoundFrontend/Public/MetasoundEnumRegistrationMacro.h
--- a/Engine/Plugins/Runtime/Metasound/Source/MetasoundFrontend/Public/MetasoundEnumRegistrationMacro.h
+++ b/Engine/Plugins/Runtime/Metasound/Source/MetasoundFrontend/Public/MetasoundEnumRegistrationMacro.h
@@ -2,9 +2,9 @@
 
 #pragma once
 
-#include "MetasoundArrayNodesRegistration.h"
 #include "MetasoundDataReference.h"
 #include "MetasoundEnum.h"
+#include "MetasoundFrontendDataTypeTraits.h"
 #include "MetasoundDataTypeRegistrationMacro.h"
 
 namespace Metasound
@@ -45,37 +45,7 @@ namespace Metasound
 		static constexpr bool Value = false;
 	};
 }
-
-// Helper macros: for defining/declaring a new Enum to be used in Metasounds.
-// * Declares/defines a typedef FEnum<YourEnum> wrapper around the enum that you pass it
-// * Declares/defines a specialization of TEnumStringHelper which is used to validate and convert between values/names.
-
-// These helps generate a couple of boiler plate functions:
-//  - GetNamespace() which will return the name of the Enum "EMyEnum" etc. 
-//  - GetAllEntries() which returns all posible entries in the enum. 
-
-/** DECLARE_METASOUND_ENUM
- * @param ENUMNAME - The typename of your EnumType you want to use for Metasounds. e.g. MyEnum
- * @param DEFAULT - A fully qualified default Enum value. e.g. EMyEnum::One
- * @param API - The module API this is declared inside e.g. METASOUNDSTANDARDNODES_API 
- * @param ENUMTYPEDEF - The name of the TEnum<YourType> wrapper type
- * @param TYPEINFO - The name of the TypeInfo type you want to define e.g. FMyEnumTypeInfo
- * @param READREF - The name of the Read Reference type you want to define. e.g FMyEnumReadRef
- * @param WRITEREF -The name of the Write Reference type you want to define e.g. FMyEnumWriteRef
- */
-#define DECLARE_METASOUND_ENUM(ENUMNAME, DEFAULT, API, ENUMTYPEDEF, TYPEINFO, READREF, WRITEREF)\
-	using ENUMTYPEDEF = Metasound::TEnum<ENUMNAME, DEFAULT>; \
-	DECLARE_METASOUND_DATA_REFERENCE_TYPES(ENUMTYPEDEF, API, TYPEINFO, READREF, WRITEREF);\
-	template<> struct API Metasound::TEnumStringHelper<ENUMNAME> : Metasound::TEnumStringHelperBase<Metasound::TEnumStringHelper<ENUMNAME>, ENUMNAME>\
-	{\
-		static FName GetNamespace()\
-		{\
-			static const FName ThisName { TEXT(#ENUMNAME) };\
-			return ThisName;\
-		}\
-		static TArrayView<const Metasound::TEnumEntry<ENUMNAME>> GetAllEntries();\
-	};
-
+// Helper macros for use with the enum declaration macro in MetasoundEnum.h
 /** DEFINE_METASOUND_ENUM_BEGIN
  * @param ENUMNAME - The typename of your raw EnumType you want to use for Metasounds. e.g. EMyType
  * @param ENUMTYPEDEF - The name of the TEnum<YourType> wrapper type
@@ -105,4 +75,4 @@ namespace Metasound
 #define DEFINE_METASOUND_ENUM_END() \
 		};\
 		return Entries;\
-	};
+	};
\ No newline at end of file
diff --git a/Engine/Plugins/Runtime/Metasound/Source/MetasoundGraphCore/Private/MetasoundEnvironment.cpp b/Engine/Plugins/Runtime/Metasound/Source/MetasoundGraphCore/Private/MetasoundEnvironment.cpp
--- a/Engine/Plugins/Runtime/Metasound/Source/MetasoundGraphCore/Private/MetasoundEnvironment.cpp
+++ b/Engine/Plugins/Runtime/Metasound/Source/MetasoundGraphCore/Private/MetasoundEnvironment.cpp
@@ -17,6 +17,7 @@ DEFINE_METASOUND_ENVIRONMENT_VARIABLE_TYPE(float);
 DEFINE_METASOUND_ENVIRONMENT_VARIABLE_TYPE(double);
 DEFINE_METASOUND_ENVIRONMENT_VARIABLE_TYPE(FString);
 DEFINE_METASOUND_ENVIRONMENT_VARIABLE_TYPE(FName);
+DEFINE_METASOUND_ENVIRONMENT_VARIABLE_TYPE(TArray<FGuid>);
 
 namespace Metasound
 {
diff --git a/Engine/Plugins/Runtime/Metasound/Source/MetasoundGraphCore/Private/MetasoundOperatorBuilder.cpp b/Engine/Plugins/Runtime/Metasound/Source/MetasoundGraphCore/Private/MetasoundOperatorBuilder.cpp
--- a/Engine/Plugins/Runtime/Metasound/Source/MetasoundGraphCore/Private/MetasoundOperatorBuilder.cpp
+++ b/Engine/Plugins/Runtime/Metasound/Source/MetasoundGraphCore/Private/MetasoundOperatorBuilder.cpp
@@ -25,6 +25,11 @@ namespace Metasound
 {
 	namespace OperatorBuilder
 	{
+		namespace Environment
+		{
+			const FName GraphHierarchy = "GraphHierarchy";
+		}
+
 		// Shared context used in the builder to maintain state of current build.
 		struct FBuildContext
 		{
@@ -202,7 +207,41 @@ namespace Metasound
 			return false;
 		}
 
-		OperatorBuilder::FBuildContext BuildContext(InParams.Graph, *AlgoAdapter, InParams.OperatorSettings, InParams.Environment, OutGraphOperatorData, OutResults);
+		// Update environment with current graph hierarchy
+		FMetasoundEnvironment NewEnvironment;
+		if (!InParams.Environment.Contains<TArray<FGuid>>(OperatorBuilder::Environment::GraphHierarchy))
+		{
+			// Copy old environment and add new environment variable
+			NewEnvironment = InParams.Environment;
+			TArray<FGuid> GraphHierarchy;
+			GraphHierarchy.Emplace(InParams.Graph.GetInstanceID());
+			NewEnvironment.SetValue<TArray<FGuid>>(OperatorBuilder::Environment::GraphHierarchy, GraphHierarchy);
+		}
+		else
+		{
+			// Copy and append to existing environment variable (environment variables aren't modifiable in place)
+			for (const auto& Iter : InParams.Environment)
+			{
+				if (Iter.Key == OperatorBuilder::Environment::GraphHierarchy)
+				{
+					TArray<FGuid> GraphHierarchy;
+					for (const FGuid Id : InParams.Environment.GetValue<TArray<FGuid>>(OperatorBuilder::Environment::GraphHierarchy))
+					{
+						GraphHierarchy.Emplace(Id);
+					}
+					GraphHierarchy.Emplace(InParams.Graph.GetInstanceID());
+
+					NewEnvironment.SetValue<TArray<FGuid>>(OperatorBuilder::Environment::GraphHierarchy, GraphHierarchy);
+				}
+				else
+				{
+					// Copy other environment variables
+					NewEnvironment.SetValue(Iter.Value->Clone());
+				}
+			}
+		}
+
+		OperatorBuilder::FBuildContext BuildContext(InParams.Graph, *AlgoAdapter, InParams.OperatorSettings, NewEnvironment, OutGraphOperatorData, OutResults);
 
 		// Sort the nodes in a valid execution order
 		BuildStatus |= DepthFirstTopologicalSort(BuildContext, OutNodeOrder);
diff --git a/Engine/Plugins/Runtime/Metasound/Source/MetasoundGraphCore/Public/MetasoundEnvironment.h b/Engine/Plugins/Runtime/Metasound/Source/MetasoundGraphCore/Public/MetasoundEnvironment.h
--- a/Engine/Plugins/Runtime/Metasound/Source/MetasoundGraphCore/Public/MetasoundEnvironment.h
+++ b/Engine/Plugins/Runtime/Metasound/Source/MetasoundGraphCore/Public/MetasoundEnvironment.h
@@ -252,6 +252,26 @@ namespace Metasound
 			Variables.Add(InVariableName, MakeUnique<TMetasoundEnvironmentVariable<VarType>>(InVariableName, InValue));
 		}
 
+		/** Sets the environment variable data from a given environment variable
+		 * 
+		 * @param InValue - Environment variable to set.
+		 */
+		void SetValue(TUniquePtr<IMetasoundEnvironmentVariable> InValue)
+		{
+			Variables.Add(InValue->GetName(), MoveTemp(InValue));
+		}
+
+		using FRangedForConstIteratorType = TSortedMap<FName, TUniquePtr<IMetasoundEnvironmentVariable>, FDefaultAllocator, FNameFastLess>::RangedForConstIteratorType;
+		FRangedForConstIteratorType begin() const
+		{
+			return Variables.begin();
+		}
+
+		FRangedForConstIteratorType end() const
+		{
+			return Variables.end();
+		}
+
 	private:
 		TSortedMap<FName, TUniquePtr<IMetasoundEnvironmentVariable>, FDefaultAllocator, FNameFastLess> Variables;
 	};
@@ -272,4 +292,4 @@ DECLARE_METASOUND_ENVIRONMENT_VARIABLE_TYPE(METASOUNDGRAPHCORE_API, float);
 DECLARE_METASOUND_ENVIRONMENT_VARIABLE_TYPE(METASOUNDGRAPHCORE_API, double);
 DECLARE_METASOUND_ENVIRONMENT_VARIABLE_TYPE(METASOUNDGRAPHCORE_API, FString);
 DECLARE_METASOUND_ENVIRONMENT_VARIABLE_TYPE(METASOUNDGRAPHCORE_API, FName);
-
+DECLARE_METASOUND_ENVIRONMENT_VARIABLE_TYPE(METASOUNDGRAPHCORE_API, TArray<FGuid>);
diff --git a/Engine/Plugins/Runtime/Metasound/Source/MetasoundGraphCore/Public/MetasoundOperatorBuilder.h b/Engine/Plugins/Runtime/Metasound/Source/MetasoundGraphCore/Public/MetasoundOperatorBuilder.h
--- a/Engine/Plugins/Runtime/Metasound/Source/MetasoundGraphCore/Public/MetasoundOperatorBuilder.h
+++ b/Engine/Plugins/Runtime/Metasound/Source/MetasoundGraphCore/Public/MetasoundOperatorBuilder.h
@@ -25,6 +25,10 @@ namespace Metasound
 	namespace OperatorBuilder
 	{
 		struct FBuildContext;
+		namespace Environment
+		{
+			METASOUNDGRAPHCORE_API const extern FName GraphHierarchy;
+		}
 	}
 
 	namespace DynamicGraph
