UE-206946 - [MetaSounds] Renaming comment nodes should not refocus / zoom out to the entire comment box

diff --git a/Engine/Plugins/Runtime/Metasound/Source/MetasoundEditor/Private/MetasoundEditor.cpp b/Engine/Plugins/Runtime/Metasound/Source/MetasoundEditor/Private/MetasoundEditor.cpp
--- a/Engine/Plugins/Runtime/Metasound/Source/MetasoundEditor/Private/MetasoundEditor.cpp
+++ b/Engine/Plugins/Runtime/Metasound/Source/MetasoundEditor/Private/MetasoundEditor.cpp
@@ -2708,7 +2708,14 @@ namespace Metasound
 				{
 					if (TSharedPtr<SGraphEditor> GraphEditor = GetGraphEditor())
 					{
-						GraphEditor->JumpToNode(Node, /*bRequestRename=*/true);
+						if (GraphEditor->IsNodeTitleVisible(Node, /*bRequestRename=*/ false))
+						{
+							GraphEditor->IsNodeTitleVisible(Node, /*bRequestRename=*/ true);
+						}
+						else
+						{
+							GraphEditor->JumpToNode(Node, /*bRequestRename=*/true);
+						}
 						return;
 					}
 				}
