UE-186195 - Add console command to mute/solo MetaSounds (2)

diff --git a/Engine/Plugins/Runtime/Metasound/Source/MetasoundEditor/Private/MetasoundWaveConsole.cpp b/Engine/Plugins/Runtime/Metasound/Source/MetasoundEditor/Private/MetasoundWaveConsole.cpp
--- a/Engine/Plugins/Runtime/Metasound/Source/MetasoundEditor/Private/MetasoundWaveConsole.cpp
+++ b/Engine/Plugins/Runtime/Metasound/Source/MetasoundEditor/Private/MetasoundWaveConsole.cpp
@@ -11,7 +11,7 @@ namespace Metasound
 {
 	namespace Console
 	{
-		static void HandleSoloMetaSoundWave(const TArray<FString>& InArgs, UWorld* World)
+		static void HandleSoloMetaSound(const TArray<FString>& InArgs, UWorld* World)
 		{
 			check(GEngine);
 			FAudioDeviceManager* DeviceManager = GEngine->GetAudioDeviceManager();
@@ -51,17 +51,17 @@ namespace Metasound
 				for (const FString& Wave : MetaSoundWaves)
 				{
 					// we can't use FAudioDebugger::SetSoloSoundWave(...) because we don't want to mute non MetaSounds
-					const bool bMute = !InArgs[0].Equals(Wave);
+					const bool bMute = !Wave.Contains(InArgs[0], ESearchCase::IgnoreCase);
 					DeviceManager->GetDebugger().SetMuteSoundWave(*Wave, bMute);
 				}
 			}
 			else
 			{
-				UE_LOG(LogMetaSound, Error, TEXT("You can solo ONLY ONE MetaSound wave!"));
+				UE_LOG(LogMetaSound, Error, TEXT("You can solo ONLY ONE MetaSound!"));
 			}
 		}
 
-		static void HandleMuteMetaSoundWave(const TArray<FString>& InArgs, UWorld* World)
+		static void HandleMuteMetaSound(const TArray<FString>& InArgs, UWorld* World)
 		{
 			check(GEngine);
 			FAudioDeviceManager* DeviceManager = GEngine->GetAudioDeviceManager();
@@ -100,7 +100,7 @@ namespace Metasound
 			{
 				for (const FString& Wave : MetaSoundWaves)
 				{
-					if (Arg.Equals(Wave))
+					if (Wave.Contains(Arg, ESearchCase::IgnoreCase))
 					{
 						DeviceManager->GetDebugger().SetMuteSoundWave(*Wave, true);
 					}
@@ -110,18 +110,18 @@ namespace Metasound
 	} // namespace Console
 } // namespace Metasound
 
-static FAutoConsoleCommandWithWorldAndArgs SoloMetaSoundWave
+static FAutoConsoleCommandWithWorldAndArgs SoloMetaSound
 (
 	TEXT("au.MetaSound.SoloMetaSound"),
-	TEXT("Mutes all other MetaSound waves. Only the first argument is accepted."),
-	FConsoleCommandWithWorldAndArgsDelegate::CreateStatic(&Metasound::Console::HandleSoloMetaSoundWave)
+	TEXT("Mutes all other MetaSounds. Only the first argument is accepted."),
+	FConsoleCommandWithWorldAndArgsDelegate::CreateStatic(&Metasound::Console::HandleSoloMetaSound)
 );
 
-static FAutoConsoleCommandWithWorldAndArgs MuteMetaSoundWave
+static FAutoConsoleCommandWithWorldAndArgs MuteMetaSound
 (
-	TEXT("au.MetaSound.MuteMetaSoundWave"),
-	TEXT("Mutes all given MetaSound waves."),
-	FConsoleCommandWithWorldAndArgsDelegate::CreateStatic(&Metasound::Console::HandleMuteMetaSoundWave)
+	TEXT("au.MetaSound.MuteMetaSound"),
+	TEXT("Mutes all given MetaSounds."),
+	FConsoleCommandWithWorldAndArgsDelegate::CreateStatic(&Metasound::Console::HandleMuteMetaSound)
 );
 
 #endif // ENABLE_AUDIO_DEBUG
