- Move MetaSoundAssetSubsystem& BuilderSubsystem implementations to monolithics that are initialized at MetaSoundEngine module load for safety
- Avoids PreDefault vs Default module Init which can interact negatively and obscurely with EngineSubsystems not being loaded prior to certain builds/scenarios where serialized assets are loaded earlier
- Sunset existing Subsystem non-UFUNCTION calls not pertaining to direct Blueprint exposition
- Optimize AssetManager to use own key vs NodeRegistryKey and TopLevelPaths vs SoftObjectPaths
- Misc module clean-up
- Misc callsite refactors to use new monolithics

#tests EngineTest & AudioQA Automation PIE Packaged Build -game

diff --git a/Engine/Plugins/Runtime/Metasound/Source/MetasoundEditor/Private/MetasoundDetailCustomization.cpp b/Engine/Plugins/Runtime/Metasound/Source/MetasoundEditor/Private/MetasoundDetailCustomization.cpp
--- a/Engine/Plugins/Runtime/Metasound/Source/MetasoundEditor/Private/MetasoundDetailCustomization.cpp
+++ b/Engine/Plugins/Runtime/Metasound/Source/MetasoundEditor/Private/MetasoundDetailCustomization.cpp
@@ -380,9 +380,7 @@ namespace Metasound
 							const FScopedTransaction Transaction(FText::Format(LOCTEXT("AddInterfaceTransactionFormat", "Add MetaSound Interface '{0}'"), FText::FromString(InterfaceToAdd.Version.ToString())));
 							MetaSound.Get()->Modify();
 							MetaSoundAsset->GetGraphChecked().Modify();
-
-							TScriptInterface<IMetaSoundDocumentInterface> MetaSoundPatchDocInterface = MetaSound.Get();
-							FMetaSoundFrontendDocumentBuilder& Builder = IMetaSoundAssetManager::GetChecked().AttachDocumentBuilderChecked(*MetaSound.Get());
+							FMetaSoundFrontendDocumentBuilder& Builder = IDocumentBuilderRegistry::GetChecked().FindOrBeginBuilding(MetaSound.Get());
 							FModifyInterfaceOptions Options({ }, { InterfaceToAdd });
 							Options.bSetDefaultNodeLocations = false; // Don't automatically add nodes to ed graph
 							Builder.ModifyInterfaces(MoveTemp(Options));
diff --git a/Engine/Plugins/Runtime/Metasound/Source/MetasoundEditor/Private/MetasoundEditor.cpp b/Engine/Plugins/Runtime/Metasound/Source/MetasoundEditor/Private/MetasoundEditor.cpp
--- a/Engine/Plugins/Runtime/Metasound/Source/MetasoundEditor/Private/MetasoundEditor.cpp
+++ b/Engine/Plugins/Runtime/Metasound/Source/MetasoundEditor/Private/MetasoundEditor.cpp
@@ -36,6 +36,7 @@
 #include "Metasound.h"
 #include "MetasoundAssetSubsystem.h"
 #include "MetasoundBuilderSubsystem.h"
+#include "MetasoundDocumentBuilderRegistry.h"
 #include "MetasoundDocumentInterface.h"
 #include "MetasoundEditorCommands.h"
 #include "MetasoundEditorDocumentClipboardUtils.h"
@@ -850,7 +851,7 @@ namespace Metasound
 				}
 			}
 
-			Builder.Reset(&UMetaSoundBuilderSubsystem::GetChecked().AttachBuilderToAssetChecked(*ObjectToEdit));
+			Builder.Reset(&Engine::FDocumentBuilderRegistry::GetChecked().FindOrBeginBuilding(*ObjectToEdit));
 
 			GEditor->RegisterForUndo(this);
 
@@ -1957,9 +1958,9 @@ namespace Metasound
 				FConstNodeHandle NodeHandle = ExternalNode->GetConstNodeHandle();
 				const FNodeRegistryKey Key(NodeHandle->GetClassMetadata());
 
-				if (const FSoftObjectPath* AssetObject = UMetaSoundAssetSubsystem::GetChecked().FindObjectPathFromKey(Key))
+				if (const FMetasoundAssetBase* Asset = IMetaSoundAssetManager::GetChecked().FindAsset(Key))
 				{
-					GEditor->GetEditorSubsystem<UAssetEditorSubsystem>()->OpenEditorForAsset(*AssetObject);
+					GEditor->GetEditorSubsystem<UAssetEditorSubsystem>()->OpenEditorForAsset(Asset->GetOwningAsset());
 				}
 			}
 		}
@@ -3095,23 +3096,18 @@ namespace Metasound
 		void FEditor::DeleteAllUnusedInSection()
 		{
 			TArray<TSharedPtr<FMetasoundGraphMemberSchemaAction>> ActionsToDelete;
-			TArray<TSharedPtr<FEdGraphSchemaAction>> Actions;		
+			TArray<TSharedPtr<FEdGraphSchemaAction>> Actions;
 			GraphMembersMenu->GetSelectedCategorySubActions(Actions);
 
 			for (TSharedPtr<FEdGraphSchemaAction> Action : Actions)
-			{				
+			{
 				const TSharedPtr<FMetasoundGraphMemberSchemaAction> MetasoundAction = StaticCastSharedPtr<FMetasoundGraphMemberSchemaAction>(Action);
 				if (MetasoundAction.IsValid())
 				{
 					if (const UMetasoundEditorGraphMember* GraphMember = MetasoundAction->GetGraphMember())
 					{
 						const TArray<UMetasoundEditorGraphMemberNode*> Nodes = GraphMember->GetNodes();
-
-						if (Nodes.Num() > 0)
-						{
-							continue;
-						}
-						else
+						if (Nodes.IsEmpty())
 						{
 							const FMetasoundFrontendVersion* InterfaceVersion = nullptr;
 							if (const UMetasoundEditorGraphVertex* Vertex = Cast<UMetasoundEditorGraphVertex>(GraphMember))
@@ -3123,10 +3119,10 @@ namespace Metasound
 							const bool bIsInterfaceMember = InterfaceVersion && InterfaceVersion->IsValid();
 							if (!bIsInterfaceMember)
 							{
-								ActionsToDelete.Add(MetasoundAction);														
+								ActionsToDelete.Add(MetasoundAction);
 							}
 						}
-					}					
+					}
 				}
 			}
 
diff --git a/Engine/Plugins/Runtime/Metasound/Source/MetasoundEditor/Private/MetasoundEditorDocumentClipboardUtils.cpp b/Engine/Plugins/Runtime/Metasound/Source/MetasoundEditor/Private/MetasoundEditorDocumentClipboardUtils.cpp
--- a/Engine/Plugins/Runtime/Metasound/Source/MetasoundEditor/Private/MetasoundEditorDocumentClipboardUtils.cpp
+++ b/Engine/Plugins/Runtime/Metasound/Source/MetasoundEditor/Private/MetasoundEditorDocumentClipboardUtils.cpp
@@ -5,7 +5,7 @@
 #include "Algo/Transform.h"
 #include "EdGraphUtilities.h"
 #include "Logging/TokenizedMessage.h"
-#include "MetasoundBuilderSubsystem.h"
+#include "MetasoundDocumentBuilderRegistry.h"
 #include "MetasoundFrontendSearchEngine.h"
 #include "ScopedTransaction.h"
 
@@ -229,9 +229,9 @@ namespace Metasound::Editor
 			const FNodeRegistryKey PastedRegistryKey(LookupMetadata);
 			UObject& MetaSound = *OutAsset.GetOwningAsset();
 
-			if (const FSoftObjectPath* AssetPath = IMetaSoundAssetManager::GetChecked().FindObjectPathFromKey(PastedRegistryKey))
+			if (const FMetasoundAssetBase* Asset = IMetaSoundAssetManager::GetChecked().FindAsset(PastedRegistryKey))
 			{
-				if (OutAsset.AddingReferenceCausesLoop(*AssetPath))
+				if (OutAsset.AddingReferenceCausesLoop(*Asset))
 				{
 					FMetasoundFrontendClass MetaSoundClass;
 					FMetasoundFrontendRegistryContainer::Get()->FindFrontendClassFromRegistered(PastedRegistryKey, MetaSoundClass);
@@ -295,10 +295,11 @@ namespace Metasound::Editor
 
 	void FDocumentClipboardUtils::ProcessPastedCommentNodes(FMetasoundAssetBase& OutAsset, const TArrayView<UMetasoundEditorGraphCommentNode*> CommentNodes)
 	{
+		using namespace Engine;
 		using namespace Frontend;
 
 		UMetasoundEditorGraph& Graph = *CastChecked<UMetasoundEditorGraph>(&OutAsset.GetGraphChecked());
-		UMetaSoundBuilderBase& Builder = UMetaSoundBuilderSubsystem::GetChecked().AttachBuilderToAssetChecked(*OutAsset.GetOwningAsset());
+		UMetaSoundBuilderBase& Builder = FDocumentBuilderRegistry::GetChecked().FindOrBeginBuilding(*OutAsset.GetOwningAsset());
 		for (UMetasoundEditorGraphCommentNode* CommentNode : CommentNodes)
 		{
 			CommentNode->CreateNewGuid();
diff --git a/Engine/Plugins/Runtime/Metasound/Source/MetasoundEditor/Private/MetasoundEditorGraph.cpp b/Engine/Plugins/Runtime/Metasound/Source/MetasoundEditor/Private/MetasoundEditorGraph.cpp
--- a/Engine/Plugins/Runtime/Metasound/Source/MetasoundEditor/Private/MetasoundEditorGraph.cpp
+++ b/Engine/Plugins/Runtime/Metasound/Source/MetasoundEditor/Private/MetasoundEditorGraph.cpp
@@ -8,7 +8,7 @@
 #include "EdGraph/EdGraphNode.h"
 #include "Interfaces/ITargetPlatform.h"
 #include "MetasoundAssetBase.h"
-#include "MetasoundBuilderSubsystem.h"
+#include "MetasoundDocumentBuilderRegistry.h"
 #include "MetasoundEditorGraphCommentNode.h"
 #include "MetasoundEditorGraphInputNode.h"
 #include "MetasoundEditorGraphMemberDefaults.h"
@@ -18,6 +18,7 @@
 #include "MetasoundEditorSubsystem.h"
 #include "MetasoundFrontendDocumentBuilder.h"
 #include "MetasoundFrontendNodeTemplateRegistry.h"
+#include "MetasoundLog.h"
 #include "MetasoundUObjectRegistry.h"
 #include "MetasoundVariableNodes.h"
 #include "MetasoundVertex.h"
@@ -137,7 +138,7 @@ void UMetasoundEditorGraphMember::InitializeLiteral()
 		UMetasoundEditorGraph* Graph = GetOwningGraph();
 		check(Graph);
 		UObject& MetaSound = Graph->GetMetasoundChecked();
-		FMetaSoundFrontendDocumentBuilder& Builder = IMetaSoundAssetManager::GetChecked().AttachDocumentBuilderChecked(MetaSound);
+		FMetaSoundFrontendDocumentBuilder& Builder = IDocumentBuilderRegistry::GetChecked().FindOrBeginBuilding(&MetaSound);
 		const bool bIsNew = UMetaSoundEditorSubsystem::GetChecked().BindMemberMetadata(Builder, *this, LiteralClass);
 		if (bIsNew)
 		{
@@ -1513,33 +1514,22 @@ void UMetasoundEditorGraph::PreSave(FObjectPreSaveContext InSaveContext)
 
 UMetaSoundBuilderBase& UMetasoundEditorGraph::GetBuilderChecked()
 {
-	return UMetaSoundBuilderSubsystem::GetChecked().AttachBuilderToAssetChecked(*GetOutermostObject());
+	using namespace Metasound::Engine;
+	return FDocumentBuilderRegistry::GetChecked().FindOrBeginBuilding(*GetOutermostObject());
 }
 
-UObject* UMetasoundEditorGraph::GetMetasound()
+UObject* UMetasoundEditorGraph::GetMetasound() const
 {
 	return GetOutermostObject();
 }
 
-const UObject* UMetasoundEditorGraph::GetMetasound() const
-{
-	return GetOutermostObject();
-}
-
-UObject& UMetasoundEditorGraph::GetMetasoundChecked()
+UObject& UMetasoundEditorGraph::GetMetasoundChecked() const
 {
 	UObject* ParentMetasound = GetMetasound();
 	check(ParentMetasound);
 	return *ParentMetasound;
 }
 
-const UObject& UMetasoundEditorGraph::GetMetasoundChecked() const
-{
-	const UObject* ParentMetasound = GetMetasound();
-	check(ParentMetasound);
-	return *ParentMetasound;
-}
-
 void UMetasoundEditorGraph::RegisterGraphWithFrontend()
 {
 	using namespace Metasound::Editor;
@@ -1958,8 +1948,13 @@ void UMetasoundEditorGraph::MigrateEditorDocumentData(FMetaSoundFrontendDocument
 	// 5. Remove input locations and ensure that all other nodes only have at most one
 	// location represented in the style/editor graph (0 is acceptable as some member
 	// node types (eg. variables) may not contain locations and that's ok).
-	const FMetasoundFrontendDocument& Document = static_cast<const FMetaSoundFrontendDocumentBuilder&>(OutBuilder).GetDocument();
+	const FMetasoundFrontendDocument& Document = OutBuilder.GetConstDocument();
 	const TArray<FMetasoundFrontendNode>& GraphNodes = Document.RootGraph.Graph.Nodes;
+	TMap<FGuid, const UMetasoundEditorGraphNode*> EdNodeMap;
+	Algo::Transform(AllMetaSoundNodes, EdNodeMap, [](const UEdGraphNode* Node)
+	{
+		return TTuple<FGuid, const UMetasoundEditorGraphNode*>(Node->NodeGuid, Cast<const UMetasoundEditorGraphNode>(Node));
+	});
 	for (const FMetasoundFrontendNode& Node : GraphNodes)
 	{
 		if (const FMetasoundFrontendClass* Class = OutBuilder.FindDependency(Node.ClassID))
@@ -1978,7 +1973,13 @@ void UMetasoundEditorGraph::MigrateEditorDocumentData(FMetaSoundFrontendDocument
 					for (const TPair<FGuid, FVector2D>& Pair : Locations)
 					{
 						DefaultLocation = Pair;
-						break;
+						if (const UMetasoundEditorGraphNode* MetaSoundNode = EdNodeMap.FindRef(Pair.Key))
+						{
+							if (MetaSoundNode->GetNodeID() == Node.GetID())
+							{
+								break;
+							}
+						}
 					}
 
 					// Remove first in case there are multiple locations and the editor guid may be different
diff --git a/Engine/Plugins/Runtime/Metasound/Source/MetasoundEditor/Private/MetasoundEditorGraphBuilder.cpp b/Engine/Plugins/Runtime/Metasound/Source/MetasoundEditor/Private/MetasoundEditorGraphBuilder.cpp
--- a/Engine/Plugins/Runtime/Metasound/Source/MetasoundEditor/Private/MetasoundEditorGraphBuilder.cpp
+++ b/Engine/Plugins/Runtime/Metasound/Source/MetasoundEditor/Private/MetasoundEditorGraphBuilder.cpp
@@ -15,7 +15,7 @@
 #include "Metasound.h"
 #include "MetasoundAssetBase.h"
 #include "MetasoundAudioBuffer.h"
-#include "MetasoundBuilderSubsystem.h"
+#include "MetasoundDocumentBuilderRegistry.h"
 #include "MetasoundEditor.h"
 #include "MetasoundEditorGraph.h"
 #include "MetasoundEditorGraphCommentNode.h"
@@ -286,12 +286,9 @@ namespace Metasound
 					bool bIsClassNative = FMetasoundFrontendRegistryContainer::Get()->IsNodeNative(RegistryKey);
 					if (!bIsClassNative)
 					{
-						if (IMetaSoundAssetManager* AssetManager = IMetaSoundAssetManager::Get())
+						if (const FTopLevelAssetPath* Path = IMetaSoundAssetManager::GetChecked().FindAssetPath(RegistryKey))
 						{
-							if (const FSoftObjectPath* Path = AssetManager->FindObjectPathFromKey(RegistryKey))
-							{
-								DisplayName = FText::FromString(Path->GetAssetName());
-							}
+							DisplayName = FText::FromName(Path->GetAssetName());
 						}
 					}
 				}
@@ -1658,11 +1655,12 @@ namespace Metasound
 
 		bool FGraphBuilder::DeleteNode(UEdGraphNode& InNode, bool bRemoveUnusedDependencies)
 		{
+			using namespace Engine;
 			using namespace Frontend;
 
 			UMetasoundEditorGraph* Graph = CastChecked<UMetasoundEditorGraph>(InNode.GetGraph());
 			UObject& MetaSound = Graph->GetMetasoundChecked();
-			UMetaSoundBuilderBase& Builder = UMetaSoundBuilderSubsystem::GetChecked().AttachBuilderToAssetChecked(MetaSound);
+			UMetaSoundBuilderBase& Builder = FDocumentBuilderRegistry::GetChecked().FindOrBeginBuilding(MetaSound);
 
 			if (UMetasoundEditorGraphCommentNode* Node = Cast<UMetasoundEditorGraphCommentNode>(&InNode))
 			{
@@ -1747,17 +1745,20 @@ namespace Metasound
 			TArray<FMetasoundAssetBase*> EditedReferencingMetaSounds;
 			if (GEditor)
 			{
-				TArray<UObject*> EditedAssets = GEditor->GetEditorSubsystem<UAssetEditorSubsystem>()->GetAllEditedAssets();
-				for (UObject* Asset : EditedAssets)
+				if (UAssetEditorSubsystem* AssetSubsystem = GEditor->GetEditorSubsystem<UAssetEditorSubsystem>())
 				{
-					if (Asset != &InMetaSound)
+					TArray<UObject*> EditedAssets = AssetSubsystem->GetAllEditedAssets();
+					for (UObject* Asset : EditedAssets)
 					{
-						if (FMetasoundAssetBase* EditedMetaSound = IMetasoundUObjectRegistry::Get().GetObjectAsAssetBase(Asset))
+						if (Asset != &InMetaSound)
 						{
-							EditedMetaSound->RebuildReferencedAssetClasses();
-							if (EditedMetaSound->IsReferencedAsset(*MetaSoundAsset))
+							if (FMetasoundAssetBase* EditedMetaSound = IMetasoundUObjectRegistry::Get().GetObjectAsAssetBase(Asset))
 							{
-								EditedReferencingMetaSounds.Add(EditedMetaSound);
+								EditedMetaSound->RebuildReferencedAssetClasses();
+								if (EditedMetaSound->IsReferencedAsset(*MetaSoundAsset))
+								{
+									EditedReferencingMetaSounds.Add(EditedMetaSound);
+								}
 							}
 						}
 					}
@@ -1767,9 +1768,6 @@ namespace Metasound
 			FMetaSoundAssetRegistrationOptions RegOptions;
 			RegOptions.bForceReregister = true;
 			RegOptions.bForceViewSynchronization = bInForceViewSynchronization;
-			// Protect against race conditions by utilizing a copy of the graph. Race conditions can happen
-			// if registration is performed asynchronously while the editor is still modifying the graph.
-			RegOptions.bRegisterCopyIfAsync = true; 
 			
 			// if EditedReferencingMetaSounds is empty, then no MetaSounds are open
 			// that reference this MetaSound, so just register this asset. Otherwise,
@@ -2634,16 +2632,16 @@ namespace Metasound
 
 			bool bEditorGraphModified = false;
 
-			FMetasoundAssetBase* MetaSoundAsset = IMetasoundUObjectRegistry::Get().GetObjectAsAssetBase(&InMetaSound);
+			const FMetasoundAssetBase* MetaSoundAsset = IMetasoundUObjectRegistry::Get().GetObjectAsAssetBase(&InMetaSound);
 			check(MetaSoundAsset);
 			UMetasoundEditorGraph* Graph = CastChecked<UMetasoundEditorGraph>(MetaSoundAsset->GetGraph());
-			FGraphHandle GraphHandle = MetaSoundAsset->GetRootGraphHandle();
+			FConstGraphHandle GraphHandle = MetaSoundAsset->GetRootGraphHandle();
 
 			TSet<UMetasoundEditorGraphInput*> Inputs;
 			TSet<UMetasoundEditorGraphOutput*> Outputs;
 
 			// Collect all editor graph inputs with corresponding frontend inputs. 
-			GraphHandle->IterateNodes([&](FNodeHandle NodeHandle)
+			GraphHandle->IterateConstNodes([&](FConstNodeHandle NodeHandle)
 			{
 				if (UMetasoundEditorGraphInput* Input = Graph->FindInput(NodeHandle->GetID()))
 				{
@@ -2681,7 +2679,7 @@ namespace Metasound
 			}
 
 			// Collect all editor graph outputs with corresponding frontend outputs. 
-			GraphHandle->IterateNodes([&](FNodeHandle NodeHandle)
+			GraphHandle->IterateConstNodes([&](FConstNodeHandle NodeHandle)
 			{
 				if (UMetasoundEditorGraphOutput* Output = Graph->FindOutput(NodeHandle->GetID()))
 				{
@@ -2749,7 +2747,7 @@ namespace Metasound
 			};
 
 			// Synchronize data types & default values for input nodes.
-			GraphHandle->IterateNodes([&](FNodeHandle NodeHandle)
+			GraphHandle->IterateConstNodes([&](FConstNodeHandle NodeHandle)
 			{
 				if (UMetasoundEditorGraphInput* Input = Graph->FindInput(NodeHandle->GetID()))
 				{
@@ -2778,7 +2776,7 @@ namespace Metasound
 			}, EMetasoundFrontendClassType::Input);
 
 			// Synchronize data types of output nodes.
-			GraphHandle->IterateNodes([&](FNodeHandle NodeHandle)
+			GraphHandle->IterateConstNodes([&](FConstNodeHandle NodeHandle)
 			{
 				if (UMetasoundEditorGraphOutput* Output = Graph->FindOutput(NodeHandle->GetID()))
 				{
@@ -2795,8 +2793,8 @@ namespace Metasound
 			// Enforce all variables exists in the edgraph
 			TSet<FGuid> VariableIDs;
 			Algo::Transform(Graph->Variables, VariableIDs, [](const TObjectPtr<UMetasoundEditorGraphVariable>& Variable) { return Variable->GetVariableID(); });
-			TArray<FVariableHandle> FrontendVariables = GraphHandle->GetVariables();
-			for (const FVariableHandle& VariableHandle : FrontendVariables)
+			TArray<FConstVariableHandle> FrontendVariables = GraphHandle->GetVariables();
+			for (const FConstVariableHandle& VariableHandle : FrontendVariables)
 			{
 				const FGuid VariableID = VariableHandle->GetID();
 				if (VariableIDs.Remove(VariableID) == 0)
diff --git a/Engine/Plugins/Runtime/Metasound/Source/MetasoundEditor/Private/MetasoundEditorGraphCommentNode.cpp b/Engine/Plugins/Runtime/Metasound/Source/MetasoundEditor/Private/MetasoundEditorGraphCommentNode.cpp
--- a/Engine/Plugins/Runtime/Metasound/Source/MetasoundEditor/Private/MetasoundEditorGraphCommentNode.cpp
+++ b/Engine/Plugins/Runtime/Metasound/Source/MetasoundEditor/Private/MetasoundEditorGraphCommentNode.cpp
@@ -5,7 +5,7 @@
 #include "Internationalization/Internationalization.h"
 #include "Kismet2/Kismet2NameValidators.h"
 #include "Layout/SlateRect.h"
-#include "MetasoundBuilderSubsystem.h"
+#include "MetasoundDocumentBuilderRegistry.h"
 #include "MetasoundEditorGraph.h"
 #include "MetasoundUObjectRegistry.h"
 #include "Styling/AppStyle.h"
@@ -72,11 +72,13 @@ void UMetasoundEditorGraphCommentNode::PostEditChangeProperty(FPropertyChangedEv
 
 void UMetasoundEditorGraphCommentNode::ResizeNode(const FVector2D& NewSize)
 {
+	using namespace Metasound::Engine;
+
 	Super::ResizeNode(NewSize);
 	if (bCanResizeNode) 
 	{
 		UMetasoundEditorGraph& Graph = *CastChecked<UMetasoundEditorGraph>(GetGraph());
-		UMetaSoundBuilderBase& Builder = UMetaSoundBuilderSubsystem::GetChecked().AttachBuilderToAssetChecked(Graph.GetMetasoundChecked());
+		UMetaSoundBuilderBase& Builder = FDocumentBuilderRegistry::GetChecked().FindOrBeginBuilding(Graph.GetMetasoundChecked());
 		Builder.FindOrAddGraphComment(CommentID).Size = FVector2D(NodeWidth, NodeHeight);
 	}
 }
diff --git a/Engine/Plugins/Runtime/Metasound/Source/MetasoundEditor/Private/MetasoundEditorGraphNode.cpp b/Engine/Plugins/Runtime/Metasound/Source/MetasoundEditor/Private/MetasoundEditorGraphNode.cpp
--- a/Engine/Plugins/Runtime/Metasound/Source/MetasoundEditor/Private/MetasoundEditorGraphNode.cpp
+++ b/Engine/Plugins/Runtime/Metasound/Source/MetasoundEditor/Private/MetasoundEditorGraphNode.cpp
@@ -11,7 +11,7 @@
 #include "Logging/TokenizedMessage.h"
 #include "Metasound.h"
 #include "MetasoundAssetManager.h"
-#include "MetasoundBuilderSubsystem.h"
+#include "MetasoundDocumentBuilderRegistry.h"
 #include "MetasoundEditorCommands.h"
 #include "MetasoundEditorGraph.h"
 #include "MetasoundEditorGraphBuilder.h"
@@ -200,16 +200,20 @@ int32 UMetasoundEditorGraphNode::EstimateNodeWidth() const
 
 UMetaSoundBuilderBase& UMetasoundEditorGraphNode::GetBuilderChecked() const
 {
+	using namespace Metasound::Engine;
+
 	UMetasoundEditorGraph* EdGraph = CastChecked<UMetasoundEditorGraph>(GetGraph());
-	return UMetaSoundBuilderSubsystem::GetChecked().AttachBuilderToAssetChecked(EdGraph->GetMetasoundChecked());
+	return FDocumentBuilderRegistry::GetChecked().FindOrBeginBuilding(EdGraph->GetMetasoundChecked());
 }
 
 const FMetasoundFrontendNode* UMetasoundEditorGraphNode::GetFrontendNode() const
 {
+	using namespace Metasound::Engine;
+
 	if (UObject* Outermost = GetOutermostObject())
 	{
 		const FGuid NodeID = GetNodeID();
-		const UMetaSoundBuilderBase& Builder = UMetaSoundBuilderSubsystem::GetChecked().AttachBuilderToAssetChecked(*Outermost);
+		const UMetaSoundBuilderBase& Builder = FDocumentBuilderRegistry::GetChecked().FindOrBeginBuilding(*Outermost);
 		return Builder.GetConstBuilder().FindNode(NodeID);
 	}
 
@@ -218,11 +222,13 @@ const FMetasoundFrontendNode* UMetasoundEditorGraphNode::GetFrontendNode() const
 
 const FMetasoundFrontendNode& UMetasoundEditorGraphNode::GetFrontendNodeChecked() const
 {
+	using namespace Metasound::Engine;
+
 	UObject* Outermost = GetOutermostObject();
 	check(Outermost);
 
 	const FGuid NodeID = GetNodeID();
-	const UMetaSoundBuilderBase& Builder = UMetaSoundBuilderSubsystem::GetChecked().AttachBuilderToAssetChecked(*Outermost);
+	const UMetaSoundBuilderBase& Builder = FDocumentBuilderRegistry::GetChecked().FindOrBeginBuilding(*Outermost);
 
 	const FMetasoundFrontendNode* FrontendNode = Builder.GetConstBuilder().FindNode(NodeID);
 	check(FrontendNode);
@@ -235,13 +241,17 @@ const FMetasoundEditorGraphNodeBreadcrumb& UMetasoundEditorGraphNode::GetBreadcr
 	return StubCrumb;
 }
 
-UObject& UMetasoundEditorGraphNode::GetMetasoundChecked()
+UObject* UMetasoundEditorGraphNode::GetMetasound() const
 {
-	UMetasoundEditorGraph* EdGraph = CastChecked<UMetasoundEditorGraph>(GetGraph());
-	return EdGraph->GetMetasoundChecked();
+	if (UMetasoundEditorGraph* EdGraph = CastChecked<UMetasoundEditorGraph>(GetGraph()))
+	{
+		return EdGraph->GetMetasound();
+	}
+
+	return nullptr;
 }
 
-const UObject& UMetasoundEditorGraphNode::GetMetasoundChecked() const
+UObject& UMetasoundEditorGraphNode::GetMetasoundChecked() const
 {
 	UMetasoundEditorGraph* EdGraph = CastChecked<UMetasoundEditorGraph>(GetGraph());
 	return EdGraph->GetMetasoundChecked();
@@ -1053,16 +1063,15 @@ void UMetasoundEditorGraphExternalNode::Validate(Metasound::Editor::FGraphNodeVa
 	FConstNodeHandle NodeHandle = GetNodeHandle();
 	const FMetasoundFrontendClassMetadata& Metadata = NodeHandle->GetClassMetadata();
 
-	// 1. Validate referenced graph recursively if defined as asset node class
-	const FNodeRegistryKey RegistryKey = FNodeRegistryKey(Metadata);
-	if (IMetaSoundAssetManager* AssetManager = IMetaSoundAssetManager::Get())
+	// 1. Validate external referenced graph or template node
+	switch(Metadata.GetType())
 	{
-		if (const FSoftObjectPath* Path = AssetManager->FindObjectPathFromKey(RegistryKey))
+		case EMetasoundFrontendClassType::External:
+		case EMetasoundFrontendClassType::Graph:
 		{
-			if (UObject* AssetObject = Path->ResolveObject())
+			const FAssetKey AssetKey(Metadata);
+			if (FMetasoundAssetBase* MetaSoundAsset = IMetaSoundAssetManager::GetChecked().TryLoadAssetFromKey(AssetKey))
 			{
-				FMetasoundAssetBase* MetaSoundAsset = Metasound::IMetasoundUObjectRegistry::Get().GetObjectAsAssetBase(AssetObject);
-				check(MetaSoundAsset);
 				const UMetasoundEditorGraph* NodeGraph = CastChecked<UMetasoundEditorGraph>(&MetaSoundAsset->GetGraphChecked());
 				const EMessageSeverity::Type MaxGraphMsg = static_cast<EMessageSeverity::Type>(NodeGraph->GetHighestMessageSeverity());
 				switch (MaxGraphMsg)
@@ -1087,38 +1096,45 @@ void UMetasoundEditorGraphExternalNode::Validate(Metasound::Editor::FGraphNodeVa
 					break;
 				}
 			}
+			break;
 		}
-	}
 
-	// 2. Validate template nodes
-	if (Metadata.GetType() == EMetasoundFrontendClassType::Template)
-	{
-		const FNodeRegistryKey Key = FNodeRegistryKey(Metadata);
-		if (const INodeTemplate* Template = INodeTemplateRegistry::Get().FindTemplate(Key))
+		case EMetasoundFrontendClassType::Template:
 		{
-			const bool bIsValidInterface = Template->IsValidNodeInterface(NodeHandle->GetNodeInterface());
-			if (!bIsValidInterface)
-			{
-				OutResult.SetMessage(EMessageSeverity::Error, FString::Format(TEXT("Cannot implement template interface for node class '{0}"), { *Metadata.GetClassName().ToString() }));
-			}
-			else
+			const FNodeRegistryKey Key = FNodeRegistryKey(Metadata);
+			if (const INodeTemplate* Template = INodeTemplateRegistry::Get().FindTemplate(Key))
 			{
-#if WITH_EDITOR
-				FString Message;
-				if (!Template->HasRequiredConnections(GetBuilderChecked().GetConstBuilder(), GetNodeID(), &Message))
+				const bool bIsValidInterface = Template->IsValidNodeInterface(NodeHandle->GetNodeInterface());
+				if (!bIsValidInterface)
 				{
-					OutResult.SetMessage(EMessageSeverity::Warning, Message);
+					OutResult.SetMessage(EMessageSeverity::Error, FString::Format(TEXT("Cannot implement template interface for node class '{0}"), { *Metadata.GetClassName().ToString() }));
 				}
+				else
+				{
+#if WITH_EDITOR
+					FString Message;
+					if (!Template->HasRequiredConnections(GetBuilderChecked().GetConstBuilder(), GetNodeID(), &Message))
+					{
+						OutResult.SetMessage(EMessageSeverity::Warning, Message);
+					}
 #endif // WITH_EDITOR
+				}
 			}
+			else
+			{
+				OutResult.SetMessage(EMessageSeverity::Error, FString::Format(TEXT("Template node interface missing for node class '{0}'"), { *Metadata.GetClassName().ToString() }));
+			}
+			break;
 		}
-		else
+
+		default:
 		{
-			OutResult.SetMessage(EMessageSeverity::Error, FString::Format(TEXT("Template node interface missing for node class '{0}'"), { *Metadata.GetClassName().ToString()}));
+			static_assert(static_cast<int32>(EMetasoundFrontendClassType::Invalid) == 10, "Possible missing EMetasoundFrontendClassType case coverage");
 		}
+		break;
 	}
 
-	// 3. Check if node is invalid, version is missing and cache if interface changes exist between the document's records and the registry
+	// 2. Check if node is invalid, version is missing and cache if interface changes exist between the document's records and the registry
 	FClassInterfaceUpdates InterfaceUpdates;
 	if (!NodeHandle->DiffAgainstRegistryInterface(InterfaceUpdates, false /* bUseHighestMinorVersion */))
 	{
@@ -1164,6 +1180,7 @@ void UMetasoundEditorGraphExternalNode::Validate(Metasound::Editor::FGraphNodeVa
 	}
 
 	// 4. Report if node was nativized
+	const FNodeRegistryKey RegistryKey(Metadata);
 	bool bNewIsClassNative = FMetasoundFrontendRegistryContainer::Get()->IsNodeNative(RegistryKey);
 	if (bIsClassNative != bNewIsClassNative)
 	{
diff --git a/Engine/Plugins/Runtime/Metasound/Source/MetasoundEditor/Private/MetasoundEditorGraphSchema.cpp b/Engine/Plugins/Runtime/Metasound/Source/MetasoundEditor/Private/MetasoundEditorGraphSchema.cpp
--- a/Engine/Plugins/Runtime/Metasound/Source/MetasoundEditor/Private/MetasoundEditorGraphSchema.cpp
+++ b/Engine/Plugins/Runtime/Metasound/Source/MetasoundEditor/Private/MetasoundEditorGraphSchema.cpp
@@ -19,6 +19,7 @@
 #include "MetasoundAssetManager.h"
 #include "MetasoundBuilderSubsystem.h"
 #include "MetasoundDataReference.h"
+#include "MetasoundDocumentInterface.h"
 #include "MetasoundEditor.h"
 #include "MetasoundEditorCommands.h"
 #include "MetasoundEditorGraph.h"
@@ -922,7 +923,6 @@ UEdGraphNode* FMetasoundGraphSchemaAction_NewReroute::PerformAction(UEdGraph* Pa
 	check(ParentGraph);
 	UMetasoundEditorGraph* MetaSoundGraph = CastChecked<UMetasoundEditorGraph>(ParentGraph);
 	UObject& ParentMetasound = MetaSoundGraph->GetMetasoundChecked();
-	UMetaSoundBuilderBase& Builder = UMetaSoundBuilderSubsystem::GetChecked().AttachBuilderToAssetChecked(ParentMetasound);
 
 	const FScopedTransaction Transaction(LOCTEXT("AddNewRerouteNode", "Add Reroute Node"));
 	ParentMetasound.Modify();
@@ -934,7 +934,7 @@ UEdGraphNode* FMetasoundGraphSchemaAction_NewReroute::PerformAction(UEdGraph* Pa
 	const INodeTemplate* RerouteTemplate = INodeTemplateRegistry::Get().FindTemplate(FRerouteNodeTemplate::ClassName);
 	check(RerouteTemplate);
 
-	const FMetaSoundFrontendDocumentBuilder& DocBuilder = Builder.GetConstBuilder();
+	const FMetaSoundFrontendDocumentBuilder& DocBuilder = IDocumentBuilderRegistry::GetChecked().FindOrBeginBuilding(&ParentMetasound);
 
 	// Provided 'FromPin' is what to connect to, so if its an input, its the output of the generated node needs to
 	// match the from pin and vice versa.
@@ -1372,6 +1372,7 @@ void UMetasoundEditorGraphSchema::OnPinConnectionDoubleCicked(UEdGraphPin* PinA,
 {
 	using namespace Metasound;
 	using namespace Metasound::Editor;
+	using namespace Metasound::Frontend;
 
 	if (!PinA || !PinB)
 	{
@@ -1387,8 +1388,7 @@ void UMetasoundEditorGraphSchema::OnPinConnectionDoubleCicked(UEdGraphPin* PinA,
 	{
 		UMetasoundEditorGraph* MetaSoundGraph = CastChecked<UMetasoundEditorGraph>(ParentGraph);
 		UObject& ParentMetasound = MetaSoundGraph->GetMetasoundChecked();
-		UMetaSoundBuilderBase& Builder = UMetaSoundBuilderSubsystem::GetChecked().AttachBuilderToAssetChecked(ParentMetasound);
-		const FMetaSoundFrontendDocumentBuilder& DocBuilder = Builder.GetConstBuilder();
+		const FMetaSoundFrontendDocumentBuilder& DocBuilder = IDocumentBuilderRegistry::GetChecked().FindOrBeginBuilding(&ParentMetasound);
 		
 		FName VertexDataType;
 		const FMetasoundFrontendVertex* Vertex = FGraphBuilder::GetPinVertex(DocBuilder, PinA);
@@ -1778,7 +1778,7 @@ void UMetasoundEditorGraphSchema::GetAssetsGraphHoverMessage(const TArray<FAsset
 				break;
 			}
 
-			if (MetaSoundAsset->AddingReferenceCausesLoop(Data.GetSoftObjectPath()))
+			if (MetaSoundAsset->AddingReferenceCausesLoop(*DroppedMetaSoundAsset))
 			{
 				OutOkIcon = false;
 				OutTooltipText = TEXT("Cannot add an asset that would create a reference loop.");
@@ -1861,7 +1861,7 @@ void UMetasoundEditorGraphSchema::DroppedAssetsOnGraph(const TArray<FAssetData>&
 				continue;
 			}
 
-			if (MetaSoundAsset->AddingReferenceCausesLoop(DroppedAsset.GetSoftObjectPath()))
+			if (MetaSoundAsset->AddingReferenceCausesLoop(*DroppedMetaSoundAsset))
 			{
 				continue;
 			}
@@ -2056,10 +2056,8 @@ void UMetasoundEditorGraphSchema::GetFunctionActions(FGraphActionMenuBuilder& Ac
 	FMetasoundAssetBase* ParentAsset = nullptr;
 	if (InGraphHandle->IsValid())
 	{
-		FMetasoundFrontendClassMetadata AssetMetadata = InGraphHandle->GetGraphMetadata();
-		AssetMetadata.SetType(EMetasoundFrontendClassType::External);
-		const FNodeRegistryKey RegistryKey = FNodeRegistryKey(AssetMetadata);
-		ParentAsset = AssetManager.TryLoadAssetFromKey(RegistryKey);
+		const FAssetKey AssetKey(InGraphHandle->GetGraphMetadata());
+		ParentAsset = AssetManager.TryLoadAssetFromKey(AssetKey);
 	}
 
 	const bool bIncludeAllVersions = static_cast<bool>(EnableAllVersionsMetaSoundNodeClassCreationCVar);
@@ -2090,16 +2088,20 @@ void UMetasoundEditorGraphSchema::GetFunctionActions(FGraphActionMenuBuilder& Ac
 		if (ParentAsset)
 		{
 			const FNodeRegistryKey RegistryKey = FNodeRegistryKey(FrontendClass.Metadata);
-			if (const FSoftObjectPath* Path = AssetManager.FindObjectPathFromKey(RegistryKey))
+			if (FMetasoundAssetBase* NodeClassAsset = AssetManager.FindAsset(RegistryKey))
 			{
 				// The path can be not an asset if it was transiently created (eg. in PIE if generated by a runtime MetaSoundDocumentBuilder)
 				// Ignore these assets and do not include as this would create a dependency on a transient MetaSound from a MetaSound asset.
-				if (!Path->IsAsset())
+				// (MetasoundAssetBase is a misnomer; it is not always an asset now that the builder API can generate transient MetaSounds
+				// at runtime).
+				UObject* OwningAsset = NodeClassAsset->GetOwningAsset();
+				check(OwningAsset);
+				if (!OwningAsset->IsAsset())
 				{
 					continue;
 				}
 
-				if (ParentAsset->AddingReferenceCausesLoop(*Path))
+				if (ParentAsset->AddingReferenceCausesLoop(*NodeClassAsset))
 				{
 					continue;
 				}
diff --git a/Engine/Plugins/Runtime/Metasound/Source/MetasoundEditor/Private/MetasoundEditorModule.cpp b/Engine/Plugins/Runtime/Metasound/Source/MetasoundEditor/Private/MetasoundEditorModule.cpp
--- a/Engine/Plugins/Runtime/Metasound/Source/MetasoundEditor/Private/MetasoundEditorModule.cpp
+++ b/Engine/Plugins/Runtime/Metasound/Source/MetasoundEditor/Private/MetasoundEditorModule.cpp
@@ -18,6 +18,7 @@
 #include "MetasoundAudioBuffer.h"
 #include "MetasoundBuilderSubsystem.h"
 #include "MetasoundDetailCustomization.h"
+#include "MetasoundDocumentInterface.h"
 #include "MetasoundEditorGraph.h"
 #include "MetasoundEditorGraphBuilder.h"
 #include "MetasoundEditorGraphConnectionDrawingPolicy.h"
@@ -267,7 +268,7 @@ namespace Metasound
 						return;
 					}
 
-					if (AssetPrimeStatus == EAssetPrimeStatus::NotRequested)
+					if (AssetPrimeStatus == EAssetPrimeStatus::NotRequested || AssetPrimeStatus == EAssetPrimeStatus::Canceled)
 					{
 						return;
 					}
@@ -277,6 +278,11 @@ namespace Metasound
 					FSoftObjectPath AssetPath = InAssetData.ToSoftObjectPath();
 					auto LoadAndRegister = [this, ObjectPath = AssetPath, RegOptions](const FName& PackageName, UPackage* Package, EAsyncLoadingResult::Type Result)
 					{
+						if (AssetPrimeStatus == EAssetPrimeStatus::Canceled)
+						{
+							return;
+						}
+
 						if (Result == EAsyncLoadingResult::Succeeded)
 						{
 							UObject* MetaSoundObj = ObjectPath.ResolveObject();
@@ -310,13 +316,9 @@ namespace Metasound
 					const bool bIsRegisteredClass = IMetasoundUObjectRegistry::Get().IsRegisteredClass(*AssetClass);
 					if (bIsRegisteredClass)
 					{
-						check(GEngine);
-						UMetaSoundAssetSubsystem* AssetSubsystem = GEngine->GetEngineSubsystem<UMetaSoundAssetSubsystem>();
-						check(AssetSubsystem);
-
-						const FNodeRegistryKey RegistryKey = AssetSubsystem->AddOrUpdateAsset(InAssetData);
+						const FNodeRegistryKey RegistryKey = IMetaSoundAssetManager::GetChecked().AddOrUpdateAsset(InAssetData);
 
-						// Can be invalid if being called for the first time on an asset before FRenameRootGraphClass is called
+						// Can be invalid if being called for the first time on an asset before its class name is generated
 						if (RegistryKey.IsValid())
 						{
 							const bool bPrimeRequested = AssetPrimeStatus > EAssetPrimeStatus::NotRequested;
@@ -342,11 +344,7 @@ namespace Metasound
 					const bool bIsRegisteredClass = IMetasoundUObjectRegistry::Get().IsRegisteredClass(*AssetClass);
 					if (bIsRegisteredClass)
 					{
-						check(GEngine);
-						UMetaSoundAssetSubsystem* AssetSubsystem = GEngine->GetEngineSubsystem<UMetaSoundAssetSubsystem>();
-						check(AssetSubsystem);
-
-						const FNodeRegistryKey RegistryKey = AssetSubsystem->AddOrUpdateAsset(InAssetData);
+						const FNodeRegistryKey RegistryKey = IMetaSoundAssetManager::GetChecked().AddOrUpdateAsset(InAssetData);
 						const bool bPrimeRequested = AssetPrimeStatus > EAssetPrimeStatus::NotRequested;
 						const bool bIsRegistered = FMetasoundFrontendRegistryContainer::Get()->IsNodeRegistered(RegistryKey);
 
@@ -383,7 +381,7 @@ namespace Metasound
 						return IMetasoundUObjectRegistry::Get().IsRegisteredClass(*AssetClass);
 					}
 					
-					return false;	
+					return false;
 				};
 
 				for (const TPair<UObject*, UObject*>& Pair : InPackageReloadedEvent->GetRepointedObjects())
@@ -392,12 +390,9 @@ namespace Metasound
 					{
 						if (IsAssetMetaSound(Obj))
 						{
-							check(GEngine);
-							UMetaSoundAssetSubsystem* AssetSubsystem = GEngine->GetEngineSubsystem<UMetaSoundAssetSubsystem>();
-							check(AssetSubsystem);
+							IMetaSoundAssetManager::GetChecked().RemoveAsset(*Pair.Key);
 
 							// Use the editor version of UnregisterWithFrontend so it refreshes any open MetaSound editors
-							AssetSubsystem->RemoveAsset(*Pair.Key);
 							FGraphBuilder::UnregisterGraphWithFrontend(*Pair.Key);
 						}
 					}
@@ -406,11 +401,9 @@ namespace Metasound
 					{
 						if (IsAssetMetaSound(Obj))
 						{
-							check(GEngine);
-							UMetaSoundAssetSubsystem* AssetSubsystem = GEngine->GetEngineSubsystem<UMetaSoundAssetSubsystem>();
-							check(AssetSubsystem);
+							IMetaSoundAssetManager::GetChecked().AddOrUpdateAsset(*Pair.Value);
+
 							// Use the editor version of RegisterWithFrontend so it refreshes any open MetaSound editors
-							AssetSubsystem->AddOrUpdateAsset(*Pair.Value);
 							FGraphBuilder::RegisterGraphWithFrontend(*Pair.Value);
 						}
 					}
@@ -439,17 +432,14 @@ namespace Metasound
 
 			void RemoveAssetFromClassRegistry(const FAssetData& InAssetData)
 			{
+				using namespace Frontend;
 				if (const UClass* AssetClass = InAssetData.GetClass())
 				{
 					const bool bIsRegisteredClass = IMetasoundUObjectRegistry::Get().IsRegisteredClass(*AssetClass);
 					if (bIsRegisteredClass)
 					{
-						check(GEngine);
-						UMetaSoundAssetSubsystem* AssetSubsystem = GEngine->GetEngineSubsystem<UMetaSoundAssetSubsystem>();
-						check(AssetSubsystem);
-
 						// Use the editor version of UnregisterWithFrontend so it refreshes any open MetaSound editors
-						AssetSubsystem->RemoveAsset(InAssetData);
+						IMetaSoundAssetManager::GetChecked().RemoveAsset(InAssetData);
 						if (UObject* AssetObject = InAssetData.GetAsset())
 						{
 							FGraphBuilder::UnregisterGraphWithFrontend(*AssetObject);
@@ -460,20 +450,18 @@ namespace Metasound
 
 			void RenameAssetInClassRegistry(const FAssetData& InAssetData, const FString& InOldObjectPath)
 			{
+				using namespace Frontend;
+
 				if (const UClass* AssetClass = InAssetData.GetClass())
 				{
 					const bool bIsRegisteredClass = IMetasoundUObjectRegistry::Get().IsRegisteredClass(*AssetClass);
 					if (bIsRegisteredClass)
 					{
-						check(GEngine);
-						UMetaSoundAssetSubsystem* AssetSubsystem = GEngine->GetEngineSubsystem<UMetaSoundAssetSubsystem>();
-						check(AssetSubsystem);
+						constexpr bool bReregisterWithFrontend = false;
+						IMetaSoundAssetManager::GetChecked().RenameAsset(InAssetData, bReregisterWithFrontend);
 
 						// Use the FGraphBuilder Register call instead of registering via the
 						// MetaSoundAssetSubsystem so as to properly refresh respective open editors.
-						constexpr bool bReregisterWithFrontend = false;
-						AssetSubsystem->RenameAsset(InAssetData, bReregisterWithFrontend);
-
 						constexpr bool bForceViewSynchronization = true;
 						UObject* AssetObject = InAssetData.GetAsset();
 						FGraphBuilder::RegisterGraphWithFrontend(*AssetObject, bForceViewSynchronization);
@@ -894,6 +882,11 @@ namespace Metasound
 			{
 				METASOUND_LLM_SCOPE;
 
+				if (AssetPrimeStatus != EAssetPrimeStatus::NotRequested && AssetPrimeStatus != EAssetPrimeStatus::Complete)
+				{
+					AssetPrimeStatus = EAssetPrimeStatus::Canceled;
+				}
+
 				if (ISettingsModule* SettingsModule = FModuleManager::GetModulePtr<ISettingsModule>("Settings"))
 				{
 					SettingsModule->UnregisterSettings("Editor", "Audio", "MetaSound Editor");
@@ -934,6 +927,8 @@ namespace Metasound
 
 			void OnPackageMigration(UE::AssetTools::FPackageMigrationContext& MigrationContext)
 			{
+				using namespace Metasound::Frontend;
+
 				// Migration can create temporary new packages that use the same name 
 				// (and therefore node registry key) as the asset migrated. 
 				// So generate new class names to avoid registry key collisions. 
@@ -941,8 +936,7 @@ namespace Metasound
 				{
 					// Gather the new MetaSound assets
 					TMap<FMetasoundFrontendClassName, FMetasoundFrontendClassName> OldToNewReferencedClassNames;
-					TArray<UMetaSoundBuilderBase*> NewMetaSoundAssetBuilders;
-					TArray<FMetasoundAssetBase*> NewMetaSoundAssets;
+					TArray<FMetaSoundFrontendDocumentBuilder> NewMetaSoundAssetBuilders;
 					for (const UE::AssetTools::FPackageMigrationContext::FMigrationPackageData& MigrationPackageData : MigrationContext.GetMigrationPackagesData())
 					{
 						UPackage* Package = MigrationPackageData.GetInstancedPackage();
@@ -952,35 +946,27 @@ namespace Metasound
 							// Only apply to MetaSound assets 
 							if (IMetasoundUObjectRegistry::Get().IsRegisteredClass(MainAsset))
 							{
-								FMetasoundAssetBase* MetaSoundAsset = IMetasoundUObjectRegistry::Get().GetObjectAsAssetBase(MainAsset);
-								if (MetaSoundAsset)
-								{
-									UMetaSoundBuilderBase& MetaSoundBuilder = UMetaSoundBuilderSubsystem::GetChecked().AttachBuilderToAssetChecked(*MainAsset);
-									NewMetaSoundAssetBuilders.Add(&MetaSoundBuilder);
-									NewMetaSoundAssets.Add(MetaSoundAsset);
-								}
+								NewMetaSoundAssetBuilders.Add(FMetaSoundFrontendDocumentBuilder(MainAsset));
 							}
 						}
 					}
 
 					// Assign new class names and cache mapping with old one
-					for (UMetaSoundBuilderBase* MetaSoundBuilder : NewMetaSoundAssetBuilders)
+					IMetaSoundAssetManager& AssetManager = IMetaSoundAssetManager::GetChecked();
+					for (FMetaSoundFrontendDocumentBuilder& MetaSoundBuilder : NewMetaSoundAssetBuilders)
 					{
-						const FMetasoundFrontendClassName NewReferencedClassName = FMetasoundFrontendClassName(/*Namespace=*/{}, FName(FGuid::NewGuid().ToString()), /*Variant=*/{});
-						const FMetasoundFrontendClassName OldReferencedClassName = MetaSoundBuilder->GetRootGraphClassName();
-
-						MetaSoundBuilder->RenameRootGraphClass(NewReferencedClassName);
+						const FMetasoundFrontendClassName OldReferencedClassName = MetaSoundBuilder.GetConstDocument().RootGraph.Metadata.GetClassName();
+						const FMetasoundFrontendClassName NewReferencedClassName = MetaSoundBuilder.GenerateNewClassName();
 						OldToNewReferencedClassNames.FindOrAdd(OldReferencedClassName) = NewReferencedClassName;
-					}
 
-					// Fix up dependencies 
-					for (UMetaSoundBuilderBase* MetaSoundBuilder : NewMetaSoundAssetBuilders)
-					{
-						MetaSoundBuilder->UpdateDependencyClassNames(OldToNewReferencedClassNames);
+						UObject& MetaSoundObject = MetaSoundBuilder.CastDocumentObjectChecked<UObject>();
+						AssetManager.AddOrUpdateAsset(MetaSoundObject);
 					}
-					for (FMetasoundAssetBase* MetaSoundAsset : NewMetaSoundAssets)
+
+					// Fix up dependencies
+					for (FMetaSoundFrontendDocumentBuilder& MetaSoundBuilder : NewMetaSoundAssetBuilders)
 					{
-						MetaSoundAsset->RebuildReferencedAssetClasses();
+						MetaSoundBuilder.UpdateDependencyClassNames(OldToNewReferencedClassNames);
 					}
 				}
 			}
diff --git a/Engine/Plugins/Runtime/Metasound/Source/MetasoundEditor/Private/MetasoundEditorSubsystem.cpp b/Engine/Plugins/Runtime/Metasound/Source/MetasoundEditor/Private/MetasoundEditorSubsystem.cpp
--- a/Engine/Plugins/Runtime/Metasound/Source/MetasoundEditor/Private/MetasoundEditorSubsystem.cpp
+++ b/Engine/Plugins/Runtime/Metasound/Source/MetasoundEditor/Private/MetasoundEditorSubsystem.cpp
@@ -3,7 +3,7 @@
 
 #include "IAssetTools.h"
 #include "Kismet/KismetSystemLibrary.h"
-#include "MetasoundBuilderSubsystem.h"
+#include "MetasoundDocumentBuilderRegistry.h"
 #include "MetasoundEditorGraph.h"
 #include "MetasoundEditorGraphBuilder.h"
 #include "MetasoundEditorGraphSchema.h"
@@ -29,6 +29,7 @@ TScriptInterface<IMetaSoundDocumentInterface> UMetaSoundEditorSubsystem::BuildTo
 )
 {
 	using namespace Metasound;
+	using namespace Metasound::Engine;
 	using namespace Metasound::Frontend;
 
 	OutResult = EMetaSoundBuilderResult::Failed;
@@ -44,9 +45,10 @@ TScriptInterface<IMetaSoundDocumentInterface> UMetaSoundEditorSubsystem::BuildTo
 		}
 
 		constexpr UFactory* Factory = nullptr;
+
 		// Not about to follow this lack of const correctness down a multidecade in the works rabbit hole.
-		UClass& BuilderUClass = const_cast<UClass&>(InBuilder->GetBuilderUClass());
-		if (UObject* NewMetaSound = IAssetTools::Get().CreateAsset(AssetName, PackagePath, &BuilderUClass, Factory))
+		UClass& MetaSoundUClass = const_cast<UClass&>(InBuilder->GetBaseMetaSoundUClass());
+		if (UObject* NewMetaSound = IAssetTools::Get().CreateAsset(AssetName, PackagePath, &MetaSoundUClass, Factory))
 		{
 			InBuilder->InitNodeLocations();
 			InBuilder->SetAuthor(Author);
@@ -62,7 +64,7 @@ TScriptInterface<IMetaSoundDocumentInterface> UMetaSoundEditorSubsystem::BuildTo
 
 			// Apply template SoundWave settings
 			{
-				const bool bIsSource = &BuilderUClass == UMetaSoundSource::StaticClass();
+				const bool bIsSource = &MetaSoundUClass == UMetaSoundSource::StaticClass();
 				if (InBuilder->IsPreset())
 				{
 					// Only use referenced UObject's SoundWave settings for sources if not overridden 
@@ -82,7 +84,7 @@ TScriptInterface<IMetaSoundDocumentInterface> UMetaSoundEditorSubsystem::BuildTo
 				}
 			}
 
-			UMetaSoundBuilderBase& NewDocBuilder = UMetaSoundBuilderSubsystem::GetChecked().AttachBuilderToAssetChecked(*NewMetaSound);
+			UMetaSoundBuilderBase& NewDocBuilder = FDocumentBuilderRegistry::GetChecked().FindOrBeginBuilding(*NewMetaSound);
 
 			EMetaSoundBuilderResult InjectResult = EMetaSoundBuilderResult::Failed;
 			constexpr bool bForceNodeCreation = true;
@@ -140,6 +142,20 @@ bool UMetaSoundEditorSubsystem::BindMemberMetadata(
 	return false;
 }
 
+UMetaSoundBuilderBase* UMetaSoundEditorSubsystem::FindOrBeginBuilding(TScriptInterface<IMetaSoundDocumentInterface> MetaSound, EMetaSoundBuilderResult& OutResult) const
+{
+	using namespace Metasound::Engine;
+
+	if (UObject* Object = MetaSound.GetObject(); Object && Object->IsAsset())
+	{
+		OutResult = EMetaSoundBuilderResult::Succeeded;
+		return &FDocumentBuilderRegistry::GetChecked().FindOrBeginBuilding(*Object);
+	}
+
+	OutResult = EMetaSoundBuilderResult::Failed;
+	return nullptr;
+}
+
 UMetaSoundEditorSubsystem& UMetaSoundEditorSubsystem::GetChecked()
 {
 	checkf(GEditor, TEXT("Cannot access UMetaSoundEditorSubsystem without editor loaded"));
diff --git a/Engine/Plugins/Runtime/Metasound/Source/MetasoundEditor/Private/SMetasoundGraphNode.cpp b/Engine/Plugins/Runtime/Metasound/Source/MetasoundEditor/Private/SMetasoundGraphNode.cpp
--- a/Engine/Plugins/Runtime/Metasound/Source/MetasoundEditor/Private/SMetasoundGraphNode.cpp
+++ b/Engine/Plugins/Runtime/Metasound/Source/MetasoundEditor/Private/SMetasoundGraphNode.cpp
@@ -459,9 +459,11 @@ namespace Metasound
 
 		void SMetaSoundGraphNode::OnCommentBubbleToggled(bool bInCommentBubbleVisible)
 		{
+			using namespace Metasound::Engine;
+
 			UMetasoundEditorGraphNode& EdNode = GetMetaSoundNode();
 			UObject& MetaSound = EdNode.GetMetasoundChecked();
-			UMetaSoundBuilderBase& Builder = UMetaSoundBuilderSubsystem::GetChecked().AttachBuilderToAssetChecked(MetaSound);
+			UMetaSoundBuilderBase& Builder = FDocumentBuilderRegistry::GetChecked().FindOrBeginBuilding(MetaSound);
 			if (const FMetasoundFrontendNode* Node = Builder.GetConstBuilder().FindNode(EdNode.GetNodeID()))
 			{
 				if (bInCommentBubbleVisible != Node->Style.Display.bCommentVisible)
@@ -477,10 +479,12 @@ namespace Metasound
 
 		void SMetaSoundGraphNode::OnCommentTextCommitted(const FText& NewComment, ETextCommit::Type CommitInfo)
 		{
+			using namespace Metasound::Engine;
+
 			FString NewCommentString = NewComment.ToString();
 			UMetasoundEditorGraphNode& EdNode = GetMetaSoundNode();
 			UObject& MetaSound = EdNode.GetMetasoundChecked();
-			UMetaSoundBuilderBase& Builder = UMetaSoundBuilderSubsystem::GetChecked().AttachBuilderToAssetChecked(MetaSound);
+			UMetaSoundBuilderBase& Builder = FDocumentBuilderRegistry::GetChecked().FindOrBeginBuilding(MetaSound);
 			if (const FMetasoundFrontendNode* Node = Builder.GetConstBuilder().FindNode(EdNode.GetNodeID()))
 			{
 				if (!Node->Style.Display.Comment.Equals(NewCommentString))
diff --git a/Engine/Plugins/Runtime/Metasound/Source/MetasoundEditor/Private/SMetasoundGraphPin.h b/Engine/Plugins/Runtime/Metasound/Source/MetasoundEditor/Private/SMetasoundGraphPin.h
--- a/Engine/Plugins/Runtime/Metasound/Source/MetasoundEditor/Private/SMetasoundGraphPin.h
+++ b/Engine/Plugins/Runtime/Metasound/Source/MetasoundEditor/Private/SMetasoundGraphPin.h
@@ -8,7 +8,7 @@
 #include "KismetPins/SGraphPinInteger.h"
 #include "KismetPins/SGraphPinObject.h"
 #include "KismetPins/SGraphPinString.h"
-#include "MetasoundBuilderSubsystem.h"
+#include "MetasoundDocumentBuilderRegistry.h"
 #include "MetasoundEditorGraph.h"
 #include "MetasoundEditorGraphBuilder.h"
 #include "MetasoundEditorGraphMemberDefaults.h"
@@ -280,21 +280,25 @@ namespace Metasound
 
 			UMetaSoundBuilderBase& GetBuilderChecked() const
 			{
+				using namespace Metasound::Engine;
+
 				const UMetasoundEditorGraphNode* Node = GetOwningMetaSoundNode();
 				check(Node);
 				UObject* Outermost = Node->GetOutermostObject();
 				check(Outermost);
-				return UMetaSoundBuilderSubsystem::GetChecked().AttachBuilderToAssetChecked(*Outermost);
+				return FDocumentBuilderRegistry::GetChecked().FindOrBeginBuilding(*Outermost);
 			}
 
 			const FMetasoundFrontendNode* GetFrontendNode() const
 			{
+				using namespace Metasound::Engine;
+
 				if (const UMetasoundEditorGraphNode* Node = GetOwningMetaSoundNode())
 				{
-					if (UObject* Outermost = Node->GetOutermostObject())
+					if (UObject* MetaSound = Node->GetMetasound())
 					{
 						const FGuid NodeID = Node->GetNodeID();
-						const UMetaSoundBuilderBase& Builder = UMetaSoundBuilderSubsystem::GetChecked().AttachBuilderToAssetChecked(*Outermost);
+						const UMetaSoundBuilderBase& Builder = FDocumentBuilderRegistry::GetChecked().FindOrBeginBuilding(*MetaSound);
 						return Builder.GetConstBuilder().FindNode(NodeID);
 					}
 				}
@@ -304,15 +308,7 @@ namespace Metasound
 
 			const FMetasoundFrontendNode& GetFrontendNodeChecked() const
 			{
-				const UMetasoundEditorGraphNode* Node = GetOwningMetaSoundNode();
-				check(Node);
-				UObject* Outermost = Node->GetOutermostObject();
-				check(Outermost);
-
-				const FGuid NodeID = Node->GetNodeID();
-				const UMetaSoundBuilderBase& Builder = UMetaSoundBuilderSubsystem::GetChecked().AttachBuilderToAssetChecked(*Outermost);
-
-				const FMetasoundFrontendNode* FrontendNode = Builder.GetConstBuilder().FindNode(NodeID);
+				const FMetasoundFrontendNode* FrontendNode = GetFrontendNode();
 				check(FrontendNode);
 				return *FrontendNode;
 			}
diff --git a/Engine/Plugins/Runtime/Metasound/Source/MetasoundEditor/Public/MetasoundEditorGraph.h b/Engine/Plugins/Runtime/Metasound/Source/MetasoundEditor/Public/MetasoundEditorGraph.h
--- a/Engine/Plugins/Runtime/Metasound/Source/MetasoundEditor/Public/MetasoundEditorGraph.h
+++ b/Engine/Plugins/Runtime/Metasound/Source/MetasoundEditor/Public/MetasoundEditorGraph.h
@@ -378,10 +378,8 @@ public:
 	virtual void PreSave(FObjectPreSaveContext InSaveContext) override;
 
 	UMetaSoundBuilderBase& GetBuilderChecked();
-	UObject* GetMetasound();
-	const UObject* GetMetasound() const;
-	UObject& GetMetasoundChecked();
-	const UObject& GetMetasoundChecked() const;
+	UObject* GetMetasound() const;
+	UObject& GetMetasoundChecked() const;
 
 	void IterateInputs(TFunctionRef<void(UMetasoundEditorGraphInput&)> InFunction) const;
 	void IterateOutputs(TFunctionRef<void(UMetasoundEditorGraphOutput&)> InFunction) const;
diff --git a/Engine/Plugins/Runtime/Metasound/Source/MetasoundEditor/Public/MetasoundEditorGraphNode.h b/Engine/Plugins/Runtime/Metasound/Source/MetasoundEditor/Public/MetasoundEditorGraphNode.h
--- a/Engine/Plugins/Runtime/Metasound/Source/MetasoundEditor/Public/MetasoundEditorGraphNode.h
+++ b/Engine/Plugins/Runtime/Metasound/Source/MetasoundEditor/Public/MetasoundEditorGraphNode.h
@@ -131,8 +131,8 @@ public:
 	// data if re-associated should the associated document data/node becomes unlinked.
 	virtual void CacheBreadcrumb() { }
 
-	UObject& GetMetasoundChecked();
-	const UObject& GetMetasoundChecked() const;
+	UObject* GetMetasound() const;
+	UObject& GetMetasoundChecked() const;
 
 	virtual bool RemoveFromDocument() const;
 
diff --git a/Engine/Plugins/Runtime/Metasound/Source/MetasoundEditor/Public/MetasoundEditorModule.h b/Engine/Plugins/Runtime/Metasound/Source/MetasoundEditor/Public/MetasoundEditorModule.h
--- a/Engine/Plugins/Runtime/Metasound/Source/MetasoundEditor/Public/MetasoundEditorModule.h
+++ b/Engine/Plugins/Runtime/Metasound/Source/MetasoundEditor/Public/MetasoundEditorModule.h
@@ -48,7 +48,8 @@ namespace Metasound
 			NotRequested = 0,
 			Requested = 1,
 			InProgress = 2,
-			Complete = 3
+			Complete = 3,
+			Canceled = 4
 		};
 
 		struct FGraphPinParams
diff --git a/Engine/Plugins/Runtime/Metasound/Source/MetasoundEditor/Public/MetasoundEditorSubsystem.h b/Engine/Plugins/Runtime/Metasound/Source/MetasoundEditor/Public/MetasoundEditorSubsystem.h
--- a/Engine/Plugins/Runtime/Metasound/Source/MetasoundEditor/Public/MetasoundEditorSubsystem.h
+++ b/Engine/Plugins/Runtime/Metasound/Source/MetasoundEditor/Public/MetasoundEditorSubsystem.h
@@ -47,10 +47,15 @@ public:
 		EMetaSoundBuilderResult& OutResult,
 		UPARAM(DisplayName = "Template SoundWave") const USoundWave* TemplateSoundWave = nullptr);
 
+	// Returns a builder for the given MetaSound asset. Returns null if provided a transient MetaSound. For finding builders for transient
+	// MetaSounds, use the UMetaSoundBuilderSubsystem's API (FindPatchBuilder, FindSourceBuilder, FindBuilderByName etc.)
+	UFUNCTION(BlueprintCallable, Category = "Audio|MetaSound|Builder|Editor", meta = (DisplayName = "Find Or Begin Building MetaSound Asset", ExpandEnumAsExecs = "OutResult"))
+	UMetaSoundBuilderBase* FindOrBeginBuilding(TScriptInterface<IMetaSoundDocumentInterface> MetaSound, EMetaSoundBuilderResult& OutResult) const;
+
 	// Sets the visual location to InLocation of a given node InNode of a given builder's document.
 	UFUNCTION(BlueprintCallable, Category = "Audio|MetaSound|Builder|Editor", meta = (ExpandEnumAsExecs = "OutResult"))
 	void SetNodeLocation(
-		UPARAM(DisplayName = "Builder") UMetaSoundBuilderBase * InBuilder,
+		UPARAM(DisplayName = "Builder") UMetaSoundBuilderBase* InBuilder,
 		UPARAM(DisplayName = "Node Handle") const FMetaSoundNodeHandle& InNode,
 		UPARAM(DisplayName = "Location") const FVector2D& InLocation,
 		EMetaSoundBuilderResult& OutResult);
diff --git a/Engine/Plugins/Runtime/Metasound/Source/MetasoundEngine/Private/Metasound.cpp b/Engine/Plugins/Runtime/Metasound/Source/MetasoundEngine/Private/Metasound.cpp
--- a/Engine/Plugins/Runtime/Metasound/Source/MetasoundEngine/Private/Metasound.cpp
+++ b/Engine/Plugins/Runtime/Metasound/Source/MetasoundEngine/Private/Metasound.cpp
@@ -7,6 +7,7 @@
 #include "MetasoundAssetBase.h"
 #include "MetasoundAudioFormats.h"
 #include "MetasoundBuilderSubsystem.h"
+#include "MetasoundDocumentInterface.h"
 #include "MetasoundEngineAsset.h"
 #include "MetasoundEngineEnvironment.h"
 #include "MetasoundEnvironment.h"
@@ -33,6 +34,7 @@
 #if WITH_EDITORONLY_DATA
 #include "EdGraph/EdGraph.h"
 #endif // WITH_EDITORONLY_DATA
+#include "../../MetasoundFrontend/Public/MetasoundAssetManager.h"
 
 #define LOCTEXT_NAMESPACE "MetaSound"
 
@@ -64,14 +66,14 @@ Metasound::Frontend::FDocumentAccessPtr UMetaSoundPatch::GetDocumentAccessPtr()
 	using namespace Metasound::Frontend;
 
 	// Mutation of a document via the soft deprecated access ptr/controller system is not tracked by
-	// the builder registry, so the document cache is invalidated here. It is discouraged to mutate
-	// documents using both systems at the same time as it can corrupt a builder document's cache.
+	// the builder registry, so the document cache is invalidated here.
 	if (UMetaSoundBuilderSubsystem* BuilderSubsystem = UMetaSoundBuilderSubsystem::Get())
 	{
-		const FMetasoundFrontendClassName& Name = RootMetaSoundDocument.RootGraph.Metadata.GetClassName();
-		BuilderSubsystem->InvalidateDocumentCache(Name);
+		if (UMetaSoundBuilderBase* Builder = BuilderSubsystem->FindBuilderOfDocument(this))
+		{
+			Builder->Reload();
+		}
 	}
-
 	// Return document using FAccessPoint to inform the TAccessPtr when the 
 	// object is no longer valid.
 	return MakeAccessPtr<FDocumentAccessPtr>(RootMetaSoundDocument.AccessPoint, RootMetaSoundDocument);
@@ -90,6 +92,11 @@ const UClass& UMetaSoundPatch::GetBaseMetaSoundUClass() const
 	return *UMetaSoundPatch::StaticClass();
 }
 
+const UClass& UMetaSoundPatch::GetBuilderUClass() const
+{
+	return *UMetaSoundPatchBuilder::StaticClass();
+}
+
 const FMetasoundFrontendDocument& UMetaSoundPatch::GetConstDocument() const
 {
 	return RootMetaSoundDocument;
@@ -99,14 +106,7 @@ const FMetasoundFrontendDocument& UMetaSoundPatch::GetConstDocument() const
 void UMetaSoundPatch::PostDuplicate(EDuplicateMode::Type InDuplicateMode)
 {
 	Super::PostDuplicate(InDuplicateMode);
-
-	// Guid is reset as asset may share implementation from
-	// asset duplicated from but should not be registered as such.
-	if (InDuplicateMode == EDuplicateMode::Normal)
-	{
-		AssetClassID = FGuid::NewGuid();
-		Metasound::Frontend::FRenameRootGraphClass::Generate(GetDocumentHandle(), AssetClassID);
-	}
+	Metasound::FMetaSoundEngineAssetHelper::PostDuplicate(this, InDuplicateMode, AssetClassID);
 }
 
 void UMetaSoundPatch::PostEditUndo()
@@ -146,23 +146,12 @@ void UMetaSoundPatch::MigrateEditorGraph(FMetaSoundFrontendDocumentBuilder& OutB
 	PRAGMA_ENABLE_DEPRECATION_WARNINGS
 }
 
-UEdGraph* UMetaSoundPatch::GetGraph()
+UEdGraph* UMetaSoundPatch::GetGraph() const
 {
 	return EditorGraph;
 }
 
-const UEdGraph* UMetaSoundPatch::GetGraph() const
-{
-	return EditorGraph;
-}
-
-UEdGraph& UMetaSoundPatch::GetGraphChecked()
-{
-	check(EditorGraph);
-	return *EditorGraph;
-}
-
-const UEdGraph& UMetaSoundPatch::GetGraphChecked() const
+UEdGraph& UMetaSoundPatch::GetGraphChecked() const
 {
 	check(EditorGraph);
 	return *EditorGraph;
