- Move MetaSoundAssetSubsystem& BuilderSubsystem implementations to monolithics that are initialized at MetaSoundEngine module load for safety

diff --git a/Engine/Plugins/Runtime/Metasound/Source/MetasoundEngine/Private/MetasoundAssetSubsystem.cpp b/Engine/Plugins/Runtime/Metasound/Source/MetasoundEngine/Private/MetasoundAssetSubsystem.cpp
--- a/Engine/Plugins/Runtime/Metasound/Source/MetasoundEngine/Private/MetasoundAssetSubsystem.cpp
+++ b/Engine/Plugins/Runtime/Metasound/Source/MetasoundEngine/Private/MetasoundAssetSubsystem.cpp
@@ -9,6 +9,7 @@
 #include "Metasound.h"
 #include "MetasoundAssetBase.h"
 #include "MetasoundBuilderSubsystem.h"
+#include "MetasoundEngineAsset.h"
 #include "MetasoundFrontendDocument.h"
 #include "MetasoundFrontendDocumentBuilder.h"
 #include "MetasoundFrontendRegistries.h"
@@ -19,13 +20,14 @@
 #include "MetasoundTrace.h"
 #include "MetasoundUObjectRegistry.h"
 #include "Misc/CoreDelegates.h"
+#include "Modules/ModuleManager.h"
 #include "UObject/NoExportTypes.h"
 
 
 #include UE_INLINE_GENERATED_CPP_BY_NAME(MetasoundAssetSubsystem)
 
 
-namespace Metasound
+namespace Metasound::Engine
 {
 	namespace AssetSubsystemPrivate
 	{
@@ -37,7 +39,7 @@ namespace Metasound
 			bool bSuccess = true;
 
 			OutInfo.Type = EMetasoundFrontendClassType::External;
-			OutInfo.AssetPath = InAssetData.AssetClassPath;
+			OutInfo.AssetPath = FTopLevelAssetPath(InAssetData.PackageName, InAssetData.AssetName);
 			FString AssetClassID;
 			bSuccess &= InAssetData.GetTagValue(AssetTags::AssetClassID, AssetClassID);
 			OutInfo.AssetClassID = FGuid(AssetClassID);
@@ -87,689 +89,949 @@ namespace Metasound
 		//
 		// This protects against scenarios where a metasound is renamed or moved 
 		// and the new entry was being erroneously removed from the PathMap
-		bool RemoveIfExactMatch(TMap<Frontend::FNodeRegistryKey, FSoftObjectPath>& InMap, const Frontend::FNodeRegistryKey& InKeyToRemove, const FSoftObjectPath& InPathToRemove)
+		bool RemoveAndWarnIfNotExactMatch(TMap<Frontend::FAssetKey, FTopLevelAssetPath>& InMap, const Frontend::FAssetKey& AssetKey, const FTopLevelAssetPath& AssetPath)
 		{
-			if (const FSoftObjectPath* Path = InMap.Find(InKeyToRemove))
+#if !NO_LOGGING
+			// Can be invalid if document on object being destroyed is being removed and MetaSound was no longer registered,
+			// in which case ignore logging as the assumption is it was a valid object to be removed.
+			if (AssetPath.IsValid())
 			{
-				if (*Path == InPathToRemove)
+				if (const FTopLevelAssetPath* MapAssetPath = InMap.Find(AssetKey))
 				{
-					InMap.Remove(InKeyToRemove);
-					return true;
-				}
-				else
-				{
-					UE_LOG(LogMetaSound, VeryVerbose, TEXT("Object paths do not match. Skipping removal of %s:%s from the asset subsystem."), *InKeyToRemove.ToString(), *InPathToRemove.ToString());
+					if (*MapAssetPath != AssetPath)
+					{
+						UE_LOG(LogMetaSound, Warning, TEXT("AssetManager removing object where path (%s) does not match registry copy (%s)"), *MapAssetPath->ToString(), *AssetPath.ToString());
+					}
 				}
 			}
-			return false;
+#endif // !NO_LOGGING
+
+			return InMap.Remove(AssetKey) > 0;
 		}
 	}
-}
-
-void UMetaSoundAssetSubsystem::Initialize(FSubsystemCollectionBase& InCollection)
-{
-	using namespace Metasound::Frontend;
-
-	IMetaSoundAssetManager::Set(*this);
-	FCoreDelegates::OnPostEngineInit.AddUObject(this, &UMetaSoundAssetSubsystem::PostEngineInit);
-}
 
-void UMetaSoundAssetSubsystem::PostEngineInit()
-{
-	check(UAssetManager::IsInitialized());
-	UAssetManager& AssetManager = UAssetManager::Get();
-	AssetManager.CallOrRegister_OnCompletedInitialScan(FSimpleMulticastDelegate::FDelegate::CreateUObject(this, &UMetaSoundAssetSubsystem::PostInitAssetScan));
-	RebuildDenyListCache(AssetManager);
-}
+	class FMetaSoundAssetManager :
+		public Frontend::IMetaSoundAssetManager,
+		public FGCObject
+	{
+	public:
+		using FAssetInfo = Frontend::IMetaSoundAssetManager::FAssetInfo;
+		using FAssetKey = Frontend::FAssetKey;
 
-void UMetaSoundAssetSubsystem::PostInitAssetScan()
-{
-	METASOUND_TRACE_CPUPROFILER_EVENT_SCOPE(UMetaSoundAssetSubsystem::PostInitAssetScan);
+		FMetaSoundAssetManager() = default;
+		virtual ~FMetaSoundAssetManager();
 
-	const UMetaSoundSettings* Settings = GetDefault<UMetaSoundSettings>();
-	if (ensureAlways(Settings))
-	{
-		SearchAndIterateDirectoryAssets(Settings->DirectoriesToRegister, [this](const FAssetData& AssetData)
+		static FMetaSoundAssetManager& GetChecked()
 		{
-			AddOrUpdateAsset(AssetData);
-		});
-	}
+			using namespace Frontend;
+			return static_cast<FMetaSoundAssetManager&>(IMetaSoundAssetManager::GetChecked());
+		}
 
-	bIsInitialAssetScanComplete = true;
-}
+		void RebuildDenyListCache(const UAssetManager& InAssetManager);
+		void RegisterAssetClassesInDirectories(const TArray<FMetaSoundAssetDirectory>& Directories);
+		void RequestAsyncLoadReferencedAssets(FMetasoundAssetBase& InAssetBase);
+		void OnAssetScanComplete();
+		void SearchAndIterateDirectoryAssets(const TArray<FDirectoryPath>& InDirectories, TFunctionRef<void(const FAssetData&)> InFunction);
+		FMetasoundAssetBase* TryLoadAsset(const FSoftObjectPath& InObjectPath) const;
+		void UnregisterAssetClassesInDirectories(const TArray<FMetaSoundAssetDirectory>& Directories);
 
+		/* IMetaSoundAssetManager Implementation */
 #if WITH_EDITORONLY_DATA
-void UMetaSoundAssetSubsystem::AddAssetReferences(FMetasoundAssetBase& InAssetBase)
-{
-	using namespace Metasound;
-	using namespace Metasound::Frontend;
-
-	const FMetasoundFrontendDocument& Document = InAssetBase.GetConstDocumentChecked();
-	const FNodeRegistryKey AssetClassKey = FNodeRegistryKey(Document.RootGraph);
+		virtual void AddAssetReferences(FMetasoundAssetBase& InAssetBase) override;
+#endif // WITH_EDITORONLY_DATA
+		virtual FAssetKey AddOrUpdateAsset(const FAssetData& InAssetData) override;
+		virtual FAssetKey AddOrUpdateAsset(const UObject& InObject) override;
+		virtual bool CanAutoUpdate(const FMetasoundFrontendClassName& InClassName) const override;
+		virtual bool ContainsKey(const FAssetKey& InKey) const override;
+		virtual FMetasoundAssetBase* FindAsset(const FAssetKey& InKey) const override;
+		virtual TScriptInterface<IMetaSoundDocumentInterface> FindAssetAsDocumentInterface(const Frontend::FAssetKey& InKey) const override;
+		virtual const FTopLevelAssetPath* FindAssetPath(const FAssetKey& InKey) const override;
+		virtual FMetasoundAssetBase* GetAsAsset(UObject& InObject) const override;
+		virtual const FMetasoundAssetBase* GetAsAsset(const UObject& InObject) const override;
+#if WITH_EDITOR
+		virtual TSet<FAssetInfo> GetReferencedAssetClasses(const FMetasoundAssetBase& InAssetBase) const override;
+#endif // WITH_EDITOR
+		virtual void RemoveAsset(const UObject& InObject) override;
+		virtual void RemoveAsset(const FAssetData& InAssetData) override;
+		virtual void RenameAsset(const FAssetData& InAssetData, bool bInReregisterWithFrontend = true) override;
+		virtual FMetasoundAssetBase* TryLoadAssetFromKey(const FAssetKey& InKey) const override;
+		virtual bool TryGetAssetIDFromClassName(const FMetasoundFrontendClassName& InClassName, FGuid& OutGuid) const override;
+		virtual bool TryLoadReferencedAssets(const FMetasoundAssetBase& InAssetBase, TArray<FMetasoundAssetBase*>& OutReferencedAssets) const override;
+		virtual void WaitUntilAsyncLoadReferencedAssetsComplete(FMetasoundAssetBase& InAssetBase) override;
+
+		/* FGCObject */
+		virtual void AddReferencedObjects(FReferenceCollector& Collector) override;
+		virtual FString GetReferencerName() const override { return TEXT("FMetaSoundAssetManager"); }
+
+	private:
+		TArray<FMetaSoundAsyncAssetDependencies> LoadingDependencies;
+
+		FMetaSoundAsyncAssetDependencies* FindLoadingDependencies(const UObject* InParentAsset);
+		FMetaSoundAsyncAssetDependencies* FindLoadingDependencies(int32 InLoadID);
+		void RemoveLoadingDependencies(int32 InLoadID);
+		void OnAssetsLoaded(int32 InLoadID);
+
+		FStreamableManager StreamableManager;
+		int32 AsyncLoadIDCounter = 0;
+
+		int32 AutoUpdateDenyListChangeID = INDEX_NONE;
+		TSet<FName> AutoUpdateDenyListCache;
+		std::atomic<bool> bIsInitialAssetScanComplete = false;
+		TMap<FAssetKey, FTopLevelAssetPath> PathMap;
+	};
 
-	if (!ContainsKey(AssetClassKey))
+	FMetaSoundAssetManager::~FMetaSoundAssetManager()
 	{
-		AddOrUpdateAsset(*InAssetBase.GetOwningAsset());
-		UE_LOG(LogMetaSound, Verbose, TEXT("Adding asset '%s' to MetaSoundAsset registry."), *InAssetBase.GetOwningAssetName());
+		if (!PathMap.IsEmpty())
+		{
+			UE_LOG(LogMetaSound, Display, TEXT("Destroying MetaSoundAssetManager while %i asset paths are still available."), PathMap.Num());
+		}
 	}
 
-	bool bAddFromReferencedAssets = false;
-	const TSet<FString>& ReferencedAssetClassKeys = InAssetBase.GetReferencedAssetClassKeys();
-	for (const FString& KeyString : ReferencedAssetClassKeys)
+	void FMetaSoundAssetManager::AddReferencedObjects(FReferenceCollector& Collector)
 	{
-		FNodeRegistryKey RegistryKey;
-		const bool bIsKey = FNodeRegistryKey::Parse(KeyString, RegistryKey);
-		if (!bIsKey || !ContainsKey(RegistryKey))
+		for (FMetaSoundAsyncAssetDependencies& Dependencies : LoadingDependencies)
 		{
-			UE_LOG(LogMetaSound, Verbose, TEXT("Missing referenced class '%s' asset entry."), *KeyString);
-			bAddFromReferencedAssets = true;
+			Collector.AddReferencedObject(Dependencies.MetaSound);
 		}
 	}
 
-	// All keys are loaded
-	if (!bAddFromReferencedAssets)
+#if WITH_EDITORONLY_DATA
+	void FMetaSoundAssetManager::AddAssetReferences(FMetasoundAssetBase& InAssetBase)
 	{
-		return;
-	}
-
-	UE_LOG(LogMetaSound, Verbose, TEXT("Attempting preemptive reference load..."));
+		using namespace Frontend;
 
-	TArray<FMetasoundAssetBase*> ReferencedAssets = InAssetBase.GetReferencedAssets();
-	for (FMetasoundAssetBase* Asset : ReferencedAssets)
-	{
-		if (Asset)
 		{
-			const FMetasoundFrontendDocument& RefDocument = Asset->GetConstDocumentChecked();
-			const FNodeRegistryKey ClassKey = FNodeRegistryKey(RefDocument.RootGraph);
-			if (!ContainsKey(ClassKey))
+			const FMetasoundFrontendDocument& Document = InAssetBase.GetConstDocumentChecked();
+			const FAssetKey AssetKey(Document.RootGraph.Metadata);
+			if (!ContainsKey(AssetKey))
 			{
-				UE_LOG(LogMetaSound, Verbose,
-					TEXT("Preemptive load of class '%s' due to early "
-						"registration request (asset scan likely not complete)."),
-					*ClassKey.ToString());
-
-				UObject* MetaSoundObject = Asset->GetOwningAsset();
-				if (ensureAlways(MetaSoundObject))
-				{
-					AddOrUpdateAsset(*MetaSoundObject);
-				}
+				AddOrUpdateAsset(*InAssetBase.GetOwningAsset());
+				UE_LOG(LogMetaSound, Verbose, TEXT("Adding asset '%s' to MetaSoundAsset registry."), *InAssetBase.GetOwningAssetName());
 			}
 		}
-		else
+
+		bool bAddFromReferencedAssets = false;
+		const TSet<FString>& ReferencedAssetClassKeys = InAssetBase.GetReferencedAssetClassKeys();
+		for (const FString& KeyString : ReferencedAssetClassKeys)
 		{
-			UE_LOG(LogMetaSound, Warning, TEXT("Null referenced dependent asset in %s. Resaving asset in editor may fix the issue"), *InAssetBase.GetOwningAssetName());
+			FNodeRegistryKey RegistryKey;
+			const bool bIsKey = FNodeRegistryKey::Parse(KeyString, RegistryKey);
+			if (!bIsKey || !ContainsKey(RegistryKey))
+			{
+				UE_LOG(LogMetaSound, Verbose, TEXT("Missing referenced class '%s' asset entry."), *KeyString);
+				bAddFromReferencedAssets = true;
+			}
 		}
-	}
-}
-#endif
 
-Metasound::Frontend::FNodeRegistryKey UMetaSoundAssetSubsystem::AddOrUpdateAsset(const UObject& InObject)
-{
-	using namespace Metasound;
-	using namespace Metasound::AssetSubsystemPrivate;
-	using namespace Metasound::Frontend;
+		// All keys are loaded
+		if (!bAddFromReferencedAssets)
+		{
+			return;
+		}
 
-	METASOUND_TRACE_CPUPROFILER_EVENT_SCOPE(UMetaSoundAssetSubsystem::AddOrUpdateAsset);
+		UE_LOG(LogMetaSound, Verbose, TEXT("Attempting preemptive reference load..."));
 
-	const FMetasoundAssetBase* MetaSoundAsset = Metasound::IMetasoundUObjectRegistry::Get().GetObjectAsAssetBase(&InObject);
-	check(MetaSoundAsset);
+		TArray<FMetasoundAssetBase*> ReferencedAssets = InAssetBase.GetReferencedAssets();
+		for (FMetasoundAssetBase* Asset : ReferencedAssets)
+		{
+			if (Asset)
+			{
+				const FMetasoundFrontendDocument& RefDocument = Asset->GetConstDocumentChecked();
+				const FNodeRegistryKey ClassKey = FNodeRegistryKey(RefDocument.RootGraph);
+				if (!ContainsKey(ClassKey))
+				{
+					UE_LOG(LogMetaSound, Verbose,
+						TEXT("Preemptive load of class '%s' due to early "
+							"registration request (asset scan likely not complete)."),
+						*ClassKey.ToString());
 
-	const FMetasoundFrontendDocument& Document = MetaSoundAsset->GetConstDocumentChecked();
-	const FNodeRegistryKey RegistryKey = FNodeRegistryKey(Document.RootGraph);
+					UObject* MetaSoundObject = Asset->GetOwningAsset();
+					if (ensureAlways(MetaSoundObject))
+					{
+						AddOrUpdateAsset(*MetaSoundObject);
+					}
+				}
+			}
+			else
+			{
+				UE_LOG(LogMetaSound, Warning, TEXT("Null referenced dependent asset in %s. Resaving asset in editor may fix the issue"), *InAssetBase.GetOwningAssetName());
+			}
+		}
+	}
+#endif // WITH_EDITORONLY_DATA
 
-	if (RegistryKey.IsValid())
+	Frontend::FAssetKey FMetaSoundAssetManager::AddOrUpdateAsset(const UObject& InObject)
 	{
-		PathMap.FindOrAdd(RegistryKey) = InObject.GetPathName();
-	}
+		using namespace AssetSubsystemPrivate;
+		using namespace Frontend;
 
-	return RegistryKey;
-}
+		METASOUND_TRACE_CPUPROFILER_EVENT_SCOPE(FMetaSoundAssetManager::AddOrUpdateAsset_UObject);
 
-Metasound::Frontend::FNodeRegistryKey UMetaSoundAssetSubsystem::AddOrUpdateAsset(const FAssetData& InAssetData)
-{
-	using namespace Metasound;
-	using namespace Metasound::Frontend;
+		const FMetasoundAssetBase* MetaSoundAsset = Metasound::IMetasoundUObjectRegistry::Get().GetObjectAsAssetBase(&InObject);
+		check(MetaSoundAsset);
+
+		const FMetasoundFrontendDocument& Document = MetaSoundAsset->GetConstDocumentChecked();
+		const FAssetKey AssetKey = FAssetKey(Document.RootGraph);
+
+		if (AssetKey.IsValid())
+		{
+			PathMap.FindOrAdd(AssetKey) = FTopLevelAssetPath(&InObject);
+		}
 
-	METASOUND_TRACE_CPUPROFILER_EVENT_SCOPE(UMetaSoundAssetSubsystem::AddOrUpdateAsset);
+		return AssetKey;
+	}
 
-	FNodeClassInfo ClassInfo;
-	bool bClassInfoFound = AssetSubsystemPrivate::GetAssetClassInfo(InAssetData, ClassInfo);
-	if (!bClassInfoFound)
+	Frontend::FAssetKey FMetaSoundAssetManager::AddOrUpdateAsset(const FAssetData& InAssetData)
 	{
-		UObject* Object = nullptr;
+		using namespace Frontend;
 
-		FSoftObjectPath Path = InAssetData.ToSoftObjectPath();
-		if (!FPackageName::GetPackageMountPoint(InAssetData.GetObjectPathString()).IsNone())
+		METASOUND_TRACE_CPUPROFILER_EVENT_SCOPE(UMetaSoundAssetSubsystem::AddOrUpdateAsset_AssetData);
+
+		FNodeClassInfo ClassInfo;
+		bool bClassInfoFound = AssetSubsystemPrivate::GetAssetClassInfo(InAssetData, ClassInfo);
+		if (!bClassInfoFound)
 		{
-			if (InAssetData.IsAssetLoaded())
+			UObject* Object = nullptr;
+
+			FSoftObjectPath Path = InAssetData.ToSoftObjectPath();
+			if (!FPackageName::GetPackageMountPoint(InAssetData.GetObjectPathString()).IsNone())
 			{
-				Object = Path.ResolveObject();
-				UE_LOG(LogMetaSound, Verbose, TEXT("Adding loaded asset '%s' to MetaSoundAsset registry."), *Object->GetName());
+				if (InAssetData.IsAssetLoaded())
+				{
+					Object = Path.ResolveObject();
+					UE_LOG(LogMetaSound, Verbose, TEXT("Adding loaded asset '%s' to MetaSoundAsset registry."), *Object->GetName());
+				}
+				else
+				{
+					Object = Path.TryLoad();
+					UE_LOG(LogMetaSound, Verbose, TEXT("Loaded asset '%s' and adding to MetaSoundAsset registry."), *Object->GetName());
+				}
 			}
-			else
+
+			if (Object)
 			{
-				Object = Path.TryLoad();
-				UE_LOG(LogMetaSound, Verbose, TEXT("Loaded asset '%s' and adding to MetaSoundAsset registry."), *Object->GetName());
+				return AddOrUpdateAsset(*Object);
 			}
 		}
 
-		if (Object)
+		if (ClassInfo.AssetClassID.IsValid())
 		{
-			return AddOrUpdateAsset(*Object);
+			const FAssetKey RegistryKey = FAssetKey(ClassInfo.ClassName, ClassInfo.Version);
+			if (RegistryKey.IsValid())
+			{
+				PathMap.FindOrAdd(RegistryKey) = ClassInfo.AssetPath;
+			}
+
+			return RegistryKey;
 		}
+
+		// Invalid ClassID means the node could not be registered.
+		// Let caller report or ensure as necessary.
+		return FAssetKey();
 	}
 
-	if (ClassInfo.AssetClassID.IsValid())
+	bool FMetaSoundAssetManager::CanAutoUpdate(const FMetasoundFrontendClassName& InClassName) const
 	{
-		const FNodeRegistryKey RegistryKey = FNodeRegistryKey(ClassInfo);
-		if (RegistryKey.IsValid())
+		const UMetaSoundSettings* Settings = GetDefault<UMetaSoundSettings>();
+		if (!Settings->bAutoUpdateEnabled)
 		{
-			PathMap.FindOrAdd(RegistryKey) = InAssetData.GetSoftObjectPath();
+			return false;
 		}
 
-		return RegistryKey;
+		return !AutoUpdateDenyListCache.Contains(InClassName.GetFullName());
 	}
 
-	// Invalid ClassID means the node could not be registered.
-	// Let caller report or ensure as necessary.
-	return FNodeRegistryKey::GetInvalid();
-}
-
-FMetaSoundFrontendDocumentBuilder& UMetaSoundAssetSubsystem::AttachDocumentBuilderChecked(UObject& InObject) const
-{
-	return UMetaSoundBuilderSubsystem::GetChecked().AttachBuilderToAssetChecked(InObject).Builder;
-}
+	bool FMetaSoundAssetManager::ContainsKey(const Metasound::Frontend::FAssetKey& InKey) const
+	{
+		return PathMap.Contains(InKey);
+	}
 
-bool UMetaSoundAssetSubsystem::CanAutoUpdate(const FMetasoundFrontendClassName& InClassName) const
-{
-	const UMetaSoundSettings* Settings = GetDefault<UMetaSoundSettings>();
-	if (!Settings->bAutoUpdateEnabled)
+	FMetaSoundAsyncAssetDependencies* FMetaSoundAssetManager::FindLoadingDependencies(const UObject* InParentAsset)
 	{
-		return false;
+		auto IsEqualMetaSoundUObject = [InParentAsset](const FMetaSoundAsyncAssetDependencies& InDependencies) -> bool
+		{
+			return (InDependencies.MetaSound == InParentAsset);
+		};
+
+		return LoadingDependencies.FindByPredicate(IsEqualMetaSoundUObject);
 	}
 
-	return !AutoUpdateDenyListCache.Contains(InClassName.GetFullName());
-}
+	FMetaSoundAsyncAssetDependencies* FMetaSoundAssetManager::FindLoadingDependencies(int32 InLoadID)
+	{
+		auto IsEqualID = [InLoadID](const FMetaSoundAsyncAssetDependencies& InDependencies) -> bool
+		{
+			return (InDependencies.LoadID == InLoadID);
+		};
 
-bool UMetaSoundAssetSubsystem::ContainsKey(const Metasound::Frontend::FNodeRegistryKey& InRegistryKey) const
-{
-	return PathMap.Contains(InRegistryKey);
-}
+		return LoadingDependencies.FindByPredicate(IsEqualID);
+	}
 
-FMetasoundAssetBase* UMetaSoundAssetSubsystem::GetAsAsset(UObject& InObject) const
-{
-	return Metasound::IMetasoundUObjectRegistry::Get().GetObjectAsAssetBase(&InObject);
-}
+	FMetasoundAssetBase* FMetaSoundAssetManager::FindAsset(const Metasound::Frontend::FAssetKey& InKey) const
+	{
+		if (const FTopLevelAssetPath* AssetPath = PathMap.Find(InKey))
+		{
+			if (UObject* Object = FSoftObjectPath(*AssetPath, { }).ResolveObject())
+			{
+				return GetAsAsset(*Object);
+			}
+		}
 
-const FMetasoundAssetBase* UMetaSoundAssetSubsystem::GetAsAsset(const UObject& InObject) const
-{
-	return Metasound::IMetasoundUObjectRegistry::Get().GetObjectAsAssetBase(&InObject);
-}
+		return nullptr;
+	}
 
-void UMetaSoundAssetSubsystem::RebuildDenyListCache(const UAssetManager& InAssetManager)
-{
-	using namespace Metasound::Frontend;
+	TScriptInterface<IMetaSoundDocumentInterface> FMetaSoundAssetManager::FindAssetAsDocumentInterface(const Metasound::Frontend::FAssetKey& InKey) const
+	{
+		if (const FTopLevelAssetPath* AssetPath = PathMap.Find(InKey))
+		{
+			if (UObject* Object = FSoftObjectPath(*AssetPath, { }).ResolveObject())
+			{
+				return TScriptInterface<IMetaSoundDocumentInterface>(Object);
+			}
+		}
 
-	const UMetaSoundSettings* Settings = GetDefault<UMetaSoundSettings>();
-	if (Settings->DenyListCacheChangeID == AutoUpdateDenyListChangeID)
+		return nullptr;
+	}
+
+	const FTopLevelAssetPath* FMetaSoundAssetManager::FindAssetPath(const Metasound::Frontend::FAssetKey& InKey) const
 	{
-		return;
+		return PathMap.Find(InKey);
 	}
 
-	AutoUpdateDenyListCache.Reset();
+	FMetasoundAssetBase* FMetaSoundAssetManager::GetAsAsset(UObject& InObject) const
+	{
+		return IMetasoundUObjectRegistry::Get().GetObjectAsAssetBase(&InObject);
+	}
 
-	for (const FMetasoundFrontendClassName& ClassName : Settings->AutoUpdateDenylist)
+	const FMetasoundAssetBase* FMetaSoundAssetManager::GetAsAsset(const UObject& InObject) const
 	{
-		AutoUpdateDenyListCache.Add(ClassName.GetFullName());
+		return IMetasoundUObjectRegistry::Get().GetObjectAsAssetBase(&InObject);
 	}
 
-	check(UAssetManager::IsInitialized());
-	UAssetManager& AssetManager = UAssetManager::Get();
-	for (const FDefaultMetaSoundAssetAutoUpdateSettings& UpdateSettings : Settings->AutoUpdateAssetDenylist)
+	void FMetaSoundAssetManager::OnAssetsLoaded(int32 InLoadID)
 	{
-		FAssetData AssetData;
-		if (AssetManager.GetAssetDataForPath(UpdateSettings.MetaSound, AssetData))
+		FMetaSoundAsyncAssetDependencies* LoadedDependencies = FindLoadingDependencies(InLoadID);
+		if (ensureMsgf(LoadedDependencies, TEXT("Call to async asset load complete with invalid IDs %d"), InLoadID))
 		{
-			FString AssetClassID;
-			if (AssetData.GetTagValue(AssetTags::AssetClassID, AssetClassID))
+			if (LoadedDependencies->StreamableHandle.IsValid())
 			{
-				const FMetasoundFrontendClassName ClassName = { FName(), *AssetClassID, FName() };
-				AutoUpdateDenyListCache.Add(ClassName.GetFullName());
+				if (LoadedDependencies->MetaSound)
+				{
+					Metasound::IMetasoundUObjectRegistry& UObjectRegistry = Metasound::IMetasoundUObjectRegistry::Get();
+					FMetasoundAssetBase* ParentAssetBase = UObjectRegistry.GetObjectAsAssetBase(LoadedDependencies->MetaSound);
+					if (ensureMsgf(ParentAssetBase, TEXT("UClass of Parent MetaSound asset %s is not registered in metasound UObject Registery"), *LoadedDependencies->MetaSound->GetPathName()))
+					{
+						// Get all async loaded assets
+						TArray<UObject*> LoadedAssets;
+						LoadedDependencies->StreamableHandle->GetLoadedAssets(LoadedAssets);
+
+						// Cast UObjects to FMetaSoundAssetBase
+						TArray<FMetasoundAssetBase*> LoadedAssetBases;
+						for (UObject* AssetDependency : LoadedAssets)
+						{
+							if (AssetDependency)
+							{
+								FMetasoundAssetBase* AssetDependencyBase = UObjectRegistry.GetObjectAsAssetBase(AssetDependency);
+								if (ensure(AssetDependencyBase))
+								{
+									LoadedAssetBases.Add(AssetDependencyBase);
+								}
+							}
+						}
+
+						// Update parent asset with loaded assets. 
+						ParentAssetBase->OnAsyncReferencedAssetsLoaded(LoadedAssetBases);
+					}
+				}
 			}
+
+			// Remove from active array of loading dependencies.
+			RemoveLoadingDependencies(InLoadID);
 		}
 	}
 
-	AutoUpdateDenyListChangeID = Settings->DenyListCacheChangeID;
-}
+	void FMetaSoundAssetManager::OnAssetScanComplete()
+	{
+		bIsInitialAssetScanComplete = true;
+	}
 
 #if WITH_EDITOR
-TSet<UMetaSoundAssetSubsystem::FAssetInfo> UMetaSoundAssetSubsystem::GetReferencedAssetClasses(const FMetasoundAssetBase& InAssetBase) const
-{
-	METASOUND_TRACE_CPUPROFILER_EVENT_SCOPE(UMetaSoundAssetSubsystem::GetReferencedAssetClasses);
-	using namespace Metasound::Frontend;
-
-	TSet<FAssetInfo> OutAssetInfos;
-	const FMetasoundFrontendDocument& Document = InAssetBase.GetConstDocumentChecked();
-	for (const FMetasoundFrontendClass& Class : Document.Dependencies)
+	TSet<UMetaSoundAssetSubsystem::FAssetInfo> FMetaSoundAssetManager::GetReferencedAssetClasses(const FMetasoundAssetBase& InAssetBase) const
 	{
-		const FNodeRegistryKey Key = FNodeRegistryKey(Class.Metadata);
-		if (const FSoftObjectPath* ObjectPath = PathMap.Find(Key))
-		{
-			OutAssetInfos.Add(FAssetInfo{Key, *ObjectPath});
-		}
-		else
+		METASOUND_TRACE_CPUPROFILER_EVENT_SCOPE(FMetaSoundAssetManager::GetReferencedAssetClasses);
+		using namespace Metasound::Frontend;
+
+		TSet<FAssetInfo> OutAssetInfos;
+		const FMetasoundFrontendDocument& Document = InAssetBase.GetConstDocumentChecked();
+		for (const FMetasoundFrontendClass& Class : Document.Dependencies)
 		{
-			const FMetasoundFrontendRegistryContainer* Registry = FMetasoundFrontendRegistryContainer::Get();
-			check(Registry);
-			const bool bIsRegistered = Registry->IsNodeRegistered(Key);
+			if (Class.Metadata.GetType() != EMetasoundFrontendClassType::External)
+			{
+				continue;
+			}
 
-			bool bReportFail = false;
-			if (bIsRegistered)
+			const FAssetKey Key(FNodeRegistryKey(Class.Metadata));
+			if (const FTopLevelAssetPath* ObjectPath = PathMap.Find(Key))
 			{
-				if (!Registry->IsNodeNative(Key))
-				{
-					bReportFail = true;
-				}
+				OutAssetInfos.Add(FAssetInfo(Key, FSoftObjectPath(*ObjectPath)));
 			}
 			else
 			{
-				// Don't report failure if a matching class with a matching major version and higher minor version exists (it will be autoupdated) 
-				FMetasoundFrontendClass FrontendClass;
-				const bool bDidFindClassWithName = ISearchEngine::Get().FindClassWithHighestVersion(Key.ClassName.ToNodeClassName(), FrontendClass);
-				if (!(bDidFindClassWithName && 
-					Key.Version.Major == FrontendClass.Metadata.GetVersion().Major && 
-					Key.Version.Minor < FrontendClass.Metadata.GetVersion().Minor))
-				{
-					bReportFail = true;
-				}
-			}
+				const FMetasoundFrontendRegistryContainer* Registry = FMetasoundFrontendRegistryContainer::Get();
+				check(Registry);
+				const bool bIsRegistered = Registry->IsNodeRegistered(Key);
 
-			if (bReportFail)
-			{
-				if (bIsInitialAssetScanComplete)
+				bool bReportFail = false;
+				if (bIsRegistered)
 				{
-					UE_LOG(LogMetaSound, Warning, TEXT("MetaSound Node Class with registry key '%s' not registered when gathering referenced asset classes from '%s': Retrieving all asset classes may not be comprehensive."), *Key.ToString(), *InAssetBase.GetOwningAssetName());
+					if (!Registry->IsNodeNative(Key))
+					{
+						bReportFail = true;
+					}
 				}
 				else
 				{
-					UE_LOG(LogMetaSound, Warning, TEXT("Attempt to get registered dependent asset with key '%s' from MetaSound asset '%s' before asset scan has completed: Asset class cannot be provided"), *Key.ToString(), *InAssetBase.GetOwningAssetName());
+					// Don't report failure if a matching class with a matching major version and higher minor version exists (it will be autoupdated) 
+					FMetasoundFrontendClass FrontendClass;
+					const bool bDidFindClassWithName = ISearchEngine::Get().FindClassWithHighestVersion(Key.ClassName.ToNodeClassName(), FrontendClass);
+					if (!(bDidFindClassWithName &&
+						Key.Version.Major == FrontendClass.Metadata.GetVersion().Major &&
+						Key.Version.Minor < FrontendClass.Metadata.GetVersion().Minor))
+					{
+						bReportFail = true;
+					}
+				}
+
+				if (bReportFail)
+				{
+					if (bIsInitialAssetScanComplete)
+					{
+						UE_LOG(LogMetaSound, Warning, TEXT("MetaSound Node Class with registry key '%s' not registered when gathering referenced asset classes from '%s': Retrieving all asset classes may not be comprehensive."), *Key.ToString(), *InAssetBase.GetOwningAssetName());
+					}
+					else
+					{
+						UE_LOG(LogMetaSound, Warning, TEXT("Attempt to get registered dependent asset with key '%s' from MetaSound asset '%s' before asset scan has completed: Asset class cannot be provided"), *Key.ToString(), *InAssetBase.GetOwningAssetName());
+					}
 				}
 			}
 		}
+		return MoveTemp(OutAssetInfos);
 	}
-	return MoveTemp(OutAssetInfos);
-}
 #endif // WITH_EDITOR
 
-void UMetaSoundAssetSubsystem::RescanAutoUpdateDenyList()
-{
-	check(UAssetManager::IsInitialized());
-	RebuildDenyListCache(UAssetManager::Get());
-}
-
-FMetasoundAssetBase* UMetaSoundAssetSubsystem::TryLoadAssetFromKey(const Metasound::Frontend::FNodeRegistryKey& RegistryKey) const
-{
-	if (const FSoftObjectPath* ObjectPath = FindObjectPathFromKey(RegistryKey))
+	void FMetaSoundAssetManager::RebuildDenyListCache(const UAssetManager& InAssetManager)
 	{
-		return TryLoadAsset(*ObjectPath);
-	}
-
-	return nullptr;
-}
+		using namespace Metasound::Frontend;
 
-bool UMetaSoundAssetSubsystem::TryLoadReferencedAssets(const FMetasoundAssetBase& InAssetBase, TArray<FMetasoundAssetBase*>& OutReferencedAssets) const
-{
-	using namespace Metasound::Frontend;
+		const UMetaSoundSettings* Settings = GetDefault<UMetaSoundSettings>();
+		if (Settings->DenyListCacheChangeID == AutoUpdateDenyListChangeID)
+		{
+			return;
+		}
 
-	bool bSucceeded = true;
-	OutReferencedAssets.Reset();
+		AutoUpdateDenyListCache.Reset();
 
-	TArray<FMetasoundAssetBase*> ReferencedAssets;
-	const TSet<FString>& AssetClassKeys = InAssetBase.GetReferencedAssetClassKeys();
-	for (const FString& KeyString : AssetClassKeys)
-	{
-		FNodeRegistryKey Key;
-		FNodeRegistryKey::Parse(KeyString, Key);
-		if (FMetasoundAssetBase* MetaSound = TryLoadAssetFromKey(Key))
+		for (const FMetasoundFrontendClassName& ClassName : Settings->AutoUpdateDenylist)
 		{
-			OutReferencedAssets.Add(MetaSound);
+			AutoUpdateDenyListCache.Add(ClassName.GetFullName());
 		}
-		else
+
+		check(UAssetManager::IsInitialized());
+		UAssetManager& AssetManager = UAssetManager::Get();
+		for (const FDefaultMetaSoundAssetAutoUpdateSettings& UpdateSettings : Settings->AutoUpdateAssetDenylist)
 		{
-			UE_LOG(LogMetaSound, Error, TEXT("Failed to find referenced MetaSound asset with key '%s'"), *KeyString);
-			bSucceeded = false;
+			FAssetData AssetData;
+			if (AssetManager.GetAssetDataForPath(UpdateSettings.MetaSound, AssetData))
+			{
+				FString AssetClassID;
+				if (AssetData.GetTagValue(AssetTags::AssetClassID, AssetClassID))
+				{
+					const FMetasoundFrontendClassName ClassName = { FName(), *AssetClassID, FName() };
+					AutoUpdateDenyListCache.Add(ClassName.GetFullName());
+				}
+			}
 		}
-	}
 
-	return bSucceeded;
-}
+		AutoUpdateDenyListChangeID = Settings->DenyListCacheChangeID;
+	}
 
-void UMetaSoundAssetSubsystem::RequestAsyncLoadReferencedAssets(FMetasoundAssetBase& InAssetBase)
-{
-	const TSet<FSoftObjectPath>& AsyncReferences = InAssetBase.GetAsyncReferencedAssetClassPaths();
-	if (AsyncReferences.Num() > 0)
+	void FMetaSoundAssetManager::RegisterAssetClassesInDirectories(const TArray<FMetaSoundAssetDirectory>& InDirectories)
 	{
-		if (UObject* OwningAsset = InAssetBase.GetOwningAsset())
+		TArray<FDirectoryPath> Directories;
+		Algo::Transform(InDirectories, Directories, [](const FMetaSoundAssetDirectory& AssetDir) { return AssetDir.Directory; });
+
+		SearchAndIterateDirectoryAssets(Directories, [this](const FAssetData& AssetData)
 		{
-			TArray<FSoftObjectPath> PathsToLoad = AsyncReferences.Array();
+			AddOrUpdateAsset(AssetData);
+			FMetasoundAssetBase* MetaSoundAsset = Metasound::IMetasoundUObjectRegistry::Get().GetObjectAsAssetBase(AssetData.GetAsset());
+			check(MetaSoundAsset);
 
-			// Protect against duplicate calls to async load assets. 
-			if (FMetaSoundAsyncAssetDependencies* ExistingAsyncLoad = FindLoadingDependencies(OwningAsset))
+			Metasound::Frontend::FMetaSoundAssetRegistrationOptions RegOptions;
+			if (const UMetaSoundSettings* Settings = GetDefault<UMetaSoundSettings>())
 			{
-				if (ExistingAsyncLoad->Dependencies == PathsToLoad)
-				{
-					// early out since these are already actively being loaded.
-					return;
-				}
+				RegOptions.bAutoUpdateLogWarningOnDroppedConnection = Settings->bAutoUpdateLogWarningOnDroppedConnection;
 			}
+			MetaSoundAsset->RegisterGraphWithFrontend(RegOptions);
+		});
+	}
 
-			int32 AsyncLoadID = AsyncLoadIDCounter++;
+	void FMetaSoundAssetManager::RemoveAsset(const UObject& InObject)
+	{
+		using namespace Frontend;
 
-			auto AssetsLoadedDelegate = [this, AsyncLoadID]()
-			{
-				this->OnAssetsLoaded(AsyncLoadID);
-			};
+		if (const FMetasoundAssetBase* MetaSoundAsset = Metasound::IMetasoundUObjectRegistry::Get().GetObjectAsAssetBase(&InObject))
+		{
+			const FMetasoundFrontendDocument& Document = MetaSoundAsset->GetConstDocumentChecked();
+			const FMetasoundFrontendClassMetadata& Metadata = Document.RootGraph.Metadata;
+
+			const FAssetKey AssetKey(Metadata.GetClassName(), Metadata.GetVersion());
+			const FGraphRegistryKey GraphKey = MetaSoundAsset->GetGraphRegistryKey();
 
-			// Store async loading data for use when async load is complete. 
-			FMetaSoundAsyncAssetDependencies& AsyncDependencies = LoadingDependencies.AddDefaulted_GetRef();
+			if (IDocumentBuilderRegistry* BuilderRegistry = IDocumentBuilderRegistry::Get())
+			{
+				BuilderRegistry->FinishBuilding(Metadata.GetClassName());
+			}
 
-			AsyncDependencies.LoadID = AsyncLoadID;
-			AsyncDependencies.MetaSound = OwningAsset;
-			AsyncDependencies.Dependencies = PathsToLoad;
-			AsyncDependencies.StreamableHandle = StreamableManager.RequestAsyncLoad(PathsToLoad, AssetsLoadedDelegate);
+			AssetSubsystemPrivate::RemoveAndWarnIfNotExactMatch(PathMap, AssetKey, GraphKey.AssetPath);
 		}
-		else
+	}
+
+	void FMetaSoundAssetManager::RemoveAsset(const FAssetData& InAssetData)
+	{
+		using namespace Frontend;
+
+		FNodeClassInfo ClassInfo;
+		if (ensureAlways(AssetSubsystemPrivate::GetAssetClassInfo(InAssetData, ClassInfo)))
 		{
-			UE_LOG(LogMetaSound, Error, TEXT("Cannot load async asset as FMetasoundAssetBase null owning UObject"), *InAssetBase.GetOwningAssetName());
+			if (IDocumentBuilderRegistry* BuilderRegistry = IDocumentBuilderRegistry::Get())
+			{
+				BuilderRegistry->FinishBuilding(ClassInfo.ClassName);
+			}
+
+			const FAssetKey AssetKey(ClassInfo.ClassName, ClassInfo.Version);
+			const FTopLevelAssetPath AssetPath(InAssetData.PackageName, InAssetData.AssetName);
+			AssetSubsystemPrivate::RemoveAndWarnIfNotExactMatch(PathMap, AssetKey, AssetPath);
 		}
 	}
 
-}
+	void FMetaSoundAssetManager::RemoveLoadingDependencies(int32 InLoadID)
+	{
+		auto IsEqualID = [InLoadID](const FMetaSoundAsyncAssetDependencies& InDependencies) -> bool
+		{
+			return (InDependencies.LoadID == InLoadID);
+		};
+		LoadingDependencies.RemoveAllSwap(IsEqualID);
+	}
 
-void UMetaSoundAssetSubsystem::WaitUntilAsyncLoadReferencedAssetsComplete(FMetasoundAssetBase& InAssetBase)
-{
-	TSet<FMetasoundAssetBase*> TransitiveReferences;
-	TArray<FMetasoundAssetBase*> TransitiveReferencesQueue;
-	TransitiveReferences.Add(&InAssetBase);
-	TransitiveReferencesQueue.Add(&InAssetBase);
-	while (!TransitiveReferencesQueue.IsEmpty())
+	void FMetaSoundAssetManager::RenameAsset(const FAssetData& InAssetData, bool bInReregisterWithFrontend)
 	{
-		FMetasoundAssetBase* Reference = TransitiveReferencesQueue.Pop();
-		UObject* OwningAsset = Reference->GetOwningAsset();
-		if (!OwningAsset)
+		using namespace Frontend;
+
+		FMetasoundAssetBase* MetaSoundAsset = GetAsAsset(*InAssetData.GetAsset());
+		check(MetaSoundAsset);
+
+		FNodeClassInfo ClassInfo;
+		if (ensureAlways(AssetSubsystemPrivate::GetAssetClassInfo(InAssetData, ClassInfo)))
 		{
-			continue;
+			IDocumentBuilderRegistry* BuilderRegistry = IDocumentBuilderRegistry::Get();
+			if (BuilderRegistry)
+			{
+				BuilderRegistry->FinishBuilding(ClassInfo.ClassName);
+			}
+
+			const FAssetKey AssetKey(ClassInfo.ClassName, ClassInfo.Version);
+			PathMap.Remove(AssetKey); // Don't warn if names do not match as by this point the incoming asset data has the newly assigned asset path assigned
+
+#if WITH_EDITOR
+			if (BuilderRegistry)
+			{
+				FMetaSoundFrontendDocumentBuilder& Builder = BuilderRegistry->FindOrBeginBuilding(MetaSoundAsset->GetOwningAsset());
+				Builder.SetDisplayName(FText());
+			}
+#endif // WITH_EDITOR
+
+			if (ClassInfo.AssetClassID.IsValid())
+			{
+				const FAssetKey NewAssetKey(ClassInfo.ClassName, ClassInfo.Version);
+				if (NewAssetKey.IsValid())
+				{
+					PathMap.FindOrAdd(NewAssetKey) = ClassInfo.AssetPath;
+				}
+			}
+		}
+
+		if (bInReregisterWithFrontend)
+		{
+			MetaSoundAsset->RegisterGraphWithFrontend();
 		}
-		while (FMetaSoundAsyncAssetDependencies* LoadingDependency = FindLoadingDependencies(OwningAsset))
+	}
+
+	void FMetaSoundAssetManager::RequestAsyncLoadReferencedAssets(FMetasoundAssetBase& InAssetBase)
+	{
+		const TSet<FSoftObjectPath>& AsyncReferences = InAssetBase.GetAsyncReferencedAssetClassPaths();
+		if (AsyncReferences.Num() > 0)
 		{
-			// Grab shared ptr to handle as LoadingDependencies may be deleted and have it's shared pointer removed. 
-			TSharedPtr<FStreamableHandle> StreamableHandle = LoadingDependency->StreamableHandle;
-			if (StreamableHandle.IsValid())
+			if (UObject* OwningAsset = InAssetBase.GetOwningAsset())
 			{
-				UE_LOG(LogMetaSound, Verbose, TEXT("Waiting on async load (id: %d) from asset %s"), LoadingDependency->LoadID, *InAssetBase.GetOwningAssetName());
+				TArray<FSoftObjectPath> PathsToLoad = AsyncReferences.Array();
 
-				EAsyncPackageState::Type LoadState = StreamableHandle->WaitUntilComplete();
-				if (EAsyncPackageState::Complete != LoadState)
+				// Protect against duplicate calls to async load assets. 
+				if (FMetaSoundAsyncAssetDependencies* ExistingAsyncLoad = FindLoadingDependencies(OwningAsset))
 				{
-					UE_LOG(LogMetaSound, Error, TEXT("Failed to complete loading of async dependent assets from parent asset %s"), *InAssetBase.GetOwningAssetName());
-					RemoveLoadingDependencies(LoadingDependency->LoadID);
+					if (ExistingAsyncLoad->Dependencies == PathsToLoad)
+					{
+						// early out since these are already actively being loaded.
+						return;
+					}
 				}
-				else
+
+				int32 AsyncLoadID = AsyncLoadIDCounter++;
+
+				auto AssetsLoadedDelegate = [this, AsyncLoadID]()
 				{
-					// This will remove the loading dependencies from internal storage
-					OnAssetsLoaded(LoadingDependency->LoadID);
-				}
+					this->OnAssetsLoaded(AsyncLoadID);
+				};
+
+				// Store async loading data for use when async load is complete. 
+				FMetaSoundAsyncAssetDependencies& AsyncDependencies = LoadingDependencies.AddDefaulted_GetRef();
 
-				// This will prevent OnAssetsLoaded from being called via the streamables
-				// internal delegate complete callback.
-				StreamableHandle->CancelHandle();
+				AsyncDependencies.LoadID = AsyncLoadID;
+				AsyncDependencies.MetaSound = OwningAsset;
+				AsyncDependencies.Dependencies = PathsToLoad;
+				AsyncDependencies.StreamableHandle = StreamableManager.RequestAsyncLoad(PathsToLoad, AssetsLoadedDelegate);
 			}
+			else
+			{
+				UE_LOG(LogMetaSound, Error, TEXT("Cannot load async asset as FMetasoundAssetBase null owning UObject"), *InAssetBase.GetOwningAssetName());
+			}
+		}
+
+	}
+
+	void FMetaSoundAssetManager::SearchAndIterateDirectoryAssets(const TArray<FDirectoryPath>& InDirectories, TFunctionRef<void(const FAssetData&)> InFunction)
+	{
+		if (InDirectories.IsEmpty())
+		{
+			return;
 		}
 
-		for (FMetasoundAssetBase* NextReference : Reference->GetReferencedAssets())
+		UAssetManager& AssetManager = UAssetManager::Get();
+
+		FAssetManagerSearchRules Rules;
+		for (const FDirectoryPath& Path : InDirectories)
 		{
-			bool bAlreadyInSet;
-			TransitiveReferences.Add(NextReference, &bAlreadyInSet);
-			if (!bAlreadyInSet)
+			Rules.AssetScanPaths.Add(*Path.Path);
+		}
+
+		Metasound::IMetasoundUObjectRegistry::Get().IterateRegisteredUClasses([&](UClass& RegisteredClass)
+		{
+			Rules.AssetBaseClass = &RegisteredClass;
+			TArray<FAssetData> MetaSoundAssets;
+			AssetManager.SearchAssetRegistryPaths(MetaSoundAssets, Rules);
+			for (const FAssetData& AssetData : MetaSoundAssets)
 			{
-				TransitiveReferencesQueue.Add(NextReference);
+				InFunction(AssetData);
 			}
+		});
+	}
+
+	FMetasoundAssetBase* FMetaSoundAssetManager::TryLoadAssetFromKey(const Metasound::Frontend::FAssetKey& RegistryKey) const
+	{
+		if (const FTopLevelAssetPath* ObjectPath = FindAssetPath(RegistryKey))
+		{
+			const FSoftObjectPath SoftPath(*ObjectPath);
+			return TryLoadAsset(SoftPath);
 		}
+
+		return nullptr;
 	}
-}
 
-FMetaSoundAsyncAssetDependencies* UMetaSoundAssetSubsystem::FindLoadingDependencies(const UObject* InParentAsset)
-{
-	auto IsEqualMetaSoundUObject = [InParentAsset](const FMetaSoundAsyncAssetDependencies& InDependencies) -> bool
+	bool FMetaSoundAssetManager::TryGetAssetIDFromClassName(const FMetasoundFrontendClassName& InClassName, FGuid& OutGuid) const
 	{
-		return (InDependencies.MetaSound == InParentAsset);
-	}; 
+		return FGuid::Parse(InClassName.Name.ToString(), OutGuid);
+	}
 
-	return LoadingDependencies.FindByPredicate(IsEqualMetaSoundUObject);
-}
+	FMetasoundAssetBase* FMetaSoundAssetManager::TryLoadAsset(const FSoftObjectPath& InObjectPath) const
+	{
+		return Metasound::IMetasoundUObjectRegistry::Get().GetObjectAsAssetBase(InObjectPath.TryLoad());
+	}
 
-FMetaSoundAsyncAssetDependencies* UMetaSoundAssetSubsystem::FindLoadingDependencies(int32 InLoadID)
-{
-	auto IsEqualID = [InLoadID](const FMetaSoundAsyncAssetDependencies& InDependencies) -> bool
+	bool FMetaSoundAssetManager::TryLoadReferencedAssets(const FMetasoundAssetBase& InAssetBase, TArray<FMetasoundAssetBase*>& OutReferencedAssets) const
 	{
-		return (InDependencies.LoadID == InLoadID);
-	};
-	
-	return LoadingDependencies.FindByPredicate(IsEqualID);
-}
+		using namespace Metasound::Frontend;
 
-void UMetaSoundAssetSubsystem::RemoveLoadingDependencies(int32 InLoadID)
-{
-	auto IsEqualID = [InLoadID](const FMetaSoundAsyncAssetDependencies& InDependencies) -> bool
+		bool bSucceeded = true;
+		OutReferencedAssets.Reset();
+
+		TArray<FMetasoundAssetBase*> ReferencedAssets;
+		const TSet<FString>& AssetClassKeys = InAssetBase.GetReferencedAssetClassKeys();
+		for (const FString& KeyString : AssetClassKeys)
+		{
+			FNodeRegistryKey Key;
+			FNodeRegistryKey::Parse(KeyString, Key);
+			if (FMetasoundAssetBase* MetaSound = TryLoadAssetFromKey(Key))
+			{
+				OutReferencedAssets.Add(MetaSound);
+			}
+			else
+			{
+				UE_LOG(LogMetaSound, Error, TEXT("Failed to find or load referenced MetaSound asset with key '%s'"), *KeyString);
+				bSucceeded = false;
+			}
+		}
+
+		return bSucceeded;
+	}
+
+	void FMetaSoundAssetManager::UnregisterAssetClassesInDirectories(const TArray<FMetaSoundAssetDirectory>& InDirectories)
 	{
-		return (InDependencies.LoadID == InLoadID);
-	};
-	LoadingDependencies.RemoveAllSwap(IsEqualID);
-}
+		TArray<FDirectoryPath> Directories;
+		Algo::Transform(InDirectories, Directories, [](const FMetaSoundAssetDirectory& AssetDir) { return AssetDir.Directory; });
 
-void UMetaSoundAssetSubsystem::OnAssetsLoaded(int32 InLoadID)
-{
-	FMetaSoundAsyncAssetDependencies* LoadedDependencies = FindLoadingDependencies(InLoadID);
-	if (ensureMsgf(LoadedDependencies, TEXT("Call to async asset load complete with invalid IDs %d"), InLoadID))
+		SearchAndIterateDirectoryAssets(Directories, [this](const FAssetData& AssetData)
+		{
+			using namespace Metasound;
+			using namespace Metasound::Frontend;
+
+			if (AssetData.IsAssetLoaded())
+			{
+				FMetasoundAssetBase* MetaSoundAsset = Metasound::IMetasoundUObjectRegistry::Get().GetObjectAsAssetBase(AssetData.GetAsset());
+				check(MetaSoundAsset);
+				MetaSoundAsset->UnregisterGraphWithFrontend();
+
+				RemoveAsset(AssetData);
+			}
+			else
+			{
+				FNodeClassInfo AssetClassInfo;
+				if (ensureAlways(AssetSubsystemPrivate::GetAssetClassInfo(AssetData, AssetClassInfo)))
+				{
+					const FNodeRegistryKey RegistryKey = FNodeRegistryKey(AssetClassInfo);
+					const bool bIsRegistered = FMetasoundFrontendRegistryContainer::Get()->IsNodeRegistered(RegistryKey);
+					if (bIsRegistered)
+					{
+						FMetasoundFrontendRegistryContainer::Get()->UnregisterNode(RegistryKey);
+						const FTopLevelAssetPath AssetPath(AssetData.PackageName, AssetData.AssetName);
+						const FAssetKey AssetKey(AssetClassInfo.ClassName, AssetClassInfo.Version);
+						AssetSubsystemPrivate::RemoveAndWarnIfNotExactMatch(PathMap, AssetKey, AssetPath);
+					}
+				}
+			}
+		});
+	}
+
+	void FMetaSoundAssetManager::WaitUntilAsyncLoadReferencedAssetsComplete(FMetasoundAssetBase& InAssetBase)
 	{
-		if (LoadedDependencies->StreamableHandle.IsValid())
+		TSet<FMetasoundAssetBase*> TransitiveReferences;
+		TArray<FMetasoundAssetBase*> TransitiveReferencesQueue;
+		TransitiveReferences.Add(&InAssetBase);
+		TransitiveReferencesQueue.Add(&InAssetBase);
+		while (!TransitiveReferencesQueue.IsEmpty())
 		{
-			if (LoadedDependencies->MetaSound)
+			FMetasoundAssetBase* Reference = TransitiveReferencesQueue.Pop();
+			UObject* OwningAsset = Reference->GetOwningAsset();
+			if (!OwningAsset)
 			{
-				Metasound::IMetasoundUObjectRegistry& UObjectRegistry = Metasound::IMetasoundUObjectRegistry::Get();
-				FMetasoundAssetBase* ParentAssetBase = UObjectRegistry.GetObjectAsAssetBase(LoadedDependencies->MetaSound);
-				if (ensureMsgf(ParentAssetBase, TEXT("UClass of Parent MetaSound asset %s is not registered in metasound UObject Registery"), *LoadedDependencies->MetaSound->GetPathName()))
+				continue;
+			}
+			while (FMetaSoundAsyncAssetDependencies* LoadingDependency = FindLoadingDependencies(OwningAsset))
+			{
+				// Grab shared ptr to handle as LoadingDependencies may be deleted and have it's shared pointer removed. 
+				TSharedPtr<FStreamableHandle> StreamableHandle = LoadingDependency->StreamableHandle;
+				if (StreamableHandle.IsValid())
 				{
-					// Get all async loaded assets
-					TArray<UObject*> LoadedAssets;
-					LoadedDependencies->StreamableHandle->GetLoadedAssets(LoadedAssets);
+					UE_LOG(LogMetaSound, Verbose, TEXT("Waiting on async load (id: %d) from asset %s"), LoadingDependency->LoadID, *InAssetBase.GetOwningAssetName());
 
-					// Cast UObjects to FMetaSoundAssetBase
-					TArray<FMetasoundAssetBase*> LoadedAssetBases;
-					for (UObject* AssetDependency : LoadedAssets)
+					EAsyncPackageState::Type LoadState = StreamableHandle->WaitUntilComplete();
+					if (EAsyncPackageState::Complete != LoadState)
 					{
-						if (AssetDependency)
-						{
-							FMetasoundAssetBase* AssetDependencyBase = UObjectRegistry.GetObjectAsAssetBase(AssetDependency);
-							if (ensure(AssetDependencyBase))
-							{
-								LoadedAssetBases.Add(AssetDependencyBase);
-							}
-						}
+						UE_LOG(LogMetaSound, Error, TEXT("Failed to complete loading of async dependent assets from parent asset %s"), *InAssetBase.GetOwningAssetName());
+						RemoveLoadingDependencies(LoadingDependency->LoadID);
 					}
+					else
+					{
+						// This will remove the loading dependencies from internal storage
+						OnAssetsLoaded(LoadingDependency->LoadID);
+					}
+
+					// This will prevent OnAssetsLoaded from being called via the streamables
+					// internal delegate complete callback.
+					StreamableHandle->CancelHandle();
+				}
+			}
 
-					// Update parent asset with loaded assets. 
-					ParentAssetBase->OnAsyncReferencedAssetsLoaded(LoadedAssetBases);
+			for (FMetasoundAssetBase* NextReference : Reference->GetReferencedAssets())
+			{
+				bool bAlreadyInSet;
+				TransitiveReferences.Add(NextReference, &bAlreadyInSet);
+				if (!bAlreadyInSet)
+				{
+					TransitiveReferencesQueue.Add(NextReference);
 				}
 			}
 		}
+	}
 
-		// Remove from active array of loading dependencies.
-		RemoveLoadingDependencies(InLoadID);
+	void DeinitializeAssetManager()
+	{
+		Frontend::IMetaSoundAssetManager::Deinitialize();
 	}
-}
 
-const FSoftObjectPath* UMetaSoundAssetSubsystem::FindObjectPathFromKey(const Metasound::Frontend::FNodeRegistryKey& InRegistryKey) const
+	void InitializeAssetManager()
+	{
+		Frontend::IMetaSoundAssetManager::Initialize(MakeUnique<FMetaSoundAssetManager>());
+	}
+} // namespace Metasound::Engine
+
+void UMetaSoundAssetSubsystem::Initialize(FSubsystemCollectionBase& InCollection)
 {
-	return PathMap.Find(InRegistryKey);
+	FCoreDelegates::OnPostEngineInit.AddUObject(this, &UMetaSoundAssetSubsystem::PostEngineInitInternal);
 }
 
-FMetasoundAssetBase* UMetaSoundAssetSubsystem::TryLoadAsset(const FSoftObjectPath& InObjectPath) const
+void UMetaSoundAssetSubsystem::PostEngineInitInternal()
 {
-	return Metasound::IMetasoundUObjectRegistry::Get().GetObjectAsAssetBase(InObjectPath.TryLoad());
+	using namespace Metasound::Engine;
+
+	check(UAssetManager::IsInitialized());
+	UAssetManager& AssetManager = UAssetManager::Get();
+	AssetManager.CallOrRegister_OnCompletedInitialScan(FSimpleMulticastDelegate::FDelegate::CreateUObject(this, &UMetaSoundAssetSubsystem::PostInitAssetScanInternal));
+	FMetaSoundAssetManager::GetChecked().RebuildDenyListCache(AssetManager);
 }
 
-void UMetaSoundAssetSubsystem::RemoveAsset(const UObject& InObject)
+void UMetaSoundAssetSubsystem::PostInitAssetScanInternal()
 {
-	using namespace Metasound;
-	using namespace Metasound::Frontend;
+	using namespace Metasound::Engine;
 
-	if (const FMetasoundAssetBase* MetaSoundAsset = Metasound::IMetasoundUObjectRegistry::Get().GetObjectAsAssetBase(&InObject))
+	METASOUND_TRACE_CPUPROFILER_EVENT_SCOPE(UMetaSoundAssetSubsystem::PostInitAssetScanInternal);
+
+	const UMetaSoundSettings* Settings = GetDefault<UMetaSoundSettings>();
+	if (ensureAlways(Settings))
 	{
-		const FSoftObjectPath ObjectPath(&InObject);
-		const FMetasoundFrontendDocument& Document = MetaSoundAsset->GetConstDocumentChecked();
-		FNodeRegistryKey RegistryKey = FNodeRegistryKey(Document.RootGraph);
-		AssetSubsystemPrivate::RemoveIfExactMatch(PathMap, RegistryKey, ObjectPath);
-		UMetaSoundBuilderSubsystem::GetChecked().DetachBuilderFromAsset(Document.RootGraph.Metadata.GetClassName());
+		FMetaSoundAssetManager& Manager = FMetaSoundAssetManager::GetChecked();
+		Manager.SearchAndIterateDirectoryAssets(Settings->DirectoriesToRegister, [&Manager](const FAssetData& AssetData)
+		{
+			Manager.AddOrUpdateAsset(AssetData);
+		});
+		Manager.OnAssetScanComplete();
 	}
 }
 
-void UMetaSoundAssetSubsystem::RemoveAsset(const FAssetData& InAssetData)
+#if WITH_EDITORONLY_DATA
+void UMetaSoundAssetSubsystem::AddAssetReferences(FMetasoundAssetBase& InAssetBase)
 {
-	using namespace Metasound;
 	using namespace Metasound::Frontend;
+	IMetaSoundAssetManager::GetChecked().AddAssetReferences(InAssetBase);
+}
+#endif // WITH_EDITORONLY_DATA
 
-	FNodeClassInfo ClassInfo;
-	if (ensureAlways(AssetSubsystemPrivate::GetAssetClassInfo(InAssetData, ClassInfo)))
-	{
-		const FNodeRegistryKey RegistryKey(ClassInfo);
-		const FSoftObjectPath ObjectPath = InAssetData.GetSoftObjectPath();
-		AssetSubsystemPrivate::RemoveIfExactMatch(PathMap, RegistryKey, ObjectPath);
-		UMetaSoundBuilderSubsystem::GetChecked().DetachBuilderFromAsset(ClassInfo.ClassName);
-	}
+Metasound::Frontend::FNodeRegistryKey UMetaSoundAssetSubsystem::AddOrUpdateAsset(const UObject& InObject)
+{
+	using namespace Metasound::Frontend;
+	return IMetaSoundAssetManager::GetChecked().AddOrUpdateAsset(InObject);
 }
 
-void UMetaSoundAssetSubsystem::RenameAsset(const FAssetData& InAssetData, bool bInReregisterWithFrontend)
+Metasound::Frontend::FNodeRegistryKey UMetaSoundAssetSubsystem::AddOrUpdateAsset(const FAssetData& InAssetData)
 {
-	auto PerformRename = [this, &InAssetData]()
-	{
-		RemoveAsset(InAssetData);
-		ResetAssetClassDisplayName(InAssetData);
-		AddOrUpdateAsset(InAssetData);
-	};
+	using namespace Metasound::Frontend;
+	return IMetaSoundAssetManager::GetChecked().AddOrUpdateAsset(InAssetData);
+}
 
-	FMetasoundAssetBase* MetaSoundAsset = Metasound::IMetasoundUObjectRegistry::Get().GetObjectAsAssetBase(InAssetData.GetAsset());
-	check(MetaSoundAsset);
-	MetaSoundAsset->UnregisterGraphWithFrontend();
-	PerformRename();
+bool UMetaSoundAssetSubsystem::CanAutoUpdate(const FMetasoundFrontendClassName& InClassName) const
+{
+	using namespace Metasound::Frontend;
+	return IMetaSoundAssetManager::GetChecked().CanAutoUpdate(InClassName);
+}
 
-	if (bInReregisterWithFrontend)
-	{
-		MetaSoundAsset->RegisterGraphWithFrontend();
-	}
+bool UMetaSoundAssetSubsystem::ContainsKey(const Metasound::Frontend::FNodeRegistryKey& InRegistryKey) const
+{
+	using namespace Metasound::Frontend;
+	return IMetaSoundAssetManager::GetChecked().ContainsKey(FAssetKey(InRegistryKey));
 }
 
-void UMetaSoundAssetSubsystem::ResetAssetClassDisplayName(const FAssetData& InAssetData)
+FMetasoundAssetBase* UMetaSoundAssetSubsystem::GetAsAsset(UObject& InObject) const
 {
-	UObject* Object = nullptr;
-	FSoftObjectPath Path = InAssetData.GetSoftObjectPath();
-	if (InAssetData.IsAssetLoaded())
-	{
-		Object = Path.ResolveObject();
-	}
-	else
-	{
-		Object = Path.TryLoad();
-	}
+	using namespace Metasound::Frontend;
+	return IMetaSoundAssetManager::GetChecked().GetAsAsset(InObject);
+}
 
-	FMetasoundAssetBase* MetaSoundAsset = Metasound::IMetasoundUObjectRegistry::Get().GetObjectAsAssetBase(Object);
-	check(MetaSoundAsset);
-	FMetasoundFrontendGraphClass& Class = MetaSoundAsset->GetDocumentChecked().RootGraph;
+const FMetasoundAssetBase* UMetaSoundAssetSubsystem::GetAsAsset(const UObject& InObject) const
+{
+	using namespace Metasound::Frontend;
+	return IMetaSoundAssetManager::GetChecked().GetAsAsset(InObject);
+}
 
 #if WITH_EDITOR
-	Class.Metadata.SetDisplayName(FText());
-#endif // WITH_EDITOR
+TSet<UMetaSoundAssetSubsystem::FAssetInfo> UMetaSoundAssetSubsystem::GetReferencedAssetClasses(const FMetasoundAssetBase& InAssetBase) const
+{
+	using namespace Metasound::Frontend;
+	return IMetaSoundAssetManager::GetChecked().GetReferencedAssetClasses(InAssetBase);
 }
+#endif // WITH_EDITOR
 
-void UMetaSoundAssetSubsystem::SearchAndIterateDirectoryAssets(const TArray<FDirectoryPath>& InDirectories, TFunctionRef<void(const FAssetData&)> InFunction)
+FMetasoundAssetBase* UMetaSoundAssetSubsystem::TryLoadAssetFromKey(const Metasound::Frontend::FNodeRegistryKey& RegistryKey) const
 {
-	if (InDirectories.IsEmpty())
-	{
-		return;
-	}
+	using namespace Metasound::Frontend;
+	return IMetaSoundAssetManager::GetChecked().TryLoadAssetFromKey(RegistryKey);
+}
 
-	UAssetManager& AssetManager = UAssetManager::Get();
+bool UMetaSoundAssetSubsystem::TryLoadReferencedAssets(const FMetasoundAssetBase& InAssetBase, TArray<FMetasoundAssetBase*>& OutReferencedAssets) const
+{
+	using namespace Metasound::Frontend;
+	return IMetaSoundAssetManager::GetChecked().TryLoadReferencedAssets(InAssetBase, OutReferencedAssets);
+}
 
-	FAssetManagerSearchRules Rules;
-	for (const FDirectoryPath& Path : InDirectories)
+const FSoftObjectPath* UMetaSoundAssetSubsystem::FindObjectPathFromKey(const Metasound::Frontend::FNodeRegistryKey& InRegistryKey) const
+{
+	using namespace Metasound::Frontend;
+	static FSoftObjectPath TempPath;
+	TempPath.Reset();
+	if (const FTopLevelAssetPath* Path = IMetaSoundAssetManager::GetChecked().FindAssetPath(InRegistryKey))
 	{
-		Rules.AssetScanPaths.Add(*Path.Path);
+		TempPath = FSoftObjectPath(*Path);
 	}
-
-	Metasound::IMetasoundUObjectRegistry::Get().IterateRegisteredUClasses([&](UClass& RegisteredClass)
-	{
-		Rules.AssetBaseClass = &RegisteredClass;
-		TArray<FAssetData> MetaSoundAssets;
-		AssetManager.SearchAssetRegistryPaths(MetaSoundAssets, Rules);
-		for (const FAssetData& AssetData : MetaSoundAssets)
-		{
-			InFunction(AssetData);
-		}
-	});
+	return &TempPath;
 }
 
-void UMetaSoundAssetSubsystem::RegisterAssetClassesInDirectories(const TArray<FMetaSoundAssetDirectory>& InDirectories)
+FMetasoundAssetBase* UMetaSoundAssetSubsystem::TryLoadAsset(const FSoftObjectPath& InObjectPath) const
 {
-	TArray<FDirectoryPath> Directories;
-	Algo::Transform(InDirectories, Directories, [](const FMetaSoundAssetDirectory& AssetDir) { return AssetDir.Directory; });
-
-	SearchAndIterateDirectoryAssets(Directories, [this](const FAssetData& AssetData)
-	{
-		AddOrUpdateAsset(AssetData);
-		FMetasoundAssetBase* MetaSoundAsset = Metasound::IMetasoundUObjectRegistry::Get().GetObjectAsAssetBase(AssetData.GetAsset());
-		check(MetaSoundAsset);
+	using namespace Metasound::Frontend;
+	return IMetaSoundAssetManager::GetChecked().TryLoadAsset(InObjectPath);
+}
 
-		Metasound::Frontend::FMetaSoundAssetRegistrationOptions RegOptions;
-		if (const UMetaSoundSettings* Settings = GetDefault<UMetaSoundSettings>())
-		{
-			RegOptions.bAutoUpdateLogWarningOnDroppedConnection = Settings->bAutoUpdateLogWarningOnDroppedConnection;
-		}
-		MetaSoundAsset->RegisterGraphWithFrontend(RegOptions);
-	});
+void UMetaSoundAssetSubsystem::RemoveAsset(const UObject& InObject)
+{
+	using namespace Metasound::Frontend;
+	IMetaSoundAssetManager::GetChecked().RemoveAsset(InObject);
 }
 
-void UMetaSoundAssetSubsystem::UnregisterAssetClassesInDirectories(const TArray<FMetaSoundAssetDirectory>& InDirectories)
+void UMetaSoundAssetSubsystem::RemoveAsset(const FAssetData& InAssetData)
 {
-	TArray<FDirectoryPath> Directories;
-	Algo::Transform(InDirectories, Directories, [](const FMetaSoundAssetDirectory& AssetDir) { return AssetDir.Directory; });
+	using namespace Metasound::Frontend;
+	IMetaSoundAssetManager::GetChecked().RemoveAsset(InAssetData);
+}
 
-	SearchAndIterateDirectoryAssets(Directories, [this](const FAssetData& AssetData)
-	{
-		using namespace Metasound;
-		using namespace Metasound::Frontend;
+void UMetaSoundAssetSubsystem::RenameAsset(const FAssetData& InAssetData, bool bInReregisterWithFrontend)
+{
+	using namespace Metasound::Frontend;
+	IMetaSoundAssetManager::GetChecked().RenameAsset(InAssetData, bInReregisterWithFrontend);
+}
 
-		if (AssetData.IsAssetLoaded())
-		{
-			FMetasoundAssetBase* MetaSoundAsset = Metasound::IMetasoundUObjectRegistry::Get().GetObjectAsAssetBase(AssetData.GetAsset());
-			check(MetaSoundAsset);
-			MetaSoundAsset->UnregisterGraphWithFrontend();
+UMetaSoundAssetSubsystem& UMetaSoundAssetSubsystem::GetChecked()
+{
+	check(GEngine);
+	UMetaSoundAssetSubsystem* Subsystem = GEngine->GetEngineSubsystem<UMetaSoundAssetSubsystem>();
+	check(Subsystem);
+	return *Subsystem;
+}
 
-			RemoveAsset(AssetData);
-		}
-		else
-		{
-			FNodeClassInfo AssetClassInfo;
-			if (ensureAlways(AssetSubsystemPrivate::GetAssetClassInfo(AssetData, AssetClassInfo)))
-			{
-				const FNodeRegistryKey RegistryKey = FNodeRegistryKey(AssetClassInfo);
-				const bool bIsRegistered = FMetasoundFrontendRegistryContainer::Get()->IsNodeRegistered(RegistryKey);
-				if (bIsRegistered)
-				{
-					FMetasoundFrontendRegistryContainer::Get()->UnregisterNode(RegistryKey);
-					AssetSubsystemPrivate::RemoveIfExactMatch(PathMap, RegistryKey, AssetData.GetSoftObjectPath());
-				}
-			}
-		}
-	});
+void UMetaSoundAssetSubsystem::RegisterAssetClassesInDirectories(const TArray<FMetaSoundAssetDirectory>& InDirectories)
+{
+	using namespace Metasound::Engine;
+	FMetaSoundAssetManager::GetChecked().RegisterAssetClassesInDirectories(InDirectories);
 }
 
+void UMetaSoundAssetSubsystem::UnregisterAssetClassesInDirectories(const TArray<FMetaSoundAssetDirectory>& InDirectories)
+{
+	using namespace Metasound::Engine;
+	FMetaSoundAssetManager::GetChecked().UnregisterAssetClassesInDirectories(InDirectories);
+}
