- Move MetaSoundAssetSubsystem& BuilderSubsystem implementations to monolithics that are initialized at MetaSoundEngine module load for safety

diff --git a/Engine/Plugins/Runtime/Metasound/Source/MetasoundEngine/Private/MetasoundBuilderSubsystem.cpp b/Engine/Plugins/Runtime/Metasound/Source/MetasoundEngine/Private/MetasoundBuilderSubsystem.cpp
--- a/Engine/Plugins/Runtime/Metasound/Source/MetasoundEngine/Private/MetasoundBuilderSubsystem.cpp
+++ b/Engine/Plugins/Runtime/Metasound/Source/MetasoundEngine/Private/MetasoundBuilderSubsystem.cpp
@@ -12,6 +12,8 @@
 #include "Metasound.h"
 #include "MetasoundAssetSubsystem.h"
 #include "MetasoundDataReference.h"
+#include "MetasoundDocumentBuilderRegistry.h"
+#include "MetasoundDocumentInterface.h"
 #include "MetasoundDynamicOperatorTransactor.h"
 #include "MetasoundFrontendDataTypeRegistry.h"
 #include "MetasoundFrontendDocument.h"
@@ -22,6 +24,7 @@
 #include "MetasoundGeneratorHandle.h"
 #include "MetasoundLog.h"
 #include "MetasoundParameterTransmitter.h"
+#include "MetasoundSettings.h"
 #include "MetasoundSource.h"
 #include "MetasoundTrace.h"
 #include "MetasoundUObjectRegistry.h"
@@ -36,1113 +39,49 @@ namespace Metasound::Engine
 {
 	namespace BuilderSubsystemPrivate
 	{
-		int32 TransactionBasedRegistrationEnabled = 1;
-
-		template <typename TLiteralType>
-		FMetasoundFrontendLiteral CreatePODMetaSoundLiteral(const TLiteralType& Value, FName& OutDataType)
-		{
-			OutDataType = GetMetasoundDataTypeName<TLiteralType>();
-
-			FMetasoundFrontendLiteral Literal;
-			Literal.Set(Value);
-			return Literal;
-		}
-
-		TUniquePtr<INode> CreateDynamicNodeFromFrontendLiteral(const FName DataType, const FMetasoundFrontendLiteral& InLiteral)
-		{
-			using namespace Frontend;
-			FLiteral ValueLiteral = InLiteral.ToLiteral(DataType);
-
-			// TODO: Node name "Literal" is always the same.  Consolidate and deprecate providing unique node name to avoid unnecessary FName table bloat.
-			FLiteralNodeConstructorParams Params { "Literal", FGuid::NewGuid(), MoveTemp(ValueLiteral) };
-			return IDataTypeRegistry::Get().CreateLiteralNode(DataType, MoveTemp(Params));
-		}
-	} // namespace BuilderSubsystemPrivate
-
-	FAutoConsoleVariableRef CVarMetaSoundBuilderForceAllFrontendRegistration(
-		TEXT("au.MetaSound.Builder.TransactionBasedRegistrationEnabled"),
-		BuilderSubsystemPrivate::TransactionBasedRegistrationEnabled,
-		TEXT("Forces all builder calls to register MetaSound objects with the Frontend.\n")
-		TEXT("Enabled (Default): !0, Disabled: 0"),
-		ECVF_Default);
-} // namespace Metasound::Engine
-
-
-void UMetaSoundBuilderBase::BeginDestroy()
-{
-	if (UMetaSoundBuilderSubsystem* Subsystem = UMetaSoundBuilderSubsystem::Get())
-	{
-		Subsystem->TransientBuilders.Remove(ClassName);
-	}
-
-	// Need to detach before destroying UPROPERTYs as the Builder 
-	// often holds a TScriptInterface<IMetaSoundDocumentInterface> of
-	// a UPROPERTY that lives on this or derived objects. 
-	Builder.FinishBuilding();
-	Super::BeginDestroy();
-}
-
-FMetaSoundBuilderNodeOutputHandle UMetaSoundBuilderBase::AddGraphInputNode(FName Name, FName DataType, FMetasoundFrontendLiteral DefaultValue, EMetaSoundBuilderResult& OutResult, bool bIsConstructorInput)
-{
-	using namespace Metasound::Frontend;
-
-	FMetaSoundBuilderNodeOutputHandle NewHandle;
-
-	if (IDataTypeRegistry::Get().FindDataTypeRegistryEntry(DataType) == nullptr)
-	{
-		UE_LOG(LogMetaSound, Error, TEXT("AddGraphInputNode Failed on builder '%s' when attempting to add '%s': '%s' is not a registered DataType"), *GetName(), *Name.ToString(), *DataType.ToString());
-	}
-	else
-	{
-		const FMetasoundFrontendNode* Node = Builder.FindGraphInputNode(Name);
-		if (Node)
-		{
-			UE_LOG(LogMetaSound, Warning, TEXT("AddGraphInputNode Failed: Input Node already exists with name '%s'; returning handle to existing node which may or may not match requested DataType '%s'"), *Name.ToString(), *DataType.ToString());
-		}
-		else
-		{
-			FDocumentIDGenerator& IDGenerator = FDocumentIDGenerator::Get();
-			const FMetasoundFrontendDocument& Doc = GetConstBuilder().GetDocument();
-
-			FMetasoundFrontendClassInput Description;
-			Description.Name = Name;
-			Description.TypeName = DataType;
-			Description.NodeID = IDGenerator.CreateNodeID(Doc);
-			Description.VertexID = IDGenerator.CreateVertexID(Doc);
-			Description.DefaultLiteral = static_cast<FMetasoundFrontendLiteral>(DefaultValue);
-			Description.AccessType = bIsConstructorInput ? EMetasoundFrontendVertexAccessType::Value : EMetasoundFrontendVertexAccessType::Reference;
-			Node = Builder.AddGraphInput(Description);
-		}
-
-		if (Node)
-		{
-			const TArray<FMetasoundFrontendVertex>& Outputs = Node->Interface.Outputs;
-			checkf(!Outputs.IsEmpty(), TEXT("Node should be initialized and have one output."));
-
-			NewHandle.NodeID = Node->GetID();
-			NewHandle.VertexID = Outputs.Last().VertexID;
-		}
-	}
-
-	OutResult = NewHandle.IsSet() ? EMetaSoundBuilderResult::Succeeded : EMetaSoundBuilderResult::Failed;
-	return NewHandle;
-}
-
-FMetaSoundBuilderNodeInputHandle UMetaSoundBuilderBase::AddGraphOutputNode(FName Name, FName DataType, FMetasoundFrontendLiteral DefaultValue, EMetaSoundBuilderResult& OutResult, bool bIsConstructorOutput)
-{
-	using namespace Metasound::Frontend;
-
-	FMetaSoundBuilderNodeInputHandle NewHandle;
-
-	if (IDataTypeRegistry::Get().FindDataTypeRegistryEntry(DataType) == nullptr)
-	{
-		UE_LOG(LogMetaSound, Error, TEXT("AddGraphOutputNode Failed on builder '%s' when attempting to add '%s': '%s' is not a registered DataType"), *GetName(), *Name.ToString(), *DataType.ToString());
-	}
-	else
-	{
-		const FMetasoundFrontendNode* Node = Builder.FindGraphOutputNode(Name);
-		if (Node)
-		{
-			UE_LOG(LogMetaSound, Warning, TEXT("AddGraphOutputNode Failed: Output Node already exists with name '%s'; returning handle to existing node which may or may not match requested DataType '%s'"), *Name.ToString(), *DataType.ToString());
-		}
-		else
-		{
-			FDocumentIDGenerator& IDGenerator = FDocumentIDGenerator::Get();
-			const FMetasoundFrontendDocument& Doc = GetConstBuilder().GetDocument();
-
-			FMetasoundFrontendClassOutput Description;
-			Description.Name = Name;
-			Description.TypeName = DataType;
-			Description.NodeID = IDGenerator.CreateNodeID(Doc);
-			Description.VertexID = IDGenerator.CreateVertexID(Doc);
-			Description.AccessType = bIsConstructorOutput ? EMetasoundFrontendVertexAccessType::Value : EMetasoundFrontendVertexAccessType::Reference;
-			Node = Builder.AddGraphOutput(Description);
-		}
-
-		if (Node)
-		{
-			const TArray<FMetasoundFrontendVertex>& Inputs = Node->Interface.Inputs;
-			checkf(!Inputs.IsEmpty(), TEXT("Node should be initialized and have one input."));
-
-			const FGuid& VertexID = Inputs.Last().VertexID;
-			if (Builder.SetNodeInputDefault(Node->GetID(), VertexID, DefaultValue))
-			{
-				NewHandle.NodeID = Node->GetID();
-				NewHandle.VertexID = VertexID;
-			}
-		}
-	}
-
-	OutResult = NewHandle.IsSet() ? EMetaSoundBuilderResult::Succeeded : EMetaSoundBuilderResult::Failed;
-	return NewHandle;
-}
-
-void UMetaSoundBuilderBase::AddInterface(FName InterfaceName, EMetaSoundBuilderResult& OutResult)
-{
-	const bool bInterfaceAdded = Builder.AddInterface(InterfaceName);
-	OutResult = bInterfaceAdded ? EMetaSoundBuilderResult::Succeeded : EMetaSoundBuilderResult::Failed;
-}
-
-FMetaSoundNodeHandle UMetaSoundBuilderBase::AddNode(const TScriptInterface<IMetaSoundDocumentInterface>& NodeClass, EMetaSoundBuilderResult& OutResult)
-{
-	using namespace Metasound;
-	using namespace Metasound::Frontend;
-
-	FMetaSoundNodeHandle NewHandle;
-
-	if (NodeClass)
-	{
-		UObject* NodeClassObject = NodeClass.GetObject();
-		check(NodeClassObject);
-
-#if WITH_EDITOR
-		// Assets that may undergo serialization cannot reference transient objects
-		const bool bIsInvalidReference = !NodeClassObject->IsAsset() && Builder.CastDocumentObjectChecked<UObject>().IsAsset();
-#else
-		constexpr bool bIsInvalidReference = false;
-#endif // WITH_EDITOR
-
-		if (bIsInvalidReference)
-		{
-			UObject& ThisBuildersObject = Builder.CastDocumentObjectChecked<UObject>();
-			UE_LOG(LogMetaSound, Warning,
-				TEXT("Failed to add node of transient asset '%s' to serialized asset '%s': "
-				"Transient object node class cannot be referenced from asset node class."),
-				*NodeClassObject->GetPathName(),
-				*ThisBuildersObject.GetPathName());
-		}
-		else
-		{
-			RegisterGraphIfOutstandingTransactions(*NodeClassObject);
-
-			const FMetasoundFrontendDocument& NodeClassDoc = NodeClass->GetConstDocument();
-			const FMetasoundFrontendGraphClass& NodeClassGraph = NodeClassDoc.RootGraph;
-			if (const FMetasoundFrontendNode* NewNode = Builder.AddGraphNode(NodeClassGraph))
-			{
-				NewHandle.NodeID = NewNode->GetID();
-			}
-		}
-	}
-
-	OutResult = NewHandle.IsSet() ? EMetaSoundBuilderResult::Succeeded : EMetaSoundBuilderResult::Failed;
-	return NewHandle;
-}
-
-FMetaSoundNodeHandle UMetaSoundBuilderBase::AddNodeByClassName(const FMetasoundFrontendClassName& InClassName, int32 MajorVersion, EMetaSoundBuilderResult& OutResult)
-{
-	return AddNodeByClassName(InClassName, OutResult, MajorVersion);
-}
-
-FMetaSoundNodeHandle UMetaSoundBuilderBase::AddNodeByClassName(const FMetasoundFrontendClassName& InClassName, EMetaSoundBuilderResult& OutResult, int32 MajorVersion)
-{
-	using namespace Metasound;
-	using namespace Metasound::Frontend;
-
-	FMetaSoundNodeHandle NewHandle;
-	if (const FMetasoundFrontendNode* NewNode = Builder.AddNodeByClassName(InClassName, MajorVersion))
-	{
-		NewHandle.NodeID = NewNode->GetID();
-	}
-
-	OutResult = NewHandle.IsSet() ? EMetaSoundBuilderResult::Succeeded : EMetaSoundBuilderResult::Failed;
-	return NewHandle;
-}
-
-void UMetaSoundBuilderBase::BuildInternal(TScriptInterface<IMetaSoundDocumentInterface> NewMetaSound, const FMetasoundFrontendClassName* InDocClassName) const
-{
-	using namespace Metasound::Frontend;
-
-	// If using existing class name, ensure that a builder does not exist for it to
-	// avoid build active flag conflation with the locally generated frontend builder below.
-	if (InDocClassName)
-	{
-		UMetaSoundBuilderSubsystem::GetChecked().DetachBuilderFromAsset(*InDocClassName);
-	}
-
-	FMetaSoundFrontendDocumentBuilder NewDocBuilder(NewMetaSound);
-
-	constexpr bool bResetVersion = false;
-	NewDocBuilder.InitDocument(&GetConstBuilder().GetDocument(), InDocClassName, bResetVersion);
-	NewMetaSound->ConformObjectToDocument();
-}
-
-#if WITH_EDITOR
-bool UMetaSoundBuilderBase::ClearMemberMetadata(const FGuid& InMemberID)
-{
-	return Builder.ClearMemberMetadata(InMemberID);
-}
-#endif // WITH_EDITOR
-
-bool UMetaSoundBuilderBase::ContainsNode(const FMetaSoundNodeHandle& NodeHandle) const
-{
-	return Builder.ContainsNode(NodeHandle.NodeID);
-}
-
-bool UMetaSoundBuilderBase::ContainsNodeInput(const FMetaSoundBuilderNodeInputHandle& InputHandle) const
-{
-	return Builder.FindNodeInput(InputHandle.NodeID, InputHandle.VertexID) != nullptr;
-}
-
-bool UMetaSoundBuilderBase::ContainsNodeOutput(const FMetaSoundBuilderNodeOutputHandle& OutputHandle) const
-{
-	return Builder.FindNodeOutput(OutputHandle.NodeID, OutputHandle.VertexID) != nullptr;
-}
-
-void UMetaSoundBuilderBase::ConnectNodes(const FMetaSoundBuilderNodeOutputHandle& NodeOutputHandle, const FMetaSoundBuilderNodeInputHandle& NodeInputHandle, EMetaSoundBuilderResult& OutResult)
-{
-	using namespace Metasound::Frontend;
-
-	OutResult = EMetaSoundBuilderResult::Failed;
-
-	FMetasoundFrontendEdge NewEdge { NodeOutputHandle.NodeID, NodeOutputHandle.VertexID, NodeInputHandle.NodeID, NodeInputHandle.VertexID };
-	const EInvalidEdgeReason InvalidEdgeReason = Builder.IsValidEdge(NewEdge);
-
-	if (InvalidEdgeReason == Metasound::Frontend::EInvalidEdgeReason::None)
-	{
-#if !NO_LOGGING
-		const FMetasoundFrontendNode* OldOutputNode = nullptr;
-		const FMetasoundFrontendVertex* OldOutputVertex = nullptr;
-		if (Builder.IsNodeInputConnected(NodeInputHandle.NodeID, NodeInputHandle.VertexID))
-		{
-			OldOutputVertex = Builder.FindNodeOutputConnectedToNodeInput(NodeInputHandle.NodeID, NodeInputHandle.VertexID, &OldOutputNode);
-		}
-#endif // !NO_LOGGING
-
-		const bool bRemovedEdge = Builder.RemoveEdgeToNodeInput(NodeInputHandle.NodeID, NodeInputHandle.VertexID);
-		Builder.AddEdge(MoveTemp(NewEdge));
-
-#if !NO_LOGGING
-		if (bRemovedEdge)
-		{
-			checkf(OldOutputNode, TEXT("MetaSound edge was removed from output but output node not found."));
-			checkf(OldOutputVertex, TEXT("MetaSound edge was removed from output but output vertex not found."));
-
-			const FMetasoundFrontendNode* InputNode = Builder.FindNode(NodeInputHandle.NodeID);
-			checkf(InputNode, TEXT("Edge was deemed valid but input parent node is missing"));
-
-			const FMetasoundFrontendVertex* InputVertex = Builder.FindNodeInput(NodeInputHandle.NodeID, NodeInputHandle.VertexID);
-			checkf(InputVertex, TEXT("Edge was deemed valid but input is missing"));
-
-			const FMetasoundFrontendNode* OutputNode = Builder.FindNode(NodeOutputHandle.NodeID);
-			checkf(OutputNode, TEXT("Edge was deemed valid but output parent node is missing"));
-
-			const FMetasoundFrontendVertex* OutputVertex = Builder.FindNodeOutput(NodeOutputHandle.NodeID, NodeOutputHandle.VertexID);
-			checkf(OutputVertex, TEXT("Edge was deemed valid but output is missing"));
-
-			UE_LOG(LogMetaSound, Verbose, TEXT("Removed connection from node output '%s:%s' to node '%s:%s' in order to connect to node output '%s:%s'"),
-				*OldOutputNode->Name.ToString(),
-				*OldOutputVertex->Name.ToString(),
-				*InputNode->Name.ToString(),
-				*InputVertex->Name.ToString(),
-				*OutputNode->Name.ToString(),
-				*OutputVertex->Name.ToString());
-		}
-#endif // !NO_LOGGING
-
-		OutResult = EMetaSoundBuilderResult::Succeeded;
-	}
-	else
-	{
-		UE_LOG(LogMetaSound, Warning, TEXT("Builder '%s' 'ConnectNodes' failed: '%s'"), *GetName(), *LexToString(InvalidEdgeReason));
-	}
-}
-
-void UMetaSoundBuilderBase::ConnectNodesByInterfaceBindings(const FMetaSoundNodeHandle& FromNodeHandle, const FMetaSoundNodeHandle& ToNodeHandle, EMetaSoundBuilderResult& OutResult)
-{
-	const bool bEdgesAdded = Builder.AddEdgesByNodeClassInterfaceBindings(FromNodeHandle.NodeID, ToNodeHandle.NodeID);
-	OutResult = bEdgesAdded ? EMetaSoundBuilderResult::Succeeded : EMetaSoundBuilderResult::Failed;
-}
-
-TArray<FMetaSoundBuilderNodeInputHandle> UMetaSoundBuilderBase::ConnectNodeOutputsToMatchingGraphInterfaceOutputs(const FMetaSoundNodeHandle& NodeHandle, EMetaSoundBuilderResult& OutResult)
-{
-	TArray<const FMetasoundFrontendEdge*> NewEdges;
-	const bool bEdgesAdded = Builder.AddEdgesFromMatchingInterfaceNodeOutputsToGraphOutputs(NodeHandle.NodeID, NewEdges);
-	OutResult = bEdgesAdded ? EMetaSoundBuilderResult::Succeeded : EMetaSoundBuilderResult::Failed;
-
-	TArray<FMetaSoundBuilderNodeInputHandle> ConnectedVertices;
-	Algo::Transform(NewEdges, ConnectedVertices, [this](const FMetasoundFrontendEdge* NewEdge)
-	{
-		const FMetasoundFrontendVertex* Vertex = Builder.FindNodeInput(NewEdge->ToNodeID, NewEdge->ToVertexID);
-		checkf(Vertex, TEXT("Edge connection reported success but vertex not found."));
-		return FMetaSoundBuilderNodeInputHandle(NewEdge->ToNodeID, Vertex->VertexID);
-	});
-
-	return ConnectedVertices;
-}
-
-TArray<FMetaSoundBuilderNodeOutputHandle> UMetaSoundBuilderBase::ConnectNodeInputsToMatchingGraphInterfaceInputs(const FMetaSoundNodeHandle& NodeHandle, EMetaSoundBuilderResult& OutResult)
-{
-	TArray<const FMetasoundFrontendEdge*> NewEdges;
-	const bool bEdgesAdded = Builder.AddEdgesFromMatchingInterfaceNodeInputsToGraphInputs(NodeHandle.NodeID, NewEdges);
-	OutResult = bEdgesAdded ? EMetaSoundBuilderResult::Succeeded : EMetaSoundBuilderResult::Failed;
-
-	TArray<FMetaSoundBuilderNodeOutputHandle> ConnectedVertices;
-	Algo::Transform(NewEdges, ConnectedVertices, [this](const FMetasoundFrontendEdge* NewEdge)
-	{
-		const FMetasoundFrontendVertex* Vertex = Builder.FindNodeOutput(NewEdge->FromNodeID, NewEdge->FromVertexID);
-		checkf(Vertex, TEXT("Edge connection reported success but vertex not found."));
-		return FMetaSoundBuilderNodeOutputHandle(NewEdge->ToNodeID, Vertex->VertexID);
-	});
-
-	return ConnectedVertices;
-}
-
-void UMetaSoundBuilderBase::ConnectNodeOutputToGraphOutput(FName GraphOutputName, const FMetaSoundBuilderNodeOutputHandle& NodeOutputHandle, EMetaSoundBuilderResult& OutResult)
-{
-	using namespace Metasound::Frontend;
-
-	OutResult = EMetaSoundBuilderResult::Failed;
-
-	if (const FMetasoundFrontendNode* GraphOutputNode = Builder.FindGraphOutputNode(GraphOutputName))
-	{
-		const FMetasoundFrontendVertex& InputVertex = GraphOutputNode->Interface.Inputs.Last();
-		FMetasoundFrontendEdge NewEdge { NodeOutputHandle.NodeID, NodeOutputHandle.VertexID, GraphOutputNode->GetID(), InputVertex.VertexID };
-		const EInvalidEdgeReason InvalidEdgeReason = Builder.IsValidEdge(NewEdge);
-		if (InvalidEdgeReason == EInvalidEdgeReason::None)
-		{
-			Builder.RemoveEdgeToNodeInput(GraphOutputNode->GetID(), InputVertex.VertexID);
-			Builder.AddEdge(MoveTemp(NewEdge));
-			OutResult = EMetaSoundBuilderResult::Succeeded;
-		}
-		else
-		{
-			UE_LOG(LogMetaSound, Warning, TEXT("Builder '%s' 'ConnectNodeOutputToGraphOutput' failed: '%s'"), *GetName(), *LexToString(InvalidEdgeReason));
-		}
-	}
-}
-
-void UMetaSoundBuilderBase::ConnectNodeInputToGraphInput(FName GraphInputName, const FMetaSoundBuilderNodeInputHandle& NodeInputHandle, EMetaSoundBuilderResult& OutResult)
-{
-	using namespace Metasound::Frontend;
-
-	OutResult = EMetaSoundBuilderResult::Failed;
-
-	if (const FMetasoundFrontendNode* GraphInputNode = Builder.FindGraphInputNode(GraphInputName))
-	{
-		const FMetasoundFrontendVertex& OutputVertex = GraphInputNode->Interface.Outputs.Last();
-		FMetasoundFrontendEdge NewEdge { GraphInputNode->GetID(), OutputVertex.VertexID, NodeInputHandle.NodeID, NodeInputHandle.VertexID };
-		const EInvalidEdgeReason InvalidEdgeReason = Builder.IsValidEdge(NewEdge);
-		if (InvalidEdgeReason == EInvalidEdgeReason::None)
-		{
-			Builder.RemoveEdgeToNodeInput(NodeInputHandle.NodeID, NodeInputHandle.VertexID);
-			Builder.AddEdge(MoveTemp(NewEdge));
-			OutResult = EMetaSoundBuilderResult::Succeeded;
-		}
-		else
-		{
-			UE_LOG(LogMetaSound, Warning, TEXT("Builder '%s' 'ConnectNodeInputToGraphInput' failed: '%s'"), *GetName(), *LexToString(InvalidEdgeReason));
-		}
-	}
-}
-
-void UMetaSoundBuilderBase::ConvertFromPreset(EMetaSoundBuilderResult& OutResult)
-{
-	const bool bSuccess = Builder.ConvertFromPreset();
-	OutResult = bSuccess ? EMetaSoundBuilderResult::Succeeded : EMetaSoundBuilderResult::Failed;
-}
-
-void UMetaSoundBuilderBase::ConvertToPreset(const TScriptInterface<IMetaSoundDocumentInterface>& ReferencedNodeClass, EMetaSoundBuilderResult& OutResult)
-{
-	using namespace Metasound::Frontend;
-
-	const IMetaSoundDocumentInterface* ReferencedInterface = ReferencedNodeClass.GetInterface();
-	if (!ReferencedInterface)
-	{
-		OutResult = EMetaSoundBuilderResult::Failed;
-		return;
-	}
-	
-	// Ensure the referenced node class isn't transient
-	if (Cast<UMetaSoundBuilderDocument>(ReferencedInterface))
-	{
-		UE_LOG(LogMetaSound, Warning, TEXT("Transient document builders cannot be referenced when converting builder '%s' to a preset. Build the referenced node class an asset first or use an existing asset instead"), *GetName());
-		OutResult = EMetaSoundBuilderResult::Failed;
-		return;
-	}
-
-	// Ensure the referenced node class is a matching object type 
-	const UClass& BaseMetaSoundClass = ReferencedInterface->GetBaseMetaSoundUClass();
-	UObject* ReferencedObject = ReferencedNodeClass.GetObject();
-	if (!ReferencedObject || (ReferencedObject && !ReferencedObject->IsA(&BaseMetaSoundClass)))
-	{
-		UE_LOG(LogMetaSound, Warning, TEXT("The referenced node type must match the base MetaSound class when converting builder '%s' to a preset (ex. source preset must reference another source)"), *GetName());
-		OutResult = EMetaSoundBuilderResult::Failed;
-		return;
-	}
-
-	// Ensure the referenced node is registered
-	if (FMetasoundAssetBase* ReferencedMetaSoundAsset = Metasound::IMetasoundUObjectRegistry::Get().GetObjectAsAssetBase(ReferencedObject))
-	{
-		ReferencedMetaSoundAsset->RegisterGraphWithFrontend();
-	}
-
-	const FMetasoundFrontendDocument& ReferencedDocument = ReferencedInterface->GetConstDocument();
-
-	TSharedRef<FDocumentModifyDelegates> DocumentDelegates = MakeShared<FDocumentModifyDelegates>();
-	InitDelegates(*DocumentDelegates);
-	const bool bConvertToPreset = Builder.ConvertToPreset(ReferencedDocument, DocumentDelegates);
-	ConformObjectToDocument();
-	OutResult = bConvertToPreset ? EMetaSoundBuilderResult::Succeeded : EMetaSoundBuilderResult::Failed;
-}
-
-void UMetaSoundBuilderBase::DisconnectNodes(const FMetaSoundBuilderNodeOutputHandle& NodeOutputHandle, const FMetaSoundBuilderNodeInputHandle& NodeInputHandle, EMetaSoundBuilderResult& OutResult)
-{
-	const bool bEdgeRemoved = Builder.RemoveEdge(FMetasoundFrontendEdge
-	{
-		NodeOutputHandle.NodeID,
-		NodeOutputHandle.VertexID,
-		NodeInputHandle.NodeID,
-		NodeInputHandle.VertexID,
-	});
-	OutResult = bEdgeRemoved ? EMetaSoundBuilderResult::Succeeded : EMetaSoundBuilderResult::Failed;
-}
-
-void UMetaSoundBuilderBase::DisconnectNodeInput(const FMetaSoundBuilderNodeInputHandle& NodeInputHandle, EMetaSoundBuilderResult& OutResult)
-{
-	const bool bEdgeRemoved = Builder.RemoveEdgeToNodeInput(NodeInputHandle.NodeID, NodeInputHandle.VertexID);
-	OutResult = bEdgeRemoved ? EMetaSoundBuilderResult::Succeeded : EMetaSoundBuilderResult::Failed;
-}
-
-void UMetaSoundBuilderBase::DisconnectNodeOutput(const FMetaSoundBuilderNodeOutputHandle& NodeOutputHandle, EMetaSoundBuilderResult& OutResult)
-{
-	const bool bEdgeRemoved = Builder.RemoveEdgesFromNodeOutput(NodeOutputHandle.NodeID, NodeOutputHandle.VertexID);
-	OutResult = bEdgeRemoved ? EMetaSoundBuilderResult::Succeeded : EMetaSoundBuilderResult::Failed;
-}
-
-void UMetaSoundBuilderBase::DisconnectNodesByInterfaceBindings(const FMetaSoundNodeHandle& FromNodeHandle, const FMetaSoundNodeHandle& ToNodeHandle, EMetaSoundBuilderResult& OutResult)
-{
-	const bool bEdgesRemoved = Builder.RemoveEdgesByNodeClassInterfaceBindings(FromNodeHandle.NodeID, ToNodeHandle.NodeID);
-	OutResult = bEdgesRemoved ? EMetaSoundBuilderResult::Succeeded : EMetaSoundBuilderResult::Failed;
-}
-
-FMetaSoundBuilderNodeInputHandle UMetaSoundBuilderBase::FindNodeInputByName(const FMetaSoundNodeHandle& NodeHandle, FName InputName, EMetaSoundBuilderResult& OutResult)
-{
-	if (const FMetasoundFrontendNode* Node = Builder.FindNode(NodeHandle.NodeID))
-	{
-		const TArray<FMetasoundFrontendVertex>& InputVertices = Node->Interface.Inputs;
-
-		auto FindByNamePredicate = [&InputName](const FMetasoundFrontendVertex& Vertex) { return Vertex.Name == InputName; };
-		if (const FMetasoundFrontendVertex* Input = InputVertices.FindByPredicate(FindByNamePredicate))
-		{
-			OutResult = EMetaSoundBuilderResult::Succeeded;
-			return FMetaSoundBuilderNodeInputHandle(Node->GetID(), Input->VertexID);
-		}
-
-		FString NodeClassName = TEXT("N/A");
-		if (const FMetasoundFrontendClass* Class = Builder.FindDependency(Node->ClassID))
-		{
-			NodeClassName = Class->Metadata.GetClassName().ToString();
-		}
-
-		UE_LOG(LogMetaSound, Display, TEXT("Builder '%s' failed to find node input '%s': Node class '%s' contains no such input"), *GetName(), *InputName.ToString(), *NodeClassName);
-	}
-	else
-	{
-		UE_LOG(LogMetaSound, Display, TEXT("Builder '%s' failed to find node input '%s': Node with ID '%s' not found"), *GetName(), *InputName.ToString(), *NodeHandle.NodeID.ToString());
-	}
-
-
-	OutResult = EMetaSoundBuilderResult::Failed;
-	return { };
-}
-
-TArray<FMetaSoundBuilderNodeInputHandle> UMetaSoundBuilderBase::FindNodeInputs(const FMetaSoundNodeHandle& NodeHandle, EMetaSoundBuilderResult& OutResult)
-{
-	return FindNodeInputsByDataType(NodeHandle, OutResult, { });
-}
-
-TArray<FMetaSoundBuilderNodeInputHandle> UMetaSoundBuilderBase::FindNodeInputsByDataType(const FMetaSoundNodeHandle& NodeHandle, EMetaSoundBuilderResult& OutResult, FName DataType)
-{
-	TArray<FMetaSoundBuilderNodeInputHandle> FoundVertices;
-	if (Builder.ContainsNode(NodeHandle.NodeID))
-	{
-		TArray<const FMetasoundFrontendVertex*> Vertices = Builder.FindNodeInputs(NodeHandle.NodeID, DataType);
-		Algo::Transform(Vertices, FoundVertices, [&NodeHandle](const FMetasoundFrontendVertex* Vertex)
-		{
-			return FMetaSoundBuilderNodeInputHandle(NodeHandle.NodeID, Vertex->VertexID);
-		});
-		OutResult = EMetaSoundBuilderResult::Succeeded;
-	}
-	else
-	{
-		UE_LOG(LogMetaSound, Display, TEXT("Failed to find node inputs by data type with builder '%s'. Node of with ID '%s' not found"), *GetName(), *NodeHandle.NodeID.ToString());
-		OutResult = EMetaSoundBuilderResult::Failed;
-	}
-
-	return FoundVertices;
-}
-
-FMetaSoundBuilderNodeOutputHandle UMetaSoundBuilderBase::FindNodeOutputByName(const FMetaSoundNodeHandle& NodeHandle, FName OutputName, EMetaSoundBuilderResult& OutResult)
-{
-	if (const FMetasoundFrontendNode* Node = Builder.FindNode(NodeHandle.NodeID))
-	{
-		const TArray<FMetasoundFrontendVertex>& OutputVertices = Node->Interface.Outputs;
-
-		auto FindByNamePredicate = [&OutputName](const FMetasoundFrontendVertex& Vertex) { return Vertex.Name == OutputName; };
-		if (const FMetasoundFrontendVertex* Output = OutputVertices.FindByPredicate(FindByNamePredicate))
-		{
-			OutResult = EMetaSoundBuilderResult::Succeeded;
-			return FMetaSoundBuilderNodeOutputHandle(Node->GetID(), Output->VertexID);
-		}
-
-		FString NodeClassName = TEXT("N/A");
-		if (const FMetasoundFrontendClass* Class = Builder.FindDependency(Node->ClassID))
-		{
-			NodeClassName = Class->Metadata.GetClassName().ToString();
-		}
-
-		UE_LOG(LogMetaSound, Display, TEXT("Builder '%s' failed to find node output '%s': Node class '%s' contains no such output"), *GetName(), *OutputName.ToString(), *NodeClassName);
-	}
-	else
-	{
-		UE_LOG(LogMetaSound, Display, TEXT("Builder '%s' failed to find node output '%s': Node with ID '%s' not found"), *GetName(), *OutputName.ToString(), *NodeHandle.NodeID.ToString());
-	}
-
-	OutResult = EMetaSoundBuilderResult::Failed;
-	return { };
-}
-
-TArray<FMetaSoundBuilderNodeOutputHandle> UMetaSoundBuilderBase::FindNodeOutputs(const FMetaSoundNodeHandle& NodeHandle, EMetaSoundBuilderResult& OutResult)
-{
-	return FindNodeOutputsByDataType(NodeHandle, OutResult, { });
-}
-
-TArray<FMetaSoundBuilderNodeOutputHandle> UMetaSoundBuilderBase::FindNodeOutputsByDataType(const FMetaSoundNodeHandle& NodeHandle, EMetaSoundBuilderResult& OutResult, FName DataType)
-{
-	TArray<FMetaSoundBuilderNodeOutputHandle> FoundVertices;
-	if (Builder.ContainsNode(NodeHandle.NodeID))
-	{
-		TArray<const FMetasoundFrontendVertex*> Vertices = Builder.FindNodeOutputs(NodeHandle.NodeID, DataType);
-		Algo::Transform(Vertices, FoundVertices, [&NodeHandle](const FMetasoundFrontendVertex* Vertex)
-		{
-			return FMetaSoundBuilderNodeOutputHandle(NodeHandle.NodeID, Vertex->VertexID);
-		});
-		OutResult = EMetaSoundBuilderResult::Succeeded;
-	}
-	else
-	{
-		UE_LOG(LogMetaSound, Display, TEXT("Failed to find node outputs by data type with builder '%s'. Node of with ID '%s' not found"), *GetName(), *NodeHandle.NodeID.ToString());
-		OutResult = EMetaSoundBuilderResult::Failed;
-	}
-
-	return FoundVertices;
-}
-
-#if WITH_EDITOR
-const FMetaSoundFrontendGraphComment* UMetaSoundBuilderBase::FindGraphComment(const FGuid& InCommentID) const
-{
-	return Builder.FindGraphComment(InCommentID);
-}
-
-FMetaSoundFrontendGraphComment* UMetaSoundBuilderBase::FindGraphComment(const FGuid& InCommentID)
-{
-	return Builder.FindGraphComment(InCommentID);
-}
-
-FMetaSoundFrontendGraphComment& UMetaSoundBuilderBase::FindOrAddGraphComment(const FGuid& InCommentID)
-{
-	return Builder.FindOrAddGraphComment(InCommentID);
-}
-#endif // WITH_EDITOR
-
-TArray<FMetaSoundNodeHandle> UMetaSoundBuilderBase::FindInterfaceInputNodes(FName InterfaceName, EMetaSoundBuilderResult& OutResult)
-{
-	TArray<FMetaSoundNodeHandle> NodeHandles;
-
-	TArray<const FMetasoundFrontendNode*> Nodes;
-	if (Builder.FindInterfaceInputNodes(InterfaceName, Nodes))
-	{
-		Algo::Transform(Nodes, NodeHandles, [this](const FMetasoundFrontendNode* Node)
-		{
-			check(Node);
-			return FMetaSoundNodeHandle { Node->GetID() };
-		});
-		OutResult = EMetaSoundBuilderResult::Succeeded;
-	}
-	else
-	{
-		UE_LOG(LogMetaSound, Display, TEXT("'%s' interface not found on builder '%s'. No input nodes returned"), *InterfaceName.ToString(), *GetName());
-		OutResult = EMetaSoundBuilderResult::Failed;
-	}
-
-	return NodeHandles;
-}
-
-TArray<FMetaSoundNodeHandle> UMetaSoundBuilderBase::FindInterfaceOutputNodes(FName InterfaceName, EMetaSoundBuilderResult& OutResult)
-{
-	TArray<FMetaSoundNodeHandle> NodeHandles;
-
-	TArray<const FMetasoundFrontendNode*> Nodes;
-	if (Builder.FindInterfaceOutputNodes(InterfaceName, Nodes))
-	{
-		Algo::Transform(Nodes, NodeHandles, [this](const FMetasoundFrontendNode* Node)
-		{
-			check(Node);
-			return FMetaSoundNodeHandle { Node->GetID() };
-		});
-		OutResult = EMetaSoundBuilderResult::Succeeded;
-	}
-	else
-	{
-		OutResult = EMetaSoundBuilderResult::Failed;
-	}
-
-	return NodeHandles;
-}
-
-FMetaSoundNodeHandle UMetaSoundBuilderBase::FindGraphInputNode(FName InputName, EMetaSoundBuilderResult& OutResult)
-{
-	if (const FMetasoundFrontendNode* GraphInputNode = Builder.FindGraphInputNode(InputName))
-	{
-		OutResult = EMetaSoundBuilderResult::Succeeded;
-		return FMetaSoundNodeHandle { GraphInputNode->GetID() };
-	}
-
-	UE_LOG(LogMetaSound, Display, TEXT("Failed to find graph input by name '%s' with builder '%s'"), *InputName.ToString(), *GetName());
-	OutResult = EMetaSoundBuilderResult::Failed;
-	return { };
-}
-
-FMetaSoundNodeHandle UMetaSoundBuilderBase::FindGraphOutputNode(FName OutputName, EMetaSoundBuilderResult& OutResult)
-{
-	if (const FMetasoundFrontendNode* GraphOutputNode = Builder.FindGraphOutputNode(OutputName))
-	{
-		OutResult = EMetaSoundBuilderResult::Succeeded;
-		return FMetaSoundNodeHandle{ GraphOutputNode->GetID() };
-	}
-
-	UE_LOG(LogMetaSound, Display, TEXT("Failed to find graph output by name '%s' with builder '%s'"), *OutputName.ToString(), *GetName());
-	OutResult = EMetaSoundBuilderResult::Failed;
-	return { };
-}
-
-#if WITH_EDITOR
-UMetaSoundFrontendMemberMetadata* UMetaSoundBuilderBase::FindMemberMetadata(const FGuid& InMemberID)
-{
-	return Builder.FindMemberMetadata(InMemberID);
-}
-#endif // WITH_EDITOR
-
-UMetaSoundBuilderDocument* UMetaSoundBuilderBase::CreateTransientDocumentObject() const
-{
-	return &UMetaSoundBuilderDocument::Create(GetBuilderUClass());
-}
-
-const FMetaSoundFrontendDocumentBuilder& UMetaSoundBuilderBase::GetConstBuilder() const
-{
-	return Builder;
-}
-
-UObject* UMetaSoundBuilderBase::GetReferencedPresetAsset() const
-{
-	using namespace Metasound::Frontend;
-	if (!IsPreset())
-	{
-		return nullptr;
-	}
-
-	// Find the single external node which is the referenced preset asset, 
-	// and find the asset with its registry key 
-	auto FindExternalNode = [this](const FMetasoundFrontendNode& Node) 
-	{
-		const FMetasoundFrontendClass* Class = Builder.FindDependency(Node.ClassID);
-		check(Class);
-		return Class->Metadata.GetType() == EMetasoundFrontendClassType::External;
-	};
-	const FMetasoundFrontendNode* Node = Builder.GetDocument().RootGraph.Graph.Nodes.FindByPredicate(FindExternalNode);
-	if (Node != nullptr)
-	{
-		const FMetasoundFrontendClass* NodeClass = Builder.FindDependency(Node->ClassID);
-		check(NodeClass);
-		const FNodeRegistryKey NodeClassRegistryKey = FNodeRegistryKey(NodeClass->Metadata);
-		if (FMetasoundAssetBase* Asset = IMetaSoundAssetManager::GetChecked().TryLoadAssetFromKey(NodeClassRegistryKey))
-		{
-			return Asset->GetOwningAsset();
-		}
-	}
-
-	return nullptr;
-}
-
-void UMetaSoundBuilderBase::InitFrontendBuilder()
-{
-	UMetaSoundBuilderDocument& DocObject = UMetaSoundBuilderDocument::Create(GetBuilderUClass());
-	Builder = FMetaSoundFrontendDocumentBuilder(&DocObject);
-	Builder.InitDocument();
-}
-
-void UMetaSoundBuilderBase::InitNodeLocations()
-{
-	Builder.InitNodeLocations();
-}
-
-#if WITH_EDITOR
-void UMetaSoundBuilderBase::InjectInputTemplateNodes(bool bForceNodeCreation, EMetaSoundBuilderResult& OutResult)
-{
-	using namespace Metasound::Frontend;
-	const bool bTemplateNodesInjected = FInputNodeTemplate::GetChecked().Inject(Builder, bForceNodeCreation);
-	OutResult = bTemplateNodesInjected ? EMetaSoundBuilderResult::Succeeded : EMetaSoundBuilderResult::Failed;
-}
-#endif // WITH_EDITOR
-
-bool UMetaSoundBuilderBase::InterfaceIsDeclared(FName InterfaceName) const
-{
-	return Builder.IsInterfaceDeclared(InterfaceName);
-}
-
-void UMetaSoundBuilderBase::InvalidateCache(bool bPrimeCache)
-{
-	using namespace Metasound::Frontend;
-
-	TSharedRef<FDocumentModifyDelegates> DocumentDelegates = MakeShared<FDocumentModifyDelegates>();
-	InitDelegates(*DocumentDelegates);
-
-	Builder.Reload(DocumentDelegates, bPrimeCache);
-}
-
-bool UMetaSoundBuilderBase::IsPreset() const
-{
-	return Builder.IsPreset();
-}
-
-bool UMetaSoundBuilderBase::NodesAreConnected(const FMetaSoundBuilderNodeOutputHandle& OutputHandle, const FMetaSoundBuilderNodeInputHandle& InputHandle) const
-{
-	const FMetasoundFrontendEdge Edge = { OutputHandle.NodeID, OutputHandle.VertexID, InputHandle.NodeID, InputHandle.VertexID };
-	return Builder.ContainsEdge(Edge);
-}
-
-bool UMetaSoundBuilderBase::NodeInputIsConnected(const FMetaSoundBuilderNodeInputHandle& InputHandle) const
-{
-	return Builder.IsNodeInputConnected(InputHandle.NodeID, InputHandle.VertexID);
-}
-
-bool UMetaSoundBuilderBase::NodeOutputIsConnected(const FMetaSoundBuilderNodeOutputHandle& OutputHandle) const
-{
-	return Builder.IsNodeOutputConnected(OutputHandle.NodeID, OutputHandle.VertexID);
-}
-
-void UMetaSoundBuilderBase::RegisterGraphIfOutstandingTransactions(UObject& InMetaSound)
-{
-	using namespace Metasound;
-	using namespace Metasound::Engine;
-	using namespace Metasound::Frontend;
-	using namespace Metasound::Engine::BuilderSubsystemPrivate;
-
-	IMetaSoundAssetManager& AssetManager = IMetaSoundAssetManager::GetChecked();
-	FMetasoundAssetBase* MetaSoundAsset = IMetasoundUObjectRegistry::Get().GetObjectAsAssetBase(&InMetaSound);
-	check(MetaSoundAsset);
-
-	FMetaSoundAssetRegistrationOptions Options;
-	if (TransactionBasedRegistrationEnabled)
-	{
-		Options.bForceReregister = false;
-		Options.bRegisterDependencies = false; // Function handles registration via own recursive functionality below
-
-		TArray<FMetasoundAssetBase*> References = MetaSoundAsset->GetReferencedAssets();
-		for (FMetasoundAssetBase* Reference : References)
-		{
-			UObject* RefMetaSound = Reference->GetOwningAsset();
-			check(RefMetaSound);
-			AssetManager.AddOrUpdateAsset(*RefMetaSound);
-			RegisterGraphIfOutstandingTransactions(*RefMetaSound);
-		}
-
-		if (UMetaSoundBuilderBase* Builder = UMetaSoundBuilderSubsystem::GetChecked().FindBuilderOfDocument(&InMetaSound))
+		template <typename TLiteralType>
+		FMetasoundFrontendLiteral CreatePODMetaSoundLiteral(const TLiteralType& Value, FName& OutDataType)
 		{
-			const int32 TransactionCount = Builder->Builder.GetTransactionCount();
+			OutDataType = GetMetasoundDataTypeName<TLiteralType>();
 
-			// Force registration if transactions occurred since now and the last time the builder registered the asset.
-			Options.bForceReregister = Builder->LastTransactionRegistered != TransactionCount;
-			Builder->LastTransactionRegistered = TransactionCount;
+			FMetasoundFrontendLiteral Literal;
+			Literal.Set(Value);
+			return Literal;
 		}
-	}
-
-	MetaSoundAsset->RegisterGraphWithFrontend(Options);
-}
-
-void UMetaSoundBuilderBase::Reload(bool bPrimeCache)
-{
-	using namespace Metasound::Frontend;
-
-	TSharedRef<FDocumentModifyDelegates> DocumentDelegates = MakeShared<FDocumentModifyDelegates>();
-	InitDelegates(*DocumentDelegates);
-	Builder.Reload(DocumentDelegates, bPrimeCache);
-}
-
-void UMetaSoundBuilderBase::ReloadCache(bool bPrimeCache)
-{
-	Reload(bPrimeCache);
-}
-
-#if WITH_EDITOR
-bool UMetaSoundBuilderBase::RemoveGraphComment(const FGuid& InCommentID)
-{
-	return Builder.RemoveGraphComment(InCommentID);
-}
-#endif // WITH_EDITOR
-
-void UMetaSoundBuilderBase::RemoveGraphInput(FName Name, EMetaSoundBuilderResult& OutResult)
-{
-	const bool bRemoved = Builder.RemoveGraphInput(Name);
-	OutResult = bRemoved ? EMetaSoundBuilderResult::Succeeded : EMetaSoundBuilderResult::Failed;
-}
-
-void UMetaSoundBuilderBase::RemoveGraphOutput(FName Name, EMetaSoundBuilderResult& OutResult)
-{
-	const bool bRemoved = Builder.RemoveGraphOutput(Name);
-	OutResult = bRemoved ? EMetaSoundBuilderResult::Succeeded : EMetaSoundBuilderResult::Failed;
-}
-
-void UMetaSoundBuilderBase::RemoveInterface(FName InterfaceName, EMetaSoundBuilderResult& OutResult)
-{
-	const bool bInterfaceRemoved = Builder.RemoveInterface(InterfaceName);
-	OutResult = bInterfaceRemoved ? EMetaSoundBuilderResult::Succeeded : EMetaSoundBuilderResult::Failed;
-}
-
-void UMetaSoundBuilderBase::RemoveNode(const FMetaSoundNodeHandle& NodeHandle, EMetaSoundBuilderResult& OutResult, bool bRemoveUnusedDependencies)
-{
-	const bool bNodeRemoved = Builder.RemoveNode(NodeHandle.NodeID);
-	if (bNodeRemoved && bRemoveUnusedDependencies)
-	{
-		Builder.RemoveUnusedDependencies();
-	}
-	OutResult = bNodeRemoved ? EMetaSoundBuilderResult::Succeeded : EMetaSoundBuilderResult::Failed;
-}
-
-void UMetaSoundBuilderBase::RemoveNodeInputDefault(const FMetaSoundBuilderNodeInputHandle& InputHandle, EMetaSoundBuilderResult& OutResult)
-{
-	const bool bInputDefaultRemoved = Builder.RemoveNodeInputDefault(InputHandle.NodeID, InputHandle.VertexID);
-	OutResult = bInputDefaultRemoved ? EMetaSoundBuilderResult::Succeeded : EMetaSoundBuilderResult::Failed;
-}
-
-void UMetaSoundBuilderBase::RemoveUnusedDependencies()
-{
-	Builder.RemoveUnusedDependencies();
-}
-
-void UMetaSoundBuilderBase::RenameRootGraphClass(const FMetasoundFrontendClassName& InName)
-{
-	Builder.RenameRootGraphClass(InName);
-}
-
-#if WITH_EDITOR
-void UMetaSoundBuilderBase::SetAuthor(const FString& InAuthor)
-{
-	Builder.SetAuthor(InAuthor);
-}
-#endif // WITH_EDITOR
 
-void UMetaSoundBuilderBase::SetGraphInputAccessType(FName InputName, EMetasoundFrontendVertexAccessType AccessType, EMetaSoundBuilderResult& OutResult)
-{
-	const bool bSet = Builder.SetGraphInputAccessType(InputName, AccessType);
-	OutResult = bSet ? EMetaSoundBuilderResult::Succeeded : EMetaSoundBuilderResult::Failed;
-}
-
-void UMetaSoundBuilderBase::SetGraphInputDataType(FName InputName, FName DataType, EMetaSoundBuilderResult& OutResult)
-{
-	const bool bSet = Builder.SetGraphInputDataType(InputName, DataType);
-	OutResult = bSet ? EMetaSoundBuilderResult::Succeeded : EMetaSoundBuilderResult::Failed;
-}
-
-void UMetaSoundBuilderBase::SetGraphInputDefault(FName InputName, const FMetasoundFrontendLiteral& Literal, EMetaSoundBuilderResult& OutResult)
-{
-	const bool bSet = Builder.SetGraphInputDefault(InputName, Literal);
-	OutResult = bSet ? EMetaSoundBuilderResult::Succeeded : EMetaSoundBuilderResult::Failed;
-}
-
-void UMetaSoundBuilderBase::SetGraphOutputAccessType(FName OutputName, EMetasoundFrontendVertexAccessType AccessType, EMetaSoundBuilderResult& OutResult)
-{
-	const bool bSet = Builder.SetGraphOutputAccessType(OutputName, AccessType);
-	OutResult = bSet ? EMetaSoundBuilderResult::Succeeded : EMetaSoundBuilderResult::Failed;
-}
-
-void UMetaSoundBuilderBase::SetGraphOutputDataType(FName OutputName, FName DataType, EMetaSoundBuilderResult& OutResult)
-{
-	const bool bSet = Builder.SetGraphOutputDataType(OutputName, DataType);
-	OutResult = bSet ? EMetaSoundBuilderResult::Succeeded : EMetaSoundBuilderResult::Failed;
-}
-
-#if WITH_EDITOR
-void UMetaSoundBuilderBase::SetMemberMetadata(UMetaSoundFrontendMemberMetadata& NewMetadata)
-{
-	Builder.SetMemberMetadata(NewMetadata);
-}
-#endif // WITH_EDITOR
-
-void UMetaSoundBuilderBase::SetNodeInputDefault(const FMetaSoundBuilderNodeInputHandle& InputHandle, const FMetasoundFrontendLiteral& Literal, EMetaSoundBuilderResult& OutResult)
-{
-	const bool bInputDefaultSet = Builder.SetNodeInputDefault(InputHandle.NodeID, InputHandle.VertexID, Literal);
-	OutResult = bInputDefaultSet ? EMetaSoundBuilderResult::Succeeded : EMetaSoundBuilderResult::Failed;
-}
-
-#if WITH_EDITOR
-void UMetaSoundBuilderBase::SetNodeComment(const FMetaSoundNodeHandle& InNodeHandle, const FString& InNewComment, EMetaSoundBuilderResult& OutResult)
-{
-	FString NewComment = InNewComment;
-	const bool bCommentSet = Builder.SetNodeComment(InNodeHandle.NodeID, MoveTemp(NewComment));
-	OutResult = bCommentSet ? EMetaSoundBuilderResult::Succeeded : EMetaSoundBuilderResult::Failed;
-}
-
-void UMetaSoundBuilderBase::SetNodeCommentVisible(const FMetaSoundNodeHandle& InNodeHandle, bool bIsVisible, EMetaSoundBuilderResult& OutResult)
-{
-	const bool bCommentSet = Builder.SetNodeCommentVisible(InNodeHandle.NodeID, bIsVisible);
-	OutResult = bCommentSet ? EMetaSoundBuilderResult::Succeeded : EMetaSoundBuilderResult::Failed;
-}
-
-void UMetaSoundBuilderBase::SetNodeLocation(const FMetaSoundNodeHandle& InNodeHandle, const FVector2D& InLocation, EMetaSoundBuilderResult& OutResult)
-{
-	const bool bLocationSet = Builder.SetNodeLocation(InNodeHandle.NodeID, InLocation);
-	OutResult = bLocationSet ? EMetaSoundBuilderResult::Succeeded : EMetaSoundBuilderResult::Failed;
-}
-
-void UMetaSoundBuilderBase::SetNodeLocation(const FMetaSoundNodeHandle& InNodeHandle, const FVector2D& InLocation, const FGuid& InLocationGuid, EMetaSoundBuilderResult& OutResult)
-{
-	const bool bLocationSet = Builder.SetNodeLocation(InNodeHandle.NodeID, InLocation, &InLocationGuid);
-	OutResult = bLocationSet ? EMetaSoundBuilderResult::Succeeded : EMetaSoundBuilderResult::Failed;
-}
-#endif // WITH_EDITOR
-
-FMetaSoundNodeHandle UMetaSoundBuilderBase::FindNodeInputParent(const FMetaSoundBuilderNodeInputHandle& InputHandle, EMetaSoundBuilderResult& OutResult)
-{
-	if (Builder.ContainsNode(InputHandle.NodeID))
-	{
-		OutResult = EMetaSoundBuilderResult::Succeeded;
-		return FMetaSoundNodeHandle { InputHandle.NodeID };
-	}
-
-	OutResult = EMetaSoundBuilderResult::Failed;
-	return { };
-}
-
-FMetaSoundNodeHandle UMetaSoundBuilderBase::FindNodeOutputParent(const FMetaSoundBuilderNodeOutputHandle& OutputHandle, EMetaSoundBuilderResult& OutResult)
-{
-	if (Builder.ContainsNode(OutputHandle.NodeID))
-	{
-		OutResult = EMetaSoundBuilderResult::Succeeded;
-		return FMetaSoundNodeHandle{ OutputHandle.NodeID };
-	}
-
-	OutResult = EMetaSoundBuilderResult::Failed;
-	return { };
-}
-
-FMetasoundFrontendVersion UMetaSoundBuilderBase::FindNodeClassVersion(const FMetaSoundNodeHandle& NodeHandle, EMetaSoundBuilderResult& OutResult)
-{
-	if (const FMetasoundFrontendNode* Node = Builder.FindNode(NodeHandle.NodeID))
-	{
-		if (const FMetasoundFrontendClass* Class = Builder.FindDependency(Node->ClassID))
+		TUniquePtr<INode> CreateDynamicNodeFromFrontendLiteral(const FName DataType, const FMetasoundFrontendLiteral& InLiteral)
 		{
-			OutResult = EMetaSoundBuilderResult::Succeeded;
-			return FMetasoundFrontendVersion { Class->Metadata.GetClassName().GetFullName(), Class->Metadata.GetVersion() };
-		}
-	}
-
-	OutResult = EMetaSoundBuilderResult::Failed;
-	return FMetasoundFrontendVersion::GetInvalid();
-}
-
-FMetasoundFrontendClassName UMetaSoundBuilderBase::GetRootGraphClassName() const
-{
-	return Builder.GetDocument().RootGraph.Metadata.GetClassName();
-}
-
-void UMetaSoundBuilderBase::GetNodeInputData(const FMetaSoundBuilderNodeInputHandle& InputHandle, FName& Name, FName& DataType, EMetaSoundBuilderResult& OutResult)
-{
-	if (const FMetasoundFrontendVertex* Vertex = Builder.FindNodeInput(InputHandle.NodeID, InputHandle.VertexID))
-	{
-		Name = Vertex->Name;
-		DataType = Vertex->TypeName;
-		OutResult = EMetaSoundBuilderResult::Succeeded;
-	}
-	else
-	{
-		Name = { };
-		DataType = { };
-		OutResult = EMetaSoundBuilderResult::Failed;
-	}
-}
-
-FMetasoundFrontendLiteral UMetaSoundBuilderBase::GetNodeInputDefault(const FMetaSoundBuilderNodeInputHandle& InputHandle, EMetaSoundBuilderResult& OutResult)
-{
-	if (const FMetasoundFrontendLiteral* Default = Builder.GetNodeInputDefault(InputHandle.NodeID, InputHandle.VertexID))
-	{
-		OutResult = EMetaSoundBuilderResult::Succeeded;
-		return *Default;
-	}
-
-	OutResult = EMetaSoundBuilderResult::Failed;
-	return { };
-}
-
-FMetasoundFrontendLiteral UMetaSoundBuilderBase::GetNodeInputClassDefault(const FMetaSoundBuilderNodeInputHandle& InputHandle, EMetaSoundBuilderResult& OutResult)
-{
-	if (const FMetasoundFrontendLiteral* Default = Builder.GetNodeInputClassDefault(InputHandle.NodeID, InputHandle.VertexID))
-	{
-		OutResult = EMetaSoundBuilderResult::Succeeded;
-		return *Default;
-	}
-
-	OutResult = EMetaSoundBuilderResult::Failed;
-	return { };
-}
-
-bool UMetaSoundBuilderBase::GetNodeInputIsConstructorPin(const FMetaSoundBuilderNodeInputHandle& InputHandle) const
-{
-	const EMetasoundFrontendVertexAccessType AccessType = Builder.GetNodeInputAccessType(InputHandle.NodeID, InputHandle.VertexID);
-	return AccessType == EMetasoundFrontendVertexAccessType::Value;
-}
-
-void UMetaSoundBuilderBase::GetNodeOutputData(const FMetaSoundBuilderNodeOutputHandle& OutputHandle, FName& Name, FName& DataType, EMetaSoundBuilderResult& OutResult)
-{
-	if (const FMetasoundFrontendVertex* Vertex = Builder.FindNodeOutput(OutputHandle.NodeID, OutputHandle.VertexID))
-	{
-		Name = Vertex->Name;
-		DataType = Vertex->TypeName;
-		OutResult = EMetaSoundBuilderResult::Succeeded;
-	}
-	else
-	{
-		Name = { };
-		DataType = { };
-		OutResult = EMetaSoundBuilderResult::Failed;
-	}
-}
+			using namespace Frontend;
+			FLiteral ValueLiteral = InLiteral.ToLiteral(DataType);
 
-bool UMetaSoundBuilderBase::GetNodeOutputIsConstructorPin(const FMetaSoundBuilderNodeOutputHandle& OutputHandle) const
-{
-	const EMetasoundFrontendVertexAccessType AccessType = Builder.GetNodeOutputAccessType(OutputHandle.NodeID, OutputHandle.VertexID);
-	return AccessType == EMetasoundFrontendVertexAccessType::Value;
-}
+			// TODO: Node name "Literal" is always the same.  Consolidate and deprecate providing unique node name to avoid unnecessary FName table bloat.
+			FLiteralNodeConstructorParams Params { "Literal", FGuid::NewGuid(), MoveTemp(ValueLiteral) };
+			return IDataTypeRegistry::Get().CreateLiteralNode(DataType, MoveTemp(Params));
+		}
+	} // namespace BuilderSubsystemPrivate
+} // namespace Metasound::Engine
 
-void UMetaSoundBuilderBase::UpdateDependencyClassNames(const TMap<FMetasoundFrontendClassName, FMetasoundFrontendClassName>& OldToNewReferencedClassNames)
-{
-	Builder.UpdateDependencyClassNames(OldToNewReferencedClassNames);
-}
 
-bool UMetaSoundPatchBuilder::ConformObjectToDocument()
+TScriptInterface<IMetaSoundDocumentInterface> UMetaSoundPatchBuilder::Build(UObject* Parent, const FMetaSoundBuilderOptions& InBuilderOptions) const
 {
-	return false;
+	return &BuildInternal<UMetaSoundPatch>(Parent, InBuilderOptions);
 }
 
-void UMetaSoundPatchBuilder::CreateTransientBuilder()
+const UClass& UMetaSoundPatchBuilder::GetBaseMetaSoundUClass() const
 {
-	using namespace Metasound::Frontend;
-
-	UMetaSoundPatch* Patch = NewObject<UMetaSoundPatch>();
-	check(Patch);
-
-	Builder = FMetaSoundFrontendDocumentBuilder(Patch);
-	Builder.InitDocument();
+	return *UMetaSoundPatch::StaticClass();
 }
 
-TScriptInterface<IMetaSoundDocumentInterface> UMetaSoundPatchBuilder::Build(UObject* Parent, const FMetaSoundBuilderOptions& InBuilderOptions) const
+void UMetaSoundPatchBuilder::OnAssetReferenceAdded(TScriptInterface<IMetaSoundDocumentInterface> DocInterface)
 {
-	return &BuildInternal<UMetaSoundPatch>(Parent, InBuilderOptions);
+	check(DocInterface.GetObject());
+	Builder.CastDocumentObjectChecked<UMetaSoundPatch>().ReferencedAssetClassObjects.Add(DocInterface.GetObject());
 }
 
-const UClass& UMetaSoundPatchBuilder::GetBuilderUClass() const
+void UMetaSoundPatchBuilder::OnRemovingAssetReference(TScriptInterface<IMetaSoundDocumentInterface> DocInterface)
 {
-	return *UMetaSoundPatch::StaticClass();
+	check(DocInterface.GetObject());
+	Builder.CastDocumentObjectChecked<UMetaSoundPatch>().ReferencedAssetClassObjects.Remove(DocInterface.GetObject());
 }
 
 UMetaSoundBuilderBase& UMetaSoundBuilderSubsystem::AttachBuilderToAssetChecked(UObject& InObject) const
@@ -1150,12 +89,16 @@ UMetaSoundBuilderBase& UMetaSoundBuilderSubsystem::AttachBuilderToAssetChecked(U
 	const UClass* BaseClass = InObject.GetClass();
 	if (BaseClass == UMetaSoundSource::StaticClass())
 	{
+PRAGMA_DISABLE_DEPRECATION_WARNINGS
 		UMetaSoundSourceBuilder* NewBuilder = AttachSourceBuilderToAsset(CastChecked<UMetaSoundSource>(&InObject));
+PRAGMA_ENABLE_DEPRECATION_WARNINGS
 		return *NewBuilder;
 	}
 	else if (BaseClass == UMetaSoundPatch::StaticClass())
 	{
+PRAGMA_DISABLE_DEPRECATION_WARNINGS
 		UMetaSoundPatchBuilder* NewBuilder = AttachPatchBuilderToAsset(CastChecked<UMetaSoundPatch>(&InObject));
+PRAGMA_ENABLE_DEPRECATION_WARNINGS
 		return *NewBuilder;
 	}
 	else
@@ -1167,21 +110,29 @@ UMetaSoundBuilderBase& UMetaSoundBuilderSubsystem::AttachBuilderToAssetChecked(U
 
 UMetaSoundPatchBuilder* UMetaSoundBuilderSubsystem::AttachPatchBuilderToAsset(UMetaSoundPatch* InPatch) const
 {
+#if WITH_EDITORONLY_DATA
+	using namespace Metasound::Engine;
+
 	if (InPatch)
 	{
-		return &AttachBuilderToAssetCheckedPrivate<UMetaSoundPatchBuilder>(InPatch);
+		return &FDocumentBuilderRegistry::GetChecked().FindOrBeginBuilding<UMetaSoundPatchBuilder>(*InPatch);
 	}
+#endif // WITH_EDITORONLY_DATA
 
 	return nullptr;
 }
 
 UMetaSoundSourceBuilder* UMetaSoundBuilderSubsystem::AttachSourceBuilderToAsset(UMetaSoundSource* InSource) const
 {
+#if WITH_EDITORONLY_DATA
+	using namespace Metasound::Engine;
+
 	if (InSource)
 	{
-		UMetaSoundSourceBuilder& SourceBuilder = AttachBuilderToAssetCheckedPrivate<UMetaSoundSourceBuilder>(InSource);
+		UMetaSoundSourceBuilder& SourceBuilder = FDocumentBuilderRegistry::GetChecked().FindOrBeginBuilding<UMetaSoundSourceBuilder>(*InSource);
 		return &SourceBuilder;
 	}
+#endif // WITH_EDITORONLY_DATA
 
 	return nullptr;
 }
@@ -1205,7 +156,7 @@ void UMetaSoundSourceBuilder::Audition(UObject* Parent, UAudioComponent* AudioCo
 	RegisterGraphIfOutstandingTransactions(MetaSoundSource);
 
 	// Must be called post register as register ensures cached runtime data passed to transactor is up-to-date
-	MetaSoundSource.SetDynamicGeneratorEnabled(MetaSoundSource.GetAssetPathChecked(), bLiveUpdatesEnabled);
+	MetaSoundSource.SetDynamicGeneratorEnabled(bLiveUpdatesEnabled);
 	MetaSoundSource.ConformObjectToDocument();
 
 	AudioComponent->SetSound(&MetaSoundSource);
@@ -1256,26 +207,6 @@ TScriptInterface<IMetaSoundDocumentInterface> UMetaSoundSourceBuilder::Build(UOb
 	return &BuildInternal<UMetaSoundSource>(Parent, InBuilderOptions);
 }
 
-bool UMetaSoundSourceBuilder::ConformObjectToDocument()
-{
-	UMetaSoundSource& Source = GetMetaSoundSource();
-	return Source.ConformObjectToDocument();
-}
-
-void UMetaSoundSourceBuilder::CreateTransientBuilder()
-{
-	using namespace Metasound::Frontend;
-
-	TSharedRef<FDocumentModifyDelegates> DocumentDelegates = MakeShared<FDocumentModifyDelegates>();
-	InitDelegates(*DocumentDelegates);
-
-	UMetaSoundSource* Source = NewObject<UMetaSoundSource>();
-	check(Source);
-
-	Builder = FMetaSoundFrontendDocumentBuilder(Source, DocumentDelegates);
-	Builder.InitDocument();
-}
-
 const Metasound::Engine::FOutputAudioFormatInfoPair* UMetaSoundSourceBuilder::FindOutputAudioFormatInfo() const
 {
 	using namespace Metasound::Engine;
@@ -1284,14 +215,14 @@ const Metasound::Engine::FOutputAudioFormatInfoPair* UMetaSoundSourceBuilder::Fi
 
 	auto Predicate = [this](const FOutputAudioFormatInfoPair& Pair)
 	{
-		const FMetasoundFrontendDocument& Document = Builder.GetDocument();
+		const FMetasoundFrontendDocument& Document = Builder.GetConstDocument();
 		return Document.Interfaces.Contains(Pair.Value.InterfaceVersion);
 	};
 
 	return Algo::FindByPredicate(FormatInfo, Predicate);
 }
 
-const UClass& UMetaSoundSourceBuilder::GetBuilderUClass() const
+const UClass& UMetaSoundSourceBuilder::GetBaseMetaSoundUClass() const
 {
 	return *UMetaSoundSource::StaticClass();
 }
@@ -1313,6 +244,8 @@ UMetaSoundSource& UMetaSoundSourceBuilder::GetMetaSoundSource()
 
 void UMetaSoundSourceBuilder::InitDelegates(Metasound::Frontend::FDocumentModifyDelegates& OutDocumentDelegates)
 {
+	Super::InitDelegates(OutDocumentDelegates);
+
 	OutDocumentDelegates.EdgeDelegates.OnEdgeAdded.AddUObject(this, &UMetaSoundSourceBuilder::OnEdgeAdded);
 	OutDocumentDelegates.EdgeDelegates.OnRemoveSwappingEdge.AddUObject(this, &UMetaSoundSourceBuilder::OnRemoveSwappingEdge);
 
@@ -1327,11 +260,17 @@ void UMetaSoundSourceBuilder::InitDelegates(Metasound::Frontend::FDocumentModify
 	OutDocumentDelegates.NodeDelegates.OnRemovingNodeInputLiteral.AddUObject(this, &UMetaSoundSourceBuilder::OnRemovingNodeInputLiteral);
 }
 
+void UMetaSoundSourceBuilder::OnAssetReferenceAdded(TScriptInterface<IMetaSoundDocumentInterface> DocInterface)
+{
+	check(DocInterface.GetObject());
+	GetMetaSoundSource().ReferencedAssetClassObjects.Add(DocInterface.GetObject());
+}
+
 void UMetaSoundSourceBuilder::OnEdgeAdded(int32 EdgeIndex) const
 {
 	using namespace Metasound::DynamicGraph;
 
-	const FMetasoundFrontendDocument& Doc = Builder.GetDocument();
+	const FMetasoundFrontendDocument& Doc = Builder.GetConstDocument();
 	const FMetasoundFrontendEdge& NewEdge = Doc.RootGraph.Graph.Edges[EdgeIndex];
 	ExecuteAuditionableTransaction([this, &NewEdge](Metasound::DynamicGraph::FDynamicOperatorTransactor& Transactor)
 	{
@@ -1369,8 +308,7 @@ void UMetaSoundSourceBuilder::OnInputAdded(int32 InputIndex)
 		using namespace Metasound;
 		using namespace Metasound::Frontend;
 
-		const FMetaSoundFrontendDocumentBuilder& ConstBuilder = Builder;
-		const FMetasoundFrontendDocument& Doc = ConstBuilder.GetDocument();
+		const FMetasoundFrontendDocument& Doc = Builder.GetConstDocument();
 		const FMetasoundFrontendGraphClass& GraphClass = Doc.RootGraph;
 		const FMetasoundFrontendClassInput& NewInput = GraphClass.Interface.Inputs[InputIndex];
 
@@ -1407,7 +345,7 @@ void UMetaSoundSourceBuilder::OnNodeAdded(int32 NodeIndex) const
 		using namespace Metasound;
 		using namespace Metasound::Frontend;
 
-		const FMetasoundFrontendDocument& Doc = Builder.GetDocument();
+		const FMetasoundFrontendDocument& Doc = Builder.GetConstDocument();
 		const FMetasoundFrontendGraphClass& GraphClass = Doc.RootGraph;
 		const FMetasoundFrontendNode& AddedNode = GraphClass.Graph.Nodes[NodeIndex];
 
@@ -1517,7 +455,7 @@ void UMetaSoundSourceBuilder::OnNodeInputLiteralSet(int32 NodeIndex, int32 Verte
 {
 	using namespace Metasound::DynamicGraph;
 
-	const FMetasoundFrontendGraphClass& GraphClass = Builder.GetDocument().RootGraph;
+	const FMetasoundFrontendGraphClass& GraphClass = Builder.GetConstDocument().RootGraph;
 	const FMetasoundFrontendNode& Node = GraphClass.Graph.Nodes[NodeIndex];
 	const FMetasoundFrontendVertex& Input = Node.Interface.Inputs[VertexIndex];
 
@@ -1546,7 +484,7 @@ void UMetaSoundSourceBuilder::OnOutputAdded(int32 OutputIndex) const
 	{
 		using namespace Metasound::Frontend;
 
-		const FMetasoundFrontendDocument& Doc = Builder.GetDocument();
+		const FMetasoundFrontendDocument& Doc = Builder.GetConstDocument();
 		const FMetasoundFrontendGraphClass& GraphClass = Doc.RootGraph;
 		const FMetasoundFrontendClassOutput& NewOutput = GraphClass.Interface.Outputs[OutputIndex];
 
@@ -1559,7 +497,7 @@ void UMetaSoundSourceBuilder::OnRemoveSwappingEdge(int32 SwapIndex, int32 LastIn
 {
 	using namespace Metasound::DynamicGraph;
 
-	const FMetasoundFrontendDocument& Doc = Builder.GetDocument();
+	const FMetasoundFrontendDocument& Doc = Builder.GetConstDocument();
 	const FMetasoundFrontendEdge& EdgeBeingRemoved = Doc.RootGraph.Graph.Edges[SwapIndex];
 	ExecuteAuditionableTransaction([this, EdgeBeingRemoved](FDynamicOperatorTransactor& Transactor)
 	{
@@ -1589,6 +527,12 @@ void UMetaSoundSourceBuilder::OnRemoveSwappingEdge(int32 SwapIndex, int32 LastIn
 	});
 }
 
+void UMetaSoundSourceBuilder::OnRemovingAssetReference(TScriptInterface<IMetaSoundDocumentInterface> DocInterface)
+{
+	check(DocInterface.GetObject());
+	GetMetaSoundSource().ReferencedAssetClassObjects.Remove(DocInterface.GetObject());
+}
+
 void UMetaSoundSourceBuilder::OnRemovingInput(int32 InputIndex)
 {
 	using namespace Metasound::DynamicGraph;
@@ -1598,8 +542,7 @@ void UMetaSoundSourceBuilder::OnRemovingInput(int32 InputIndex)
 		using namespace Metasound;
 		using namespace Metasound::Frontend;
 
-		const FMetaSoundFrontendDocumentBuilder& ConstBuilder = Builder;
-		const FMetasoundFrontendDocument& Doc = ConstBuilder.GetDocument();
+		const FMetasoundFrontendDocument& Doc = Builder.GetConstDocument();
 		const FMetasoundFrontendGraphClass& GraphClass = Doc.RootGraph;
 		const FMetasoundFrontendClassInput& InputBeingRemoved = GraphClass.Interface.Inputs[InputIndex];
 
@@ -1637,7 +580,7 @@ void UMetaSoundSourceBuilder::OnRemoveSwappingNode(int32 SwapIndex, int32 LastIn
 	{
 		using namespace Metasound::Frontend;
 
-		const FMetasoundFrontendDocument& Doc = Builder.GetDocument();
+		const FMetasoundFrontendDocument& Doc = Builder.GetConstDocument();
 		const FMetasoundFrontendGraphClass& GraphClass = Doc.RootGraph;
 		const FMetasoundFrontendNode& NodeBeingRemoved = GraphClass.Graph.Nodes[SwapIndex];
 		const FGuid& NodeID = NodeBeingRemoved.GetID();
@@ -1650,7 +593,7 @@ void UMetaSoundSourceBuilder::OnRemovingNodeInputLiteral(int32 NodeIndex, int32
 {
 	using namespace Metasound::DynamicGraph;
 
-	const FMetasoundFrontendGraphClass& GraphClass = Builder.GetDocument().RootGraph;
+	const FMetasoundFrontendGraphClass& GraphClass = Builder.GetConstDocument().RootGraph;
 	const FMetasoundFrontendNode& Node = GraphClass.Graph.Nodes[NodeIndex];
 	const FMetasoundFrontendVertex& Input = Node.Interface.Inputs[VertexIndex];
 
@@ -1663,7 +606,7 @@ void UMetaSoundSourceBuilder::OnRemovingNodeInputLiteral(int32 NodeIndex, int32
 			using namespace Metasound::Engine;
 			using namespace Metasound::Frontend;
 
-			const TArray<FMetasoundFrontendNode>& Nodes = Builder.GetDocument().RootGraph.Graph.Nodes;
+			const TArray<FMetasoundFrontendNode>& Nodes = Builder.GetConstDocument().RootGraph.Graph.Nodes;
 			const FMetasoundFrontendNode& Node = Nodes[NodeIndex];
 			const FMetasoundFrontendVertex& Input = Node.Interface.Inputs[VertexIndex];
 
@@ -1688,7 +631,7 @@ void UMetaSoundSourceBuilder::OnRemovingOutput(int32 OutputIndex) const
 	{
 		using namespace Metasound::Frontend;
 
-		const FMetasoundFrontendDocument& Doc = Builder.GetDocument();
+		const FMetasoundFrontendDocument& Doc = Builder.GetConstDocument();
 		const FMetasoundFrontendGraphClass& GraphClass = Doc.RootGraph;
 		const FMetasoundFrontendClassOutput& OutputBeingRemoved = GraphClass.Interface.Outputs[OutputIndex];
 
@@ -1724,7 +667,7 @@ void UMetaSoundSourceBuilder::SetFormat(EMetaSoundOutputAudioFormat OutputFormat
 
 	TArray<FMetasoundFrontendVersion> OutputFormatsToRemove;
 
-	const FMetasoundFrontendDocument& Document = GetConstBuilder().GetDocument();
+	const FMetasoundFrontendDocument& Document = GetConstBuilder().GetConstDocument();
 	for (const FOutputAudioFormatInfoPair& Pair : FormatMap)
 	{
 		const FMetasoundFrontendVersion& FormatVersion = Pair.Value.InterfaceVersion;
@@ -1775,8 +718,11 @@ void UMetaSoundSourceBuilder::SetSampleRateOverride(int32 SampleRate)
 
 UMetaSoundPatchBuilder* UMetaSoundBuilderSubsystem::CreatePatchBuilder(FName BuilderName, EMetaSoundBuilderResult& OutResult)
 {
+	using namespace Metasound::Engine;
+
 	OutResult = EMetaSoundBuilderResult::Succeeded;
-	return &CreateTransientBuilder<UMetaSoundPatchBuilder>(BuilderName);
+	UMetaSoundPatchBuilder& NewBuilder = FDocumentBuilderRegistry::GetChecked().CreateTransientBuilder<UMetaSoundPatchBuilder>(BuilderName);
+	return &NewBuilder;
 }
 
 UMetaSoundSourceBuilder* UMetaSoundBuilderSubsystem::CreateSourceBuilder(
@@ -1788,14 +734,15 @@ UMetaSoundSourceBuilder* UMetaSoundBuilderSubsystem::CreateSourceBuilder(
 	EMetaSoundOutputAudioFormat OutputFormat,
 	bool bIsOneShot)
 {
-	using namespace Metasound::Frontend;
+	using namespace Metasound::Engine;
 	using namespace Metasound::Engine::BuilderSubsystemPrivate;
+	using namespace Metasound::Frontend;
 
 	OnPlayNodeOutput = { };
 	OnFinishedNodeInput = { };
 	AudioOutNodeInputs.Reset();
 
-	UMetaSoundSourceBuilder& NewBuilder = CreateTransientBuilder<UMetaSoundSourceBuilder>(BuilderName);
+	UMetaSoundSourceBuilder& NewBuilder = FDocumentBuilderRegistry::GetChecked().CreateTransientBuilder<UMetaSoundSourceBuilder>(BuilderName);
 	OutResult = EMetaSoundBuilderResult::Succeeded;
 	if (OutputFormat != EMetaSoundOutputAudioFormat::Mono)
 	{
@@ -1888,9 +835,11 @@ UMetaSoundSourceBuilder* UMetaSoundBuilderSubsystem::CreateSourceBuilder(
 
 UMetaSoundPatchBuilder* UMetaSoundBuilderSubsystem::CreatePatchPresetBuilder(FName BuilderName, const TScriptInterface<IMetaSoundDocumentInterface>& ReferencedNodeClass, EMetaSoundBuilderResult& OutResult)
 {
+	using namespace Metasound::Engine;
+
 	if (ReferencedNodeClass)
 	{
-		UMetaSoundPatchBuilder& Builder = CreateTransientBuilder<UMetaSoundPatchBuilder>(BuilderName);
+		UMetaSoundPatchBuilder& Builder = FDocumentBuilderRegistry::GetChecked().CreateTransientBuilder<UMetaSoundPatchBuilder>(BuilderName);
 		Builder.ConvertToPreset(ReferencedNodeClass, OutResult);
 		return &Builder;
 	}
@@ -1919,9 +868,11 @@ UMetaSoundBuilderBase& UMetaSoundBuilderSubsystem::CreatePresetBuilder(FName Bui
 
 UMetaSoundSourceBuilder* UMetaSoundBuilderSubsystem::CreateSourcePresetBuilder(FName BuilderName, const TScriptInterface<IMetaSoundDocumentInterface>& ReferencedNodeClass, EMetaSoundBuilderResult& OutResult)
 {
+	using namespace Metasound::Engine;
+
 	if (ReferencedNodeClass)
 	{
-		UMetaSoundSourceBuilder& Builder = CreateTransientBuilder<UMetaSoundSourceBuilder>();
+		UMetaSoundSourceBuilder& Builder = FDocumentBuilderRegistry::GetChecked().CreateTransientBuilder<UMetaSoundSourceBuilder>();
 		Builder.ConvertToPreset(ReferencedNodeClass, OutResult);
 		return &Builder;
 	}
@@ -2032,42 +983,9 @@ FMetasoundFrontendLiteral UMetaSoundBuilderSubsystem::CreateMetaSoundLiteralFrom
 }
 
 bool UMetaSoundBuilderSubsystem::DetachBuilderFromAsset(const FMetasoundFrontendClassName& InClassName) const
-{
-	using namespace Metasound;
-	using namespace Metasound::Frontend;
-
-	TWeakObjectPtr<UMetaSoundBuilderBase> Builder = AssetBuilders.FindRef(InClassName);
-	if (Builder.IsValid())
-	{
-		// If the builder has applied transactions to its document object that are not mirrored in the frontend registry,
-		// unregister version in registry. This will ensure that future requests for the builder's associated asset will
-		// register a fresh version from the object as the transaction history is intrinsically lost once this builder
-		// is destroyed.
-		if (Builder->LastTransactionRegistered != Builder->Builder.GetTransactionCount())
-		{
-			UObject& MetaSound = Builder->Builder.CastDocumentObjectChecked<UObject>();
-			if (FMetasoundAssetBase* MetaSoundAsset = IMetasoundUObjectRegistry::Get().GetObjectAsAssetBase(&MetaSound))
-			{
-				MetaSoundAsset->UnregisterGraphWithFrontend();
-			}
-		}
-
-		ensureAlways(AssetBuilders.Remove(InClassName));
-		return true;
-	}
-
-	return false;
-}
-
-void UMetaSoundBuilderSubsystem::InvalidateDocumentCache(const FMetasoundFrontendClassName& InClassName) const
 {
 	using namespace Metasound::Frontend;
-
-	TWeakObjectPtr<UMetaSoundBuilderBase> BuilderPtr = AssetBuilders.FindRef(InClassName);
-	if (BuilderPtr.IsValid())
-	{
-		BuilderPtr->Reload();
-	}
+	return IDocumentBuilderRegistry::GetChecked().FinishBuilding(InClassName);
 }
 
 UMetaSoundBuilderBase* UMetaSoundBuilderSubsystem::FindBuilder(FName BuilderName)
@@ -2077,29 +995,17 @@ UMetaSoundBuilderBase* UMetaSoundBuilderSubsystem::FindBuilder(FName BuilderName
 
 UMetaSoundBuilderBase* UMetaSoundBuilderSubsystem::FindBuilderOfDocument(TScriptInterface<const IMetaSoundDocumentInterface> InMetaSound) const
 {
-	TWeakObjectPtr<UMetaSoundBuilderBase> Builder;
+	using namespace Metasound::Engine;
+
 	if (const UObject* MetaSoundObject = InMetaSound.GetObject())
 	{
 		const FMetasoundFrontendDocument& Document = InMetaSound->GetConstDocument();
 		const FMetasoundFrontendClassName& InClassName = Document.RootGraph.Metadata.GetClassName();
 
-		Builder = AssetBuilders.FindRef(InClassName);
-		if (Builder.IsValid())
-		{
-			ensureAlwaysMsgf(MetaSoundObject->IsAsset(), TEXT("MetaSound is asset but Builder was registered with Subsystem as transient"));
-		}
-		else
-		{
-			Builder = TransientBuilders.FindRef(InClassName);
-			if (Builder.IsValid())
-			{
-				ensureAlwaysMsgf(!MetaSoundObject->IsAsset(), TEXT("MetaSound is transient object but Builder was registered with Subsystem as asset"));
-			}
-		}
-
+		return FDocumentBuilderRegistry::GetChecked().FindBuilderObject(InClassName);
 	}
 
-	return Builder.Get();
+	return nullptr;
 }
 
 UMetaSoundPatchBuilder* UMetaSoundBuilderSubsystem::FindPatchBuilder(FName BuilderName)
@@ -2122,17 +1028,13 @@ UMetaSoundSourceBuilder* UMetaSoundBuilderSubsystem::FindSourceBuilder(FName Bui
 	return nullptr;
 }
 
-void UMetaSoundBuilderSubsystem::Initialize(FSubsystemCollectionBase& Collection)
+void UMetaSoundBuilderSubsystem::InvalidateDocumentCache(const FMetasoundFrontendClassName& InClassName) const
 {
-	using namespace Metasound::Frontend;
-
-	IDocumentBuilderRegistry::Set([]() -> IDocumentBuilderRegistry&
+	using namespace Metasound::Engine;
+	if (FMetaSoundFrontendDocumentBuilder* Builder = FDocumentBuilderRegistry::GetChecked().FindBuilder(InClassName))
 	{
-		check(GEngine);
-		UMetaSoundBuilderSubsystem* BuilderSubsystem = GEngine->GetEngineSubsystem<UMetaSoundBuilderSubsystem>();
-		check(BuilderSubsystem);
-		return static_cast<IDocumentBuilderRegistry&>(*BuilderSubsystem);
-	});
+		Builder->Reload();
+	}
 }
 
 bool UMetaSoundBuilderSubsystem::IsInterfaceRegistered(FName InInterfaceName) const
@@ -2143,17 +1045,6 @@ bool UMetaSoundBuilderSubsystem::IsInterfaceRegistered(FName InInterfaceName) co
 	return ISearchEngine::Get().FindInterfaceWithHighestVersion(InInterfaceName, Interface);
 }
 
-#if WITH_EDITOR
-void UMetaSoundBuilderSubsystem::PostBuilderAssetTransaction(const FMetasoundFrontendClassName& InClassName)
-{
-	TWeakObjectPtr<UMetaSoundBuilderBase> BuilderPtr = AssetBuilders.FindRef(InClassName);
-	if (UMetaSoundBuilderBase* Builder = BuilderPtr.Get())
-	{
-		Builder->Reload();
-	}
-}
-#endif // WITH_EDITOR
-
 void UMetaSoundBuilderSubsystem::RegisterBuilder(FName BuilderName, UMetaSoundBuilderBase* Builder)
 {
 	if (Builder)
diff --git a/Engine/Plugins/Runtime/Metasound/Source/MetasoundEngine/Private/MetasoundDocumentBuilderRegistry.cpp b/Engine/Plugins/Runtime/Metasound/Source/MetasoundEngine/Private/MetasoundDocumentBuilderRegistry.cpp
new file mode 100644
--- /dev/null
+++ b/Engine/Plugins/Runtime/Metasound/Source/MetasoundEngine/Private/MetasoundDocumentBuilderRegistry.cpp
@@ -0,0 +1,111 @@
+// Copyright Epic Games, Inc. All Rights Reserved.
+
+#include "MetasoundDocumentBuilderRegistry.h"
+
+#include "MetasoundTrace.h"
+#include "MetasoundUObjectRegistry.h"
+
+
+namespace Metasound::Engine
+{
+	FDocumentBuilderRegistry::~FDocumentBuilderRegistry()
+	{
+		for(const TPair<FMetasoundFrontendClassName, TWeakObjectPtr<UMetaSoundBuilderBase>>& Pair : Builders)
+		{
+			TWeakObjectPtr<UMetaSoundBuilderBase> Builder = Pair.Value;
+			if (Builder.IsValid())
+			{
+				// If the builder has applied transactions to its document object that are not mirrored in the frontend registry,
+				// unregister version in registry. This will ensure that future requests for the builder's associated asset will
+				// register a fresh version from the object as the transaction history is intrinsically lost once this builder
+				// is destroyed.
+				if (Builder->GetLastTransactionRegistered() != Builder->GetConstBuilder().GetTransactionCount())
+				{
+					UObject& MetaSound = Builder->GetConstBuilder().CastDocumentObjectChecked<UObject>();
+					if (FMetasoundAssetBase* MetaSoundAsset = IMetasoundUObjectRegistry::Get().GetObjectAsAssetBase(&MetaSound))
+					{
+						MetaSoundAsset->UnregisterGraphWithFrontend();
+					}
+				}
+			}
+		}
+		Builders.Reset();
+	}
+
+#if WITH_EDITORONLY_DATA
+	FMetaSoundFrontendDocumentBuilder& FDocumentBuilderRegistry::FindOrBeginBuilding(TScriptInterface<IMetaSoundDocumentInterface> MetaSound)
+	{
+		UObject* Object = MetaSound.GetObject();
+		check(Object);
+
+		return FindOrBeginBuilding(*Object).GetBuilder();
+	}
+#endif // WITH_EDITORONLY_DATA
+
+	FMetaSoundFrontendDocumentBuilder* FDocumentBuilderRegistry::FindBuilder(TScriptInterface<IMetaSoundDocumentInterface> MetaSound) const
+	{
+		if (UObject* Object = MetaSound.GetObject())
+		{
+			return FindBuilder(MetaSound->GetConstDocument().RootGraph.Metadata.GetClassName());
+		}
+
+		return nullptr;
+	}
+
+	FMetaSoundFrontendDocumentBuilder* FDocumentBuilderRegistry::FindBuilder(const FMetasoundFrontendClassName& InClassName) const
+	{
+		using namespace Metasound::Engine;
+		if (UMetaSoundBuilderBase* Builder = Builders.FindRef(InClassName).Get())
+		{
+			return &Builder->GetBuilder();
+		}
+
+		return nullptr;
+	}
+
+	UMetaSoundBuilderBase* FDocumentBuilderRegistry::FindBuilderObject(TScriptInterface<const IMetaSoundDocumentInterface> MetaSound) const
+	{
+		if (const UObject* MetaSoundObject = MetaSound.GetObject())
+		{
+			const FMetasoundFrontendDocument& Document = MetaSound->GetConstDocument();
+			const FMetasoundFrontendClassName& ClassName = Document.RootGraph.Metadata.GetClassName();
+			return Builders.FindRef(ClassName).Get();
+		}
+
+		return nullptr;
+	}
+
+	UMetaSoundBuilderBase* FDocumentBuilderRegistry::FindBuilderObject(const FMetasoundFrontendClassName& ClassName) const
+	{
+		return Builders.FindRef(ClassName).Get();
+	}
+
+	bool FDocumentBuilderRegistry::FinishBuilding(const FMetasoundFrontendClassName& InClassName) const
+	{
+		using namespace Metasound;
+		using namespace Metasound::Engine;
+
+		TWeakObjectPtr<UMetaSoundBuilderBase> Builder = Builders.FindRef(InClassName);
+		if (Builder.IsValid())
+		{
+			// If the builder has applied transactions to its document object that are not mirrored in the frontend registry,
+			// unregister version in registry. This will ensure that future requests for the builder's associated asset will
+			// register a fresh version from the object as the transaction history is intrinsically lost once this builder
+			// is destroyed.
+			if (Builder->GetLastTransactionRegistered() != Builder->GetConstBuilder().GetTransactionCount())
+			{
+				UObject& MetaSound = Builder->GetConstBuilder().CastDocumentObjectChecked<UObject>();
+				if (FMetasoundAssetBase* MetaSoundAsset = IMetasoundUObjectRegistry::Get().GetObjectAsAssetBase(&MetaSound))
+				{
+					MetaSoundAsset->UnregisterGraphWithFrontend();
+				}
+			}
+
+			Builder->GetBuilder().FinishBuilding();
+			ensureAlways(Builders.Remove(InClassName));
+			return true;
+		}
+
+		return false;
+	}
+} // namespace Metasound::Engine
\ No newline at end of file
diff --git a/Engine/Plugins/Runtime/Metasound/Source/MetasoundEngine/Private/MetasoundEngineAsset.h b/Engine/Plugins/Runtime/Metasound/Source/MetasoundEngine/Private/MetasoundEngineAsset.h
--- a/Engine/Plugins/Runtime/Metasound/Source/MetasoundEngine/Private/MetasoundEngineAsset.h
+++ b/Engine/Plugins/Runtime/Metasound/Source/MetasoundEngine/Private/MetasoundEngineAsset.h
@@ -26,14 +26,32 @@ namespace Metasound
 	struct FMetaSoundEngineAssetHelper
 	{
 #if WITH_EDITOR
+		static void PostDuplicate(TScriptInterface<IMetaSoundDocumentInterface> MetaSound, EDuplicateMode::Type InDuplicateMode, FGuid& OutAssetClassID)
+		{
+			using namespace Metasound::Frontend;
+
+			if (InDuplicateMode == EDuplicateMode::Normal)
+			{
+				// Uses a bespoke builder instead of a registered one to avoid class name collision
+				// that can result in accessing the wrong pre-existing builder for the pre-existing
+				// MetaSound the provided MetaSound was duplicated from. Being that this is a new
+				// MetaSound, it is perfectly safe and should not result in multiple builders existing
+				// for the given asset.
+				FMetaSoundFrontendDocumentBuilder DuplicateBuilder(MetaSound);
+				const FMetasoundFrontendClassName NewName = DuplicateBuilder.GenerateNewClassName();
+				ensureAlwaysMsgf(IMetaSoundAssetManager::GetChecked().TryGetAssetIDFromClassName(NewName, OutAssetClassID), TEXT("Failed to retrieve newly duplicated MetaSoundClassName AssetID"));
+			}
+		}
+
 		template <typename TMetaSoundObject>
 		static void PostEditUndo(TMetaSoundObject& InMetaSound)
 		{
 			InMetaSound.GetModifyContext().SetForceRefreshViews();
 
-			const FMetasoundFrontendDocument& Document = InMetaSound.GetConstDocumentChecked();
-			const FMetasoundFrontendClassName& ClassName = Document.RootGraph.Metadata.GetClassName();
-			UMetaSoundBuilderSubsystem::GetChecked().PostBuilderAssetTransaction(ClassName);
+			if (UMetaSoundBuilderBase* Builder = UMetaSoundBuilderSubsystem::GetChecked().FindBuilderOfDocument(&InMetaSound))
+			{
+				Builder->Reload();
+			}
 
 			if (UMetasoundEditorGraphBase* Graph = Cast<UMetasoundEditorGraphBase>(InMetaSound.GetGraph()))
 			{
@@ -100,8 +118,6 @@ namespace Metasound
 #if WITH_EDITORONLY_DATA
 			using namespace Frontend;
 
-			// Do not call asset manager on CDO objects which may be loaded before asset
-			// manager is set.
 			if (IMetaSoundAssetManager* AssetManager = IMetaSoundAssetManager::Get())
 			{
 				AssetManager->WaitUntilAsyncLoadReferencedAssetsComplete(InMetaSound);
@@ -155,14 +171,7 @@ namespace Metasound
 			{
 				if (InMetaSound.GetAsyncReferencedAssetClassPaths().Num() > 0)
 				{
-					if (IMetaSoundAssetManager* AssetManager = IMetaSoundAssetManager::Get())
-					{
-						AssetManager->RequestAsyncLoadReferencedAssets(InMetaSound);
-					}
-					else
-					{
-						UE_LOG(LogMetaSound, Warning, TEXT("Request for to load async references ignored from asset %s. Likely due loading before the MetaSoundEngine module."), *InMetaSound.GetPathName());
-					}
+					IMetaSoundAssetManager::GetChecked().RequestAsyncLoadReferencedAssets(InMetaSound);
 				}
 			}
 		}
diff --git a/Engine/Plugins/Runtime/Metasound/Source/MetasoundEngine/Private/MetasoundSource.cpp b/Engine/Plugins/Runtime/Metasound/Source/MetasoundEngine/Private/MetasoundSource.cpp
--- a/Engine/Plugins/Runtime/Metasound/Source/MetasoundEngine/Private/MetasoundSource.cpp
+++ b/Engine/Plugins/Runtime/Metasound/Source/MetasoundEngine/Private/MetasoundSource.cpp
@@ -16,6 +16,7 @@
 #include "MetasoundAssetManager.h"
 #include "MetasoundAudioFormats.h"
 #include "MetasoundBuilderSubsystem.h"
+#include "MetasoundDocumentBuilderRegistry.h"
 #include "MetasoundDocumentInterface.h"
 #include "MetasoundDynamicOperatorTransactor.h"
 #include "MetasoundEngineAsset.h"
@@ -343,6 +344,11 @@ const UClass& UMetaSoundSource::GetBaseMetaSoundUClass() const
 	return *UMetaSoundSource::StaticClass();
 }
 
+const UClass& UMetaSoundSource::GetBuilderUClass() const
+{
+	return *UMetaSoundSourceBuilder::StaticClass();
+}
+
 const FMetasoundFrontendDocument& UMetaSoundSource::GetConstDocument() const
 {
 	return RootMetasoundDocument;
@@ -358,14 +364,7 @@ void UMetaSoundSource::PostEditUndo()
 void UMetaSoundSource::PostDuplicate(EDuplicateMode::Type InDuplicateMode)
 {
 	Super::PostDuplicate(InDuplicateMode);
-
-	// Guid is reset as asset may share implementation from
-	// asset duplicated from but should not be registered as such.
-	if (InDuplicateMode == EDuplicateMode::Normal)
-	{
-		AssetClassID = FGuid::NewGuid();
-		Metasound::Frontend::FRenameRootGraphClass::Generate(GetDocumentHandle(), AssetClassID);
-	}
+	Metasound::FMetaSoundEngineAssetHelper::PostDuplicate(this, InDuplicateMode, AssetClassID);
 }
 
 void UMetaSoundSource::PostEditChangeProperty(FPropertyChangedEvent& InEvent)
@@ -421,16 +420,11 @@ bool UMetaSoundSource::CanEditChange(const FProperty* InProperty) const
 
 void UMetaSoundSource::PostEditChangeOutputFormat()
 {
-	using namespace Metasound::Frontend;
+	using namespace Metasound::Engine;
 
 	EMetaSoundBuilderResult Result = EMetaSoundBuilderResult::Failed;
-	{
-		const UMetaSoundBuilderSubsystem& BuilderSubsystem = UMetaSoundBuilderSubsystem::GetConstChecked();
-
-		UMetaSoundSourceBuilder* SourceBuilder = BuilderSubsystem.AttachSourceBuilderToAsset(this);
-		check(SourceBuilder);
-		SourceBuilder->SetFormat(OutputFormat, Result);
-	}
+	UMetaSoundSourceBuilder& SourceBuilder = FDocumentBuilderRegistry::GetChecked().FindOrBeginBuilding<UMetaSoundSourceBuilder>(*this);
+	SourceBuilder.SetFormat(OutputFormat, Result);
 
 	if (Result == EMetaSoundBuilderResult::Succeeded)
 	{
@@ -499,10 +493,6 @@ FTopLevelAssetPath UMetaSoundSource::GetAssetPathChecked() const
 
 void UMetaSoundSource::BeginDestroy()
 {
-#if WITH_SERVER_CODE
-	Metasound::Frontend::IMetaSoundAssetManager::OnManagerSet.RemoveAll(this);
-#endif //WITH_SERVER_CODE
-
 	OnNotifyBeginDestroy();
 	Super::BeginDestroy();
 }
@@ -646,23 +636,12 @@ void UMetaSoundSource::MigrateEditorGraph(FMetaSoundFrontendDocumentBuilder& Out
 	PRAGMA_ENABLE_DEPRECATION_WARNINGS
 }
 
-UEdGraph* UMetaSoundSource::GetGraph()
+UEdGraph* UMetaSoundSource::GetGraph() const
 {
 	return EditorGraph;
 }
 
-const UEdGraph* UMetaSoundSource::GetGraph() const
-{
-	return EditorGraph;
-}
-
-UEdGraph& UMetaSoundSource::GetGraphChecked()
-{
-	check(EditorGraph);
-	return *EditorGraph;
-}
-
-const UEdGraph& UMetaSoundSource::GetGraphChecked() const
+UEdGraph& UMetaSoundSource::GetGraphChecked() const
 {
 	check(EditorGraph);
 	return *EditorGraph;
@@ -828,20 +807,7 @@ void UMetaSoundSource::InitResources()
 
 	if (IsInGameThread())
 	{
-#if WITH_SERVER_CODE
-		
-		if (IMetaSoundAssetManager::Get())
-		{
-			RegisterGraphWithFrontend(GetInitRegistrationOptions());
-		}
-		else
-		{
-			//to allow the server to preload UMetaSoundSource objects before the manager has been set up
-			IMetaSoundAssetManager::OnManagerSet.AddUObject(this, &UMetaSoundSource::InitResources);
-		}
-#else
 		RegisterGraphWithFrontend(GetInitRegistrationOptions());
-#endif //WITH_SERVER_CODE
 	}
 	else
 	{
@@ -897,12 +863,13 @@ Metasound::Frontend::FDocumentAccessPtr UMetaSoundSource::GetDocumentAccessPtr()
 	using namespace Metasound::Frontend;
 
 	// Mutation of a document via the soft deprecated access ptr/controller system is not tracked by
-	// the builder registry, so the document cache is invalidated here. It is discouraged to mutate
-	// documents using both systems at the same time as it can corrupt a builder document's cache.
+	// the builder registry, so the document cache is invalidated here.
 	if (UMetaSoundBuilderSubsystem* BuilderSubsystem = UMetaSoundBuilderSubsystem::Get())
 	{
-		const FMetasoundFrontendClassName& Name = RootMetasoundDocument.RootGraph.Metadata.GetClassName();
-		BuilderSubsystem->InvalidateDocumentCache(Name);
+		if (UMetaSoundBuilderBase* Builder = BuilderSubsystem->FindBuilderOfDocument(this))
+		{
+			Builder->Reload();
+		}
 	}
 
 	// Return document using FAccessPoint to inform the TAccessPtr when the 
@@ -1633,7 +1600,7 @@ void UMetaSoundSource::OnFinishActiveBuilder()
 	bIsBuilderActive = false;
 }
 
-TSharedPtr<Metasound::DynamicGraph::FDynamicOperatorTransactor> UMetaSoundSource::SetDynamicGeneratorEnabled(const FTopLevelAssetPath& InAssetPath, bool bInIsEnabled)
+TSharedPtr<Metasound::DynamicGraph::FDynamicOperatorTransactor> UMetaSoundSource::SetDynamicGeneratorEnabled(bool bInIsEnabled)
 {
 	using namespace Metasound;
 	using namespace Metasound::DynamicGraph;
diff --git a/Engine/Plugins/Runtime/Metasound/Source/MetasoundEngine/Public/Metasound.h b/Engine/Plugins/Runtime/Metasound/Source/MetasoundEngine/Public/Metasound.h
--- a/Engine/Plugins/Runtime/Metasound/Source/MetasoundEngine/Public/Metasound.h
+++ b/Engine/Plugins/Runtime/Metasound/Source/MetasoundEngine/Public/Metasound.h
@@ -70,6 +70,8 @@ class METASOUNDENGINE_API UMetaSoundPatch : public UObject, public FMetasoundAss
 	GENERATED_BODY()
 
 	friend struct Metasound::FMetaSoundEngineAssetHelper;
+	friend class UMetaSoundPatchBuilder;
+
 protected:
 	UPROPERTY(EditAnywhere, Category = CustomView)
 	FMetasoundFrontendDocument RootMetaSoundDocument;
@@ -129,10 +131,8 @@ public:
 	// Returns the graph associated with this Metasound. Graph is required to be referenced on
 	// Metasound UObject for editor serialization purposes.
 	// @return Editor graph associated with UMetaSoundSource.
-	virtual UEdGraph* GetGraph() override;
-	virtual const UEdGraph* GetGraph() const override;
-	virtual UEdGraph& GetGraphChecked() override;
-	virtual const UEdGraph& GetGraphChecked() const override;
+	virtual UEdGraph* GetGraph() const override;
+	virtual UEdGraph& GetGraphChecked() const override;
 	virtual void MigrateEditorGraph(FMetaSoundFrontendDocumentBuilder& OutBuilder) override;
 
 	// Sets the graph associated with this Metasound. Graph is required to be referenced on
@@ -146,6 +146,7 @@ public:
 
 	virtual FTopLevelAssetPath GetAssetPathChecked() const override;
 	virtual const UClass& GetBaseMetaSoundUClass() const final override;
+	virtual const UClass& GetBuilderUClass() const final override;
 	virtual const FMetasoundFrontendDocument& GetConstDocument() const override;
 
 #if WITH_EDITOR
diff --git a/Engine/Plugins/Runtime/Metasound/Source/MetasoundEngine/Public/MetasoundAssetSubsystem.h b/Engine/Plugins/Runtime/Metasound/Source/MetasoundEngine/Public/MetasoundAssetSubsystem.h
--- a/Engine/Plugins/Runtime/Metasound/Source/MetasoundEngine/Public/MetasoundAssetSubsystem.h
+++ b/Engine/Plugins/Runtime/Metasound/Source/MetasoundEngine/Public/MetasoundAssetSubsystem.h
@@ -5,6 +5,7 @@
 #include "Engine/StreamableManager.h"
 #include "MetasoundAssetBase.h"
 #include "MetasoundAssetManager.h"
+#include "MetasoundBuilderBase.h"
 #include "Subsystems/EngineSubsystem.h"
 #include "UObject/Object.h"
 
@@ -47,9 +48,15 @@ struct FMetaSoundAsyncAssetDependencies
 	TSharedPtr<FStreamableHandle> StreamableHandle;
 };
 
-/** The subsystem in charge of the MetaSound asset registry */
+namespace Metasound::Engine
+{
+	void DeinitializeAssetManager();
+	void InitializeAssetManager();
+} // namespace Metasound::Engine
+
+/** DEPRECATED: Subsystem in charge of the MetaSound asset registry */
 UCLASS()
-class METASOUNDENGINE_API UMetaSoundAssetSubsystem : public UEngineSubsystem, public Metasound::Frontend::IMetaSoundAssetManager
+class METASOUNDENGINE_API UMetaSoundAssetSubsystem : public UEngineSubsystem
 {
 	GENERATED_BODY()
 
@@ -58,33 +65,65 @@ public:
 
 	virtual void Initialize(FSubsystemCollectionBase& InCollection) override;
 
-	Metasound::Frontend::FNodeRegistryKey AddOrUpdateAsset(const FAssetData& InAssetData);
-	void RemoveAsset(const UObject& InObject);
-	void RemoveAsset(const FAssetData& InAssetData);
-	void RenameAsset(const FAssetData& InAssetData, bool bInReregisterWithFrontend = true);
+	UE_DEPRECATED(5.5, "Moved to internal implementation, use IMetaSoundAssetManager::GetChecked() and analogous call")
+	virtual void RemoveAsset(const UObject& InObject);
+
+	UE_DEPRECATED(5.5, "Moved to internal implementation, use IMetaSoundAssetManager::GetChecked() and analogous call")
+	virtual void RemoveAsset(const FAssetData& InAssetData);
+
+	UE_DEPRECATED(5.5, "Moved to internal implementation, use IMetaSoundAssetManager::GetChecked() and analogous call")
+	virtual void RenameAsset(const FAssetData& InAssetData, bool bInReregisterWithFrontend = true);
 
 #if WITH_EDITORONLY_DATA
-	virtual void AddAssetReferences(FMetasoundAssetBase& InAssetBase) override;
+	UE_DEPRECATED(5.5, "Moved to internal implementation, use IMetaSoundAssetManager::GetChecked() and analogous call")
+	virtual void AddAssetReferences(FMetasoundAssetBase& InAssetBase);
 #endif
-	virtual Metasound::Frontend::FNodeRegistryKey AddOrUpdateAsset(const UObject& InObject) override;
-	virtual FMetaSoundFrontendDocumentBuilder& AttachDocumentBuilderChecked(UObject& InObject) const override;
 
-	virtual bool CanAutoUpdate(const FMetasoundFrontendClassName& InClassName) const override;
-	virtual bool ContainsKey(const Metasound::Frontend::FNodeRegistryKey& InRegistryKey) const override;
-	virtual const FSoftObjectPath* FindObjectPathFromKey(const Metasound::Frontend::FNodeRegistryKey& RegistryKey) const override;
+	UE_DEPRECATED(5.5, "Moved to internal implementation, use IMetaSoundAssetManager::GetChecked() and analogous call")
+	Metasound::Frontend::FNodeRegistryKey AddOrUpdateAsset(const FAssetData& InAssetData);
+
+	UE_DEPRECATED(5.5, "Moved to internal implementation, use IMetaSoundAssetManager::GetChecked() and analogous call")
+	virtual Metasound::Frontend::FNodeRegistryKey AddOrUpdateAsset(const UObject& InObject);
+
+	UE_DEPRECATED(5.5, "Moved to internal implementation, use IMetaSoundAssetManager::GetChecked() and analogous call")
+	virtual bool CanAutoUpdate(const FMetasoundFrontendClassName& InClassName) const;
+
+	UE_DEPRECATED(5.5, "Moved to internal implementation, use IMetaSoundAssetManager::GetChecked() and analogous call")
+	virtual bool ContainsKey(const Metasound::Frontend::FNodeRegistryKey& InRegistryKey) const;
+
+	UE_DEPRECATED(5.5, "Moved to internal implementation, use IMetaSoundAssetManager::GetChecked() and analogous call")
+	virtual const FSoftObjectPath* FindObjectPathFromKey(const Metasound::Frontend::FNodeRegistryKey& RegistryKey) const;
+
+	UE_DEPRECATED(5.5, "Moved to internal implementation, use IMetaSoundAssetManager::GetChecked() and analogous call")
+	virtual FMetasoundAssetBase* GetAsAsset(UObject& InObject) const;
 
-	virtual FMetasoundAssetBase* GetAsAsset(UObject& InObject) const override;
-	virtual const FMetasoundAssetBase* GetAsAsset(const UObject& InObject) const override;
+	UE_DEPRECATED(5.5, "Moved to internal implementation, use IMetaSoundAssetManager::GetChecked() and analogous call")
+	virtual const FMetasoundAssetBase* GetAsAsset(const UObject& InObject) const;
+
+	UE_DEPRECATED(5.5, "Implementation of MetaSound asset management has been moved to raw c++ implementation for more reliable, monolithic "
+		"lifetime management. This subsystem continues to exist only for Blueprint-related asset functionality. "
+		"Use IMetaSoundAssetManager::GetChecked() instead")
+	static UMetaSoundAssetSubsystem& GetChecked();
 
 #if WITH_EDITOR
-	virtual TSet<FAssetInfo> GetReferencedAssetClasses(const FMetasoundAssetBase& InAssetBase) const override;
+	UE_DEPRECATED(5.5, "Moved to internal implementation, use IMetaSoundAssetManager::GetChecked() and analogous call")
+	virtual TSet<FAssetInfo> GetReferencedAssetClasses(const FMetasoundAssetBase& InAssetBase) const;
 #endif
-	virtual void RescanAutoUpdateDenyList() override;
-	virtual FMetasoundAssetBase* TryLoadAsset(const FSoftObjectPath& InObjectPath) const override;
-	virtual FMetasoundAssetBase* TryLoadAssetFromKey(const Metasound::Frontend::FNodeRegistryKey& RegistryKey) const override;
-	virtual bool TryLoadReferencedAssets(const FMetasoundAssetBase& InAssetBase, TArray<FMetasoundAssetBase*>& OutReferencedAssets) const override;
-	virtual void RequestAsyncLoadReferencedAssets(FMetasoundAssetBase& InAssetBase) override;
-	virtual void WaitUntilAsyncLoadReferencedAssetsComplete(FMetasoundAssetBase& InAssetBase) override;
+
+	UE_DEPRECATED(5.5, "Moved to internal implementation, use IMetaSoundAssetManager::GetChecked() and analogous call")
+	virtual FMetasoundAssetBase* TryLoadAsset(const FSoftObjectPath& InObjectPath) const;
+
+	UE_DEPRECATED(5.5, "Moved to internal implementation, use IMetaSoundAssetManager::GetChecked() and analogous call")
+	virtual FMetasoundAssetBase* TryLoadAssetFromKey(const Metasound::Frontend::FNodeRegistryKey& RegistryKey) const;
+
+	UE_DEPRECATED(5.5, "Moved to internal implementation, use IMetaSoundAssetManager::GetChecked() and analogous call")
+	virtual bool TryLoadReferencedAssets(const FMetasoundAssetBase& InAssetBase, TArray<FMetasoundAssetBase*>& OutReferencedAssets) const;
+
+	UE_DEPRECATED(5.5, "Moved to internal implementation, use IMetaSoundAssetManager::GetChecked() and analogous call")
+	virtual void RequestAsyncLoadReferencedAssets(FMetasoundAssetBase& InAssetBase) { }
+
+	UE_DEPRECATED(5.5, "Moved to internal implementation, use IMetaSoundAssetManager::GetChecked() and analogous call")
+	virtual void WaitUntilAsyncLoadReferencedAssetsComplete(FMetasoundAssetBase& InAssetBase) { }
 
 	UFUNCTION(BlueprintCallable, Category = "MetaSounds|Registration")
 	void RegisterAssetClassesInDirectories(const TArray<FMetaSoundAssetDirectory>& Directories);
@@ -93,27 +132,22 @@ public:
 	void UnregisterAssetClassesInDirectories(const TArray<FMetaSoundAssetDirectory>& Directories);
 
 protected:
+	UE_DEPRECATED(5.5, "Moved to private implementation")
+	void PostEngineInit() { }
 
-	void PostEngineInit();
-	void PostInitAssetScan();
-	void RebuildDenyListCache(const UAssetManager& InAssetManager);
-	void ResetAssetClassDisplayName(const FAssetData& InAssetData);
-	void SearchAndIterateDirectoryAssets(const TArray<FDirectoryPath>& InDirectories, TFunctionRef<void(const FAssetData&)> InFunction);
+	UE_DEPRECATED(5.5, "Moved to private implementation")
+	void PostInitAssetScan() { }
 
-private:
+	UE_DEPRECATED(5.5, "Moved to internal implementation, use IMetaSoundAssetManager::GetChecked() and analogous call")
+	void RebuildDenyListCache(const UAssetManager& InAssetManager) { }
 
-	UPROPERTY(Transient)
-	TArray<FMetaSoundAsyncAssetDependencies> LoadingDependencies;
-
-	FMetaSoundAsyncAssetDependencies* FindLoadingDependencies(const UObject* InParentAsset);
-	FMetaSoundAsyncAssetDependencies* FindLoadingDependencies(int32 InLoadID);
-	void RemoveLoadingDependencies(int32 InLoadID);
-	void OnAssetsLoaded(int32 InLoadID);
-
-	FStreamableManager StreamableManager;
-	int32 AsyncLoadIDCounter = 0;
-	int32 AutoUpdateDenyListChangeID = INDEX_NONE;
-	TSet<FName> AutoUpdateDenyListCache;
-	TMap<Metasound::Frontend::FNodeRegistryKey, FSoftObjectPath> PathMap;
-	std::atomic<bool> bIsInitialAssetScanComplete = false;
+	UE_DEPRECATED(5.5, "Use FMetaSoundDocumentBuilder::SetDisplayName instead (call now only available with editor compiled)")
+	void ResetAssetClassDisplayName(const FAssetData& InAssetData) { }
+
+	UE_DEPRECATED(5.5, "Moved to internal implementation, use IMetaSoundAssetManager::GetChecked() and analogous call")
+	void SearchAndIterateDirectoryAssets(const TArray<FDirectoryPath>& InDirectories, TFunctionRef<void(const FAssetData&)> InFunction) { }
+
+private:
+	void PostEngineInitInternal();
+	void PostInitAssetScanInternal();
 };
