- Move MetaSoundAssetSubsystem& BuilderSubsystem implementations to monolithics that are initialized at MetaSoundEngine module load for safety

diff --git a/Engine/Plugins/Runtime/Metasound/Source/MetasoundEngine/Public/MetasoundBuilderBase.h b/Engine/Plugins/Runtime/Metasound/Source/MetasoundEngine/Public/MetasoundBuilderBase.h
new file mode 100644
--- /dev/null
+++ b/Engine/Plugins/Runtime/Metasound/Source/MetasoundEngine/Public/MetasoundBuilderBase.h
@@ -0,0 +1,541 @@
+// Copyright Epic Games, Inc. All Rights Reserved.
+#pragma once
+
+#include "Containers/Map.h"
+#include "Interfaces/MetasoundOutputFormatInterfaces.h"
+#include "MetasoundDocumentInterface.h"
+#include "MetasoundDynamicOperatorTransactor.h"
+#include "MetasoundFrontendController.h"
+#include "MetasoundFrontendDocument.h"
+#include "MetasoundFrontendDocumentBuilder.h"
+#include "Subsystems/EngineSubsystem.h"
+#include "Templates/Function.h"
+#include "UObject/NoExportTypes.h"
+#include "UObject/Object.h"
+#include "UObject/ObjectMacros.h"
+
+#include "MetasoundBuilderBase.generated.h"
+
+
+// Forward Declarations
+struct FMetasoundFrontendClassName;
+struct FMetasoundFrontendVersion;
+
+
+USTRUCT(BlueprintType, meta = (DisplayName = "MetaSound Node Input Handle"))
+struct METASOUNDENGINE_API FMetaSoundBuilderNodeInputHandle : public FMetasoundFrontendVertexHandle
+{
+	GENERATED_BODY()
+
+public:
+	FMetaSoundBuilderNodeInputHandle() = default;
+	FMetaSoundBuilderNodeInputHandle(const FGuid InNodeID, const FGuid& InVertexID)
+	{
+		NodeID = InNodeID;
+		VertexID = InVertexID;
+	}
+};
+
+USTRUCT(BlueprintType, meta = (DisplayName = "MetaSound Node Output Handle"))
+struct METASOUNDENGINE_API FMetaSoundBuilderNodeOutputHandle : public FMetasoundFrontendVertexHandle
+{
+	GENERATED_BODY()
+
+public:
+	FMetaSoundBuilderNodeOutputHandle() = default;
+	FMetaSoundBuilderNodeOutputHandle(const FGuid InNodeID, const FGuid& InVertexID)
+	{
+		NodeID = InNodeID;
+		VertexID = InVertexID;
+	}
+};
+
+USTRUCT(BlueprintType)
+struct METASOUNDENGINE_API FMetaSoundNodeHandle
+{
+	GENERATED_BODY()
+
+	UPROPERTY()
+	FGuid NodeID;
+
+public:
+	FMetaSoundNodeHandle() = default;
+	FMetaSoundNodeHandle(const FGuid& InNodeID)
+		: NodeID(InNodeID)
+	{
+	}
+
+	// Returns whether or not the vertex handle is set (may or may not be
+	// valid depending on what builder context it is referenced against)
+	bool IsSet() const
+	{
+		return NodeID.IsValid();
+	}
+};
+
+USTRUCT(BlueprintType)
+struct METASOUNDENGINE_API FMetaSoundBuilderOptions
+{
+	GENERATED_BODY()
+
+	// Name of generated object. If object already exists, used as the base name to ensure new object is unique.
+	UPROPERTY(EditAnywhere, BlueprintReadWrite, Category = "MetaSound|Builder")
+	FName Name;
+
+	// If the resulting MetaSound is building over an existing document, a unique class name will be generated,
+	// invalidating any referencing MetaSounds and registering the MetaSound as a new entry in the Frontend. If
+	// building a new document, option is ignored.
+	UPROPERTY(EditAnywhere, BlueprintReadWrite, Category = "MetaSound|Builder", meta = (AdvancedDisplay))
+	bool bForceUniqueClassName = false;
+
+	// If true, adds MetaSound to node registry, making it available
+	// for reference by other dynamically created MetaSounds.
+	UPROPERTY()
+	bool bAddToRegistry = true;
+
+	// If set, builder overwrites the given MetaSound's document with the builder's copy
+	// (ignores the Name field above).
+	UPROPERTY()
+	TScriptInterface<IMetaSoundDocumentInterface> ExistingMetaSound;
+};
+
+UENUM(BlueprintType)
+enum class EMetaSoundBuilderResult : uint8
+{
+	Succeeded,
+	Failed
+};
+
+/** Base implementation of MetaSound builder */
+UCLASS(Abstract)
+class METASOUNDENGINE_API UMetaSoundBuilderBase : public UObject
+{
+	GENERATED_BODY()
+
+public:
+	// Begin UObject interface
+	virtual void BeginDestroy() override;
+	// End UObject interface
+
+	// Adds a graph input node with the given name, DataType, and sets the graph input to default value.
+	// Returns the new input node's output handle if it was successfully created, or an invalid handle if it failed.
+	UFUNCTION(BlueprintCallable, Category = "Audio|MetaSound|Builder", meta = (ExpandEnumAsExecs = "OutResult", AdvancedDisplay = "3"))
+	UPARAM(DisplayName = "Output Handle") FMetaSoundBuilderNodeOutputHandle AddGraphInputNode(FName Name, FName DataType, FMetasoundFrontendLiteral DefaultValue, EMetaSoundBuilderResult& OutResult, bool bIsConstructorInput = false);
+
+	// Adds a graph output node with the given name, DataType, and sets output node's input to default value.
+	// Returns the new output node's input handle if it was successfully created, or an invalid handle if it failed.
+	UFUNCTION(BlueprintCallable, Category = "Audio|MetaSound|Builder", meta = (ExpandEnumAsExecs = "OutResult", AdvancedDisplay = "3"))
+	UPARAM(DisplayName = "Input Handle") FMetaSoundBuilderNodeInputHandle AddGraphOutputNode(FName Name, FName DataType, FMetasoundFrontendLiteral DefaultValue, EMetaSoundBuilderResult& OutResult, bool bIsConstructorOutput = false);
+
+	// Adds an interface registered with the given name to the graph, adding associated input and output nodes.
+	UFUNCTION(BlueprintCallable, Category = "Audio|MetaSound|Builder", meta = (ExpandEnumAsExecs = "OutResult"))
+	void AddInterface(FName InterfaceName, EMetaSoundBuilderResult& OutResult);
+
+	// Adds a node to the graph using the provided MetaSound asset as its defining NodeClass.
+	// Returns a node handle to the created node if successful, or an invalid handle if it failed.
+	UFUNCTION(BlueprintCallable, Category = "Audio|MetaSound|Builder", DisplayName = "Add MetaSound Node From Asset Class", meta = (ExpandEnumAsExecs = "OutResult"))
+	UPARAM(DisplayName = "Node Handle") FMetaSoundNodeHandle AddNode(const TScriptInterface<IMetaSoundDocumentInterface>& NodeClass, EMetaSoundBuilderResult& OutResult);
+
+	// Adds node referencing the highest native class version of the given class name to the document.
+	// Returns a node handle to the created node if successful, or an invalid handle if it failed.
+	UFUNCTION(BlueprintCallable, Category = "Audio|MetaSound|Builder", DisplayName = "Add MetaSound Node By ClassName", meta = (ExpandEnumAsExecs = "OutResult", AdvancedDisplay = "2"))
+	UPARAM(DisplayName = "Node Handle") FMetaSoundNodeHandle AddNodeByClassName(const FMetasoundFrontendClassName& ClassName, EMetaSoundBuilderResult& OutResult, int32 MajorVersion = 1);
+	
+	UE_DEPRECATED(5.4, "This version of AddNodeByClassName is deprecated. Use the one with a default MajorVersion of 1.")
+	UPARAM(DisplayName = "Node Handle") FMetaSoundNodeHandle AddNodeByClassName(const FMetasoundFrontendClassName& ClassName, int32 MajorVersion, EMetaSoundBuilderResult& OutResult);
+
+#if WITH_EDITOR
+	bool ClearMemberMetadata(const FGuid& InMemberID);
+#endif // WITH_EDITOR
+
+	// Connects node output to a node input. Does *NOT* provide loop detection for performance reasons.  Loop detection is checked on class registration when built or played.
+	// Returns succeeded if connection made, failed if connection already exists with input, the data types do not match, or the connection is not supported due to access type
+	// incompatibility (ex. constructor input to non-constructor input).
+	UFUNCTION(BlueprintCallable, Category = "Audio|MetaSound|Builder", meta = (ExpandEnumAsExecs = "OutResult"))
+	void ConnectNodes(const FMetaSoundBuilderNodeOutputHandle& NodeOutputHandle, const FMetaSoundBuilderNodeInputHandle& NodeInputHandle, EMetaSoundBuilderResult& OutResult);
+
+	// Connects two nodes using defined MetaSound Interface Bindings registered with the MetaSound Interface registry.
+	UFUNCTION(BlueprintCallable, Category = "Audio|MetaSound|Builder", meta = (ExpandEnumAsExecs = "OutResult"))
+	void ConnectNodesByInterfaceBindings(const FMetaSoundNodeHandle& FromNodeHandle, const FMetaSoundNodeHandle& ToNodeHandle, EMetaSoundBuilderResult& OutResult);
+
+	// Connects a given node's outputs to all graph outputs for shared interfaces implemented on both the node's referenced class and the builder's MetaSound graph. Returns inputs of connected output nodes.
+	UFUNCTION(BlueprintCallable, Category = "Audio|MetaSound|Builder", meta = (ExpandEnumAsExecs = "OutResult"))
+	UPARAM(DisplayName = "Connected Graph Output Node Inputs") TArray<FMetaSoundBuilderNodeInputHandle> ConnectNodeOutputsToMatchingGraphInterfaceOutputs(const FMetaSoundNodeHandle& NodeHandle, EMetaSoundBuilderResult& OutResult);
+
+	// Connects a given node's inputs to all graph inputs for shared interfaces implemented on both the node's referenced class and the builder's MetaSound graph. Returns outputs of connected input nodes.
+	UFUNCTION(BlueprintCallable, Category = "Audio|MetaSound|Builder", meta = (ExpandEnumAsExecs = "OutResult"))
+	UPARAM(DisplayName = "Connected Graph Input Node Outputs") TArray<FMetaSoundBuilderNodeOutputHandle> ConnectNodeInputsToMatchingGraphInterfaceInputs(const FMetaSoundNodeHandle& NodeHandle, EMetaSoundBuilderResult& OutResult);
+
+	// Connects a given node output to the graph output with the given name.
+	UFUNCTION(BlueprintCallable, Category = "Audio|MetaSound|Builder", meta = (ExpandEnumAsExecs = "OutResult"))
+	void ConnectNodeOutputToGraphOutput(FName GraphOutputName, const FMetaSoundBuilderNodeOutputHandle& NodeOutputHandle, EMetaSoundBuilderResult& OutResult);
+
+	// Connects a given node input to the graph input with the given name.
+	UFUNCTION(BlueprintCallable, Category = "Audio|MetaSound|Builder", meta = (ExpandEnumAsExecs = "OutResult"))
+	void ConnectNodeInputToGraphInput(FName GraphInputName, const FMetaSoundBuilderNodeInputHandle& NodeInputHandle, EMetaSoundBuilderResult& OutResult);
+
+	// Returns whether node exists.
+	UFUNCTION(BlueprintCallable, Category = "Audio|MetaSound|Builder")
+	UPARAM(DisplayName = "IsValid") bool ContainsNode(const FMetaSoundNodeHandle& Node) const;
+
+	// Returns whether node input exists.
+	UFUNCTION(BlueprintCallable, Category = "Audio|MetaSound|Builder")
+	UPARAM(DisplayName = "IsValid") bool ContainsNodeInput(const FMetaSoundBuilderNodeInputHandle& Input) const;
+
+	// Returns whether node output exists.
+	UFUNCTION(BlueprintCallable, Category = "Audio|MetaSound|Builder")
+	UPARAM(DisplayName = "IsValid") bool ContainsNodeOutput(const FMetaSoundBuilderNodeOutputHandle& Output) const;
+
+	// Disconnects node output to a node input. Returns success if connection was removed, failed if not.
+	UFUNCTION(BlueprintCallable, Category = "Audio|MetaSound|Builder", meta = (ExpandEnumAsExecs = "OutResult"))
+	void DisconnectNodes(const FMetaSoundBuilderNodeOutputHandle& NodeOutputHandle, const FMetaSoundBuilderNodeInputHandle& NodeInputHandle, EMetaSoundBuilderResult& OutResult);
+
+	// Removes connection to a given node input. Returns success if connection was removed, failed if not.
+	UFUNCTION(BlueprintCallable, Category = "Audio|MetaSound|Builder", meta = (ExpandEnumAsExecs = "OutResult"))
+	void DisconnectNodeInput(const FMetaSoundBuilderNodeInputHandle& NodeInputHandle, EMetaSoundBuilderResult& OutResult);
+
+	// Removes all connections from a given node output. Returns success if all connections were removed, failed if not.
+	UFUNCTION(BlueprintCallable, Category = "Audio|MetaSound|Builder", meta = (ExpandEnumAsExecs = "OutResult"))
+	void DisconnectNodeOutput(const FMetaSoundBuilderNodeOutputHandle& NodeOutputHandle, EMetaSoundBuilderResult& OutResult);
+
+	// Disconnects two nodes using defined MetaSound Interface Bindings registered with the MetaSound Interface registry. Returns success if
+	// all connections were found and removed, failed if any connections were not.
+	UFUNCTION(BlueprintCallable, Category = "Audio|MetaSound|Builder", meta = (ExpandEnumAsExecs = "OutResult"))
+	void DisconnectNodesByInterfaceBindings(const FMetaSoundNodeHandle& FromNodeHandle, const FMetaSoundNodeHandle& ToNodeHandle, EMetaSoundBuilderResult& OutResult);
+
+	// Returns graph input node by the given name if it exists, or an invalid handle if not found.
+	UFUNCTION(BlueprintCallable, Category = "Audio|MetaSound|Builder", meta = (ExpandEnumAsExecs = "OutResult"))
+	UPARAM(DisplayName = "Node Handle") FMetaSoundNodeHandle FindGraphInputNode(FName InputName, EMetaSoundBuilderResult& OutResult);
+
+	// Returns graph output node by the given name if it exists, or an invalid handle if not found.
+	UFUNCTION(BlueprintCallable, Category = "Audio|MetaSound|Builder", meta = (ExpandEnumAsExecs = "OutResult"))
+	UPARAM(DisplayName = "Node Handle") FMetaSoundNodeHandle FindGraphOutputNode(FName OutputName, EMetaSoundBuilderResult& OutResult);
+
+#if WITH_EDITOR
+	UMetaSoundFrontendMemberMetadata* FindMemberMetadata(const FGuid& InMemberID);
+#endif // WITH_EDITOR
+
+	// Returns node input by the given name if it exists, or an invalid handle if not found.
+	UFUNCTION(BlueprintCallable, Category = "Audio|MetaSound|Builder", meta = (ExpandEnumAsExecs = "OutResult"))
+	UPARAM(DisplayName = "Input Handle") FMetaSoundBuilderNodeInputHandle FindNodeInputByName(const FMetaSoundNodeHandle& NodeHandle, FName InputName, EMetaSoundBuilderResult& OutResult);
+
+	// Returns all node inputs.
+	UFUNCTION(BlueprintCallable, Category = "Audio|MetaSound|Builder", meta = (ExpandEnumAsExecs = "OutResult"))
+	UPARAM(DisplayName = "Input  Handles") TArray<FMetaSoundBuilderNodeInputHandle> FindNodeInputs(const FMetaSoundNodeHandle& NodeHandle, EMetaSoundBuilderResult& OutResult);
+
+	// Returns node inputs by the given DataType (ex. "Audio", "Trigger", "String", "Bool", "Float", "Int32", etc.).
+	UFUNCTION(BlueprintCallable, Category = "Audio|MetaSound|Builder", meta = (ExpandEnumAsExecs = "OutResult"))
+	UPARAM(DisplayName = "Input Handles") TArray<FMetaSoundBuilderNodeInputHandle> FindNodeInputsByDataType(const FMetaSoundNodeHandle& NodeHandle, EMetaSoundBuilderResult& OutResult, FName DataType);
+
+	// Returns node output by the given name.
+	UFUNCTION(BlueprintCallable, Category = "Audio|MetaSound|Builder", meta = (ExpandEnumAsExecs = "OutResult"))
+	UPARAM(DisplayName = "Output Handle") FMetaSoundBuilderNodeOutputHandle FindNodeOutputByName(const FMetaSoundNodeHandle& NodeHandle, FName OutputName, EMetaSoundBuilderResult& OutResult);
+
+	// Returns all node outputs.
+	UFUNCTION(BlueprintCallable, Category = "Audio|MetaSound|Builder", meta = (ExpandEnumAsExecs = "OutResult"))
+	UPARAM(DisplayName = "Output Handles") TArray<FMetaSoundBuilderNodeOutputHandle> FindNodeOutputs(const FMetaSoundNodeHandle& NodeHandle, EMetaSoundBuilderResult& OutResult);
+
+	// Returns node outputs by the given DataType (ex. "Audio", "Trigger", "String", "Bool", "Float", "Int32", etc.).
+	UFUNCTION(BlueprintCallable, Category = "Audio|MetaSound|Builder", meta = (ExpandEnumAsExecs = "OutResult"))
+	UPARAM(DisplayName = "Output Handles") TArray<FMetaSoundBuilderNodeOutputHandle> FindNodeOutputsByDataType(const FMetaSoundNodeHandle& NodeHandle, EMetaSoundBuilderResult& OutResult, FName DataType);
+
+	// Returns input nodes associated with a given interface.
+	UFUNCTION(BlueprintCallable, Category = "Audio|MetaSound|Builder", meta = (ExpandEnumAsExecs = "OutResult"))
+	UPARAM(DisplayName = "Input Node Handles") TArray<FMetaSoundNodeHandle> FindInterfaceInputNodes(FName InterfaceName, EMetaSoundBuilderResult& OutResult);
+
+	// Returns output nodes associated with a given interface.
+	UFUNCTION(BlueprintCallable, Category = "Audio|MetaSound|Builder", meta = (ExpandEnumAsExecs = "OutResult"))
+	UPARAM(DisplayName = "Output Node Handles") TArray<FMetaSoundNodeHandle> FindInterfaceOutputNodes(FName InterfaceName, EMetaSoundBuilderResult& OutResult);
+
+	// Returns input's parent node if the input is valid, otherwise returns invalid node handle.
+	UFUNCTION(BlueprintCallable, Category = "Audio|MetaSound|Builder", meta = (ExpandEnumAsExecs = "OutResult"))
+	UPARAM(DisplayName = "Node Handle") FMetaSoundNodeHandle FindNodeInputParent(const FMetaSoundBuilderNodeInputHandle& InputHandle, EMetaSoundBuilderResult& OutResult);
+
+	// Returns output's parent node if the input is valid, otherwise returns invalid node handle.
+	UFUNCTION(BlueprintCallable, Category = "Audio|MetaSound|Builder", meta = (ExpandEnumAsExecs = "OutResult"))
+	UPARAM(DisplayName = "Node Handle") FMetaSoundNodeHandle FindNodeOutputParent(const FMetaSoundBuilderNodeOutputHandle& OutputHandle, EMetaSoundBuilderResult& OutResult);
+
+	// Returns output's parent node if the input is valid, otherwise returns invalid node handle.
+	UFUNCTION(BlueprintCallable, Category = "Audio|MetaSound|Builder", meta = (ExpandEnumAsExecs = "OutResult"))
+	UPARAM(DisplayName = "Node ClassVersion") FMetasoundFrontendVersion FindNodeClassVersion(const FMetaSoundNodeHandle& NodeHandle, EMetaSoundBuilderResult& OutResult);
+
+	// Returns the document's root graph class name
+	FMetasoundFrontendClassName GetRootGraphClassName() const;
+
+	// Returns node input's data if valid (including things like name and datatype).
+	UFUNCTION(BlueprintCallable, Category = "Audio|MetaSound|Builder", meta = (ExpandEnumAsExecs = "OutResult"))
+	void GetNodeInputData(const FMetaSoundBuilderNodeInputHandle& InputHandle, FName& Name, FName& DataType, EMetaSoundBuilderResult& OutResult);
+
+	// Returns node input's literal value if set on graph, otherwise fails and returns default literal.
+	UFUNCTION(BlueprintCallable, Category = "Audio|MetaSound|Builder", meta = (ExpandEnumAsExecs = "OutResult"))
+	FMetasoundFrontendLiteral GetNodeInputDefault(const FMetaSoundBuilderNodeInputHandle& InputHandle, EMetaSoundBuilderResult& OutResult);
+
+	// Returns node input's class literal value if set, otherwise fails and returns default literal.
+	UFUNCTION(BlueprintCallable, Category = "Audio|MetaSound|Builder", meta = (ExpandEnumAsExecs = "OutResult"))
+	FMetasoundFrontendLiteral GetNodeInputClassDefault(const FMetaSoundBuilderNodeInputHandle& InputHandle, EMetaSoundBuilderResult& OutResult);
+
+	// Returns whether the given node input is a constructor pin
+	UFUNCTION(BlueprintCallable, Category = "Audio|MetaSound|Builder")
+	bool GetNodeInputIsConstructorPin(const FMetaSoundBuilderNodeInputHandle& InputHandle) const;
+
+	// Returns node output's data if valid (including things like name and datatype).
+	UFUNCTION(BlueprintCallable, Category = "Audio|MetaSound|Builder", meta = (ExpandEnumAsExecs = "OutResult"))
+	void GetNodeOutputData(const FMetaSoundBuilderNodeOutputHandle& OutputHandle, FName& Name, FName& DataType, EMetaSoundBuilderResult& OutResult);
+	
+	// Returns whether the given node output is a constructor pin
+	UFUNCTION(BlueprintCallable, Category = "Audio|MetaSound|Builder")
+	bool GetNodeOutputIsConstructorPin(const FMetaSoundBuilderNodeOutputHandle& OutputHandle) const;
+
+	// Return the asset referenced by this preset builder. Returns nullptr if the builder is not a preset.
+	UFUNCTION(BlueprintCallable, Category = "Audio|MetaSound|Builder")
+	UObject* GetReferencedPresetAsset() const;
+
+	// Returns if a given interface is declared.
+	UFUNCTION(BlueprintCallable, Category = "Audio|MetaSound|Builder")
+	UPARAM(DisplayName = "IsDeclared") bool InterfaceIsDeclared(FName InterfaceName) const;
+
+	// Returns if a given node output and node input are connected.
+	UFUNCTION(BlueprintCallable, Category = "Audio|MetaSound|Builder")
+	UPARAM(DisplayName = "Connected") bool NodesAreConnected(const FMetaSoundBuilderNodeOutputHandle& OutputHandle, const FMetaSoundBuilderNodeInputHandle& InputHandle) const;
+
+	// Returns if a given node input has connections.
+	UFUNCTION(BlueprintCallable, Category = "Audio|MetaSound|Builder")
+	UPARAM(DisplayName = "Connected") bool NodeInputIsConnected(const FMetaSoundBuilderNodeInputHandle& InputHandle) const;
+
+	// Returns if a given node output is connected.
+	UFUNCTION(BlueprintCallable, Category = "Audio|MetaSound|Builder")
+	UPARAM(DisplayName = "Connected") bool NodeOutputIsConnected(const FMetaSoundBuilderNodeOutputHandle& OutputHandle) const;
+
+	// Returns whether this is a preset.
+	UFUNCTION(BlueprintCallable, Category = "Audio|MetaSound|Builder")
+	bool IsPreset() const;
+
+	// Converts this preset to a fully accessible MetaSound; sets result to succeeded if it was converted successfully and failed if it was not.
+	UFUNCTION(BlueprintCallable, Category = "Audio|MetaSound|Builder", meta = (ExpandEnumAsExecs = "OutResult"))
+	void ConvertFromPreset(EMetaSoundBuilderResult& OutResult);
+
+	// Convert this builder to a MetaSound source preset with the given referenced source builder 
+	UFUNCTION(BlueprintCallable, Category = "Audio|MetaSound|Builder", meta = (ExpandEnumAsExecs = "OutResult"))
+	void ConvertToPreset(const TScriptInterface<IMetaSoundDocumentInterface>& ReferencedNodeClass, EMetaSoundBuilderResult& OutResult);
+
+	// Removes graph input if it exists; sets result to succeeded if it was removed and failed if it was not.
+	UFUNCTION(BlueprintCallable, Category = "Audio|MetaSound|Builder", meta = (ExpandEnumAsExecs = "OutResult"))
+	void RemoveGraphInput(FName Name, EMetaSoundBuilderResult& OutResult);
+
+	// Removes graph output if it exists; sets result to succeeded if it was removed and failed if it was not.
+	UFUNCTION(BlueprintCallable, Category = "Audio|MetaSound|Builder", meta = (ExpandEnumAsExecs = "OutResult"))
+	void RemoveGraphOutput(FName Name, EMetaSoundBuilderResult& OutResult);
+
+	// Removes the interface with the given name from the builder's MetaSound. Removes any graph inputs
+	// and outputs associated with the given interface and their respective connections (if any).
+	UFUNCTION(BlueprintCallable, Category = "Audio|MetaSound|Builder", meta = (ExpandEnumAsExecs = "OutResult"))
+	void RemoveInterface(FName InterfaceName, EMetaSoundBuilderResult& OutResult);
+
+	// Removes node and any associated connections from the builder's MetaSound.
+	UFUNCTION(BlueprintCallable, Category = "Audio|MetaSound|Builder", meta = (ExpandEnumAsExecs = "OutResult"))
+	void RemoveNode(const FMetaSoundNodeHandle& NodeHandle, EMetaSoundBuilderResult& OutResult, bool bRemoveUnusedDependencies = true);
+
+	// Removes node input literal default if set, reverting the value to be whatever the node class defaults the value to.
+	// Returns success if value was removed, false if not removed (i.e. wasn't set to begin with).
+	UFUNCTION(BlueprintCallable, Category = "Audio|MetaSound|Builder", meta = (ExpandEnumAsExecs = "OutResult"))
+	void RemoveNodeInputDefault(const FMetaSoundBuilderNodeInputHandle& InputHandle, EMetaSoundBuilderResult& OutResult);
+
+	// Removes dependencies in document that are no longer referenced by nodes
+	UFUNCTION(BlueprintCallable, Category = "Audio|MetaSound|Builder")
+	void RemoveUnusedDependencies();
+
+	UE_DEPRECATED(5.5, "Use IDocumentBuilderRegistry::GenerateNewClassName instead to maintain registry mappings.")
+	void RenameRootGraphClass(const FMetasoundFrontendClassName& InName);
+
+	UE_DEPRECATED(5.5, "Moved to 'Reload' to discourage priming by defaulting to false. "
+		"In addition, internal delegates and respective bindings are now regenerated when reloading, "
+		"effectively resetting the builder's entire state and not just strictly the underlying cache.")
+	void ReloadCache(bool bPrimeCache = true);
+
+	// Reloads the builder, freeing the internal cache and rebuilding delegate bindings
+	void Reload(bool bPrimeCache = false);
+
+#if WITH_EDITOR
+	// Sets the author of the MetaSound.
+	void SetAuthor(const FString& InAuthor);
+#endif // WITH_EDITOR
+
+	// Sets the node's input default value (used if no connection to the given node input is present)
+	UFUNCTION(BlueprintCallable, Category = "Audio|MetaSound|Builder", meta = (ExpandEnumAsExecs = "OutResult"))
+	void SetNodeInputDefault(const FMetaSoundBuilderNodeInputHandle& NodeInputHandle, const FMetasoundFrontendLiteral& Literal, EMetaSoundBuilderResult& OutResult);
+
+	// Disconnects the given graph input's respective template nodes and sets the graph input's AccessType should it not match the current AccessType.
+	// Result succeeds if the AccessType was successfully changed or if the provided AccessType is already the input's current AccessType.
+	UFUNCTION(BlueprintCallable, Category = "Audio|MetaSound|Builder", meta = (ExpandEnumAsExecs = "OutResult"))
+	void SetGraphInputAccessType(FName InputName, EMetasoundFrontendVertexAccessType AccessType, EMetaSoundBuilderResult& OutResult);
+
+	// Disconnects the given graph input's respective template nodes and sets the graph input's DataType should it not match the current DataType.
+	// Result succeeds if the DataType was successfully changed or if the provided DataType is already the input's current DataType.
+	UFUNCTION(BlueprintCallable, Category = "Audio|MetaSound|Builder", meta = (ExpandEnumAsExecs = "OutResult"))
+	void SetGraphInputDataType(FName InputName, FName DataType, EMetaSoundBuilderResult& OutResult);
+
+	// Sets the input node's default value, overriding the default provided by the referenced graph if the graph is a preset.
+	UFUNCTION(BlueprintCallable, Category = "Audio|MetaSound|Builder", meta = (ExpandEnumAsExecs = "OutResult"))
+	void SetGraphInputDefault(FName InputName, const FMetasoundFrontendLiteral& Literal, EMetaSoundBuilderResult& OutResult);
+
+	// Disconnects the given graph output's respective template nodes and sets the graph output's AccessType should it not match the current AccessType.
+	// Result succeeds if the AccessType was successfully changed or if the provided AccessType is already the output's current AccessType.
+	UFUNCTION(BlueprintCallable, Category = "Audio|MetaSound|Builder", meta = (ExpandEnumAsExecs = "OutResult"))
+	void SetGraphOutputAccessType(FName OutputName, EMetasoundFrontendVertexAccessType AccessType, EMetaSoundBuilderResult& OutResult);
+
+	// Disconnects the given graph output's respective template nodes and sets the graph output's DataType should it not match the current DataType.
+	// Result succeeds if the DataType was successfully changed or if the provided DataType is already the output's current DataType.
+	UFUNCTION(BlueprintCallable, Category = "Audio|MetaSound|Builder", meta = (ExpandEnumAsExecs = "OutResult"))
+	void SetGraphOutputDataType(FName OutputName, FName DataType, EMetaSoundBuilderResult& OutResult);
+
+	void SetMemberMetadata(UMetaSoundFrontendMemberMetadata& NewMetadata);
+
+	UE_DEPRECATED(5.5, "Call directly on Frontend builder using 'GetBuilder'")
+	void UpdateDependencyClassNames(const TMap<FMetasoundFrontendClassName, FMetasoundFrontendClassName>& OldToNewReferencedClassNames);
+
+	virtual TScriptInterface<IMetaSoundDocumentInterface> Build(UObject* Parent, const FMetaSoundBuilderOptions& Options) const PURE_VIRTUAL(UMetaSoundBuilderBase::Build, return { }; );
+
+	virtual bool ConformObjectToDocument();
+
+	// Returns the base class registered with the MetaSound UObject registry.
+	virtual const UClass& GetBaseMetaSoundUClass() const PURE_VIRTUAL(UMetaSoundBuilderBase::GetBaseMetaSoundUClass, return *UClass::StaticClass(); );
+
+	UE_DEPRECATED(5.5, "Renamed to 'GetBaseMetaSoundUClass' for consistency")
+	virtual const UClass& GetBuilderUClass() const { return GetBaseMetaSoundUClass(); }
+
+	UE_DEPRECATED(5.4, "Moved to Initialize and should only be called by internal implementation (i.e. IDocumentBuilderRegistry'")
+	virtual void InitFrontendBuilder();
+
+	// Initializes and ensures all nodes have a position (required prior to exporting to an asset if expected to be viewed in the editor).
+	virtual void InitNodeLocations();
+
+#if WITH_EDITOR
+	// Injects template nodes between builder's document inputs not connected
+	// to existing template inputs, copying locational data from the represented
+	// input metadata. If bForceNodeCreation is false, only generates a template
+	// input node if a connection between the input and other nodes exists. If true,
+	// will inject template node irrespective of whether or not the input has connections.
+	void InjectInputTemplateNodes(bool bForceNodeCreation, EMetaSoundBuilderResult& OutResult);
+
+	const FMetaSoundFrontendGraphComment* FindGraphComment(const FGuid& InCommentID) const;
+	FMetaSoundFrontendGraphComment* FindGraphComment(const FGuid & InCommentID);
+	FMetaSoundFrontendGraphComment& FindOrAddGraphComment(const FGuid& InCommentID);
+	bool RemoveGraphComment(const FGuid& InCommentID);
+
+	void SetNodeComment(const FMetaSoundNodeHandle& InNodeHandle, const FString& InNewComment, EMetaSoundBuilderResult& OutResult);
+	void SetNodeCommentVisible(const FMetaSoundNodeHandle& InNodeHandle, bool bIsVisible, EMetaSoundBuilderResult& OutResult);
+	void SetNodeLocation(const FMetaSoundNodeHandle& InNodeHandle, const FVector2D& InLocation, EMetaSoundBuilderResult& OutResult);
+	void SetNodeLocation(const FMetaSoundNodeHandle & InNodeHandle, const FVector2D& InLocation, const FGuid& InLocationGuid, EMetaSoundBuilderResult& OutResult);
+#endif // WITH_EDITOR
+
+	FMetaSoundFrontendDocumentBuilder& GetBuilder();
+	const FMetaSoundFrontendDocumentBuilder& GetConstBuilder() const;
+	int32 GetLastTransactionRegistered() const;
+
+	// Resets FrontendBuilder instance, creating a transient MetaSound document that is internal to this UObject Builder.
+	void Initialize();
+
+protected:
+	virtual void InitDelegates(Metasound::Frontend::FDocumentModifyDelegates & OutDocumentDelegates);
+
+	virtual void OnAssetReferenceAdded(TScriptInterface<IMetaSoundDocumentInterface> DocInterface) PURE_VIRTUAL(UMetaSoundBuilderBase::OnAssetReferenceAdded, );
+	virtual void OnRemovingAssetReference(TScriptInterface<IMetaSoundDocumentInterface> DocInterface) PURE_VIRTUAL(UMetaSoundBuilderBase::OnRemovingAssetReference, );
+
+	UE_DEPRECATED(5.5, "Moved to internal implementation")
+	virtual void CreateTransientBuilder() { };
+
+	UE_DEPRECATED(5.5, "Moved to 'Reload', to enforce generation of new delegates")
+	void InvalidateCache(bool bPrimeCache = false);
+
+	// Runs build, conforming the document and corresponding object data on a MetaSound UObject to that managed by this builder.
+	template <typename UClassType>
+	UClassType& BuildInternal(UObject* Parent, const FMetaSoundBuilderOptions& BuilderOptions) const
+	{
+		using namespace Metasound::Frontend;
+
+		UClassType* MetaSound = nullptr;
+		const FMetasoundFrontendClassName* DocClassName = nullptr;
+		if (BuilderOptions.ExistingMetaSound)
+		{
+			MetaSound = CastChecked<UClassType>(BuilderOptions.ExistingMetaSound.GetObject());
+
+			// Always unregister if mutating existing object. If bAddToRegistry is set to false,
+			// leaving registered would result in any references to this MetaSound executing on
+			// out-of-date data. If bAddToRegistry is set, then it needs to be unregistered before
+			// being registered as it is below.
+			if (MetaSound)
+			{
+				// If MetaSound already exists, preserve the class name to avoid
+				// nametable bloat & preserve potentially existing references.
+				const FMetasoundFrontendDocument& ExistingDoc = CastChecked<const UClassType>(MetaSound)->GetConstDocumentChecked();
+				if (!BuilderOptions.bForceUniqueClassName)
+				{
+					DocClassName = &ExistingDoc.RootGraph.Metadata.GetClassName();
+				}
+
+				MetaSound->UnregisterGraphWithFrontend();
+			}
+		}
+		else
+		{
+			FName ObjectName = BuilderOptions.Name;
+			if (!ObjectName.IsNone())
+			{
+				ObjectName = MakeUniqueObjectName(Parent, UClassType::StaticClass(), BuilderOptions.Name);
+			}
+
+			if (!Parent)
+			{
+				Parent = GetTransientPackage();
+			}
+
+			MetaSound = NewObject<UClassType>(Parent, ObjectName, RF_Public | RF_Transient);
+		}
+
+		checkf(MetaSound, TEXT("Failed to build MetaSound from builder '%s'"), *GetPathName());
+
+		TScriptInterface<IMetaSoundDocumentInterface> MetaSoundInterface = MetaSound;
+		BuildInternal(MetaSoundInterface, DocClassName);
+
+		if (BuilderOptions.bAddToRegistry)
+		{
+			MetaSound->RegisterGraphWithFrontend();
+		}
+
+		UE_LOG(LogMetaSound, VeryVerbose, TEXT("MetaSound '%s' built from '%s'"), *BuilderOptions.Name.ToString(), *GetFullName());
+		return *MetaSound;
+	}
+
+	UE_DEPRECATED(5.4, "Use UMetaSoundBuilderDocument::Create instead")
+	UMetaSoundBuilderDocument* CreateTransientDocumentObject() const;
+
+	// Only registers provided MetaSound's graph class and referenced graphs recursively if
+	// it has yet to be registered or if it has an attached builder reporting outstanding
+	// transactions that have yet to be registered.
+	static void RegisterGraphIfOutstandingTransactions(UObject& InMetaSound);
+
+	UPROPERTY()
+	FMetaSoundFrontendDocumentBuilder Builder;
+
+	// ClassName of associated document object.  Cached in case DocumentObject is deleted but continues to provide an identifier after destruction.
+	UPROPERTY()
+	FMetasoundFrontendClassName ClassName;
+
+#if WITH_EDITORONLY_DATA
+	UPROPERTY(meta = (DeprecatedProperty, DeprecationMessage = "5.4 - All source builders now operate on an underlying document source document that is also used to audition."))
+	bool bIsAttached = false;
+#endif // WITH_EDITORONLY_DATA
+
+private:
+	void BuildInternal(TScriptInterface<IMetaSoundDocumentInterface> NewMetaSound, const FMetasoundFrontendClassName* InDocClassName) const;
+	void OnDependencyAdded(int32 Index);
+	void OnRemoveSwappingDependency(int32 Index, int32 LastIndex);
+	void OnRenamingDependencyInternal(const int32 Index, const FMetasoundFrontendClassName& NewName);
+
+private:
+	int32 LastTransactionRegistered = 0;
+
+	// Friending allows for swapping the builder in certain circumstances where desired (eg. attaching a builder to an existing asset)
+	// or performing editor-only manipulations to the builder (ex. adding template nodes) and versioning the document data where necessary.
+	friend class UMetaSoundBuilderSubsystem;
+};
\ No newline at end of file
diff --git a/Engine/Plugins/Runtime/Metasound/Source/MetasoundEngine/Public/MetasoundBuilderSubsystem.h b/Engine/Plugins/Runtime/Metasound/Source/MetasoundEngine/Public/MetasoundBuilderSubsystem.h
--- a/Engine/Plugins/Runtime/Metasound/Source/MetasoundEngine/Public/MetasoundBuilderSubsystem.h
+++ b/Engine/Plugins/Runtime/Metasound/Source/MetasoundEngine/Public/MetasoundBuilderSubsystem.h
@@ -5,6 +5,7 @@
 #include "Delegates/IDelegateInstance.h"
 #include "Engine/Engine.h"
 #include "Interfaces/MetasoundOutputFormatInterfaces.h"
+#include "MetasoundBuilderBase.h"
 #include "MetasoundDocumentInterface.h"
 #include "MetasoundDynamicOperatorTransactor.h"
 #include "MetasoundFrontendController.h"
@@ -37,519 +38,14 @@ enum class EMetaSoundOutputAudioFormat : uint8;
 
 namespace Metasound::Engine
 {
+	// Forward Declarations
 	struct FOutputAudioFormatInfo;
+	class FMetaSoundAssetManager;
 } // namespace Metasound::Engine
 
 
 DECLARE_DYNAMIC_DELEGATE_OneParam(FOnCreateAuditionGeneratorHandleDelegate, UMetasoundGeneratorHandle*, GeneratorHandle);
 
-USTRUCT(BlueprintType, meta = (DisplayName = "MetaSound Node Input Handle"))
-struct METASOUNDENGINE_API FMetaSoundBuilderNodeInputHandle : public FMetasoundFrontendVertexHandle
-{
-	GENERATED_BODY()
-
-public:
-	FMetaSoundBuilderNodeInputHandle() = default;
-	FMetaSoundBuilderNodeInputHandle(const FGuid InNodeID, const FGuid& InVertexID)
-	{
-		NodeID = InNodeID;
-		VertexID = InVertexID;
-	}
-};
-
-USTRUCT(BlueprintType, meta = (DisplayName = "MetaSound Node Output Handle"))
-struct METASOUNDENGINE_API FMetaSoundBuilderNodeOutputHandle : public FMetasoundFrontendVertexHandle
-{
-	GENERATED_BODY()
-
-public:
-	FMetaSoundBuilderNodeOutputHandle() = default;
-	FMetaSoundBuilderNodeOutputHandle(const FGuid InNodeID, const FGuid& InVertexID)
-	{
-		NodeID = InNodeID;
-		VertexID = InVertexID;
-	}
-};
-
-USTRUCT(BlueprintType)
-struct METASOUNDENGINE_API FMetaSoundNodeHandle
-{
-	GENERATED_BODY()
-
-	UPROPERTY()
-	FGuid NodeID;
-
-public:
-	FMetaSoundNodeHandle() = default;
-	FMetaSoundNodeHandle(const FGuid& InNodeID)
-		: NodeID(InNodeID)
-	{
-	}
-
-	// Returns whether or not the vertex handle is set (may or may not be
-	// valid depending on what builder context it is referenced against)
-	bool IsSet() const
-	{
-		return NodeID.IsValid();
-	}
-};
-
-USTRUCT(BlueprintType)
-struct METASOUNDENGINE_API FMetaSoundBuilderOptions
-{
-	GENERATED_BODY()
-
-	// Name of generated object. If object already exists, used as the base name to ensure new object is unique.
-	UPROPERTY(EditAnywhere, BlueprintReadWrite, Category = "MetaSound|Builder")
-	FName Name;
-
-	// If the resulting MetaSound is building over an existing document, a unique class name will be generated,
-	// invalidating any referencing MetaSounds and registering the MetaSound as a new entry in the Frontend. If
-	// building a new document, option is ignored.
-	UPROPERTY(EditAnywhere, BlueprintReadWrite, Category = "MetaSound|Builder", meta = (AdvancedDisplay))
-	bool bForceUniqueClassName = false;
-
-	// If true, adds MetaSound to node registry, making it available
-	// for reference by other dynamically created MetaSounds.
-	UPROPERTY()
-	bool bAddToRegistry = true;
-
-	// If set, builder overwrites the given MetaSound's document with the builder's copy
-	// (ignores the Name field above).
-	UPROPERTY()
-	TScriptInterface<IMetaSoundDocumentInterface> ExistingMetaSound;
-};
-
-UENUM(BlueprintType)
-enum class EMetaSoundBuilderResult : uint8
-{
-	Succeeded,
-	Failed
-};
-
-/** Base implementation of MetaSound builder */
-UCLASS(Abstract)
-class METASOUNDENGINE_API UMetaSoundBuilderBase : public UObject
-{
-	GENERATED_BODY()
-
-public:
-	// Begin UObject interface
-	virtual void BeginDestroy() override;
-	// End UObject interface
-
-	// Adds a graph input node with the given name, DataType, and sets the graph input to default value.
-	// Returns the new input node's output handle if it was successfully created, or an invalid handle if it failed.
-	UFUNCTION(BlueprintCallable, Category = "Audio|MetaSound|Builder", meta = (ExpandEnumAsExecs = "OutResult", AdvancedDisplay = "3"))
-	UPARAM(DisplayName = "Output Handle") FMetaSoundBuilderNodeOutputHandle AddGraphInputNode(FName Name, FName DataType, FMetasoundFrontendLiteral DefaultValue, EMetaSoundBuilderResult& OutResult, bool bIsConstructorInput = false);
-
-	// Adds a graph output node with the given name, DataType, and sets output node's input to default value.
-	// Returns the new output node's input handle if it was successfully created, or an invalid handle if it failed.
-	UFUNCTION(BlueprintCallable, Category = "Audio|MetaSound|Builder", meta = (ExpandEnumAsExecs = "OutResult", AdvancedDisplay = "3"))
-	UPARAM(DisplayName = "Input Handle") FMetaSoundBuilderNodeInputHandle AddGraphOutputNode(FName Name, FName DataType, FMetasoundFrontendLiteral DefaultValue, EMetaSoundBuilderResult& OutResult, bool bIsConstructorOutput = false);
-
-	// Adds an interface registered with the given name to the graph, adding associated input and output nodes.
-	UFUNCTION(BlueprintCallable, Category = "Audio|MetaSound|Builder", meta = (ExpandEnumAsExecs = "OutResult"))
-	void AddInterface(FName InterfaceName, EMetaSoundBuilderResult& OutResult);
-
-	// Adds a node to the graph using the provided MetaSound asset as its defining NodeClass.
-	// Returns a node handle to the created node if successful, or an invalid handle if it failed.
-	UFUNCTION(BlueprintCallable, Category = "Audio|MetaSound|Builder", DisplayName = "Add MetaSound Node From Asset Class", meta = (ExpandEnumAsExecs = "OutResult"))
-	UPARAM(DisplayName = "Node Handle") FMetaSoundNodeHandle AddNode(const TScriptInterface<IMetaSoundDocumentInterface>& NodeClass, EMetaSoundBuilderResult& OutResult);
-
-	// Adds node referencing the highest native class version of the given class name to the document.
-	// Returns a node handle to the created node if successful, or an invalid handle if it failed.
-	UFUNCTION(BlueprintCallable, Category = "Audio|MetaSound|Builder", DisplayName = "Add MetaSound Node By ClassName", meta = (ExpandEnumAsExecs = "OutResult", AdvancedDisplay = "2"))
-	UPARAM(DisplayName = "Node Handle") FMetaSoundNodeHandle AddNodeByClassName(const FMetasoundFrontendClassName& ClassName, EMetaSoundBuilderResult& OutResult, int32 MajorVersion = 1);
-	
-	UE_DEPRECATED(5.4, "This version of AddNodeByClassName is deprecated. Use the one with a default MajorVersion of 1.")
-	UPARAM(DisplayName = "Node Handle") FMetaSoundNodeHandle AddNodeByClassName(const FMetasoundFrontendClassName& ClassName, int32 MajorVersion, EMetaSoundBuilderResult& OutResult);
-
-#if WITH_EDITOR
-	bool ClearMemberMetadata(const FGuid& InMemberID);
-#endif // WITH_EDITOR
-
-	// Connects node output to a node input. Does *NOT* provide loop detection for performance reasons.  Loop detection is checked on class registration when built or played.
-	// Returns succeeded if connection made, failed if connection already exists with input, the data types do not match, or the connection is not supported due to access type
-	// incompatibility (ex. constructor input to non-constructor input).
-	UFUNCTION(BlueprintCallable, Category = "Audio|MetaSound|Builder", meta = (ExpandEnumAsExecs = "OutResult"))
-	void ConnectNodes(const FMetaSoundBuilderNodeOutputHandle& NodeOutputHandle, const FMetaSoundBuilderNodeInputHandle& NodeInputHandle, EMetaSoundBuilderResult& OutResult);
-
-	// Connects two nodes using defined MetaSound Interface Bindings registered with the MetaSound Interface registry.
-	UFUNCTION(BlueprintCallable, Category = "Audio|MetaSound|Builder", meta = (ExpandEnumAsExecs = "OutResult"))
-	void ConnectNodesByInterfaceBindings(const FMetaSoundNodeHandle& FromNodeHandle, const FMetaSoundNodeHandle& ToNodeHandle, EMetaSoundBuilderResult& OutResult);
-
-	// Connects a given node's outputs to all graph outputs for shared interfaces implemented on both the node's referenced class and the builder's MetaSound graph. Returns inputs of connected output nodes.
-	UFUNCTION(BlueprintCallable, Category = "Audio|MetaSound|Builder", meta = (ExpandEnumAsExecs = "OutResult"))
-	UPARAM(DisplayName = "Connected Graph Output Node Inputs") TArray<FMetaSoundBuilderNodeInputHandle> ConnectNodeOutputsToMatchingGraphInterfaceOutputs(const FMetaSoundNodeHandle& NodeHandle, EMetaSoundBuilderResult& OutResult);
-
-	// Connects a given node's inputs to all graph inputs for shared interfaces implemented on both the node's referenced class and the builder's MetaSound graph. Returns outputs of connected input nodes.
-	UFUNCTION(BlueprintCallable, Category = "Audio|MetaSound|Builder", meta = (ExpandEnumAsExecs = "OutResult"))
-	UPARAM(DisplayName = "Connected Graph Input Node Outputs") TArray<FMetaSoundBuilderNodeOutputHandle> ConnectNodeInputsToMatchingGraphInterfaceInputs(const FMetaSoundNodeHandle& NodeHandle, EMetaSoundBuilderResult& OutResult);
-
-	// Connects a given node output to the graph output with the given name.
-	UFUNCTION(BlueprintCallable, Category = "Audio|MetaSound|Builder", meta = (ExpandEnumAsExecs = "OutResult"))
-	void ConnectNodeOutputToGraphOutput(FName GraphOutputName, const FMetaSoundBuilderNodeOutputHandle& NodeOutputHandle, EMetaSoundBuilderResult& OutResult);
-
-	// Connects a given node input to the graph input with the given name.
-	UFUNCTION(BlueprintCallable, Category = "Audio|MetaSound|Builder", meta = (ExpandEnumAsExecs = "OutResult"))
-	void ConnectNodeInputToGraphInput(FName GraphInputName, const FMetaSoundBuilderNodeInputHandle& NodeInputHandle, EMetaSoundBuilderResult& OutResult);
-
-	// Returns whether node exists.
-	UFUNCTION(BlueprintCallable, Category = "Audio|MetaSound|Builder")
-	UPARAM(DisplayName = "IsValid") bool ContainsNode(const FMetaSoundNodeHandle& Node) const;
-
-	// Returns whether node input exists.
-	UFUNCTION(BlueprintCallable, Category = "Audio|MetaSound|Builder")
-	UPARAM(DisplayName = "IsValid") bool ContainsNodeInput(const FMetaSoundBuilderNodeInputHandle& Input) const;
-
-	// Returns whether node output exists.
-	UFUNCTION(BlueprintCallable, Category = "Audio|MetaSound|Builder")
-	UPARAM(DisplayName = "IsValid") bool ContainsNodeOutput(const FMetaSoundBuilderNodeOutputHandle& Output) const;
-
-	// Disconnects node output to a node input. Returns success if connection was removed, failed if not.
-	UFUNCTION(BlueprintCallable, Category = "Audio|MetaSound|Builder", meta = (ExpandEnumAsExecs = "OutResult"))
-	void DisconnectNodes(const FMetaSoundBuilderNodeOutputHandle& NodeOutputHandle, const FMetaSoundBuilderNodeInputHandle& NodeInputHandle, EMetaSoundBuilderResult& OutResult);
-
-	// Removes connection to a given node input. Returns success if connection was removed, failed if not.
-	UFUNCTION(BlueprintCallable, Category = "Audio|MetaSound|Builder", meta = (ExpandEnumAsExecs = "OutResult"))
-	void DisconnectNodeInput(const FMetaSoundBuilderNodeInputHandle& NodeInputHandle, EMetaSoundBuilderResult& OutResult);
-
-	// Removes all connections from a given node output. Returns success if all connections were removed, failed if not.
-	UFUNCTION(BlueprintCallable, Category = "Audio|MetaSound|Builder", meta = (ExpandEnumAsExecs = "OutResult"))
-	void DisconnectNodeOutput(const FMetaSoundBuilderNodeOutputHandle& NodeOutputHandle, EMetaSoundBuilderResult& OutResult);
-
-	// Disconnects two nodes using defined MetaSound Interface Bindings registered with the MetaSound Interface registry. Returns success if
-	// all connections were found and removed, failed if any connections were not.
-	UFUNCTION(BlueprintCallable, Category = "Audio|MetaSound|Builder", meta = (ExpandEnumAsExecs = "OutResult"))
-	void DisconnectNodesByInterfaceBindings(const FMetaSoundNodeHandle& FromNodeHandle, const FMetaSoundNodeHandle& ToNodeHandle, EMetaSoundBuilderResult& OutResult);
-
-	// Returns graph input node by the given name if it exists, or an invalid handle if not found.
-	UFUNCTION(BlueprintCallable, Category = "Audio|MetaSound|Builder", meta = (ExpandEnumAsExecs = "OutResult"))
-	UPARAM(DisplayName = "Node Handle") FMetaSoundNodeHandle FindGraphInputNode(FName InputName, EMetaSoundBuilderResult& OutResult);
-
-	// Returns graph output node by the given name if it exists, or an invalid handle if not found.
-	UFUNCTION(BlueprintCallable, Category = "Audio|MetaSound|Builder", meta = (ExpandEnumAsExecs = "OutResult"))
-	UPARAM(DisplayName = "Node Handle") FMetaSoundNodeHandle FindGraphOutputNode(FName OutputName, EMetaSoundBuilderResult& OutResult);
-
-#if WITH_EDITOR
-	UMetaSoundFrontendMemberMetadata* FindMemberMetadata(const FGuid& InMemberID);
-#endif // WITH_EDITOR
-
-	// Returns node input by the given name if it exists, or an invalid handle if not found.
-	UFUNCTION(BlueprintCallable, Category = "Audio|MetaSound|Builder", meta = (ExpandEnumAsExecs = "OutResult"))
-	UPARAM(DisplayName = "Input Handle") FMetaSoundBuilderNodeInputHandle FindNodeInputByName(const FMetaSoundNodeHandle& NodeHandle, FName InputName, EMetaSoundBuilderResult& OutResult);
-
-	// Returns all node inputs.
-	UFUNCTION(BlueprintCallable, Category = "Audio|MetaSound|Builder", meta = (ExpandEnumAsExecs = "OutResult"))
-	UPARAM(DisplayName = "Input  Handles") TArray<FMetaSoundBuilderNodeInputHandle> FindNodeInputs(const FMetaSoundNodeHandle& NodeHandle, EMetaSoundBuilderResult& OutResult);
-
-	// Returns node inputs by the given DataType (ex. "Audio", "Trigger", "String", "Bool", "Float", "Int32", etc.).
-	UFUNCTION(BlueprintCallable, Category = "Audio|MetaSound|Builder", meta = (ExpandEnumAsExecs = "OutResult"))
-	UPARAM(DisplayName = "Input Handles") TArray<FMetaSoundBuilderNodeInputHandle> FindNodeInputsByDataType(const FMetaSoundNodeHandle& NodeHandle, EMetaSoundBuilderResult& OutResult, FName DataType);
-
-	// Returns node output by the given name.
-	UFUNCTION(BlueprintCallable, Category = "Audio|MetaSound|Builder", meta = (ExpandEnumAsExecs = "OutResult"))
-	UPARAM(DisplayName = "Output Handle") FMetaSoundBuilderNodeOutputHandle FindNodeOutputByName(const FMetaSoundNodeHandle& NodeHandle, FName OutputName, EMetaSoundBuilderResult& OutResult);
-
-	// Returns all node outputs.
-	UFUNCTION(BlueprintCallable, Category = "Audio|MetaSound|Builder", meta = (ExpandEnumAsExecs = "OutResult"))
-	UPARAM(DisplayName = "Output Handles") TArray<FMetaSoundBuilderNodeOutputHandle> FindNodeOutputs(const FMetaSoundNodeHandle& NodeHandle, EMetaSoundBuilderResult& OutResult);
-
-	// Returns node outputs by the given DataType (ex. "Audio", "Trigger", "String", "Bool", "Float", "Int32", etc.).
-	UFUNCTION(BlueprintCallable, Category = "Audio|MetaSound|Builder", meta = (ExpandEnumAsExecs = "OutResult"))
-	UPARAM(DisplayName = "Output Handles") TArray<FMetaSoundBuilderNodeOutputHandle> FindNodeOutputsByDataType(const FMetaSoundNodeHandle& NodeHandle, EMetaSoundBuilderResult& OutResult, FName DataType);
-
-	// Returns input nodes associated with a given interface.
-	UFUNCTION(BlueprintCallable, Category = "Audio|MetaSound|Builder", meta = (ExpandEnumAsExecs = "OutResult"))
-	UPARAM(DisplayName = "Input Node Handles") TArray<FMetaSoundNodeHandle> FindInterfaceInputNodes(FName InterfaceName, EMetaSoundBuilderResult& OutResult);
-
-	// Returns output nodes associated with a given interface.
-	UFUNCTION(BlueprintCallable, Category = "Audio|MetaSound|Builder", meta = (ExpandEnumAsExecs = "OutResult"))
-	UPARAM(DisplayName = "Output Node Handles") TArray<FMetaSoundNodeHandle> FindInterfaceOutputNodes(FName InterfaceName, EMetaSoundBuilderResult& OutResult);
-
-	// Returns input's parent node if the input is valid, otherwise returns invalid node handle.
-	UFUNCTION(BlueprintCallable, Category = "Audio|MetaSound|Builder", meta = (ExpandEnumAsExecs = "OutResult"))
-	UPARAM(DisplayName = "Node Handle") FMetaSoundNodeHandle FindNodeInputParent(const FMetaSoundBuilderNodeInputHandle& InputHandle, EMetaSoundBuilderResult& OutResult);
-
-	// Returns output's parent node if the input is valid, otherwise returns invalid node handle.
-	UFUNCTION(BlueprintCallable, Category = "Audio|MetaSound|Builder", meta = (ExpandEnumAsExecs = "OutResult"))
-	UPARAM(DisplayName = "Node Handle") FMetaSoundNodeHandle FindNodeOutputParent(const FMetaSoundBuilderNodeOutputHandle& OutputHandle, EMetaSoundBuilderResult& OutResult);
-
-	// Returns output's parent node if the input is valid, otherwise returns invalid node handle.
-	UFUNCTION(BlueprintCallable, Category = "Audio|MetaSound|Builder", meta = (ExpandEnumAsExecs = "OutResult"))
-	UPARAM(DisplayName = "Node ClassVersion") FMetasoundFrontendVersion FindNodeClassVersion(const FMetaSoundNodeHandle& NodeHandle, EMetaSoundBuilderResult& OutResult);
-
-	// Returns the document's root graph class name
-	FMetasoundFrontendClassName GetRootGraphClassName() const;
-
-	// Returns node input's data if valid (including things like name and datatype).
-	UFUNCTION(BlueprintCallable, Category = "Audio|MetaSound|Builder", meta = (ExpandEnumAsExecs = "OutResult"))
-	void GetNodeInputData(const FMetaSoundBuilderNodeInputHandle& InputHandle, FName& Name, FName& DataType, EMetaSoundBuilderResult& OutResult);
-
-	// Returns node input's literal value if set on graph, otherwise fails and returns default literal.
-	UFUNCTION(BlueprintCallable, Category = "Audio|MetaSound|Builder", meta = (ExpandEnumAsExecs = "OutResult"))
-	FMetasoundFrontendLiteral GetNodeInputDefault(const FMetaSoundBuilderNodeInputHandle& InputHandle, EMetaSoundBuilderResult& OutResult);
-
-	// Returns node input's class literal value if set, otherwise fails and returns default literal.
-	UFUNCTION(BlueprintCallable, Category = "Audio|MetaSound|Builder", meta = (ExpandEnumAsExecs = "OutResult"))
-	FMetasoundFrontendLiteral GetNodeInputClassDefault(const FMetaSoundBuilderNodeInputHandle& InputHandle, EMetaSoundBuilderResult& OutResult);
-
-	// Returns whether the given node input is a constructor pin
-	UFUNCTION(BlueprintCallable, Category = "Audio|MetaSound|Builder")
-	bool GetNodeInputIsConstructorPin(const FMetaSoundBuilderNodeInputHandle& InputHandle) const;
-
-	// Returns node output's data if valid (including things like name and datatype).
-	UFUNCTION(BlueprintCallable, Category = "Audio|MetaSound|Builder", meta = (ExpandEnumAsExecs = "OutResult"))
-	void GetNodeOutputData(const FMetaSoundBuilderNodeOutputHandle& OutputHandle, FName& Name, FName& DataType, EMetaSoundBuilderResult& OutResult);
-	
-	// Returns whether the given node output is a constructor pin
-	UFUNCTION(BlueprintCallable, Category = "Audio|MetaSound|Builder")
-	bool GetNodeOutputIsConstructorPin(const FMetaSoundBuilderNodeOutputHandle& OutputHandle) const;
-
-	// Return the asset referenced by this preset builder. Returns nullptr if the builder is not a preset.
-	UFUNCTION(BlueprintCallable, Category = "Audio|MetaSound|Builder")
-	UObject* GetReferencedPresetAsset() const;
-
-	// Returns if a given interface is declared.
-	UFUNCTION(BlueprintCallable, Category = "Audio|MetaSound|Builder")
-	UPARAM(DisplayName = "IsDeclared") bool InterfaceIsDeclared(FName InterfaceName) const;
-
-	// Returns if a given node output and node input are connected.
-	UFUNCTION(BlueprintCallable, Category = "Audio|MetaSound|Builder")
-	UPARAM(DisplayName = "Connected") bool NodesAreConnected(const FMetaSoundBuilderNodeOutputHandle& OutputHandle, const FMetaSoundBuilderNodeInputHandle& InputHandle) const;
-
-	// Returns if a given node input has connections.
-	UFUNCTION(BlueprintCallable, Category = "Audio|MetaSound|Builder")
-	UPARAM(DisplayName = "Connected") bool NodeInputIsConnected(const FMetaSoundBuilderNodeInputHandle& InputHandle) const;
-
-	// Returns if a given node output is connected.
-	UFUNCTION(BlueprintCallable, Category = "Audio|MetaSound|Builder")
-	UPARAM(DisplayName = "Connected") bool NodeOutputIsConnected(const FMetaSoundBuilderNodeOutputHandle& OutputHandle) const;
-
-	// Returns whether this is a preset.
-	UFUNCTION(BlueprintCallable, Category = "Audio|MetaSound|Builder")
-	bool IsPreset() const;
-
-	// Converts this preset to a fully accessible MetaSound; sets result to succeeded if it was converted successfully and failed if it was not.
-	UFUNCTION(BlueprintCallable, Category = "Audio|MetaSound|Builder", meta = (ExpandEnumAsExecs = "OutResult"))
-	void ConvertFromPreset(EMetaSoundBuilderResult& OutResult);
-
-	// Convert this builder to a MetaSound source preset with the given referenced source builder 
-	UFUNCTION(BlueprintCallable, Category = "Audio|MetaSound|Builder", meta = (ExpandEnumAsExecs = "OutResult"))
-	void ConvertToPreset(const TScriptInterface<IMetaSoundDocumentInterface>& ReferencedNodeClass, EMetaSoundBuilderResult& OutResult);
-
-	// Removes graph input if it exists; sets result to succeeded if it was removed and failed if it was not.
-	UFUNCTION(BlueprintCallable, Category = "Audio|MetaSound|Builder", meta = (ExpandEnumAsExecs = "OutResult"))
-	void RemoveGraphInput(FName Name, EMetaSoundBuilderResult& OutResult);
-
-	// Removes graph output if it exists; sets result to succeeded if it was removed and failed if it was not.
-	UFUNCTION(BlueprintCallable, Category = "Audio|MetaSound|Builder", meta = (ExpandEnumAsExecs = "OutResult"))
-	void RemoveGraphOutput(FName Name, EMetaSoundBuilderResult& OutResult);
-
-	// Removes the interface with the given name from the builder's MetaSound. Removes any graph inputs
-	// and outputs associated with the given interface and their respective connections (if any).
-	UFUNCTION(BlueprintCallable, Category = "Audio|MetaSound|Builder", meta = (ExpandEnumAsExecs = "OutResult"))
-	void RemoveInterface(FName InterfaceName, EMetaSoundBuilderResult& OutResult);
-
-	// Removes node and any associated connections from the builder's MetaSound.
-	UFUNCTION(BlueprintCallable, Category = "Audio|MetaSound|Builder", meta = (ExpandEnumAsExecs = "OutResult"))
-	void RemoveNode(const FMetaSoundNodeHandle& NodeHandle, EMetaSoundBuilderResult& OutResult, bool bRemoveUnusedDependencies = true);
-
-	// Removes node input literal default if set, reverting the value to be whatever the node class defaults the value to.
-	// Returns success if value was removed, false if not removed (i.e. wasn't set to begin with).
-	UFUNCTION(BlueprintCallable, Category = "Audio|MetaSound|Builder", meta = (ExpandEnumAsExecs = "OutResult"))
-	void RemoveNodeInputDefault(const FMetaSoundBuilderNodeInputHandle& InputHandle, EMetaSoundBuilderResult& OutResult);
-
-	// Removes dependencies in document that are no longer referenced by nodes
-	UFUNCTION(BlueprintCallable, Category = "Audio|MetaSound|Builder")
-	void RemoveUnusedDependencies();
-
-	// Rename the document's root graph class with a guid and optional namespace and variant
-	void RenameRootGraphClass(const FMetasoundFrontendClassName& InName);
-
-	UE_DEPRECATED(5.5, "Moved to 'Reload' to discourage priming by defaulting to false. "
-		"In addition, internal delegates and respective bindings are now regenerated when reloading, "
-		"effectively resetting the builder's entire state and not just strictly the underlying cache.")
-	void ReloadCache(bool bPrimeCache = true);
-
-	// Reloads the builder, freeing the internal cache and rebuilding delegate bindings
-	void Reload(bool bPrimeCache = false);
-
-#if WITH_EDITOR
-	// Sets the author of the MetaSound.
-	void SetAuthor(const FString& InAuthor);
-#endif // WITH_EDITOR
-
-	// Sets the node's input default value (used if no connection to the given node input is present)
-	UFUNCTION(BlueprintCallable, Category = "Audio|MetaSound|Builder", meta = (ExpandEnumAsExecs = "OutResult"))
-	void SetNodeInputDefault(const FMetaSoundBuilderNodeInputHandle& NodeInputHandle, const FMetasoundFrontendLiteral& Literal, EMetaSoundBuilderResult& OutResult);
-
-	// Disconnects the given graph input's respective template nodes and sets the graph input's AccessType should it not match the current AccessType.
-	// Result succeeds if the AccessType was successfully changed or if the provided AccessType is already the input's current AccessType.
-	UFUNCTION(BlueprintCallable, Category = "Audio|MetaSound|Builder", meta = (ExpandEnumAsExecs = "OutResult"))
-	void SetGraphInputAccessType(FName InputName, EMetasoundFrontendVertexAccessType AccessType, EMetaSoundBuilderResult& OutResult);
-
-	// Disconnects the given graph input's respective template nodes and sets the graph input's DataType should it not match the current DataType.
-	// Result succeeds if the DataType was successfully changed or if the provided DataType is already the input's current DataType.
-	UFUNCTION(BlueprintCallable, Category = "Audio|MetaSound|Builder", meta = (ExpandEnumAsExecs = "OutResult"))
-	void SetGraphInputDataType(FName InputName, FName DataType, EMetaSoundBuilderResult& OutResult);
-
-	// Sets the input node's default value, overriding the default provided by the referenced graph if the graph is a preset.
-	UFUNCTION(BlueprintCallable, Category = "Audio|MetaSound|Builder", meta = (ExpandEnumAsExecs = "OutResult"))
-	void SetGraphInputDefault(FName InputName, const FMetasoundFrontendLiteral& Literal, EMetaSoundBuilderResult& OutResult);
-
-	// Disconnects the given graph output's respective template nodes and sets the graph output's AccessType should it not match the current AccessType.
-	// Result succeeds if the AccessType was successfully changed or if the provided AccessType is already the output's current AccessType.
-	UFUNCTION(BlueprintCallable, Category = "Audio|MetaSound|Builder", meta = (ExpandEnumAsExecs = "OutResult"))
-	void SetGraphOutputAccessType(FName OutputName, EMetasoundFrontendVertexAccessType AccessType, EMetaSoundBuilderResult& OutResult);
-
-	// Disconnects the given graph output's respective template nodes and sets the graph output's DataType should it not match the current DataType.
-	// Result succeeds if the DataType was successfully changed or if the provided DataType is already the output's current DataType.
-	UFUNCTION(BlueprintCallable, Category = "Audio|MetaSound|Builder", meta = (ExpandEnumAsExecs = "OutResult"))
-	void SetGraphOutputDataType(FName OutputName, FName DataType, EMetaSoundBuilderResult& OutResult);
-
-	void SetMemberMetadata(UMetaSoundFrontendMemberMetadata& NewMetadata);
-
-	// Update dependency class names given a map of old to new referenced class names. 
-	void UpdateDependencyClassNames(const TMap<FMetasoundFrontendClassName, FMetasoundFrontendClassName>& OldToNewReferencedClassNames);
-
-	virtual TScriptInterface<IMetaSoundDocumentInterface> Build(UObject* Parent, const FMetaSoundBuilderOptions& Options) const PURE_VIRTUAL(UMetaSoundBuilderBase::Build, return { }; );
-
-	// Returns the base MetaSound UClass the builder is operating on
-	virtual const UClass& GetBuilderUClass() const PURE_VIRTUAL(UMetaSoundBuilderBase::Build, return *UClass::StaticClass(); );
-
-	UE_DEPRECATED(5.4, "Moved to protected pure virtual 'CreateTransientBuilder'")
-	virtual void InitFrontendBuilder();
-
-	virtual bool ConformObjectToDocument() PURE_VIRTUAL(UMetaSoundBuilderBase::ConformObjectToDocument, return false; );
-
-	// Initializes and ensures all nodes have a position (required prior to exporting to an asset if expected to be viewed in the editor).
-	virtual void InitNodeLocations();
-
-#if WITH_EDITOR
-	// Injects template nodes between builder's document inputs not connected
-	// to existing template inputs, copying locational data from the represented
-	// input metadata. If bForceNodeCreation is false, only generates a template
-	// input node if a connection between the input and other nodes exists. If true,
-	// will inject template node irrespective of whether or not the input has connections.
-	void InjectInputTemplateNodes(bool bForceNodeCreation, EMetaSoundBuilderResult& OutResult);
-
-	const FMetaSoundFrontendGraphComment* FindGraphComment(const FGuid& InCommentID) const;
-	FMetaSoundFrontendGraphComment* FindGraphComment(const FGuid & InCommentID);
-	FMetaSoundFrontendGraphComment& FindOrAddGraphComment(const FGuid& InCommentID);
-	bool RemoveGraphComment(const FGuid& InCommentID);
-
-	void SetNodeComment(const FMetaSoundNodeHandle& InNodeHandle, const FString& InNewComment, EMetaSoundBuilderResult& OutResult);
-	void SetNodeCommentVisible(const FMetaSoundNodeHandle& InNodeHandle, bool bIsVisible, EMetaSoundBuilderResult& OutResult);
-	void SetNodeLocation(const FMetaSoundNodeHandle& InNodeHandle, const FVector2D& InLocation, EMetaSoundBuilderResult& OutResult);
-	void SetNodeLocation(const FMetaSoundNodeHandle & InNodeHandle, const FVector2D& InLocation, const FGuid& InLocationGuid, EMetaSoundBuilderResult& OutResult);
-#endif // WITH_EDITOR
-
-	const FMetaSoundFrontendDocumentBuilder& GetConstBuilder() const;
-
-protected:
-	virtual void InitDelegates(Metasound::Frontend::FDocumentModifyDelegates & OutDocumentDelegates) { };
-
-	// Creates a FrontendBuilder wrapping a transient MetaSound object of the class supported by the given subsystem builder (See GetBuilderUClass).
-	// Should only be used when builders is not being attached to an existing, serialized MetaSound asset.
-	virtual void CreateTransientBuilder() PURE_VIRTUAL(UMetaSoundBuilderBase::CreateTransientBuilder, );
-
-	UE_DEPRECATED(5.5, "Moved to 'Reload', to enforce generation of new delegates")
-	void InvalidateCache(bool bPrimeCache = false);
-
-	// Runs build, conforming the document and corresponding object data on a MetaSound UObject to that managed by this builder.
-	template <typename UClassType>
-	UClassType& BuildInternal(UObject* Parent, const FMetaSoundBuilderOptions& BuilderOptions) const
-	{
-		using namespace Metasound::Frontend;
-
-		UClassType* MetaSound = nullptr;
-		const FMetasoundFrontendClassName* DocClassName = nullptr;
-		if (BuilderOptions.ExistingMetaSound)
-		{
-			MetaSound = CastChecked<UClassType>(BuilderOptions.ExistingMetaSound.GetObject());
-
-			// Always unregister if mutating existing object. If bAddToRegistry is set to false,
-			// leaving registered would result in any references to this MetaSound executing on
-			// out-of-date data. If bAddToRegistry is set, then it needs to be unregistered before
-			// being registered as it is below.
-			if (MetaSound)
-			{
-				// If MetaSound already exists, preserve the class name to avoid
-				// nametable bloat & preserve potentially existing references.
-				const FMetasoundFrontendDocument& ExistingDoc = CastChecked<const UClassType>(MetaSound)->GetConstDocumentChecked();
-				if (!BuilderOptions.bForceUniqueClassName)
-				{
-					DocClassName = &ExistingDoc.RootGraph.Metadata.GetClassName();
-				}
-
-				MetaSound->UnregisterGraphWithFrontend();
-			}
-		}
-		else
-		{
-			FName ObjectName = BuilderOptions.Name;
-			if (!ObjectName.IsNone())
-			{
-				ObjectName = MakeUniqueObjectName(Parent, UClassType::StaticClass(), BuilderOptions.Name);
-			}
-
-			if (!Parent)
-			{
-				Parent = GetTransientPackage();
-			}
-
-			MetaSound = NewObject<UClassType>(Parent, ObjectName, RF_Public | RF_Transient);
-		}
-
-		checkf(MetaSound, TEXT("Failed to build MetaSound from builder '%s'"), *GetPathName());
-
-		TScriptInterface<IMetaSoundDocumentInterface> MetaSoundInterface = MetaSound;
-		BuildInternal(MetaSoundInterface, DocClassName);
-
-		if (BuilderOptions.bAddToRegistry)
-		{
-			MetaSound->RegisterGraphWithFrontend();
-		}
-
-		UE_LOG(LogMetaSound, VeryVerbose, TEXT("MetaSound '%s' built from '%s'"), *BuilderOptions.Name.ToString(), *GetFullName());
-		return *MetaSound;
-	}
-
-	UE_DEPRECATED(5.4, "Use UMetaSoundBuilderDocument::Create instead")
-	UMetaSoundBuilderDocument* CreateTransientDocumentObject() const;
-
-	// Only registers provided MetaSound's graph class and referenced graphs recursively if
-	// it has yet to be registered or if it has an attached builder reporting outstanding
-	// transactions that have yet to be registered.
-	static void RegisterGraphIfOutstandingTransactions(UObject& InMetaSound);
-
-	UPROPERTY()
-	FMetaSoundFrontendDocumentBuilder Builder;
-
-	// ClassName of associated document object.  Cached in case DocumentObject is deleted but continues to provide an identifier after destruction.
-	UPROPERTY()
-	FMetasoundFrontendClassName ClassName;
-
-#if WITH_EDITORONLY_DATA
-	UPROPERTY(meta = (DeprecatedProperty, DeprecationMessage = "5.4 - All source builders now operate on an underlying document source document that is also used to audition."))
-	bool bIsAttached = false;
-#endif // WITH_EDITORONLY_DATA
-
-
-private:
-	void BuildInternal(TScriptInterface<IMetaSoundDocumentInterface> NewMetaSound, const FMetasoundFrontendClassName* InDocClassName) const;
-
-private:
-	int32 LastTransactionRegistered = 0;
-
-	// Friending allows for swapping the builder in certain circumstances where desired (eg. attaching a builder to an existing asset)
-	// or performing editor-only manipulations to the builder (ex. adding template nodes) and versioning the document data where necessary.
-	friend class UMetaSoundAssetSubsystem;
-	friend class UMetaSoundBuilderSubsystem;
-};
-
 /** Builder in charge of building a MetaSound Patch */
 UCLASS(Transient, BlueprintType)
 class METASOUNDENGINE_API UMetaSoundPatchBuilder : public UMetaSoundBuilderBase
@@ -560,11 +56,11 @@ public:
 	UFUNCTION(BlueprintCallable, Category = "Audio|MetaSound|Builder", meta = (WorldContext = "Parent"))
 	virtual UPARAM(DisplayName = "MetaSound") TScriptInterface<IMetaSoundDocumentInterface> Build(UObject* Parent, const FMetaSoundBuilderOptions& Options) const override;
 
-	virtual bool ConformObjectToDocument() override;
-	virtual const UClass& GetBuilderUClass() const override;
+	virtual const UClass& GetBaseMetaSoundUClass() const override;
 
 protected:
-	virtual void CreateTransientBuilder() override;
+	virtual void OnAssetReferenceAdded(TScriptInterface<IMetaSoundDocumentInterface> DocInterface) override;
+	virtual void OnRemovingAssetReference(TScriptInterface<IMetaSoundDocumentInterface> DocInterface) override;
 
 	friend class UMetaSoundBuilderSubsystem;
 };
@@ -600,8 +96,7 @@ public:
 
 	const Metasound::Engine::FOutputAudioFormatInfoPair* FindOutputAudioFormatInfo() const;
 
-	virtual bool ConformObjectToDocument() override;
-	virtual const UClass& GetBuilderUClass() const override;
+	virtual const UClass& GetBaseMetaSoundUClass() const override;
 
 #if WITH_EDITORONLY_DATA
 	// Sets the MetaSound's BlockRate override (editor only, to allow setting per-platform values)
@@ -616,8 +111,9 @@ public:
 	void SetQuality(FName Quality);
 
 protected:
-	virtual void CreateTransientBuilder() override;
 	virtual void InitDelegates(Metasound::Frontend::FDocumentModifyDelegates& OutDocumentDelegates) override;
+	virtual void OnAssetReferenceAdded(TScriptInterface<IMetaSoundDocumentInterface> DocInterface) override;
+	virtual void OnRemovingAssetReference(TScriptInterface<IMetaSoundDocumentInterface> DocInterface) override;
 
 private:
 	static TOptional<Metasound::FAnyDataReference> CreateDataReference(const Metasound::FOperatorSettings& InOperatorSettings, FName DataType, const Metasound::FLiteral& InLiteral, Metasound::EDataReferenceAccessType AccessType);
@@ -649,7 +145,7 @@ private:
 
 /** The subsystem in charge of tracking MetaSound builders */
 UCLASS()
-class METASOUNDENGINE_API UMetaSoundBuilderSubsystem : public UEngineSubsystem, public Metasound::Frontend::IDocumentBuilderRegistry
+class METASOUNDENGINE_API UMetaSoundBuilderSubsystem : public UEngineSubsystem
 {
 	GENERATED_BODY()
 
@@ -657,29 +153,25 @@ private:
 	UPROPERTY()
 	TMap<FName, TObjectPtr<UMetaSoundBuilderBase>> NamedBuilders;
 
-	UPROPERTY()
-	mutable TMap<FMetasoundFrontendClassName, TWeakObjectPtr<UMetaSoundBuilderBase>> AssetBuilders;
-
-	UPROPERTY()
-	mutable TMap<FMetasoundFrontendClassName, TWeakObjectPtr<UMetaSoundBuilderBase>> TransientBuilders;
-
-public:
-	virtual void Initialize(FSubsystemCollectionBase& Collection) override;
-	virtual void InvalidateDocumentCache(const FMetasoundFrontendClassName& InClassName) const override;
+public:	
+	UE_DEPRECATED(5.5, "Direct call of 'Reload' can be applied to builder directly")
+	virtual void InvalidateDocumentCache(const FMetasoundFrontendClassName& InClassName) const;
 
 	static UMetaSoundBuilderSubsystem* Get();
 	static UMetaSoundBuilderSubsystem& GetChecked();
 	static const UMetaSoundBuilderSubsystem* GetConst();
 	static const UMetaSoundBuilderSubsystem& GetConstChecked();
 
+	UE_DEPRECATED(5.5, "Use FDocumentBuilderRegistry::FindOrBeginBuilding, which is now only supported in builds loading editor-only data.")
 	UMetaSoundBuilderBase& AttachBuilderToAssetChecked(UObject& InObject) const;
 
-	// Moving toward becoming UFUNCTIONS in the future, but right now unguarded case
-	// where a user could improperly reference transient assets. Also, unsafe where
-	// legacy controller system are potentially mutating and not notifying delegates
-	// appropriately.
+	UE_DEPRECATED(5.5, "Use FDocumentBuilderRegistry::FindOrBeginBuilding (when editor only data is loaded) or MetaSoundEditorSubsystem::FindOrBeginBuilding call")
 	UMetaSoundPatchBuilder* AttachPatchBuilderToAsset(UMetaSoundPatch* InPatch) const;
+
+	UE_DEPRECATED(5.5, "Use FDocumentBuilderRegistry::FindOrBeginBuilding (when editor only data is loaded) or MetaSoundEditorSubsystem::FindOrBeginBuilding call")
 	UMetaSoundSourceBuilder* AttachSourceBuilderToAsset(UMetaSoundSource* InSource) const;
+
+	UE_DEPRECATED(5.5, "Moved to IDocumentBuilderRegistry::RemoveBuilderFromAsset")
 	bool DetachBuilderFromAsset(const FMetasoundFrontendClassName& InClassName) const;
 
 	UFUNCTION(BlueprintCallable, Category = "Audio|MetaSound|Builder",  meta = (ExpandEnumAsExecs = "OutResult"))
@@ -756,8 +248,8 @@ public:
 	UPARAM(DisplayName = "Is Registered") bool IsInterfaceRegistered(FName InInterfaceName) const;
 
 #if WITH_EDITOR
-	// Updates builder data should a builder-managed asset be updated via the transaction stack.
-	void PostBuilderAssetTransaction(const FMetasoundFrontendClassName& InClassName);
+	UE_DEPRECATED(5.5, "No longer required as reload is now just directly called on a given builder.")
+	void PostBuilderAssetTransaction(const FMetasoundFrontendClassName& InClassName) { }
 #endif // WITH_EDITOR
 
 	// Adds builder to subsystem's registry to make it persistent and easily accessible by multiple systems or Blueprints
@@ -782,56 +274,5 @@ public:
 	UPARAM(DisplayName = "Unregistered") bool UnregisterSourceBuilder(FName BuilderName);
 
 private:
-	template <typename BuilderClass>
-	BuilderClass& CreateTransientBuilder(FName BuilderName = FName())
-	{
-		const EObjectFlags NewObjectFlags = RF_Public | RF_Transient;
-		UPackage* TransientPackage = GetTransientPackage();
-		const FName ObjectName = MakeUniqueObjectName(TransientPackage, BuilderClass::StaticClass(), BuilderName);
-		TObjectPtr<BuilderClass> NewBuilder = NewObject<BuilderClass>(TransientPackage, ObjectName, NewObjectFlags);
-		check(NewBuilder);
-		NewBuilder->CreateTransientBuilder();
-
-		const FMetasoundFrontendDocument& Document = NewBuilder->GetConstBuilder().GetDocument();
-		const FMetasoundFrontendClassName& ClassName = Document.RootGraph.Metadata.GetClassName();
-		NewBuilder->ClassName = ClassName;
-		TransientBuilders.Add(ClassName, NewBuilder);
-		return *NewBuilder.Get();
-	}
-
-	template <typename BuilderClass>
-	BuilderClass& AttachBuilderToAssetCheckedPrivate(UObject* InMetaSoundObject) const
-	{
-		check(InMetaSoundObject);
-		check(InMetaSoundObject->IsAsset());
-
-		TScriptInterface<IMetaSoundDocumentInterface> DocInterface = InMetaSoundObject;
-		const FMetasoundFrontendDocument& Document = DocInterface->GetConstDocument();
-		const FMetasoundFrontendClassName& FullClassName = Document.RootGraph.Metadata.GetClassName();
-
-		if (FullClassName.IsValid())
-		{
-			TWeakObjectPtr<UMetaSoundBuilderBase> Builder = AssetBuilders.FindRef(FullClassName);
-			if (Builder.IsValid())
-			{
-				return *CastChecked<BuilderClass>(Builder.Get());
-			}
-		}
-
-		TObjectPtr<BuilderClass> NewBuilder = NewObject<BuilderClass>(InMetaSoundObject);
-		check(NewBuilder);
-		NewBuilder->Builder = FMetaSoundFrontendDocumentBuilder(DocInterface);
-
-		if (!FullClassName.IsValid())
-		{
-			NewBuilder->Builder.InitDocument();
-		}
-
-		checkf(FullClassName.IsValid(), TEXT("Document initialization must result in a valid class name being generated"));
-		TObjectPtr<UMetaSoundBuilderBase> NewBuilderBase = CastChecked<UMetaSoundBuilderBase>(NewBuilder);
-		AssetBuilders.Add(FullClassName, NewBuilderBase);
-		return *NewBuilder;
-	}
-
 	friend class UMetaSoundBuilderBase;
 };
diff --git a/Engine/Plugins/Runtime/Metasound/Source/MetasoundEngine/Public/MetasoundDocumentBuilderRegistry.h b/Engine/Plugins/Runtime/Metasound/Source/MetasoundEngine/Public/MetasoundDocumentBuilderRegistry.h
new file mode 100644
--- /dev/null
+++ b/Engine/Plugins/Runtime/Metasound/Source/MetasoundEngine/Public/MetasoundDocumentBuilderRegistry.h
@@ -0,0 +1,96 @@
+// Copyright Epic Games, Inc. All Rights Reserved.
+#pragma once
+
+#include "Algo/Find.h"
+#include "MetasoundBuilderBase.h"
+#include "MetasoundDocumentInterface.h"
+#include "MetasoundFrontendDocument.h"
+#include "MetasoundFrontendRegistries.h"
+#include "MetasoundLog.h"
+
+
+namespace Metasound::Engine
+{
+	class METASOUNDENGINE_API FDocumentBuilderRegistry : public Frontend::IDocumentBuilderRegistry
+	{
+		mutable TMap<FMetasoundFrontendClassName, TWeakObjectPtr<UMetaSoundBuilderBase>> Builders;
+
+	public:
+		FDocumentBuilderRegistry() = default;
+		virtual ~FDocumentBuilderRegistry();
+
+		static FDocumentBuilderRegistry& GetChecked()
+		{
+			return static_cast<FDocumentBuilderRegistry&>(IDocumentBuilderRegistry::GetChecked());
+		}
+
+		template <typename BuilderClass>
+		BuilderClass& CreateTransientBuilder(FName BuilderName = FName())
+		{
+			using namespace Metasound::Frontend;
+
+			checkf(IsInGameThread(), TEXT("Transient MetaSound builder cannot be created in non - game thread as it may result in UObject creation"));
+
+			const EObjectFlags NewObjectFlags = RF_Public | RF_Transient;
+			UPackage* TransientPackage = GetTransientPackage();
+			const FName ObjectName = MakeUniqueObjectName(TransientPackage, BuilderClass::StaticClass(), BuilderName);
+			TObjectPtr<BuilderClass> NewBuilder = NewObject<BuilderClass>(TransientPackage, ObjectName, NewObjectFlags);
+			check(NewBuilder);
+			NewBuilder->Initialize();
+			const FMetasoundFrontendDocument& Document = NewBuilder->GetConstBuilder().GetConstDocument();
+			const FMetasoundFrontendClassName& ClassName = Document.RootGraph.Metadata.GetClassName();
+			Builders.Add(ClassName, NewBuilder);
+			return *NewBuilder.Get();
+		}
+
+#if WITH_EDITORONLY_DATA
+		template <typename BuilderClass = UMetaSoundBuilderBase>
+		BuilderClass& FindOrBeginBuilding(UObject& InMetaSoundObject) const
+		{
+			check(InMetaSoundObject.IsAsset());
+			checkf(IsInGameThread(), TEXT("Asset MetaSound Builder cannot be created in non-game thread as it may result in UObject creation"));
+
+			TScriptInterface<IMetaSoundDocumentInterface> DocInterface = &InMetaSoundObject;
+			check(DocInterface.GetObject());
+
+			const FMetasoundFrontendDocument& Document = DocInterface->GetConstDocument();
+			const FMetasoundFrontendClassName& FullClassName = Document.RootGraph.Metadata.GetClassName();
+
+			if (FullClassName.IsValid())
+			{
+				TWeakObjectPtr<UMetaSoundBuilderBase> Builder = Builders.FindRef(FullClassName);
+				if (Builder.IsValid())
+				{
+					return *CastChecked<BuilderClass>(Builder.Get());
+				}
+			}
+
+			TObjectPtr<UMetaSoundBuilderBase> NewBuilder = CastChecked<UMetaSoundBuilderBase>(NewObject<UObject>(&InMetaSoundObject, &DocInterface->GetBuilderUClass()));
+			FMetaSoundFrontendDocumentBuilder& BuilderRef = NewBuilder->GetBuilder();
+			BuilderRef = FMetaSoundFrontendDocumentBuilder(DocInterface);
+
+			if (!FullClassName.IsValid())
+			{
+				BuilderRef.InitDocument();
+			}
+
+			checkf(FullClassName.IsValid(), TEXT("Document initialization must result in a valid class name being generated"));
+			TObjectPtr<UMetaSoundBuilderBase> NewBuilderBase = CastChecked<UMetaSoundBuilderBase>(NewBuilder);
+			Builders.Add(FullClassName, NewBuilderBase);
+			return *CastChecked<BuilderClass>(NewBuilder);
+		}
+#endif // WITH_EDITORONLY_DATA
+
+		// Frontend::IDocumentBuilderRegistry Implementation
+#if WITH_EDITORONLY_DATA
+		virtual FMetaSoundFrontendDocumentBuilder& FindOrBeginBuilding(TScriptInterface<IMetaSoundDocumentInterface> MetaSound) override;
+#endif // WITH_EDITORONLY_DATA
+
+		virtual FMetaSoundFrontendDocumentBuilder* FindBuilder(TScriptInterface<IMetaSoundDocumentInterface> MetaSound) const override;
+		virtual FMetaSoundFrontendDocumentBuilder* FindBuilder(const FMetasoundFrontendClassName& InClassName) const override;
+		virtual bool FinishBuilding(const FMetasoundFrontendClassName& InClassName) const override;
+
+		UMetaSoundBuilderBase* FindBuilderObject(TScriptInterface<const IMetaSoundDocumentInterface> MetaSound) const;
+		UMetaSoundBuilderBase* FindBuilderObject(const FMetasoundFrontendClassName& ClassName) const;
+	};
+} // namespace Metasound::Engine
\ No newline at end of file
diff --git a/Engine/Plugins/Runtime/Metasound/Source/MetasoundEngine/Public/MetasoundSource.h b/Engine/Plugins/Runtime/Metasound/Source/MetasoundEngine/Public/MetasoundSource.h
--- a/Engine/Plugins/Runtime/Metasound/Source/MetasoundEngine/Public/MetasoundSource.h
+++ b/Engine/Plugins/Runtime/Metasound/Source/MetasoundEngine/Public/MetasoundSource.h
@@ -140,7 +140,7 @@ public:
 	FPerPlatformInt SampleRateOverride = 0;
 	
 #endif //WITH_EDITOR_DATA
-	
+
 	UPROPERTY(AssetRegistrySearchable)
 	FGuid AssetClassID;
 
@@ -176,10 +176,8 @@ public:
 	// Returns the graph associated with this Metasound. Graph is required to be referenced on
 	// Metasound UObject for editor serialization purposes.
 	// @return Editor graph associated with UMetaSoundSource.
-	virtual UEdGraph* GetGraph() override;
-	virtual const UEdGraph* GetGraph() const override;
-	virtual UEdGraph& GetGraphChecked() override;
-	virtual const UEdGraph& GetGraphChecked() const override;
+	virtual UEdGraph* GetGraph() const override;
+	virtual UEdGraph& GetGraphChecked() const override;
 	virtual void MigrateEditorGraph(FMetaSoundFrontendDocumentBuilder& OutBuilder) override;
 
 	// Sets the graph associated with this Metasound. Graph is required to be referenced on
@@ -275,12 +273,13 @@ public:
 	virtual const FMetasoundFrontendDocument& GetConstDocument() const override;
 	virtual bool IsActivelyBuilding() const override;
 
+	virtual const UClass& GetBaseMetaSoundUClass() const final override;
+	virtual const UClass& GetBuilderUClass() const final override;
+
 protected:
 	Metasound::Frontend::FDocumentAccessPtr GetDocumentAccessPtr() override;
 	Metasound::Frontend::FConstDocumentAccessPtr GetDocumentConstAccessPtr() const override;
 
-	virtual const UClass& GetBaseMetaSoundUClass() const final override;
-
 	/** Gets all the default parameters for this Asset.  */
 	virtual bool GetAllDefaultParameters(TArray<FAudioParameter>& OutParameters) const override;
 
@@ -333,7 +332,7 @@ private:
 	 * Note: Disabling the dynamic generator will sever the communication between any active generators
 	 * even if the dynamic generator is re-enabled during the lifetime of the active generators
 	 */
-	TSharedPtr<Metasound::DynamicGraph::FDynamicOperatorTransactor> SetDynamicGeneratorEnabled(const FTopLevelAssetPath& InAssetPath, bool bInIsEnabled);
+	TSharedPtr<Metasound::DynamicGraph::FDynamicOperatorTransactor> SetDynamicGeneratorEnabled(bool bInIsEnabled);
 
 	/** Get dynamic transactor
 	 *
diff --git a/Engine/Plugins/Runtime/Metasound/Source/MetasoundEngineTest/Private/EngineTestMetaSoundAutomatedNodeTest.cpp b/Engine/Plugins/Runtime/Metasound/Source/MetasoundEngineTest/Private/EngineTestMetaSoundAutomatedNodeTest.cpp
--- a/Engine/Plugins/Runtime/Metasound/Source/MetasoundEngineTest/Private/EngineTestMetaSoundAutomatedNodeTest.cpp
+++ b/Engine/Plugins/Runtime/Metasound/Source/MetasoundEngineTest/Private/EngineTestMetaSoundAutomatedNodeTest.cpp
@@ -67,13 +67,9 @@ namespace Metasound::EngineTest{
 
 	FString GetPrettyName(const Frontend::FNodeRegistryKey& InRegistryKey)
 	{
-		Frontend::IMetaSoundAssetManager* AssetManager = Frontend::IMetaSoundAssetManager::Get();
-		if (ensure(AssetManager))
+		if (const FTopLevelAssetPath* ObjectPath = Frontend::IMetaSoundAssetManager::GetChecked().FindAssetPath(InRegistryKey))
 		{
-			if (const FSoftObjectPath* ObjectPath = AssetManager->FindObjectPathFromKey(InRegistryKey))
-			{
-				return ObjectPath->ToString();
-			}
+			return ObjectPath->ToString();
 		}
 
 		FMetasoundFrontendRegistryContainer* NodeRegistry = FMetasoundFrontendRegistryContainer::Get();
diff --git a/Engine/Plugins/Runtime/Metasound/Source/MetasoundFrontend/Private/MetasoundAssetBase.cpp b/Engine/Plugins/Runtime/Metasound/Source/MetasoundFrontend/Private/MetasoundAssetBase.cpp
--- a/Engine/Plugins/Runtime/Metasound/Source/MetasoundFrontend/Private/MetasoundAssetBase.cpp
+++ b/Engine/Plugins/Runtime/Metasound/Source/MetasoundFrontend/Private/MetasoundAssetBase.cpp
@@ -280,7 +280,6 @@ void FMetasoundAssetBase::RegisterGraphWithFrontend(Metasound::Frontend::FMetaSo
 		}
 	}
 
-
 #if WITH_EDITOR
 	if (InRegistrationOptions.bRebuildReferencedAssetClasses)
 	{
@@ -292,7 +291,14 @@ void FMetasoundAssetBase::RegisterGraphWithFrontend(Metasound::Frontend::FMetaSo
 	{
 		RegisterAssetDependencies(InRegistrationOptions);
 	}
-	IMetaSoundAssetManager::GetChecked().AddOrUpdateAsset(*GetOwningAsset());
+
+	// This should not be necessary as it should be added on asset load,
+	// but currently registration is required to be called prior to adding
+	// an object-defined graph class to the registry so it was placed here.
+	if (UObject* OwningObject = GetOwningAsset())
+	{
+		IMetaSoundAssetManager::GetChecked().AddOrUpdateAsset(*OwningObject);
+	}
 
 	// Auto update must be done after all referenced asset classes are registered
 	if (InRegistrationOptions.bAutoUpdate)
@@ -374,13 +380,15 @@ void FMetasoundAssetBase::CookMetaSound()
 	check(Owner);
 
 	{
-		// Performs document transforms on local copy, which reduces document footprint & renders transforming unnecessary unless altered at runtime when registering
-		FMetaSoundFrontendDocumentBuilder& DocBuilder = IMetaSoundAssetManager::GetChecked().AttachDocumentBuilderChecked(*Owner);
+#if WITH_EDITORONLY_DATA
+		// Performs document transforms on local copy, which reduces document footprint & renders transforming unnecessary at runtime
+		FMetaSoundFrontendDocumentBuilder& DocBuilder = IDocumentBuilderRegistry::GetChecked().FindOrBeginBuilding(Owner);
 		const bool bContainsTemplateDependency = DocBuilder.ContainsDependencyOfType(EMetasoundFrontendClassType::Template);
 		if (bContainsTemplateDependency)
 		{
 			DocBuilder.TransformTemplateNodes();
 		}
+#endif // WITH_EDITORONLY_DATA
 
 		if (GraphRegistryKey.IsValid())
 		{
@@ -431,6 +439,19 @@ void FMetasoundAssetBase::OnNotifyBeginDestroy()
 	}
 	else
 	{
+		// The editor module removal and addition of assets with the MetaSoundAssetManager directly,
+		// so it can update open editors and load possible assets to reference at will, as opposed
+		// to managing by direct load requests and associated asset references.
+		// (Owning "asset" is a misnomer as we support dynamically generated MetaSounds at runtime now.)
+#if WITH_EDITOR
+		UObject* OwningAsset = GetOwningAsset();
+		check(OwningAsset);
+		if (IMetaSoundAssetManager* AssetManager = IMetaSoundAssetManager::Get())
+		{
+			AssetManager->RemoveAsset(*OwningAsset);
+		}
+#endif // !WITH_EDITOR
+
 		UnregisterGraphWithFrontend();
 	}
 }
@@ -529,7 +550,6 @@ bool FMetasoundAssetBase::VersionAsset()
 			for (const FMetasoundFrontendVersion& Version : Versions)
 			{
 				FUpdateRootGraphInterface UpdateTransform(Version, GetOwningAssetName());
-				FConstDocumentHandle ConstDoc = GetDocumentHandle();
 				bPassUpdated = TryUpdateInterfaceFromVersion(Version);
 			}
 
@@ -709,19 +729,13 @@ bool FMetasoundAssetBase::IsReferencedAsset(const FMetasoundAssetBase& InAsset)
 	return bIsReferenced;
 }
 
-bool FMetasoundAssetBase::AddingReferenceCausesLoop(const FSoftObjectPath& InReferencePath) const
+bool FMetasoundAssetBase::AddingReferenceCausesLoop(const FMetasoundAssetBase& InMetaSound) const
 {
 	using namespace Metasound::Frontend;
 
-	const FMetasoundAssetBase* ReferenceAsset = IMetaSoundAssetManager::GetChecked().TryLoadAsset(InReferencePath);
-	if (!ensureAlways(ReferenceAsset))
-	{
-		return false;
-	}
-
 	bool bCausesLoop = false;
 	const FMetasoundAssetBase* Parent = this;
-	AssetBasePrivate::DepthFirstTraversal(*ReferenceAsset, [&](const FMetasoundAssetBase& ChildAsset)
+	AssetBasePrivate::DepthFirstTraversal(InMetaSound, [&](const FMetasoundAssetBase& ChildAsset)
 	{
 		TSet<const FMetasoundAssetBase*> Children;
 		if (Parent == &ChildAsset)
@@ -732,13 +746,26 @@ bool FMetasoundAssetBase::AddingReferenceCausesLoop(const FSoftObjectPath& InRef
 
 		TArray<FMetasoundAssetBase*> ChildRefs;
 		ensureAlways(IMetaSoundAssetManager::GetChecked().TryLoadReferencedAssets(ChildAsset, ChildRefs));
-		Algo::Transform(ChildRefs, Children, [] (FMetasoundAssetBase* Child) { return Child; });
+		Algo::Transform(ChildRefs, Children, [](FMetasoundAssetBase* Child) { return Child; });
 		return Children;
 	});
 
 	return bCausesLoop;
 }
 
+bool FMetasoundAssetBase::AddingReferenceCausesLoop(const FSoftObjectPath& InReferencePath) const
+{
+	using namespace Metasound::Frontend;
+
+	const FMetasoundAssetBase* ReferenceAsset = IMetaSoundAssetManager::GetChecked().TryLoadAsset(InReferencePath);
+	if (!ensureAlways(ReferenceAsset))
+	{
+		return false;
+	}
+
+	return AddingReferenceCausesLoop(*ReferenceAsset);
+}
+
 TArray<FMetasoundAssetBase::FSendInfoAndVertexName> FMetasoundAssetBase::GetSendInfos(uint64 InInstanceID) const
 {
 	return TArray<FSendInfoAndVertexName>();
@@ -876,7 +903,6 @@ void FMetasoundAssetBase::RebuildReferencedAssetClasses()
 	using namespace Metasound::Frontend;
 
 	IMetaSoundAssetManager& AssetManager = IMetaSoundAssetManager::GetChecked();
-
 	AssetManager.AddAssetReferences(*this);
 	TSet<IMetaSoundAssetManager::FAssetInfo> ReferencedAssetClasses = AssetManager.GetReferencedAssetClasses(*this);
 	SetReferencedAssetClasses(MoveTemp(ReferencedAssetClasses));
@@ -934,11 +960,12 @@ void FMetasoundAssetBase::UpdateAssetRegistry()
 
 	UObject* Owner = GetOwningAsset();
 	check(Owner);
+	const FMetasoundFrontendGraphClass& DocumentClassGraph = GetDocumentChecked().RootGraph;
+
 	TScriptInterface<IMetaSoundDocumentInterface> DocInterface(Owner);
-	FNodeClassInfo AssetClassInfo(GetDocumentChecked().RootGraph, DocInterface->GetAssetPathChecked());
+	FNodeClassInfo AssetClassInfo(DocumentClassGraph, DocInterface->GetAssetPathChecked());
 
 	// Refresh Asset Registry Info if successfully registered with Frontend
-	const FMetasoundFrontendGraphClass& DocumentClassGraph = GetDocumentHandle()->GetRootGraphClass();
 	const FMetasoundFrontendClassMetadata& DocumentClassMetadata = DocumentClassGraph.Metadata;
 	AssetClassInfo.AssetClassID = FGuid(DocumentClassMetadata.GetClassName().Name.ToString());
 	FNodeClassName ClassName = DocumentClassMetadata.GetClassName().ToNodeClassName();
diff --git a/Engine/Plugins/Runtime/Metasound/Source/MetasoundFrontend/Private/MetasoundAssetManager.cpp b/Engine/Plugins/Runtime/Metasound/Source/MetasoundFrontend/Private/MetasoundAssetManager.cpp
--- a/Engine/Plugins/Runtime/Metasound/Source/MetasoundFrontend/Private/MetasoundAssetManager.cpp
+++ b/Engine/Plugins/Runtime/Metasound/Source/MetasoundFrontend/Private/MetasoundAssetManager.cpp
@@ -1,34 +1,106 @@
 // Copyright Epic Games, Inc. All Rights Reserved.
 
 #include "MetasoundAssetManager.h"
+#include "MetasoundFrontendRegistryKey.h"
 
-namespace Metasound
+namespace Metasound::Frontend
 {
-	namespace Frontend
+	namespace AssetManagerPrivate
 	{
-		namespace AssetTags
+		static TUniquePtr<IMetaSoundAssetManager> Instance;
+
+		bool IsAssetClassType(EMetasoundFrontendClassType ClassType)
 		{
-			const FString ArrayDelim = TEXT(",");
+			switch (ClassType)
+			{
+				case EMetasoundFrontendClassType::External:
+				case EMetasoundFrontendClassType::Graph:
+				case EMetasoundFrontendClassType::Invalid:
+					return true;
+
+				default:
+					return false;
+			}
+		}
+	} // AssetManagerPrivate
+
+	namespace AssetTags
+	{
+		const FString ArrayDelim = TEXT(",");
 
-			const FName AssetClassID = "AssetClassID";
+		const FName AssetClassID = "AssetClassID";
 
 #if WITH_EDITORONLY_DATA
-			const FName IsPreset = "bIsPreset";
+		const FName IsPreset = "bIsPreset";
 #endif // WITH_EDITORONLY_DATA
 
-			const FName RegistryVersionMajor = "RegistryVersionMajor";
-			const FName RegistryVersionMinor = "RegistryVersionMinor";
+		const FName RegistryVersionMajor = "RegistryVersionMajor";
+		const FName RegistryVersionMinor = "RegistryVersionMinor";
 
 #if WITH_EDITORONLY_DATA
-			const FName RegistryInputTypes = "RegistryInputTypes";
-			const FName RegistryOutputTypes = "RegistryOutputTypes";
+		const FName RegistryInputTypes = "RegistryInputTypes";
+		const FName RegistryOutputTypes = "RegistryOutputTypes";
 #endif // WITH_EDITORONLY_DATA
-		} // namespace AssetTags
+	} // namespace AssetTags
+
+	FAssetKey::FAssetKey(const FMetasoundFrontendClassName& InClassName, const FMetasoundFrontendVersionNumber& InVersion)
+		: ClassName(InClassName)
+		, Version(InVersion)
+	{
+	}
 
-		IMetaSoundAssetManager* IMetaSoundAssetManager::Instance = nullptr;
+	FAssetKey::FAssetKey(const FNodeRegistryKey& RegKey)
+		: ClassName(RegKey.ClassName)
+		, Version(RegKey.Version)
+	{
+		checkf(AssetManagerPrivate::IsAssetClassType(RegKey.Type), TEXT("Invalid ClassType '%s' for Registry Key"), LexToString(RegKey.Type));
+	}
+
+	FAssetKey::FAssetKey(const FMetasoundFrontendClassMetadata& InMetadata)
+		: ClassName(InMetadata.GetClassName())
+		, Version(InMetadata.GetVersion())
+	{
+		checkf(AssetManagerPrivate::IsAssetClassType(InMetadata.GetType()), TEXT("Invalid ClassType '%s' for Registry Key"), LexToString(InMetadata.GetType()));
+	}
+
+	FString FAssetKey::ToString() const
+	{
+		TStringBuilder<128> KeyStringBuilder;
+		KeyStringBuilder.Append(ClassName.GetFullName().ToString());
+		KeyStringBuilder.AppendChar('_');
+		KeyStringBuilder.Append(FString::FromInt(Version.Major));
+		KeyStringBuilder.AppendChar('.');
+		KeyStringBuilder.Append(FString::FromInt(Version.Minor));
+		return KeyStringBuilder.ToString();
+	}
 
-#if WITH_SERVER_CODE
-		TMulticastDelegate<void()> IMetaSoundAssetManager::OnManagerSet;
-#endif //WITH_SERVER_CODE
-	} // namespace Frontend
-} // namespace Metasound
+	IMetaSoundAssetManager* IMetaSoundAssetManager::Get()
+	{
+		return AssetManagerPrivate::Instance.Get();
+	}
+
+	IMetaSoundAssetManager& IMetaSoundAssetManager::GetChecked()
+	{
+		check(AssetManagerPrivate::Instance.IsValid());
+		return *AssetManagerPrivate::Instance;
+	}
+
+	void IMetaSoundAssetManager::Deinitialize()
+	{
+		if (AssetManagerPrivate::Instance.IsValid())
+		{
+			AssetManagerPrivate::Instance.Reset();
+		}
+	}
+
+	void IMetaSoundAssetManager::Initialize(TUniquePtr<IMetaSoundAssetManager>&& InInterface)
+	{
+		check(!AssetManagerPrivate::Instance.IsValid());
+		AssetManagerPrivate::Instance = MoveTemp(InInterface);
+	}
+
+	bool IMetaSoundAssetManager::IsTesting() const
+	{
+		return false;
+	}
+} // namespace Metasound::Frontend
diff --git a/Engine/Plugins/Runtime/Metasound/Source/MetasoundFrontend/Private/MetasoundDocumentInterface.cpp b/Engine/Plugins/Runtime/Metasound/Source/MetasoundFrontend/Private/MetasoundDocumentInterface.cpp
--- a/Engine/Plugins/Runtime/Metasound/Source/MetasoundFrontend/Private/MetasoundDocumentInterface.cpp
+++ b/Engine/Plugins/Runtime/Metasound/Source/MetasoundFrontend/Private/MetasoundDocumentInterface.cpp
@@ -5,39 +5,31 @@
 
 namespace Metasound::Frontend
 {
-	namespace DocumentBuilderRegistryPrivate
+	namespace DocumentInterfacePrivate
 	{
-		static bool bInitialized = false;
-		TUniqueFunction<IDocumentBuilderRegistry&()> GetInstance;
-	} // namespace DocumentBuilderRegistryPrivate
+		TUniquePtr<IDocumentBuilderRegistry> Instance;
+	}
 
 	IDocumentBuilderRegistry* IDocumentBuilderRegistry::Get()
 	{
-		using namespace DocumentBuilderRegistryPrivate;
-
-		if (!DocumentBuilderRegistryPrivate::bInitialized)
-		{
-			return nullptr;
-		}
-		
-		return &GetInstance();
+		return DocumentInterfacePrivate::Instance.Get();
 	}
 
 	IDocumentBuilderRegistry& IDocumentBuilderRegistry::GetChecked()
 	{
-		using namespace DocumentBuilderRegistryPrivate;
-
-		checkf(GetInstance, TEXT("Failed to return MetaSoundDocumentBuilderRegistry instance: Registry has not been initialized"));
-		return GetInstance();
+		return *DocumentInterfacePrivate::Instance.Get();
 	}
 
-	void IDocumentBuilderRegistry::Set(TUniqueFunction<IDocumentBuilderRegistry&()>&& InGetInstance)
+	void IDocumentBuilderRegistry::Deinitialize()
 	{
-		using namespace DocumentBuilderRegistryPrivate;
+		check(DocumentInterfacePrivate::Instance.IsValid());
+		DocumentInterfacePrivate::Instance.Reset();
+	}
 
-		checkf(!GetInstance, TEXT("Failed to initialize MetaSoundDocumentBuilderRegistry getter: Cannot reinitialize once initialized."))
-		GetInstance = MoveTemp(InGetInstance);
-		DocumentBuilderRegistryPrivate::bInitialized = true;
+	void IDocumentBuilderRegistry::Initialize(TUniquePtr<IDocumentBuilderRegistry>&& InInstance)
+	{
+		check(!DocumentInterfacePrivate::Instance.IsValid());
+		DocumentInterfacePrivate::Instance = MoveTemp(InInstance);
 	}
 
 	IMetaSoundDocumentBuilderRegistry& IMetaSoundDocumentBuilderRegistry::GetChecked()
diff --git a/Engine/Plugins/Runtime/Metasound/Source/MetasoundFrontend/Private/MetasoundFrontendController.cpp b/Engine/Plugins/Runtime/Metasound/Source/MetasoundFrontend/Private/MetasoundFrontendController.cpp
--- a/Engine/Plugins/Runtime/Metasound/Source/MetasoundFrontend/Private/MetasoundFrontendController.cpp
+++ b/Engine/Plugins/Runtime/Metasound/Source/MetasoundFrontend/Private/MetasoundFrontendController.cpp
@@ -67,22 +67,13 @@ namespace Metasound
 
 		FDocumentHandle IDocumentController::CreateDocumentHandle(FMetasoundFrontendDocument& InDocument)
 		{
-			// Mutation of a document via the soft deprecated access ptr/controller system is not tracked by
-			// the builder registry, so the document cache is invalidated here. It is discouraged to mutate
-			// documents using both systems at the same time as it can corrupt a builder document's cache.
-			if (IDocumentBuilderRegistry* DocRegistry = IDocumentBuilderRegistry::Get())
-			{
-				const FMetasoundFrontendClassName& Name = InDocument.RootGraph.Metadata.GetClassName();
-				DocRegistry->InvalidateDocumentCache(Name);
-			}
-
 			return CreateDocumentHandle(MakeAccessPtr<FDocumentAccessPtr>(InDocument.AccessPoint, InDocument));
 		}
 
 		FConstDocumentHandle IDocumentController::CreateDocumentHandle(FConstDocumentAccessPtr InDocument)
 		{
 			// Create using standard document controller. 
-			return FDocumentController::CreateDocumentHandle(ConstCastAccessPtr<FDocumentAccessPtr>(InDocument));
+			return FDocumentController::CreateDocumentHandle(InDocument);
 		}
 
 		FConstDocumentHandle IDocumentController::CreateDocumentHandle(const FMetasoundFrontendDocument& InDocument)
