- Move MetaSoundAssetSubsystem& BuilderSubsystem implementations to monolithics that are initialized at MetaSoundEngine module load for safety

diff --git a/Engine/Plugins/Runtime/Metasound/Source/MetasoundFrontend/Private/MetasoundFrontendDocumentBuilder.cpp b/Engine/Plugins/Runtime/Metasound/Source/MetasoundFrontend/Private/MetasoundFrontendDocumentBuilder.cpp
--- a/Engine/Plugins/Runtime/Metasound/Source/MetasoundFrontend/Private/MetasoundFrontendDocumentBuilder.cpp
+++ b/Engine/Plugins/Runtime/Metasound/Source/MetasoundFrontend/Private/MetasoundFrontendDocumentBuilder.cpp
@@ -27,6 +27,7 @@
 #include "MetasoundTrace.h"
 #include "MetasoundVariableNodes.h"
 #include "NodeTemplates/MetasoundFrontendNodeTemplateInput.h"
+#include "../Public/MetasoundDocumentInterface.h"
 
 #include UE_INLINE_GENERATED_CPP_BY_NAME(MetasoundFrontendDocumentBuilder)
 
@@ -189,19 +190,6 @@ namespace Metasound::Frontend
 			}
 		}
 
-		const FString GetDebugName(const IMetaSoundDocumentInterface& DocumentInterface)
-		{
-			const FMetasoundFrontendClassMetadata& Metadata = DocumentInterface.GetConstDocument().RootGraph.Metadata;
-			const FMetasoundFrontendVersionNumber& Version = Metadata.GetVersion();
-			const FNodeRegistryKey RegKey(EMetasoundFrontendClassType::External, Metadata.GetClassName(), Version.Major, Version.Minor);
-			if (const FMetasoundAssetBase* Asset = IMetaSoundAssetManager::GetChecked().TryLoadAssetFromKey(RegKey))
-			{
-				return Asset->GetOwningAssetName();
-			}
-
-			return Metadata.GetClassName().ToString();
-		}
-
 		class FModifyInterfacesImpl
 		{
 		public:
@@ -630,19 +618,21 @@ namespace Metasound::Frontend
 	}
 } // namespace Metasound::Frontend
 
-
-UMetaSoundBuilderDocument& UMetaSoundBuilderDocument::Create(const UClass& InBuilderClass)
+UMetaSoundBuilderDocument& UMetaSoundBuilderDocument::Create(const UClass& InMetaSoundUClass)
 {
 	UMetaSoundBuilderDocument* DocObject = NewObject<UMetaSoundBuilderDocument>();
-	DocObject->MetaSoundUClass = InBuilderClass;
+	check(DocObject);
+	DocObject->MetaSoundUClass = &InMetaSoundUClass;
 	return *DocObject;
 }
 
 UMetaSoundBuilderDocument& UMetaSoundBuilderDocument::Create(const IMetaSoundDocumentInterface& InDocToCopy)
 {
 	UMetaSoundBuilderDocument* DocObject = NewObject<UMetaSoundBuilderDocument>();
+	check(DocObject);
 	DocObject->Document = InDocToCopy.GetConstDocument();
 	DocObject->MetaSoundUClass = InDocToCopy.GetBaseMetaSoundUClass();
+	DocObject->BuilderUClass = InDocToCopy.GetBuilderUClass();
 	return *DocObject;
 }
 
@@ -670,6 +660,12 @@ const UClass& UMetaSoundBuilderDocument::GetBaseMetaSoundUClass() const
 	return *MetaSoundUClass;
 }
 
+const UClass& UMetaSoundBuilderDocument::GetBuilderUClass() const
+{
+	checkf(BuilderUClass, TEXT("BuilderUClass must be set upon creation of UMetaSoundBuilderDocument instance"));
+	return *BuilderUClass;
+}
+
 bool UMetaSoundBuilderDocument::IsActivelyBuilding() const
 {
 	return true;
@@ -1144,13 +1140,10 @@ const FMetasoundFrontendNode* FMetaSoundFrontendDocumentBuilder::AddGraphNode(co
 		using namespace Metasound::Frontend;
 
 		// Cache the asset name on the node if it node is reference to asset-defined graph.
-		if (IMetaSoundAssetManager* AssetManager = IMetaSoundAssetManager::Get())
+		if (const FTopLevelAssetPath* Path = IMetaSoundAssetManager::GetChecked().FindAssetPath(FAssetKey(ClassKey.ClassName, ClassKey.Version)))
 		{
-			if (const FSoftObjectPath* Path = AssetManager->FindObjectPathFromKey(ClassKey))
-			{
-				InOutNode.Name = Path->GetAssetFName();
-				return;
-			}
+			InOutNode.Name = Path->GetAssetName();
+			return;
 		}
 
 		InOutNode.Name = ClassKey.ClassName.GetFullName();
@@ -1269,7 +1262,7 @@ bool FMetaSoundFrontendDocumentBuilder::CanAddEdge(const FMetasoundFrontendEdge&
 {
 	using namespace Metasound::Frontend;
 
-	const FMetasoundFrontendDocument& Document = GetDocument();
+	const FMetasoundFrontendDocument& Document = GetConstDocument();
 	const IDocumentGraphEdgeCache& EdgeCache = DocumentCache->GetEdgeCache();
 
 	if (!EdgeCache.IsNodeInputConnected(InEdge.ToNodeID, InEdge.ToVertexID))
@@ -1510,7 +1503,7 @@ bool FMetaSoundFrontendDocumentBuilder::ConvertToPreset(const FMetasoundFrontend
 
 bool FMetaSoundFrontendDocumentBuilder::FindDeclaredInterfaces(TArray<const Metasound::Frontend::IInterfaceRegistryEntry*>& OutInterfaces) const
 {
-	return FindDeclaredInterfaces(GetDocument(), OutInterfaces);
+	return FindDeclaredInterfaces(GetConstDocument(), OutInterfaces);
 }
 
 bool FMetaSoundFrontendDocumentBuilder::FindDeclaredInterfaces(const FMetasoundFrontendDocument& InDocument, TArray<const Metasound::Frontend::IInterfaceRegistryEntry*>& OutInterfaces)
@@ -1562,7 +1555,7 @@ TArray<const FMetasoundFrontendEdge*> FMetaSoundFrontendDocumentBuilder::FindEdg
 #if WITH_EDITOR
 const FMetaSoundFrontendGraphComment* FMetaSoundFrontendDocumentBuilder::FindGraphComment(const FGuid& InCommentID) const
 {
-	const TMap<FGuid, FMetaSoundFrontendGraphComment>& Comments = GetDocument().RootGraph.Graph.Style.Comments;
+	const TMap<FGuid, FMetaSoundFrontendGraphComment>& Comments = GetConstDocument().RootGraph.Graph.Style.Comments;
 	return Comments.Find(InCommentID);
 }
 
@@ -1582,7 +1575,7 @@ bool FMetaSoundFrontendDocumentBuilder::FindInterfaceInputNodes(FName InterfaceN
 	FMetasoundFrontendInterface Interface;
 	if (ISearchEngine::Get().FindInterfaceWithHighestVersion(InterfaceName, Interface))
 	{
-		if (GetDocument().Interfaces.Contains(Interface.Version))
+		if (GetConstDocument().Interfaces.Contains(Interface.Version))
 		{
 			const IDocumentGraphNodeCache& NodeCache = DocumentCache->GetNodeCache();
 			const IDocumentGraphInterfaceCache& InterfaceCache = DocumentCache->GetInterfaceCache();
@@ -1623,7 +1616,7 @@ bool FMetaSoundFrontendDocumentBuilder::FindInterfaceOutputNodes(FName Interface
 	FMetasoundFrontendInterface Interface;
 	if (ISearchEngine::Get().FindInterfaceWithHighestVersion(InterfaceName, Interface))
 	{
-		if (GetDocument().Interfaces.Contains(Interface.Version))
+		if (GetConstDocument().Interfaces.Contains(Interface.Version))
 		{
 			const IDocumentGraphNodeCache& NodeCache = DocumentCache->GetNodeCache();
 			const IDocumentGraphInterfaceCache& InterfaceCache = DocumentCache->GetInterfaceCache();
@@ -1713,7 +1706,7 @@ bool FMetaSoundFrontendDocumentBuilder::FindNodeClassInterfaces(const FGuid& InN
 	using namespace Metasound;
 	using namespace Metasound::Frontend;
 
-	const FMetasoundFrontendDocument& Document = GetDocument();
+	const FMetasoundFrontendDocument& Document = GetConstDocument();
 	const IDocumentGraphNodeCache& NodeCache = DocumentCache->GetNodeCache();
 	if (const FMetasoundFrontendNode* Node = NodeCache.FindNode(InNodeID))
 	{
@@ -1753,7 +1746,7 @@ TArray<const FMetasoundFrontendVertex*> FMetaSoundFrontendDocumentBuilder::FindN
 	const IDocumentGraphEdgeCache& EdgeCache = DocumentCache->GetEdgeCache();
 	const IDocumentGraphNodeCache& NodeCache = DocumentCache->GetNodeCache();
 
-	const FMetasoundFrontendDocument& Document = GetDocument();
+	const FMetasoundFrontendDocument& Document = GetConstDocument();
 
 	if (ConnectedInputNodes)
 	{
@@ -1816,7 +1809,7 @@ const FMetasoundFrontendVertex* FMetaSoundFrontendDocumentBuilder::FindNodeOutpu
 	const IDocumentGraphEdgeCache& EdgeCache = DocumentCache->GetEdgeCache();
 	if (const int32* Index = EdgeCache.FindEdgeIndexToNodeInput(InInputNodeID, InInputVertexID))
 	{
-		const FMetasoundFrontendDocument& Document = GetDocument();
+		const FMetasoundFrontendDocument& Document = GetConstDocument();
 		const FMetasoundFrontendEdge& Edge = Document.RootGraph.Graph.Edges[*Index];
 		const IDocumentGraphNodeCache& NodeCache = DocumentCache->GetNodeCache();
 		if (ConnectedOutputNode)
@@ -1842,6 +1835,16 @@ FMetaSoundFrontendGraphComment& FMetaSoundFrontendDocumentBuilder::FindOrAddGrap
 }
 #endif // WITH_EDITOR
 
+FMetasoundFrontendClassName FMetaSoundFrontendDocumentBuilder::GenerateNewClassName()
+{
+	using namespace Metasound::Frontend;
+	FMetasoundFrontendClassMetadata& Metadata = GetDocument().RootGraph.Metadata;
+	const FMetasoundFrontendClassName OldClassName = Metadata.GetClassName();
+	const FMetasoundFrontendClassName NewClassName(FName(), FName(*FGuid::NewGuid().ToString()), FName());
+	Metadata.SetClassName(NewClassName);
+	return NewClassName;
+}
+
 const FTopLevelAssetPath FMetaSoundFrontendDocumentBuilder::GetBuilderClassPath() const
 {
 	IMetaSoundDocumentInterface* Interface = DocumentInterface.GetInterface();
@@ -1849,11 +1852,24 @@ const FTopLevelAssetPath FMetaSoundFrontendDocumentBuilder::GetBuilderClassPath(
 	return Interface->GetBaseMetaSoundUClass().GetClassPathName();
 }
 
+const FMetasoundFrontendDocument& FMetaSoundFrontendDocumentBuilder::GetConstDocument() const
+{
+	return DocumentInterface->GetConstDocument();
+}
+
 const FString FMetaSoundFrontendDocumentBuilder::GetDebugName() const
 {
 	using namespace Metasound::Frontend;
 
-	return DocumentBuilderPrivate::GetDebugName(GetDocumentInterface());
+	const FMetasoundFrontendClassMetadata& Metadata = GetConstDocument().RootGraph.Metadata;
+	const FMetasoundFrontendVersionNumber& Version = Metadata.GetVersion();
+	const FAssetKey AssetKey(Metadata.GetClassName(), Version);
+	if (const FMetasoundAssetBase* Asset = IMetaSoundAssetManager::GetChecked().FindAsset(AssetKey))
+	{
+		return Asset->GetOwningAssetName();
+	}
+
+	return Metadata.GetClassName().ToString();
 }
 
 FMetasoundFrontendDocument& FMetaSoundFrontendDocumentBuilder::GetDocument()
@@ -1909,7 +1925,7 @@ EMetasoundFrontendVertexAccessType FMetaSoundFrontendDocumentBuilder::GetNodeInp
 	const IDocumentGraphNodeCache& NodeCache = DocumentCache->GetNodeCache();
 	if (const int32* NodeIndex = NodeCache.FindNodeIndex(InNodeID))
 	{
-		const FMetasoundFrontendGraph& Graph = GetDocument().RootGraph.Graph;
+		const FMetasoundFrontendGraph& Graph = GetConstDocument().RootGraph.Graph;
 		const FMetasoundFrontendNode& Node = Graph.Nodes[*NodeIndex];
 		auto IsVertexID = [&InVertexID](const FMetasoundFrontendVertex& Vertex) { return Vertex.VertexID == InVertexID; };
 		if (const FMetasoundFrontendClass* Class = DocumentCache->FindDependency(Node.ClassID))
@@ -1963,7 +1979,7 @@ const FMetasoundFrontendLiteral* FMetaSoundFrontendDocumentBuilder::GetNodeInput
 	const IDocumentGraphNodeCache& NodeCache = DocumentCache->GetNodeCache();
 	if (const int32* NodeIndex = NodeCache.FindNodeIndex(InNodeID))
 	{
-		const FMetasoundFrontendGraph& Graph = GetDocument().RootGraph.Graph;
+		const FMetasoundFrontendGraph& Graph = GetConstDocument().RootGraph.Graph;
 		const FMetasoundFrontendNode& Node = Graph.Nodes[*NodeIndex];
 		auto IsVertexID = [&InVertexID](const FMetasoundFrontendVertex& Vertex) { return Vertex.VertexID == InVertexID; };
 		if (const FMetasoundFrontendVertex* Vertex = Node.Interface.Inputs.FindByPredicate(IsVertexID))
@@ -2007,7 +2023,7 @@ const FMetasoundFrontendLiteral* FMetaSoundFrontendDocumentBuilder::GetNodeInput
 	const IDocumentGraphNodeCache& NodeCache = DocumentCache->GetNodeCache();
 	if (const int32* NodeIndex = NodeCache.FindNodeIndex(InNodeID))
 	{
-		const FMetasoundFrontendGraph& Graph = GetDocument().RootGraph.Graph;
+		const FMetasoundFrontendGraph& Graph = GetConstDocument().RootGraph.Graph;
 		const FMetasoundFrontendNode& Node = Graph.Nodes[*NodeIndex];
 
 		auto IsVertex = [&InVertexID](const FMetasoundFrontendVertex& Vertex) { return Vertex.VertexID == InVertexID; };
@@ -2036,7 +2052,7 @@ EMetasoundFrontendVertexAccessType FMetaSoundFrontendDocumentBuilder::GetNodeOut
 	const IDocumentGraphNodeCache& NodeCache = DocumentCache->GetNodeCache();
 	if (const int32* NodeIndex = NodeCache.FindNodeIndex(InNodeID))
 	{
-		const FMetasoundFrontendGraph& Graph = GetDocument().RootGraph.Graph;
+		const FMetasoundFrontendGraph& Graph = GetConstDocument().RootGraph.Graph;
 		const FMetasoundFrontendNode& Node = Graph.Nodes[*NodeIndex];
 		if (const FMetasoundFrontendClass* Class = DocumentCache->FindDependency(Node.ClassID))
 		{
@@ -2248,12 +2264,12 @@ bool FMetaSoundFrontendDocumentBuilder::IsInterfaceDeclared(FName InInterfaceNam
 
 bool FMetaSoundFrontendDocumentBuilder::IsInterfaceDeclared(const FMetasoundFrontendVersion& InInterfaceVersion) const
 {
-	return GetDocument().Interfaces.Contains(InInterfaceVersion);
+	return GetConstDocument().Interfaces.Contains(InInterfaceVersion);
 }
 
 bool FMetaSoundFrontendDocumentBuilder::IsPreset() const
 {
-	return GetDocument().RootGraph.PresetOptions.bIsPreset;
+	return GetConstDocument().RootGraph.PresetOptions.bIsPreset;
 }
 
 bool FMetaSoundFrontendDocumentBuilder::IsVariableClass(EMetasoundFrontendClassType ClassType) const
@@ -2330,6 +2346,7 @@ bool FMetaSoundFrontendDocumentBuilder::ModifyInterfaces(Metasound::Frontend::FM
 	return Context.Execute(*this, *DocumentDelegates);
 }
 
+#if WITH_EDITORONLY_DATA
 bool FMetaSoundFrontendDocumentBuilder::TransformTemplateNodes()
 {
 	using namespace Metasound::Frontend;
@@ -2408,6 +2425,7 @@ bool FMetaSoundFrontendDocumentBuilder::TransformTemplateNodes()
 
 	return bModified;
 }
+#endif // WITH_EDITORONLY_DATA
 
 void FMetaSoundFrontendDocumentBuilder::BeginBuilding(TSharedPtr<Metasound::Frontend::FDocumentModifyDelegates> Delegates, bool bPrimeCache)
 {
@@ -3021,16 +3039,7 @@ int32 FMetaSoundFrontendDocumentBuilder::RemoveNodeLocation(const FGuid& InNodeI
 
 bool FMetaSoundFrontendDocumentBuilder::RenameRootGraphClass(const FMetasoundFrontendClassName& InName)
 {
-	FGuid NameGuid;
-	if (FGuid::Parse(InName.Name.ToString(), NameGuid))
-	{
-		return Metasound::Frontend::FRenameRootGraphClass::Generate(GetDocument(), NameGuid, InName.Namespace, InName.Variant);
-	}
-	else
-	{
-		UE_LOG(LogMetaSound, Error, TEXT("Attempting to rename a root graph class with class name '%s' which contains an invalid name guid."), *InName.GetFullName().ToString());
-		return false;
-	}
+	return false;
 }
 
 void FMetaSoundFrontendDocumentBuilder::ReloadCache()
@@ -3337,6 +3346,11 @@ bool FMetaSoundFrontendDocumentBuilder::SetGraphOutputDataType(FName OutputName,
 }
 
 #if WITH_EDITOR
+void FMetaSoundFrontendDocumentBuilder::SetDisplayName(const FText& InDisplayName)
+{
+	DocumentInterface->GetDocument().RootGraph.Metadata.SetDisplayName(InDisplayName);
+}
+
 void FMetaSoundFrontendDocumentBuilder::SetMemberMetadata(UMetaSoundFrontendMemberMetadata& NewMetadata)
 {
 	check(NewMetadata.MemberID.IsValid());
diff --git a/Engine/Plugins/Runtime/Metasound/Source/MetasoundFrontend/Private/MetasoundFrontendDocumentController.h b/Engine/Plugins/Runtime/Metasound/Source/MetasoundFrontend/Private/MetasoundFrontendDocumentController.h
--- a/Engine/Plugins/Runtime/Metasound/Source/MetasoundFrontend/Private/MetasoundFrontendDocumentController.h
+++ b/Engine/Plugins/Runtime/Metasound/Source/MetasoundFrontend/Private/MetasoundFrontendDocumentController.h
@@ -2,6 +2,8 @@
 
 #pragma once
 
+#include "MetasoundFrontendDocumentBuilder.h"
+#include "MetasoundDocumentInterface.h"
 #include "MetasoundFrontendController.h"
 #include "MetasoundFrontendDocument.h"
 
@@ -20,15 +22,31 @@ namespace Metasound
 			 */
 			FDocumentController(FDocumentAccessPtr InDocumentPtr);
 
+			static FDocumentHandle CreateDocumentHandle(FDocumentAccessPtr InDocument)
+			{
+				// Unit test builds may not load the builder registry (i.e. via the engine module).  Creating and
+				// manipulating documents via controllers/handles must be supported for backward compat in this context,
+				// so registry is not required to exist.
+				if (IDocumentBuilderRegistry* Registry = IDocumentBuilderRegistry::Get())
+				{
+					const FMetasoundFrontendClassName& ClassName = InDocument.Get()->RootGraph.Metadata.GetClassName();
+
+					PRAGMA_DISABLE_DEPRECATION_WARNINGS
+					Registry->InvalidateDocumentCache(ClassName);
+					PRAGMA_ENABLE_DEPRECATION_WARNINGS
+				}
+				return MakeShared<FDocumentController>(InDocument);
+			}
+
 			/** Create a FDocumentController.
 			 *
 			 * @param InDocument - Document to be manipulated.
 			 *
 			 * @return A document handle. 
 			 */
-			static FDocumentHandle CreateDocumentHandle(FDocumentAccessPtr InDocument)
+			static FConstDocumentHandle CreateDocumentHandle(FConstDocumentAccessPtr InDocument)
 			{
-				return MakeShared<FDocumentController>(InDocument);
+				return MakeShared<const FDocumentController>(ConstCastAccessPtr<FDocumentAccessPtr>(InDocument));
 			}
 
 			virtual ~FDocumentController() = default;
diff --git a/Engine/Plugins/Runtime/Metasound/Source/MetasoundFrontend/Private/MetasoundFrontendDocumentVersioning.cpp b/Engine/Plugins/Runtime/Metasound/Source/MetasoundFrontend/Private/MetasoundFrontendDocumentVersioning.cpp
--- a/Engine/Plugins/Runtime/Metasound/Source/MetasoundFrontend/Private/MetasoundFrontendDocumentVersioning.cpp
+++ b/Engine/Plugins/Runtime/Metasound/Source/MetasoundFrontend/Private/MetasoundFrontendDocumentVersioning.cpp
@@ -28,14 +28,6 @@ namespace Metasound::Frontend
 
 			virtual void TransformInternal(FMetasoundFrontendDocument& OutDocument) const
 			{
-				// Mutation of a document via the soft deprecated access ptr/controller system is not tracked by
-				// the builder registry, so the document cache is invalidated here. It is discouraged to mutate
-				// documents using both systems at the same time as it can corrupt a builder document's cache.
-				if (IDocumentBuilderRegistry* DocRegistry = IDocumentBuilderRegistry::Get())
-				{
-					const FMetasoundFrontendClassName& Name = OutDocument.RootGraph.Metadata.GetClassName();
-					DocRegistry->InvalidateDocumentCache(Name);
-				}
 				FDocumentAccessPtr DocAccessPtr = MakeAccessPtr<FDocumentAccessPtr>(OutDocument.AccessPoint, OutDocument);
 				return TransformInternal(FDocumentController::CreateDocumentHandle(DocAccessPtr));
 			}
@@ -63,8 +55,7 @@ namespace Metasound::Frontend
 
 			bool Transform(FMetaSoundFrontendDocumentBuilder& OutDocumentBuilder) const
 			{
-				const FMetaSoundFrontendDocumentBuilder& ConstBuilder = OutDocumentBuilder;
-				const FMetasoundFrontendDocumentMetadata& Metadata = ConstBuilder.GetDocument().Metadata;
+				const FMetasoundFrontendDocumentMetadata& Metadata = OutDocumentBuilder.GetConstDocument().Metadata;
 
 				const FMetasoundFrontendVersionNumber TargetVersion = GetTargetVersion();
 				if (Metadata.Version.Number < TargetVersion)
@@ -325,7 +316,8 @@ namespace Metasound::Frontend
 		void TransformInternal(FDocumentHandle InDocument) const override
 		{
 			const FGuid NewAssetClassID = FGuid::NewGuid();
-			FRenameRootGraphClass::Generate(InDocument, NewAssetClassID);
+			FMetasoundFrontendGraphClass Class = InDocument->GetRootGraphClass();
+			Class.Metadata.SetClassName(FMetasoundFrontendClassName({ }, FName(*NewAssetClassID.ToString()), { }));
 		}
 	};
 
diff --git a/Engine/Plugins/Runtime/Metasound/Source/MetasoundFrontend/Private/MetasoundFrontendGraphController.cpp b/Engine/Plugins/Runtime/Metasound/Source/MetasoundFrontend/Private/MetasoundFrontendGraphController.cpp
--- a/Engine/Plugins/Runtime/Metasound/Source/MetasoundFrontend/Private/MetasoundFrontendGraphController.cpp
+++ b/Engine/Plugins/Runtime/Metasound/Source/MetasoundFrontend/Private/MetasoundFrontendGraphController.cpp
@@ -1630,15 +1630,17 @@ namespace Metasound
 					FMetasoundFrontendNode& Node = GraphClass->Graph.Nodes.Emplace_GetRef(*NodeClass);
 
 					// Cache the asset name on the node if it node is reference to asset-defined graph.
-					if (NodeClass->Metadata.GetType() == EMetasoundFrontendClassType::External)
+					// AssetManager may not exist if this is called from a build that does not load the
+					// engine module (ex. unit test builds which only load frontend and core), or any build
+					// where the AssetManager more generally has not been implemented.
+					if (IMetaSoundAssetManager* AssetManager = IMetaSoundAssetManager::Get())
 					{
-						if (IMetaSoundAssetManager* AssetManager = IMetaSoundAssetManager::Get())
+						if (NodeClass->Metadata.GetType() == EMetasoundFrontendClassType::External)
 						{
 							const FNodeRegistryKey RegistryKey = FNodeRegistryKey(NodeClass->Metadata);
-							if (const FSoftObjectPath* Path = AssetManager->FindObjectPathFromKey(RegistryKey))
+							if (const FTopLevelAssetPath* Path = AssetManager->FindAssetPath(RegistryKey))
 							{
-								const FString& AssetName = Path->GetAssetName();
-								Node.Name = *AssetName;
+								Node.Name = Path->GetAssetName();
 							}
 						}
 					}
diff --git a/Engine/Plugins/Runtime/Metasound/Source/MetasoundFrontend/Private/MetasoundFrontendNodeController.cpp b/Engine/Plugins/Runtime/Metasound/Source/MetasoundFrontend/Private/MetasoundFrontendNodeController.cpp
--- a/Engine/Plugins/Runtime/Metasound/Source/MetasoundFrontend/Private/MetasoundFrontendNodeController.cpp
+++ b/Engine/Plugins/Runtime/Metasound/Source/MetasoundFrontend/Private/MetasoundFrontendNodeController.cpp
@@ -1014,12 +1014,9 @@ namespace Metasound
 			OutInterfaceUpdates = { };
 
 			const FMetasoundFrontendClassMetadata& NodeClassMetadata = GetClassMetadata();
-			if (IMetaSoundAssetManager* AssetManager = IMetaSoundAssetManager::Get())
+			if (!IMetaSoundAssetManager::GetChecked().CanAutoUpdate(NodeClassMetadata.GetClassName()))
 			{
-				if (!AssetManager->CanAutoUpdate(NodeClassMetadata.GetClassName()))
-				{
-					return false;
-				}
+				return false;
 			}
 
 			FMetasoundFrontendClass RegistryClass;
diff --git a/Engine/Plugins/Runtime/Metasound/Source/MetasoundFrontend/Private/MetasoundFrontendRegistryContainerImpl.cpp b/Engine/Plugins/Runtime/Metasound/Source/MetasoundFrontend/Private/MetasoundFrontendRegistryContainerImpl.cpp
--- a/Engine/Plugins/Runtime/Metasound/Source/MetasoundFrontend/Private/MetasoundFrontendRegistryContainerImpl.cpp
+++ b/Engine/Plugins/Runtime/Metasound/Source/MetasoundFrontend/Private/MetasoundFrontendRegistryContainerImpl.cpp
@@ -32,12 +32,6 @@
 #include "MetasoundRouter.h"
 #include "MetasoundTrace.h"
 
-bool bBusyWaitOnAsyncRegistrationTasks = true;
-static FAutoConsoleVariableRef CVarAsyncRegistrationTasksBusyWait(
-	TEXT("au.MetaSound.BusyWaitOnAsyncRegistrationTasks"),
-	bBusyWaitOnAsyncRegistrationTasks,
-	TEXT("Use TaskGraph BusyWait instead of simple Wait. Required to avoid hangs on platforms with low number of cores."),
-	ECVF_Default);
 
 namespace Metasound::Frontend
 {
@@ -62,6 +56,7 @@ namespace Metasound::Frontend
 
 			UObject* DocObject = DocumentInterface.GetObject();
 			check(DocObject);
+			const FMetasoundFrontendDocument& Document = DocumentInterface->GetConstDocument();
 
 #if WITH_EDITOR
 			// Only assets require template node processing
@@ -74,33 +69,37 @@ namespace Metasound::Frontend
 
 			return DocumentInterface;
 #else
+			const bool bIsBuilding = DocumentInterface->IsActivelyBuilding();
+			const bool bForceCopy = bIsBuilding && bAsync;
+
 	#if !NO_LOGGING
+		// Force a copy if async registration is enabled and we need to protect against race conditions from external modifications.
+
+		#if WITH_EDITORONLY_DATA
 			// Only assets require template node processing and support document attachment
-			if (DocObject->IsAsset()) 
+			if (DocObject->IsAsset())
 			{
-				if (IMetaSoundAssetManager* AssetManager = IMetaSoundAssetManager::Get())
+				const FMetaSoundFrontendDocumentBuilder& OriginalDocBuilder = IDocumentBuilderRegistry::GetChecked().FindOrBeginBuilding(DocumentInterface);
+				const bool bContainsTemplateDependency = OriginalDocBuilder.ContainsDependencyOfType(EMetasoundFrontendClassType::Template);
+				if (bContainsTemplateDependency)
 				{
-					// Hack to ensure the builder is available prior to testing validity of cooked build
-					// in order to report if cooked asset is properly processed. This will be replaced with
-					// actual accessor directly from registry and AssetManager::AttachDocumentBuilderChecked
-					// will be deprecated.
-					if (IDocumentBuilderRegistry* BuilderRegistry = IDocumentBuilderRegistry::Get())
-					{
-						const FMetaSoundFrontendDocumentBuilder& OriginalDocBuilder = AssetManager->AttachDocumentBuilderChecked(*DocObject);
-						const bool bContainsTemplateDependency = OriginalDocBuilder.ContainsDependencyOfType(EMetasoundFrontendClassType::Template);
-						if (bContainsTemplateDependency)
-						{
-							UE_LOG(LogMetaSound, Error,
-								TEXT("Template node processing disabled but provided asset class at '%s' to register contains template nodes. Runtime graph will fail to build."),
-								*OriginalDocBuilder.GetDebugName());
-						}
-					}
+					UE_LOG(LogMetaSound, Error,
+						TEXT("Template node processing disabled but provided asset class at '%s' to register contains template nodes. Runtime graph will fail to build."),
+						*OriginalDocBuilder.GetDebugName());
+				}
+
+				// Destroy builder if one didn't exist before running template check to ensure
+				// that builder existence doesn't inadvertently cause potential future re-registration
+				// calls to perform unnecessary document copy below.
+				if (!bIsBuilding)
+				{
+					const FMetasoundFrontendClassName& ClassName = Document.RootGraph.Metadata.GetClassName();
+					IDocumentBuilderRegistry::GetChecked().FinishBuilding(ClassName);
 				}
 			}
+		#endif // WITH_EDITORONLY_DATA
 	#endif // !NO_LOGGING
 
-			// Force a copy if async registration is enabled and we need to protect against race conditions from external modifications.
-			const bool bForceCopy = DocumentInterface->IsActivelyBuilding() && bAsync;
 			if (bForceCopy)
 			{
 				return &UMetaSoundBuilderDocument::Create(*DocumentInterface.GetInterface());
@@ -474,12 +473,10 @@ namespace Metasound::Frontend
 
 		METASOUND_TRACE_CPUPROFILER_EVENT_SCOPE_TEXT(*FString::Printf(TEXT("FRegistryContainerImpl::RegisterGraph key:%s, asset %s"), *RegistryKey.ToString(), *AssetPath.ToString()));
 
-		// Wait for any async tasks that are in flight which correspond to the same graph prior to building as building, if force copy is false, may mutate the serialized document
+		// Wait for any async tasks that are in flight which correspond to the same graph prior to building, even if this is a synchronous call.
+		WaitForAsyncGraphRegistration(RegistryKey);
+
 		const bool bAsync = !ConsoleVariables::bDisableAsyncGraphRegistration;
-		if (bAsync)
-		{
-			WaitForAsyncGraphRegistration(RegistryKey);
-		}
 
 		// Use the asset path of the provided document interface object for identification, *NOT* the
 		// built version as the build process may in fact create a new object with a transient path.
@@ -488,8 +485,6 @@ namespace Metasound::Frontend
 		UObject* OwningObject = RegistryDocInterface.GetObject();
 		check(OwningObject);
 
-		FNodeClassInfo NodeClassInfo(Document.RootGraph, AssetPath);
-
 		// Proxies are created synchronously to avoid creating proxies in async tasks. Proxies
 		// are created from UObjects which need to be protected from GC and non-GT access.
 		FProxyDataCache ProxyDataCache;
@@ -497,6 +492,7 @@ namespace Metasound::Frontend
 
 		// Store update to newly registered node in history so nodes
 		// can be queried by transaction ID
+		FNodeClassInfo NodeClassInfo(Document.RootGraph, AssetPath);
 		{
 			FNodeRegistryTransaction::FTimeType Timestamp = FPlatformTime::Cycles64();
 			TransactionBuffer->AddTransaction(FNodeRegistryTransaction(FNodeRegistryTransaction::ETransactionType::NodeRegistration, NodeClassInfo, Timestamp));
@@ -1224,15 +1220,8 @@ namespace Metasound::Frontend
 			METASOUND_TRACE_CPUPROFILER_EVENT_SCOPE(Metasound::FRegistryContainerImpl::WaitForRegistrationTaskToComplete);
 			if (Task.IsValid())
 			{
-				if (bBusyWaitOnAsyncRegistrationTasks)
-				{
-					Task.BusyWait();
-				}
-				else
-				{
-					Task.Wait();
-				}
-			}	
+				Task.BusyWait();
+			}
 		}
 	}
 } // namespace Metasound::Frontend
diff --git a/Engine/Plugins/Runtime/Metasound/Source/MetasoundFrontend/Private/MetasoundFrontendRegistryKey.cpp b/Engine/Plugins/Runtime/Metasound/Source/MetasoundFrontend/Private/MetasoundFrontendRegistryKey.cpp
--- a/Engine/Plugins/Runtime/Metasound/Source/MetasoundFrontend/Private/MetasoundFrontendRegistryKey.cpp
+++ b/Engine/Plugins/Runtime/Metasound/Source/MetasoundFrontend/Private/MetasoundFrontendRegistryKey.cpp
@@ -1,6 +1,7 @@
 // Copyright Epic Games, Inc. All Rights Reserved.
 
 #include "MetasoundFrontendRegistryKey.h"
+#include "MetasoundAssetManager.h"
 
 
 namespace Metasound::Frontend
@@ -107,6 +108,13 @@ namespace Metasound::Frontend
 		Parse(InKeyString, *this);
 	}
 
+	FNodeRegistryKey::FNodeRegistryKey(const FAssetKey& AssetKey)
+		: Type(EMetasoundFrontendClassType::External)
+		, ClassName(AssetKey.ClassName)
+		, Version(AssetKey.Version)
+	{
+	}
+
 	const FNodeRegistryKey& FNodeRegistryKey::GetInvalid()
 	{
 		static const FNodeRegistryKey InvalidKey;
diff --git a/Engine/Plugins/Runtime/Metasound/Source/MetasoundFrontend/Private/MetasoundFrontendTransform.cpp b/Engine/Plugins/Runtime/Metasound/Source/MetasoundFrontend/Private/MetasoundFrontendTransform.cpp
--- a/Engine/Plugins/Runtime/Metasound/Source/MetasoundFrontend/Private/MetasoundFrontendTransform.cpp
+++ b/Engine/Plugins/Runtime/Metasound/Source/MetasoundFrontend/Private/MetasoundFrontendTransform.cpp
@@ -56,12 +56,6 @@ namespace Metasound
 
 		bool IDocumentTransform::Transform(FMetasoundFrontendDocument& InOutDocument) const
 		{
-			if (IDocumentBuilderRegistry* DocRegistry = IDocumentBuilderRegistry::Get())
-			{
-				const FMetasoundFrontendClassName& Name = InOutDocument.RootGraph.Metadata.GetClassName();
-				DocRegistry->InvalidateDocumentCache(Name);
-			}
-
 			FDocumentAccessPtr DocAccessPtr = MakeAccessPtr<FDocumentAccessPtr>(InOutDocument.AccessPoint, InOutDocument);
 			return Transform(FDocumentController::CreateDocumentHandle(DocAccessPtr));
 		}
@@ -413,10 +407,16 @@ namespace Metasound
 
 		bool FModifyRootGraphInterfaces::Transform(FMetasoundFrontendDocument& InOutDocument) const
 		{
-			if (IDocumentBuilderRegistry* DocRegistry = IDocumentBuilderRegistry::Get())
 			{
-				const FMetasoundFrontendClassName& Name = InOutDocument.RootGraph.Metadata.GetClassName();
-				DocRegistry->InvalidateDocumentCache(Name);
+				// Mutation of a document via the soft deprecated access ptr/controller system is not tracked by
+				// the builder registry, so the document cache is invalidated here. It is discouraged to mutate
+				// documents using controllers at this point as it disconnects delegates applied at the object level
+				// (i.e. disconnecting changes to any MetaSound instances being auditioned).
+				const FMetasoundFrontendClassName& ClassName = InOutDocument.RootGraph.Metadata.GetClassName();
+				if (FMetaSoundFrontendDocumentBuilder* Builder = IDocumentBuilderRegistry::GetChecked().FindBuilder(ClassName))
+				{
+					Builder->Reload();
+				}
 			}
 
 			FDocumentAccessPtr DocAccessPtr = MakeAccessPtr<FDocumentAccessPtr>(InOutDocument.AccessPoint, InOutDocument);
@@ -537,7 +537,7 @@ namespace Metasound
 				const FMetasoundFrontendClassMetadata& ClassMetadata = NodeHandle->GetClassMetadata();
 				const FNodeRegistryKey RegistryKey(ClassMetadata);
 
-				if (FMetasoundAssetBase* ReferencedMetaSoundAsset = IMetaSoundAssetManager::GetChecked().TryLoadAssetFromKey(RegistryKey))
+				if (FMetasoundAssetBase* ReferencedMetaSoundAsset = IMetaSoundAssetManager::GetChecked().FindAsset(RegistryKey))
 				{
 					if (bIsPreset)
 					{
@@ -966,21 +966,12 @@ namespace Metasound
 
 		bool FRenameRootGraphClass::Transform(FDocumentHandle InDocument) const
 		{
-			FGraphHandle RootGraph = InDocument->GetRootGraph();
-
-			FMetasoundFrontendClassMetadata Metadata = RootGraph->GetGraphMetadata();
-			Metadata.SetClassName(NewClassName);
-			RootGraph->SetGraphMetadata(Metadata);
-
-			return true;
+			return false;
 		}
 
 		bool FRenameRootGraphClass::Transform(FMetasoundFrontendDocument& InOutDocument) const
 		{
-			FMetasoundFrontendClassMetadata& Metadata = InOutDocument.RootGraph.Metadata;
-			Metadata.SetClassName(NewClassName);
-
-			return true;
+			return false;
 		}
 	} // namespace Frontend
 } // namespace Metasound
diff --git a/Engine/Plugins/Runtime/Metasound/Source/MetasoundFrontend/Private/NodeTemplates/MetasoundFrontendNodeTemplateInput.cpp b/Engine/Plugins/Runtime/Metasound/Source/MetasoundFrontend/Private/NodeTemplates/MetasoundFrontendNodeTemplateInput.cpp
--- a/Engine/Plugins/Runtime/Metasound/Source/MetasoundFrontend/Private/NodeTemplates/MetasoundFrontendNodeTemplateInput.cpp
+++ b/Engine/Plugins/Runtime/Metasound/Source/MetasoundFrontend/Private/NodeTemplates/MetasoundFrontendNodeTemplateInput.cpp
@@ -206,7 +206,7 @@ namespace Metasound::Frontend
 		bool bInjectedNodes = false;
 
 		const FMetaSoundFrontendDocumentBuilder& ConstBuilder = static_cast<const FMetaSoundFrontendDocumentBuilder&>(InOutBuilder);
-		const TArray<FMetasoundFrontendClassInput>& Inputs = ConstBuilder.GetDocument().RootGraph.Interface.Inputs;
+		const TArray<FMetasoundFrontendClassInput>& Inputs = ConstBuilder.GetConstDocument().RootGraph.Interface.Inputs;
 		for (const FMetasoundFrontendClassInput& Input : Inputs)
 		{
 			// Potentially not used input, which is perfectly valid so early out
diff --git a/Engine/Plugins/Runtime/Metasound/Source/MetasoundFrontend/Public/MetasoundAssetBase.h b/Engine/Plugins/Runtime/Metasound/Source/MetasoundFrontend/Public/MetasoundAssetBase.h
--- a/Engine/Plugins/Runtime/Metasound/Source/MetasoundFrontend/Public/MetasoundAssetBase.h
+++ b/Engine/Plugins/Runtime/Metasound/Source/MetasoundFrontend/Public/MetasoundAssetBase.h
@@ -63,10 +63,8 @@ public:
 	// Returns the graph associated with this Metasound. Graph is required to be referenced on
 	// Metasound UObject for editor serialization purposes.
 	// @return Editor graph associated with this metasound uobject.
-	virtual UEdGraph* GetGraph() = 0;
-	virtual const UEdGraph* GetGraph() const = 0;
-	virtual UEdGraph& GetGraphChecked() = 0;
-	virtual const UEdGraph& GetGraphChecked() const = 0;
+	virtual UEdGraph* GetGraph() const = 0;
+	virtual UEdGraph& GetGraphChecked() const = 0;
 	virtual void MigrateEditorGraph(FMetaSoundFrontendDocumentBuilder& OutBuilder) = 0;
 
 	// Sets the graph associated with this Metasound. Graph is required to be referenced on
@@ -116,8 +114,11 @@ public:
 	// Called when async assets have finished loading.
 	virtual void OnAsyncReferencedAssetsLoaded(const TArray<FMetasoundAssetBase*>& InAsyncReferences) = 0;
 
+	bool AddingReferenceCausesLoop(const FMetasoundAssetBase& InMetaSound) const;
 
+	UE_DEPRECATED(5.5, "Use overload that is provided an AssetBase")
 	bool AddingReferenceCausesLoop(const FSoftObjectPath& InReferencePath) const;
+
 	bool IsReferencedAsset(const FMetasoundAssetBase& InAssetToCheck) const;
 
 	bool IsRegistered() const;
diff --git a/Engine/Plugins/Runtime/Metasound/Source/MetasoundFrontend/Public/MetasoundAssetManager.h b/Engine/Plugins/Runtime/Metasound/Source/MetasoundFrontend/Public/MetasoundAssetManager.h
--- a/Engine/Plugins/Runtime/Metasound/Source/MetasoundFrontend/Public/MetasoundAssetManager.h
+++ b/Engine/Plugins/Runtime/Metasound/Source/MetasoundFrontend/Public/MetasoundAssetManager.h
@@ -2,8 +2,8 @@
 #pragma once
 
 #include "MetasoundFrontendDocument.h"
-#include "MetasoundVertex.h"
 #include "MetasoundFrontendRegistryKey.h"
+#include "MetasoundVertex.h"
 #include "UObject/SoftObjectPath.h"
 #include "UObject/WeakObjectPtrTemplates.h"
 
@@ -11,165 +11,204 @@
 class FMetasoundAssetBase;
 class UEdGraph;
 
+struct FAssetData;
 struct FMetasoundFrontendClassName;
 struct FMetaSoundFrontendDocumentBuilder;
+class IMetaSoundDocumentInterface;
 
-namespace Metasound
+namespace Metasound::Frontend
 {
-	namespace Frontend
+	namespace AssetTags
 	{
-		namespace AssetTags
-		{
-			extern const FString METASOUNDFRONTEND_API ArrayDelim;
+		extern const FString METASOUNDFRONTEND_API ArrayDelim;
 
 #if WITH_EDITORONLY_DATA
-			extern const FName METASOUNDFRONTEND_API IsPreset;
+		extern const FName METASOUNDFRONTEND_API IsPreset;
 #endif // WITH_EDITORONLY_DATA
 
-			extern const FName METASOUNDFRONTEND_API AssetClassID;
-			extern const FName METASOUNDFRONTEND_API RegistryVersionMajor;
-			extern const FName METASOUNDFRONTEND_API RegistryVersionMinor;
+		extern const FName METASOUNDFRONTEND_API AssetClassID;
+		extern const FName METASOUNDFRONTEND_API RegistryVersionMajor;
+		extern const FName METASOUNDFRONTEND_API RegistryVersionMinor;
 
 #if WITH_EDITORONLY_DATA
-			extern const FName METASOUNDFRONTEND_API RegistryInputTypes;
-			extern const FName METASOUNDFRONTEND_API RegistryOutputTypes;
+		extern const FName METASOUNDFRONTEND_API RegistryInputTypes;
+		extern const FName METASOUNDFRONTEND_API RegistryOutputTypes;
 #endif // WITH_EDITORONLY_DATA
-		} // namespace AssetTags
+	} // namespace AssetTags
 
-		struct METASOUNDFRONTEND_API FMetaSoundAssetRegistrationOptions
-		{
-			// If true, forces a re-register of this class (and all class dependencies
-			// if the following option 'bRegisterDependencies' is enabled).
-			bool bForceReregister = true;
+	struct METASOUNDFRONTEND_API FMetaSoundAssetRegistrationOptions
+	{
+		// If true, forces a re-register of this class (and all class dependencies
+		// if the following option 'bRegisterDependencies' is enabled).
+		bool bForceReregister = true;
 
-			// If true, forces flag to resync all view (editor) data pertaining to the given asset(s) being registered.
-			bool bForceViewSynchronization = true;
+		// If true, forces flag to resync all view (editor) data pertaining to the given asset(s) being registered.
+		bool bForceViewSynchronization = true;
 
-			// If true, recursively attempts to register dependencies. (TODO: Determine if this option should be removed.
-			// Must validate that failed dependency updates due to auto-update for ex. being disabled is handled gracefully
-			// at runtime.)
-			bool bRegisterDependencies = true;
+		// If true, recursively attempts to register dependencies. (TODO: Determine if this option should be removed.
+		// Must validate that failed dependency updates due to auto-update for ex. being disabled is handled gracefully
+		// at runtime.)
+		bool bRegisterDependencies = true;
 
-			// Attempt to auto-update (Only runs if class not registered or set to force re-register.
-			// Will not respect being set to true if project-level MetaSoundSettings specify to not run auto-update.)
-			bool bAutoUpdate = true;
+		// Attempt to auto-update (Only runs if class not registered or set to force re-register.
+		// Will not respect being set to true if project-level MetaSoundSettings specify to not run auto-update.)
+		bool bAutoUpdate = true;
 
-			// If true, warnings will be logged if updating a node results in existing connections being discarded.
-			bool bAutoUpdateLogWarningOnDroppedConnection = false;
+		// If true, warnings will be logged if updating a node results in existing connections being discarded.
+		bool bAutoUpdateLogWarningOnDroppedConnection = false;
 
 #if WITH_EDITOR
-			// Soft deprecated. Preprocessing now handled contextually if cooking or serializing.
-			bool bPreprocessDocument = true;
+		// Soft deprecated. Preprocessing now handled contextually if cooking or serializing.
+		bool bPreprocessDocument = true;
 
-			// Attempt to rebuild referenced classes (only run if class not registered or set to force re-register)
-			bool bRebuildReferencedAssetClasses = true;
+		// Attempt to rebuild referenced classes (only run if class not registered or set to force re-register)
+		bool bRebuildReferencedAssetClasses = true;
 
-			// If true and the registration will be perform asynchronously, a copy of the graph will be used for
-			// async registration. 
-			bool bRegisterCopyIfAsync = false;
+		// No longer used. Memory management of document (i.e. copying or using object's version) inferred internally
+		bool bRegisterCopyIfAsync = false;
 #endif // WITH_EDITOR
-		};
+	};
+
+	struct METASOUNDFRONTEND_API FAssetKey
+	{
+		FMetasoundFrontendClassName ClassName;
+		FMetasoundFrontendVersionNumber Version;
 
-		class METASOUNDFRONTEND_API IMetaSoundAssetManager
+		FAssetKey() = default;
+		FAssetKey(const FMetasoundFrontendClassName& InClassName, const FMetasoundFrontendVersionNumber& InVersion);
+		FAssetKey(const FNodeRegistryKey& RegKey);
+		FAssetKey(const FMetasoundFrontendClassMetadata& InMetadata);
+
+		FORCEINLINE friend bool operator==(const FAssetKey& InLHS, const FAssetKey& InRHS)
 		{
-			static IMetaSoundAssetManager* Instance;
+			return (InLHS.ClassName == InRHS.ClassName) && (InLHS.Version == InRHS.Version);
+		}
 
-		public:
-			static void Set(IMetaSoundAssetManager& InInterface)
+		FORCEINLINE friend bool operator<(const FAssetKey& InLHS, const FAssetKey& InRHS)
+		{
+			if (InLHS.ClassName == InRHS.ClassName)
 			{
-				if (!InInterface.IsTesting())
-				{
-					check(!Instance);
-				}
-				Instance = &InInterface;
-
-#if WITH_SERVER_CODE
-				OnManagerSet.Broadcast();
-#endif //WITH_SERVER_CODE
+				return InLHS.Version < InRHS.Version;
 			}
 
-			static IMetaSoundAssetManager* Get()
-			{
-				return Instance;
-			}
+			return InLHS.ClassName < InRHS.ClassName;
+		}
 
-			static IMetaSoundAssetManager& GetChecked()
-			{
-				check(Instance);
-				return *Instance;
-			}
+		friend FORCEINLINE uint32 GetTypeHash(const FAssetKey& InKey)
+		{
+			return HashCombineFast(GetTypeHash(InKey.ClassName), GetTypeHash(InKey.Version));
+		}
 
-#if WITH_SERVER_CODE
-			//to allow the server to preload UMetaSoundSource objects before the manager has been set up
-			static TMulticastDelegate<void()> OnManagerSet;
-#endif //WITH_SERVER_CODE
+		static const FAssetKey& GetInvalid()
+		{
+			static FAssetKey Invalid;
+			return Invalid;
+		}
 
-			struct FAssetInfo
-			{
-				const Metasound::Frontend::FNodeRegistryKey RegistryKey;
-				FSoftObjectPath AssetPath;
+		bool IsValid() const
+		{
+			return ClassName.IsValid() && Version.IsValid();
+		}
 
-				FORCEINLINE friend bool operator==(const FAssetInfo& InLHS, const FAssetInfo& InRHS)
-				{
-					return (InLHS.RegistryKey == InRHS.RegistryKey) && (InLHS.AssetPath == InRHS.AssetPath);
-				}
+		FString ToString() const;
+	};
 
-				FORCEINLINE friend uint32 GetTypeHash(const IMetaSoundAssetManager::FAssetInfo& InInfo)
-				{
-					return HashCombineFast(GetTypeHash(InInfo.RegistryKey), GetTypeHash(InInfo.AssetPath));
-				}
-			};
 
-			// Whether or not manager is being used to run tests or not (enabling instances to be reset without asserting.)
-			virtual bool IsTesting() const { return false; }
+	class METASOUNDFRONTEND_API IMetaSoundAssetManager
+	{
+	public:
+		virtual ~IMetaSoundAssetManager() = default;
+
+		static IMetaSoundAssetManager* Get();
+		static IMetaSoundAssetManager& GetChecked();
+		static void Deinitialize();
+		static void Initialize(TUniquePtr<IMetaSoundAssetManager>&& InInterface);
+
+		UE_DEPRECATED(5.5, "Use Initialize/Deinitialize instead")
+		virtual bool IsTesting() const;
+
+		UE_DEPRECATED(5.5, "Use Initialize/Deinitialize instead")
+		static void Set(IMetaSoundAssetManager& InInterface) { checkNoEntry(); }
+
+		struct FAssetInfo
+		{
+			const FNodeRegistryKey RegistryKey;
+			FSoftObjectPath AssetPath;
+
+			FORCEINLINE friend bool operator==(const FAssetInfo& InLHS, const FAssetInfo& InRHS)
+			{
+				return (InLHS.RegistryKey == InRHS.RegistryKey) && (InLHS.AssetPath == InRHS.AssetPath);
+			}
+
+			FORCEINLINE friend uint32 GetTypeHash(const IMetaSoundAssetManager::FAssetInfo& InInfo)
+			{
+				return HashCombineFast(GetTypeHash(InInfo.RegistryKey), GetTypeHash(InInfo.AssetPath));
+			}
+		};
 
 #if WITH_EDITORONLY_DATA
-			// Adds missing assets using the provided asset's local reference class cache. Used
-			// to prime system from asset attempting to register prior to asset scan being complete.
-			virtual void AddAssetReferences(FMetasoundAssetBase& InAssetBase) = 0;
+		// Adds missing assets using the provided asset's local reference class cache. Used
+		// to prime system from asset attempting to register prior to asset scan being complete.
+		virtual void AddAssetReferences(FMetasoundAssetBase& InAssetBase) = 0;
 #endif
 
-			// Add or Update a MetaSound Asset's entry data
-			virtual Metasound::Frontend::FNodeRegistryKey AddOrUpdateAsset(const UObject& InObject) = 0;
+		// Add or Update a MetaSound Asset's entry data
+		virtual FAssetKey AddOrUpdateAsset(const UObject& InObject) = 0;
+		virtual FAssetKey AddOrUpdateAsset(const FAssetData& InAssetData) = 0;
 
-			virtual FMetaSoundFrontendDocumentBuilder& AttachDocumentBuilderChecked(UObject& InObject) const = 0;
+		// Whether or not the class is eligible for auto-update
+		virtual bool CanAutoUpdate(const FMetasoundFrontendClassName& InClassName) const = 0;
 
-			// Whether or not the class is eligible for auto-update
-			virtual bool CanAutoUpdate(const FMetasoundFrontendClassName& InClassName) const = 0;
+		// Whether or not the asset manager has loaded the given asset
+		virtual bool ContainsKey(const Metasound::Frontend::FAssetKey& InAssetKey) const = 0;
 
-			// Whether or not the asset manager has loaded the given asset
-			virtual bool ContainsKey(const Metasound::Frontend::FNodeRegistryKey& InRegistryKey) const = 0;
+		// Returns object (if loaded) associated with the given key (null if key is not registered with the AssetManager)
+		virtual FMetasoundAssetBase* FindAsset(const Metasound::Frontend::FAssetKey& InAssetKey) const = 0;
 
-			// Returns path associated with the given key (null if key is not registered with the AssetManager or was not loaded from asset)
-			virtual const FSoftObjectPath* FindObjectPathFromKey(const Metasound::Frontend::FNodeRegistryKey& InRegistryKey) const = 0;
+		// Returns object (if loaded) associated with the given key as a Document Interface (null if key is not registered with the AssetManager)
+		virtual TScriptInterface<IMetaSoundDocumentInterface> FindAssetAsDocumentInterface(const Frontend::FAssetKey& InKey) const = 0;
 
-			// Converts an object to an AssetBase if its a registered asset
-			virtual FMetasoundAssetBase* GetAsAsset(UObject& InObject) const = 0;
-			virtual const FMetasoundAssetBase* GetAsAsset(const UObject& InObject) const = 0;
+		// Returns path associated with the given key (null if key is not registered with the AssetManager or was not loaded from asset)
+		virtual const FTopLevelAssetPath* FindAssetPath(const Metasound::Frontend::FAssetKey& InAssetKey) const = 0;
+
+		// Converts an object to an AssetBase if its a registered asset
+		virtual FMetasoundAssetBase* GetAsAsset(UObject& InObject) const = 0;
+		virtual const FMetasoundAssetBase* GetAsAsset(const UObject& InObject) const = 0;
 
 #if WITH_EDITOR
-			// Generates all asset info associated with registered assets that are referenced by the provided asset's graph.
-			virtual TSet<FAssetInfo> GetReferencedAssetClasses(const FMetasoundAssetBase& InAssetBase) const = 0;
+		// Generates all asset info associated with registered assets that are referenced by the provided asset's graph.
+		virtual TSet<FAssetInfo> GetReferencedAssetClasses(const FMetasoundAssetBase& InAssetBase) const = 0;
 #endif // WITH_EDITOR
 
-			// Rescans settings for denied assets not to run reference auto-update against.
-			virtual void RescanAutoUpdateDenyList() = 0;
+		UE_DEPRECATED(5.5, "Rescan no longer supported nor required by Frontend")
+		virtual void RescanAutoUpdateDenyList() { }
+
+		// Attempts to retrieve the AssetID from the given ClassName if the ClassName is from a valid asset.
+		virtual bool TryGetAssetIDFromClassName(const FMetasoundFrontendClassName& InClassName, FGuid& OutGuid) const = 0;
 
-			// Attempts to load an FMetasoundAssetBase from the given path, or returns it if its already loaded
-			virtual FMetasoundAssetBase* TryLoadAsset(const FSoftObjectPath& InObjectPath) const = 0;
+		// Attempts to load an FMetasoundAssetBase from the given path, or returns it if its already loaded
+		virtual FMetasoundAssetBase* TryLoadAsset(const FSoftObjectPath& InObjectPath) const = 0;
 
-			// Returns asset associated with the given key (null if key is not registered with the AssetManager or was not loaded from asset)
-			virtual FMetasoundAssetBase* TryLoadAssetFromKey(const Metasound::Frontend::FNodeRegistryKey& InRegistryKey) const = 0;
+		// Returns asset associated with the given key (null if key is not registered with the AssetManager or was not loaded from asset)
+		virtual FMetasoundAssetBase* TryLoadAssetFromKey(const Metasound::Frontend::FAssetKey& InAssetKey) const = 0;
 
-			// Try to load referenced assets of the given asset or return them if they are already loaded (non-recursive).
-			// @return - True if all referenced assets successfully loaded, false if not.
-			virtual bool TryLoadReferencedAssets(const FMetasoundAssetBase& InAssetBase, TArray<FMetasoundAssetBase*>& OutReferencedAssets) const = 0;
+		// Try to load referenced assets of the given asset or return them if they are already loaded (non-recursive).
+		// @return - True if all referenced assets successfully loaded, false if not.
+		virtual bool TryLoadReferencedAssets(const FMetasoundAssetBase& InAssetBase, TArray<FMetasoundAssetBase*>& OutReferencedAssets) const = 0;
 			
-			// Requests an async load of all async referenced assets of the input asset.
-			virtual void RequestAsyncLoadReferencedAssets(FMetasoundAssetBase& InAssetBase) = 0;
+		// Requests an async load of all async referenced assets of the input asset.
+		virtual void RequestAsyncLoadReferencedAssets(FMetasoundAssetBase& InAssetBase) = 0;
 
-			// Waits until all async load requests related to this asset are complete. 
- 			virtual void WaitUntilAsyncLoadReferencedAssetsComplete(FMetasoundAssetBase& InAssetBase) = 0;
-		};
-	} // namespace Frontend
-};
+		// Removes object from MetaSound asset manager
+		virtual void RemoveAsset(const UObject& InObject) = 0;
+
+		// Removes object from MetaSound asset manager
+		virtual void RemoveAsset(const FAssetData& InAssetData) = 0;
+
+		// Updates the given MetaSound's asset record with the new name and optionally reregisters it with the Frontend Node Class Registry.
+		virtual void RenameAsset(const FAssetData& InAssetData, bool bInReregisterWithFrontend = true) = 0;
+
+		// Waits until all async load requests related to this asset are complete.
+ 		virtual void WaitUntilAsyncLoadReferencedAssetsComplete(FMetasoundAssetBase& InAssetBase) = 0;
+	};
+} // namespace Metasound::Frontend
diff --git a/Engine/Plugins/Runtime/Metasound/Source/MetasoundFrontend/Public/MetasoundDocumentInterface.h b/Engine/Plugins/Runtime/Metasound/Source/MetasoundFrontend/Public/MetasoundDocumentInterface.h
--- a/Engine/Plugins/Runtime/Metasound/Source/MetasoundFrontend/Public/MetasoundDocumentInterface.h
+++ b/Engine/Plugins/Runtime/Metasound/Source/MetasoundFrontend/Public/MetasoundDocumentInterface.h
@@ -13,6 +13,10 @@
 // Forward Declarations
 struct FMetaSoundFrontendDocumentBuilder;
 
+namespace Metasound::Frontend
+{
+	struct FAssetKey;
+}
 
 // UInterface for all MetaSound UClasses that implement a MetaSound document
 // as a means for accessing data via code, scripting, execution, or node
@@ -45,6 +49,9 @@ public:
 	// Returns the parent class registered with the MetaSound UObject registry.
 	virtual const UClass& GetBaseMetaSoundUClass() const = 0;
 
+	// Returns the builder class used to modify the given document.
+	virtual const UClass& GetBuilderUClass() const = 0;
+
 	// Conforms UProperty data outside the Frontend Document Model to the document's data.
 	// Returns whether or not object data was modified.
 	virtual bool ConformObjectToDocument() = 0;
@@ -69,23 +76,42 @@ namespace Metasound::Frontend
 	class METASOUNDFRONTEND_API IDocumentBuilderRegistry
 	{
 	public:
-		// Invalidates the cache of a given document's builder should one be registered, causing it to be rebuilt.  Not recommended
-		// for general use, and is only available in case the document is modified by a system other than an active, registered builder
-		// (ex. via the soft deprecated controller/handle API).
-		virtual void InvalidateDocumentCache(const FMetasoundFrontendClassName& InClassName) const = 0;
+		virtual ~IDocumentBuilderRegistry() = default;
+
+		virtual FMetaSoundFrontendDocumentBuilder* FindBuilder(TScriptInterface<IMetaSoundDocumentInterface> MetaSound) const = 0;
+		virtual FMetaSoundFrontendDocumentBuilder* FindBuilder(const FMetasoundFrontendClassName& ClassName) const = 0;
+
+#if WITH_EDITORONLY_DATA
+		// Find the existing builder for the given MetaSound, or optionally begin building by attaching a new builder.  Only available
+		// in builds with editor only data as building serialized assets (which may have template nodes, cooked builds do not) is only
+		// supported when editor data is loaded. Creating transient builders can simply be done by passing a new MetaSound asset to
+		// a FMetaSoundFrontendDocumentBuilder constructor, or this registry's implementation may supply its own create call for tracking
+		// and reuse purposes.
+		virtual FMetaSoundFrontendDocumentBuilder& FindOrBeginBuilding(TScriptInterface<IMetaSoundDocumentInterface> MetaSound) = 0;
+#endif // WITH_EDITORONLY_DATA
+
+		virtual bool FinishBuilding(const FMetasoundFrontendClassName& InClassName) const = 0;
+
+		UE_DEPRECATED(5.5, "Document cache can now be invalidated by retrieving an asset builder and calling 'Reload'")
+		virtual void InvalidateDocumentCache(const FMetasoundFrontendClassName& InClassName) const { }
 
 		static IDocumentBuilderRegistry* Get();
 		static IDocumentBuilderRegistry& GetChecked();
+		static void Deinitialize();
+		static void Initialize(TUniquePtr<IDocumentBuilderRegistry>&& InBuilderRegistry);
 
 	protected:
-		static void Set(TUniqueFunction<IDocumentBuilderRegistry&()>&& InGetInstance);
+		UE_DEPRECATED(5.5, "Use 'Initialize' instead")
+		static void Set(TUniqueFunction<IDocumentBuilderRegistry&()>&& InGetInstance) { checkNoEntry(); }
 	};
 
 	class METASOUNDFRONTEND_API IMetaSoundDocumentBuilderRegistry : public IDocumentBuilderRegistry
 	{
 	public:
+		virtual ~IMetaSoundDocumentBuilderRegistry() = default;
+
 		UE_DEPRECATED(5.4, "Public exposition of modify delegates no longer available to discourage unsafe manipulation of builder document cache")
-		virtual const FDocumentModifyDelegates* FindModifyDelegates(const FMetasoundFrontendClassName& InClassName) const = 0;
+		virtual const FDocumentModifyDelegates* FindModifyDelegates(const FMetasoundFrontendClassName& InClassName) const { return nullptr; }
 
 		UE_DEPRECATED(5.4, "Use 'IDocumentBuilderRegistry' instead")
 		static IMetaSoundDocumentBuilderRegistry& GetChecked();
diff --git a/Engine/Plugins/Runtime/Metasound/Source/MetasoundFrontend/Public/MetasoundFrontendDocumentBuilder.h b/Engine/Plugins/Runtime/Metasound/Source/MetasoundFrontend/Public/MetasoundFrontendDocumentBuilder.h
--- a/Engine/Plugins/Runtime/Metasound/Source/MetasoundFrontend/Public/MetasoundFrontendDocumentBuilder.h
+++ b/Engine/Plugins/Runtime/Metasound/Source/MetasoundFrontend/Public/MetasoundFrontendDocumentBuilder.h
@@ -89,13 +89,15 @@ namespace Metasound::Frontend
 } // namespace Metasound::Frontend
 
 
+// Builder Document UObject, which is only used for registration purposes when attempting
+// async registration whereby the original document is serialized and must not be mutated.
 UCLASS()
 class METASOUNDFRONTEND_API UMetaSoundBuilderDocument : public UObject, public IMetaSoundDocumentInterface
 {
 	GENERATED_BODY()
 
 public:
-	// Create and return a valid builder document with the provided class
+	UE_DEPRECATED(5.5, "Use overload supplying MetaSound to copy (builder documents no longer supported for cases outside of cloned document registration.")
 	static UMetaSoundBuilderDocument& Create(const UClass& InBuilderClass);
 
 	// Create and return a valid builder document which copies the provided interface's document & class
@@ -109,9 +111,12 @@ public:
 	// Returns temp path of builder document
 	virtual FTopLevelAssetPath GetAssetPathChecked() const override;
 
-	// Base MetaSoundClass that document is published to.
+	// Returns the base class registered with the MetaSound UObject registry.
 	virtual const UClass& GetBaseMetaSoundUClass() const final override;
 
+	// Returns the builder class used to modify the given document.
+	virtual const UClass& GetBuilderUClass() const final override;
+
 	// Returns if the document is being actively built (always true as builder documents are always being actively built)
 	virtual bool IsActivelyBuilding() const final override;
 
@@ -126,6 +131,9 @@ private:
 
 	UPROPERTY(Transient)
 	TObjectPtr<const UClass> MetaSoundUClass = nullptr;
+
+	UPROPERTY(Transient)
+	TObjectPtr<const UClass> BuilderUClass = nullptr;
 };
 
 // Builder used to support dynamically generating MetaSound documents at runtime. Builder contains caches that speed up
@@ -213,23 +221,26 @@ public:
 	TArray<const FMetasoundFrontendVertex*> FindNodeOutputs(const FGuid& InNodeID, FName TypeName = FName()) const;
 	const FMetasoundFrontendVertex* FindNodeOutputConnectedToNodeInput(const FGuid& InInputNodeID, const FGuid& InInputVertexID, const FMetasoundFrontendNode** ConnectedOutputNode = nullptr) const;
 
+	const FMetasoundFrontendDocument& GetConstDocument() const;
 	const FString GetDebugName() const;
-	const FMetasoundFrontendDocument& GetDocument() const;
 
-	template<typename TObjectType>
-	const TObjectType& CastDocumentObjectChecked() const
-	{
-		const UObject* Owner = DocumentInterface.GetObject();
-		return *CastChecked<TObjectType>(Owner);
-	}
+	UE_DEPRECATED(5.5, "Use GetConstDocument() instead")
+	const FMetasoundFrontendDocument& GetDocument() const;
 
 	template<typename TObjectType>
-	TObjectType& CastDocumentObjectChecked()
+	TObjectType& CastDocumentObjectChecked() const
 	{
 		UObject* Owner = DocumentInterface.GetObject();
 		return *CastChecked<TObjectType>(Owner);
 	}
 
+	// Generates and returns new class name for the given builder's document. Should ONLY be called on new assets
+	// using transient frontend builder, as using a persistent builder registered with the DocumentBuilderRegistry
+	// may result in corrupt records keyed off of undefined class name being potentially generated over.  In addition,
+	// this can potentially leave existing node references in an abandoned state to this class causing MetaSound generator
+	// build errors.
+	FMetasoundFrontendClassName GenerateNewClassName();
+
 	const Metasound::Frontend::FDocumentModifyDelegates& GetDocumentDelegates() const;
 	const IMetaSoundDocumentInterface& GetDocumentInterface() const;
 	FMetasoundAssetBase& GetMetasoundAsset();
@@ -303,6 +314,8 @@ public:
 
 	bool RemoveNodeInputDefault(const FGuid& InNodeID, const FGuid& InVertexID);
 	bool RemoveUnusedDependencies();
+
+	UE_DEPRECATED(5.5, "Use GenerateNewClassName instead")
 	bool RenameRootGraphClass(const FMetasoundFrontendClassName& InName);
 
 #if WITH_EDITOR
@@ -318,6 +331,8 @@ public:
 	bool SetGraphOutputDataType(FName OutputName, FName DataType);
 
 #if WITH_EDITOR
+	void SetDisplayName(const FText& InDisplayName);
+
 	void SetMemberMetadata(UMetaSoundFrontendMemberMetadata& NewMetadata);
 
 	// Sets the editor-only comment to the provided value.
@@ -342,9 +357,11 @@ public:
 	bool SwapGraphOutput(const FMetasoundFrontendClassVertex& InExistingOutputVertex, const FMetasoundFrontendClassVertex& NewOutputVertex);
 	bool UpdateDependencyClassNames(const TMap<FMetasoundFrontendClassName, FMetasoundFrontendClassName>& OldToNewReferencedClassNames);
 
+#if WITH_EDITORONLY_DATA
 	// Transforms template nodes within the given builder's document, which can include swapping associated edges and/or
 	// replacing nodes with other, registry-defined concrete node class instances. Returns true if any template nodes were processed.
 	bool TransformTemplateNodes();
+#endif // WITH_EDITORONLY_DATA
 
 private:
 	using FFinalizeNodeFunctionRef = TFunctionRef<void(FMetasoundFrontendNode&, const Metasound::Frontend::FNodeRegistryKey&)>;
diff --git a/Engine/Plugins/Runtime/Metasound/Source/MetasoundFrontend/Public/MetasoundFrontendRegistryKey.h b/Engine/Plugins/Runtime/Metasound/Source/MetasoundFrontend/Public/MetasoundFrontendRegistryKey.h
--- a/Engine/Plugins/Runtime/Metasound/Source/MetasoundFrontend/Public/MetasoundFrontendRegistryKey.h
+++ b/Engine/Plugins/Runtime/Metasound/Source/MetasoundFrontend/Public/MetasoundFrontendRegistryKey.h
@@ -10,6 +10,9 @@
 
 namespace Metasound::Frontend
 {
+	// Forward Declarations
+	struct FAssetKey;
+
 	/** FNodeClassInfo contains a minimal set of information needed to find and query node classes. */
 	struct METASOUNDFRONTEND_API FNodeClassInfo
 	{
@@ -69,6 +72,7 @@ namespace Metasound::Frontend
 		FNodeRegistryKey(const FMetasoundFrontendClassMetadata& InNodeMetadata);
 		FNodeRegistryKey(const FMetasoundFrontendGraphClass& InGraphClass);
 		FNodeRegistryKey(const FNodeClassInfo& InClassInfo);
+		FNodeRegistryKey(const FAssetKey& AssetKey);
 
 		UE_DEPRECATED(5.4, "Implicit String ctor is no longer supported.")
 		FNodeRegistryKey(const FString& InKeyString);
diff --git a/Engine/Plugins/Runtime/Metasound/Source/MetasoundFrontend/Public/MetasoundFrontendTransform.h b/Engine/Plugins/Runtime/Metasound/Source/MetasoundFrontend/Public/MetasoundFrontendTransform.h
--- a/Engine/Plugins/Runtime/Metasound/Source/MetasoundFrontend/Public/MetasoundFrontendTransform.h
+++ b/Engine/Plugins/Runtime/Metasound/Source/MetasoundFrontend/Public/MetasoundFrontendTransform.h
@@ -210,28 +210,28 @@ namespace Metasound::Frontend
 		const FMetasoundFrontendClassName NewClassName;
 
 	public:
-		/* Generates and assigns the document's root graph class a unique name using the provided guid as the name field and
-		 * (optionally) namespace and variant.
-		 * Returns true if transform succeeded, false if not.
-		*/
+		UE_DEPRECATED(5.5, "Use FMetasoundFrontendDocumentBuilder::GenerateNewClassName instead")
 		static bool Generate(FDocumentHandle InDocument, const FGuid& InGuid, const FName Namespace = { }, const FName Variant = { })
 		{
-			const FMetasoundFrontendClassName GeneratedClassName = { Namespace, *InGuid.ToString(), Variant };
-			return FRenameRootGraphClass(GeneratedClassName).Transform(InDocument);
+			return false;
 		}
 
+		UE_DEPRECATED(5.5, "Use FMetasoundFrontendDocumentBuilder::GenerateNewClassName instead")
 		static bool Generate(FMetasoundFrontendDocument& InDocument, const FGuid& InGuid, const FName Namespace = { }, const FName Variant = { })
 		{
-			const FMetasoundFrontendClassName GeneratedClassName = { Namespace, *InGuid.ToString(), Variant };
-			return FRenameRootGraphClass(GeneratedClassName).Transform(InDocument);
+			return false;
 		}
 
+		UE_DEPRECATED(5.5, "Use FMetasoundFrontendDocumentBuilder::GenerateNewClassName instead")
 		FRenameRootGraphClass(const FMetasoundFrontendClassName InClassName)
 			: NewClassName(InClassName)
 		{
 		}
 
+		UE_DEPRECATED(5.5, "Use FMetasoundFrontendDocumentBuilder::GenerateNewClassName instead")
 		bool Transform(FDocumentHandle InDocument) const override;
+
+		UE_DEPRECATED(5.5, "Use FMetasoundFrontendDocumentBuilder::GenerateNewClassName instead")
 		bool Transform(FMetasoundFrontendDocument& InOutDocument) const override;
 	};
 } // Metasound::Frontend
