Fix for build error due to implicit key cast
#rnx

diff --git a/Engine/Plugins/Runtime/Metasound/Source/MetasoundEngine/Private/MetasoundAssetSubsystem.cpp b/Engine/Plugins/Runtime/Metasound/Source/MetasoundEngine/Private/MetasoundAssetSubsystem.cpp
--- a/Engine/Plugins/Runtime/Metasound/Source/MetasoundEngine/Private/MetasoundAssetSubsystem.cpp
+++ b/Engine/Plugins/Runtime/Metasound/Source/MetasoundEngine/Private/MetasoundAssetSubsystem.cpp
@@ -470,21 +470,23 @@ namespace Metasound::Engine
 				continue;
 			}
 
-			const FAssetKey Key(FNodeRegistryKey(Class.Metadata));
-			if (const FTopLevelAssetPath* ObjectPath = PathMap.Find(Key))
+			const FAssetKey AssetKey(Class.Metadata);
+			if (const FTopLevelAssetPath* ObjectPath = PathMap.Find(AssetKey))
 			{
-				OutAssetInfos.Add(FAssetInfo(Key, FSoftObjectPath(*ObjectPath)));
+				FAssetInfo AssetInfo(FNodeRegistryKey(Class.Metadata), FSoftObjectPath(*ObjectPath));
+				OutAssetInfos.Add(MoveTemp(AssetInfo));
 			}
 			else
 			{
+				FNodeRegistryKey RegistryKey(Class.Metadata);
 				const FMetasoundFrontendRegistryContainer* Registry = FMetasoundFrontendRegistryContainer::Get();
 				check(Registry);
-				const bool bIsRegistered = Registry->IsNodeRegistered(Key);
+				const bool bIsRegistered = Registry->IsNodeRegistered(RegistryKey);
 
 				bool bReportFail = false;
 				if (bIsRegistered)
 				{
-					if (!Registry->IsNodeNative(Key))
+					if (!Registry->IsNodeNative(RegistryKey))
 					{
 						bReportFail = true;
 					}
@@ -493,10 +495,10 @@ namespace Metasound::Engine
 				{
 					// Don't report failure if a matching class with a matching major version and higher minor version exists (it will be autoupdated) 
 					FMetasoundFrontendClass FrontendClass;
-					const bool bDidFindClassWithName = ISearchEngine::Get().FindClassWithHighestVersion(Key.ClassName.ToNodeClassName(), FrontendClass);
+					const bool bDidFindClassWithName = ISearchEngine::Get().FindClassWithHighestVersion(AssetKey.ClassName.ToNodeClassName(), FrontendClass);
 					if (!(bDidFindClassWithName &&
-						Key.Version.Major == FrontendClass.Metadata.GetVersion().Major &&
-						Key.Version.Minor < FrontendClass.Metadata.GetVersion().Minor))
+						AssetKey.Version.Major == FrontendClass.Metadata.GetVersion().Major &&
+						AssetKey.Version.Minor < FrontendClass.Metadata.GetVersion().Minor))
 					{
 						bReportFail = true;
 					}
@@ -506,11 +508,11 @@ namespace Metasound::Engine
 				{
 					if (bIsInitialAssetScanComplete)
 					{
-						UE_LOG(LogMetaSound, Warning, TEXT("MetaSound Node Class with registry key '%s' not registered when gathering referenced asset classes from '%s': Retrieving all asset classes may not be comprehensive."), *Key.ToString(), *InAssetBase.GetOwningAssetName());
+						UE_LOG(LogMetaSound, Warning, TEXT("MetaSound Node Class with registry key '%s' not registered when gathering referenced asset classes from '%s': Retrieving all asset classes may not be comprehensive."), *AssetKey.ToString(), *InAssetBase.GetOwningAssetName());
 					}
 					else
 					{
-						UE_LOG(LogMetaSound, Warning, TEXT("Attempt to get registered dependent asset with key '%s' from MetaSound asset '%s' before asset scan has completed: Asset class cannot be provided"), *Key.ToString(), *InAssetBase.GetOwningAssetName());
+						UE_LOG(LogMetaSound, Warning, TEXT("Attempt to get registered dependent asset with key '%s' from MetaSound asset '%s' before asset scan has completed: Asset class cannot be provided"), *AssetKey.ToString(), *InAssetBase.GetOwningAssetName());
 					}
 				}
 			}
