Implicit constructor issue: round 2
#rnx

diff --git a/Engine/Plugins/Runtime/Metasound/Source/MetasoundEngine/Private/MetasoundAssetSubsystem.cpp b/Engine/Plugins/Runtime/Metasound/Source/MetasoundEngine/Private/MetasoundAssetSubsystem.cpp
--- a/Engine/Plugins/Runtime/Metasound/Source/MetasoundEngine/Private/MetasoundAssetSubsystem.cpp
+++ b/Engine/Plugins/Runtime/Metasound/Source/MetasoundEngine/Private/MetasoundAssetSubsystem.cpp
@@ -473,7 +473,7 @@ namespace Metasound::Engine
 			const FAssetKey AssetKey(Class.Metadata);
 			if (const FTopLevelAssetPath* ObjectPath = PathMap.Find(AssetKey))
 			{
-				FAssetInfo AssetInfo(FNodeRegistryKey(Class.Metadata), FSoftObjectPath(*ObjectPath));
+				FAssetInfo AssetInfo { FNodeRegistryKey(Class.Metadata), FSoftObjectPath(*ObjectPath) };
 				OutAssetInfos.Add(MoveTemp(AssetInfo));
 			}
 			else
diff --git a/Engine/Plugins/Runtime/Metasound/Source/MetasoundFrontend/Public/MetasoundAssetManager.h b/Engine/Plugins/Runtime/Metasound/Source/MetasoundFrontend/Public/MetasoundAssetManager.h
--- a/Engine/Plugins/Runtime/Metasound/Source/MetasoundFrontend/Public/MetasoundAssetManager.h
+++ b/Engine/Plugins/Runtime/Metasound/Source/MetasoundFrontend/Public/MetasoundAssetManager.h
@@ -132,7 +132,7 @@ namespace Metasound::Frontend
 
 		struct FAssetInfo
 		{
-			const FNodeRegistryKey RegistryKey;
+			FNodeRegistryKey RegistryKey;
 			FSoftObjectPath AssetPath;
 
 			FORCEINLINE friend bool operator==(const FAssetInfo& InLHS, const FAssetInfo& InRHS)
