Fix duplicate/bad include path format (shame on you VAX)
#rnx

diff --git a/Engine/Plugins/Runtime/Metasound/Source/MetasoundFrontend/Private/MetasoundFrontendDocumentBuilder.cpp b/Engine/Plugins/Runtime/Metasound/Source/MetasoundFrontend/Private/MetasoundFrontendDocumentBuilder.cpp
--- a/Engine/Plugins/Runtime/Metasound/Source/MetasoundFrontend/Private/MetasoundFrontendDocumentBuilder.cpp
+++ b/Engine/Plugins/Runtime/Metasound/Source/MetasoundFrontend/Private/MetasoundFrontendDocumentBuilder.cpp
@@ -27,7 +27,6 @@
 #include "MetasoundTrace.h"
 #include "MetasoundVariableNodes.h"
 #include "NodeTemplates/MetasoundFrontendNodeTemplateInput.h"
-#include "../Public/MetasoundDocumentInterface.h"
 
 #include UE_INLINE_GENERATED_CPP_BY_NAME(MetasoundFrontendDocumentBuilder)
 
