- Fix for rename ensure regression from AssetManager refactor
- Add logging and track case when multiple assets are registered with the AssetManager with the same key for better debugging
- Add ability to force unregistration when request to finish building & utilize on destruction of MetaSound assets
#rnx

diff --git a/Engine/Plugins/Runtime/Metasound/Source/MetasoundEditor/Private/MetasoundEditorModule.cpp b/Engine/Plugins/Runtime/Metasound/Source/MetasoundEditor/Private/MetasoundEditorModule.cpp
--- a/Engine/Plugins/Runtime/Metasound/Source/MetasoundEditor/Private/MetasoundEditorModule.cpp
+++ b/Engine/Plugins/Runtime/Metasound/Source/MetasoundEditor/Private/MetasoundEditorModule.cpp
@@ -439,11 +439,15 @@ namespace Metasound
 					if (bIsRegisteredClass)
 					{
 						// Use the editor version of UnregisterWithFrontend so it refreshes any open MetaSound editors
-						IMetaSoundAssetManager::GetChecked().RemoveAsset(InAssetData);
-						if (UObject* AssetObject = InAssetData.GetAsset())
+						// Doesn't use AssetData::GetAsset() as this can result in attempting to reload the object.
+						// If this call is hit after the asset is removed, the assumption is unregistration already
+						// occurred on object destroy.
+						if (UObject* AssetObject = InAssetData.GetSoftObjectPath().ResolveObject())
 						{
 							FGraphBuilder::UnregisterGraphWithFrontend(*AssetObject);
 						}
+
+						IMetaSoundAssetManager::GetChecked().RemoveAsset(InAssetData);
 					}
 				}
 			}
@@ -457,14 +461,29 @@ namespace Metasound
 					const bool bIsRegisteredClass = IMetasoundUObjectRegistry::Get().IsRegisteredClass(*AssetClass);
 					if (bIsRegisteredClass)
 					{
-						constexpr bool bReregisterWithFrontend = false;
-						IMetaSoundAssetManager::GetChecked().RenameAsset(InAssetData, bReregisterWithFrontend);
+						IMetaSoundAssetManager& AssetManager = IMetaSoundAssetManager::GetChecked();
 
-						// Use the FGraphBuilder Register call instead of registering via the
-						// MetaSoundAssetSubsystem so as to properly refresh respective open editors.
-						constexpr bool bForceViewSynchronization = true;
+						// Unregister using the new asset data even though the old object was last to be registered
+						// as the old asset is no longer accessible by the time rename is called. The asset at this
+						// point is identical however to its prior counterpart.
 						UObject* AssetObject = InAssetData.GetAsset();
-						FGraphBuilder::RegisterGraphWithFrontend(*AssetObject, bForceViewSynchronization);
+						check(AssetObject);
+
+						FMetasoundAssetBase* AssetBase = AssetManager.GetAsAsset(*AssetObject);
+						check(AssetBase);
+						bool bIsRegistered = AssetBase->IsRegistered();
+						if (bIsRegistered)
+						{
+							FGraphBuilder::UnregisterGraphWithFrontend(*AssetObject);
+						}
+
+						IMetaSoundAssetManager::GetChecked().RenameAsset(InAssetData, InOldObjectPath);
+
+						if (bIsRegistered)
+						{
+							constexpr bool bForceViewSynchronization = true;
+							FGraphBuilder::RegisterGraphWithFrontend(*AssetObject, bForceViewSynchronization);
+						}
 					}
 				}
 			}
diff --git a/Engine/Plugins/Runtime/Metasound/Source/MetasoundEngine/Private/MetasoundAssetSubsystem.cpp b/Engine/Plugins/Runtime/Metasound/Source/MetasoundEngine/Private/MetasoundAssetSubsystem.cpp
--- a/Engine/Plugins/Runtime/Metasound/Source/MetasoundEngine/Private/MetasoundAssetSubsystem.cpp
+++ b/Engine/Plugins/Runtime/Metasound/Source/MetasoundEngine/Private/MetasoundAssetSubsystem.cpp
@@ -89,24 +89,22 @@ namespace Metasound::Engine
 		//
 		// This protects against scenarios where a metasound is renamed or moved 
 		// and the new entry was being erroneously removed from the PathMap
-		bool RemoveAndWarnIfNotExactMatch(TMap<Frontend::FAssetKey, FTopLevelAssetPath>& InMap, const Frontend::FAssetKey& AssetKey, const FTopLevelAssetPath& AssetPath)
+		bool RemovePath(TMap<Frontend::FAssetKey, TArray<FTopLevelAssetPath>>& InMap, const Frontend::FAssetKey& AssetKey, const FTopLevelAssetPath& AssetPath)
 		{
-#if !NO_LOGGING
-			// Can be invalid if document on object being destroyed is being removed and MetaSound was no longer registered,
-			// in which case ignore logging as the assumption is it was a valid object to be removed.
 			if (AssetPath.IsValid())
 			{
-				if (const FTopLevelAssetPath* MapAssetPath = InMap.Find(AssetKey))
+				if (TArray<FTopLevelAssetPath>* MapAssetPaths = InMap.Find(AssetKey))
 				{
-					if (*MapAssetPath != AssetPath)
+					const int32 NumRemoved = MapAssetPaths->RemoveAllSwap([&AssetPath](const FTopLevelAssetPath& Path) { return Path == AssetPath; }, EAllowShrinking::No);
+					if (MapAssetPaths->IsEmpty())
 					{
-						UE_LOG(LogMetaSound, Warning, TEXT("AssetManager removing object where path (%s) does not match registry copy (%s)"), *MapAssetPath->ToString(), *AssetPath.ToString());
+						InMap.Remove(AssetKey);
 					}
+					return NumRemoved > 0;
 				}
 			}
-#endif // !NO_LOGGING
 
-			return InMap.Remove(AssetKey) > 0;
+			return false;
 		}
 	}
 
@@ -153,7 +151,7 @@ namespace Metasound::Engine
 #endif // WITH_EDITOR
 		virtual void RemoveAsset(const UObject& InObject) override;
 		virtual void RemoveAsset(const FAssetData& InAssetData) override;
-		virtual void RenameAsset(const FAssetData& InAssetData, bool bInReregisterWithFrontend = true) override;
+		virtual void RenameAsset(const FAssetData& InAssetData, const FString& InOldObjectPath) override;
 		virtual FMetasoundAssetBase* TryLoadAssetFromKey(const FAssetKey& InKey) const override;
 		virtual bool TryGetAssetIDFromClassName(const FMetasoundFrontendClassName& InClassName, FGuid& OutGuid) const override;
 		virtual bool TryLoadReferencedAssets(const FMetasoundAssetBase& InAssetBase, TArray<FMetasoundAssetBase*>& OutReferencedAssets) const override;
@@ -177,14 +175,14 @@ namespace Metasound::Engine
 		int32 AutoUpdateDenyListChangeID = INDEX_NONE;
 		TSet<FName> AutoUpdateDenyListCache;
 		std::atomic<bool> bIsInitialAssetScanComplete = false;
-		TMap<FAssetKey, FTopLevelAssetPath> PathMap;
+		TMap<FAssetKey, TArray<FTopLevelAssetPath>> PathMap;
 	};
 
 	FMetaSoundAssetManager::~FMetaSoundAssetManager()
 	{
 		if (!PathMap.IsEmpty())
 		{
-			UE_LOG(LogMetaSound, Display, TEXT("Destroying MetaSoundAssetManager while %i asset paths are still available."), PathMap.Num());
+			UE_LOG(LogMetaSound, Display, TEXT("Destroying MetaSoundAssetManager while %i asset keys are still registered."), PathMap.Num());
 		}
 	}
 
@@ -276,7 +274,19 @@ namespace Metasound::Engine
 
 		if (AssetKey.IsValid())
 		{
-			PathMap.FindOrAdd(AssetKey) = FTopLevelAssetPath(&InObject);
+			TArray<FTopLevelAssetPath> Paths = PathMap.FindOrAdd(AssetKey);
+			Paths.AddUnique(FTopLevelAssetPath(&InObject));
+#if !NO_LOGGING
+			if (Paths.Num() > 1)
+			{
+				TArray<FString> PathStrings;
+				Algo::Transform(Paths, PathStrings, [](const FTopLevelAssetPath& Path) { return Path.ToString(); });
+				UE_LOG(LogMetaSound, Warning,
+					TEXT("Registering multiple assets with key '%s' with AssetManager:\n%s\n"),
+				*AssetKey.ToString(),
+				*FString::Join(PathStrings, TEXT("\n")));
+			}
+#endif // !NO_LOGGING
 		}
 
 		return AssetKey;
@@ -320,7 +330,7 @@ namespace Metasound::Engine
 			const FAssetKey RegistryKey = FAssetKey(ClassInfo.ClassName, ClassInfo.Version);
 			if (RegistryKey.IsValid())
 			{
-				PathMap.FindOrAdd(RegistryKey) = ClassInfo.AssetPath;
+				PathMap.FindOrAdd(RegistryKey).AddUnique(ClassInfo.AssetPath);
 			}
 
 			return RegistryKey;
@@ -369,7 +379,7 @@ namespace Metasound::Engine
 
 	FMetasoundAssetBase* FMetaSoundAssetManager::FindAsset(const Metasound::Frontend::FAssetKey& InKey) const
 	{
-		if (const FTopLevelAssetPath* AssetPath = PathMap.Find(InKey))
+		if (const FTopLevelAssetPath* AssetPath = FindAssetPath(InKey))
 		{
 			if (UObject* Object = FSoftObjectPath(*AssetPath, { }).ResolveObject())
 			{
@@ -382,7 +392,7 @@ namespace Metasound::Engine
 
 	TScriptInterface<IMetaSoundDocumentInterface> FMetaSoundAssetManager::FindAssetAsDocumentInterface(const Metasound::Frontend::FAssetKey& InKey) const
 	{
-		if (const FTopLevelAssetPath* AssetPath = PathMap.Find(InKey))
+		if (const FTopLevelAssetPath* AssetPath = FindAssetPath(InKey))
 		{
 			if (UObject* Object = FSoftObjectPath(*AssetPath, { }).ResolveObject())
 			{
@@ -395,7 +405,15 @@ namespace Metasound::Engine
 
 	const FTopLevelAssetPath* FMetaSoundAssetManager::FindAssetPath(const Metasound::Frontend::FAssetKey& InKey) const
 	{
-		return PathMap.Find(InKey);
+		if (const TArray<FTopLevelAssetPath>* Paths = PathMap.Find(InKey))
+		{
+			if (!Paths->IsEmpty())
+			{
+				return &Paths->Last();
+			}
+		}
+
+		return nullptr;
 	}
 
 	FMetasoundAssetBase* FMetaSoundAssetManager::GetAsAsset(UObject& InObject) const
@@ -471,7 +489,7 @@ namespace Metasound::Engine
 			}
 
 			const FAssetKey AssetKey(Class.Metadata);
-			if (const FTopLevelAssetPath* ObjectPath = PathMap.Find(AssetKey))
+			if (const FTopLevelAssetPath* ObjectPath = FindAssetPath(AssetKey))
 			{
 				FAssetInfo AssetInfo { FNodeRegistryKey(Class.Metadata), FSoftObjectPath(*ObjectPath) };
 				OutAssetInfos.Add(MoveTemp(AssetInfo));
@@ -591,10 +609,11 @@ namespace Metasound::Engine
 
 			if (IDocumentBuilderRegistry* BuilderRegistry = IDocumentBuilderRegistry::Get())
 			{
-				BuilderRegistry->FinishBuilding(Metadata.GetClassName());
+				constexpr bool bForceUnregister = true;
+				BuilderRegistry->FinishBuilding(Metadata.GetClassName(), bForceUnregister);
 			}
 
-			AssetSubsystemPrivate::RemoveAndWarnIfNotExactMatch(PathMap, AssetKey, GraphKey.AssetPath);
+			AssetSubsystemPrivate::RemovePath(PathMap, AssetKey, GraphKey.AssetPath);
 		}
 	}
 
@@ -607,12 +626,13 @@ namespace Metasound::Engine
 		{
 			if (IDocumentBuilderRegistry* BuilderRegistry = IDocumentBuilderRegistry::Get())
 			{
-				BuilderRegistry->FinishBuilding(ClassInfo.ClassName);
+				constexpr bool bForceUnregister = true;
+				BuilderRegistry->FinishBuilding(ClassInfo.ClassName, bForceUnregister);
 			}
 
 			const FAssetKey AssetKey(ClassInfo.ClassName, ClassInfo.Version);
 			const FTopLevelAssetPath AssetPath(InAssetData.PackageName, InAssetData.AssetName);
-			AssetSubsystemPrivate::RemoveAndWarnIfNotExactMatch(PathMap, AssetKey, AssetPath);
+			AssetSubsystemPrivate::RemovePath(PathMap, AssetKey, AssetPath);
 		}
 	}
 
@@ -625,7 +645,7 @@ namespace Metasound::Engine
 		LoadingDependencies.RemoveAllSwap(IsEqualID);
 	}
 
-	void FMetaSoundAssetManager::RenameAsset(const FAssetData& InAssetData, bool bInReregisterWithFrontend)
+	void FMetaSoundAssetManager::RenameAsset(const FAssetData& InAssetData, const FString& InOldObjectPath)
 	{
 		using namespace Frontend;
 
@@ -635,16 +655,12 @@ namespace Metasound::Engine
 		FNodeClassInfo ClassInfo;
 		if (ensureAlways(AssetSubsystemPrivate::GetAssetClassInfo(InAssetData, ClassInfo)))
 		{
-			IDocumentBuilderRegistry* BuilderRegistry = IDocumentBuilderRegistry::Get();
-			if (BuilderRegistry)
-			{
-				BuilderRegistry->FinishBuilding(ClassInfo.ClassName);
-			}
-
 			const FAssetKey AssetKey(ClassInfo.ClassName, ClassInfo.Version);
-			PathMap.Remove(AssetKey); // Don't warn if names do not match as by this point the incoming asset data has the newly assigned asset path assigned
+			const FTopLevelAssetPath OldPath(InOldObjectPath);
+			AssetSubsystemPrivate::RemovePath(PathMap, AssetKey, OldPath);
 
 #if WITH_EDITOR
+			IDocumentBuilderRegistry* BuilderRegistry = IDocumentBuilderRegistry::Get();
 			if (BuilderRegistry)
 			{
 				FMetaSoundFrontendDocumentBuilder& Builder = BuilderRegistry->FindOrBeginBuilding(MetaSoundAsset->GetOwningAsset());
@@ -654,18 +670,12 @@ namespace Metasound::Engine
 
 			if (ClassInfo.AssetClassID.IsValid())
 			{
-				const FAssetKey NewAssetKey(ClassInfo.ClassName, ClassInfo.Version);
-				if (NewAssetKey.IsValid())
+				if (AssetKey.IsValid())
 				{
-					PathMap.FindOrAdd(NewAssetKey) = ClassInfo.AssetPath;
+					PathMap.FindOrAdd(AssetKey).AddUnique(ClassInfo.AssetPath);
 				}
 			}
 		}
-
-		if (bInReregisterWithFrontend)
-		{
-			MetaSoundAsset->RegisterGraphWithFrontend();
-		}
 	}
 
 	void FMetaSoundAssetManager::RequestAsyncLoadReferencedAssets(FMetasoundAssetBase& InAssetBase)
@@ -737,9 +747,9 @@ namespace Metasound::Engine
 		});
 	}
 
-	FMetasoundAssetBase* FMetaSoundAssetManager::TryLoadAssetFromKey(const Metasound::Frontend::FAssetKey& RegistryKey) const
+	FMetasoundAssetBase* FMetaSoundAssetManager::TryLoadAssetFromKey(const Metasound::Frontend::FAssetKey& InAssetKey) const
 	{
-		if (const FTopLevelAssetPath* ObjectPath = FindAssetPath(RegistryKey))
+		if (const FTopLevelAssetPath* ObjectPath = FindAssetPath(InAssetKey))
 		{
 			const FSoftObjectPath SoftPath(*ObjectPath);
 			return TryLoadAsset(SoftPath);
@@ -815,7 +825,7 @@ namespace Metasound::Engine
 						FMetasoundFrontendRegistryContainer::Get()->UnregisterNode(RegistryKey);
 						const FTopLevelAssetPath AssetPath(AssetData.PackageName, AssetData.AssetName);
 						const FAssetKey AssetKey(AssetClassInfo.ClassName, AssetClassInfo.Version);
-						AssetSubsystemPrivate::RemoveAndWarnIfNotExactMatch(PathMap, AssetKey, AssetPath);
+						AssetSubsystemPrivate::RemovePath(PathMap, AssetKey, AssetPath);
 					}
 				}
 			}
@@ -1012,10 +1022,10 @@ void UMetaSoundAssetSubsystem::RemoveAsset(const FAssetData& InAssetData)
 	IMetaSoundAssetManager::GetChecked().RemoveAsset(InAssetData);
 }
 
-void UMetaSoundAssetSubsystem::RenameAsset(const FAssetData& InAssetData, bool bInReregisterWithFrontend)
+void UMetaSoundAssetSubsystem::RenameAsset(const FAssetData& InAssetData, bool /* bInReregisterWithFrontend */)
 {
 	using namespace Metasound::Frontend;
-	IMetaSoundAssetManager::GetChecked().RenameAsset(InAssetData, bInReregisterWithFrontend);
+	IMetaSoundAssetManager::GetChecked().RenameAsset(InAssetData, { });
 }
 
 UMetaSoundAssetSubsystem& UMetaSoundAssetSubsystem::GetChecked()
diff --git a/Engine/Plugins/Runtime/Metasound/Source/MetasoundEngine/Private/MetasoundBuilderSubsystem.cpp b/Engine/Plugins/Runtime/Metasound/Source/MetasoundEngine/Private/MetasoundBuilderSubsystem.cpp
--- a/Engine/Plugins/Runtime/Metasound/Source/MetasoundEngine/Private/MetasoundBuilderSubsystem.cpp
+++ b/Engine/Plugins/Runtime/Metasound/Source/MetasoundEngine/Private/MetasoundBuilderSubsystem.cpp
@@ -74,14 +74,24 @@ const UClass& UMetaSoundPatchBuilder::GetBaseMetaSoundUClass() const
 
 void UMetaSoundPatchBuilder::OnAssetReferenceAdded(TScriptInterface<IMetaSoundDocumentInterface> DocInterface)
 {
+	using namespace Metasound::Frontend;
+
 	check(DocInterface.GetObject());
-	Builder.CastDocumentObjectChecked<UMetaSoundPatch>().ReferencedAssetClassObjects.Add(DocInterface.GetObject());
+	UMetaSoundPatch& Patch = Builder.CastDocumentObjectChecked<UMetaSoundPatch>();
+	FNodeRegistryKey RegistryKey(Patch.GetConstDocument().RootGraph.Metadata);
+	Patch.ReferencedAssetClassObjects.Add(DocInterface.GetObject());
+	Patch.ReferencedAssetClassKeys.Add(RegistryKey.ToString());
 }
 
 void UMetaSoundPatchBuilder::OnRemovingAssetReference(TScriptInterface<IMetaSoundDocumentInterface> DocInterface)
 {
+	using namespace Metasound::Frontend;
+
 	check(DocInterface.GetObject());
-	Builder.CastDocumentObjectChecked<UMetaSoundPatch>().ReferencedAssetClassObjects.Remove(DocInterface.GetObject());
+	UMetaSoundPatch& Patch = Builder.CastDocumentObjectChecked<UMetaSoundPatch>();
+	FNodeRegistryKey RegistryKey(Patch.GetConstDocument().RootGraph.Metadata);
+	Patch.ReferencedAssetClassObjects.Remove(DocInterface.GetObject());
+	Patch.ReferencedAssetClassKeys.Remove(RegistryKey.ToString());
 }
 
 UMetaSoundBuilderBase& UMetaSoundBuilderSubsystem::AttachBuilderToAssetChecked(UObject& InObject) const
@@ -262,8 +272,13 @@ void UMetaSoundSourceBuilder::InitDelegates(Metasound::Frontend::FDocumentModify
 
 void UMetaSoundSourceBuilder::OnAssetReferenceAdded(TScriptInterface<IMetaSoundDocumentInterface> DocInterface)
 {
+	using namespace Metasound::Frontend;
+
 	check(DocInterface.GetObject());
-	GetMetaSoundSource().ReferencedAssetClassObjects.Add(DocInterface.GetObject());
+	UMetaSoundSource& Source = GetMetaSoundSource();
+	FNodeRegistryKey RegistryKey(Source.GetConstDocument().RootGraph.Metadata);
+	Source.ReferencedAssetClassObjects.Add(DocInterface.GetObject());
+	Source.ReferencedAssetClassKeys.Add(RegistryKey.ToString());
 }
 
 void UMetaSoundSourceBuilder::OnEdgeAdded(int32 EdgeIndex) const
@@ -529,8 +544,13 @@ void UMetaSoundSourceBuilder::OnRemoveSwappingEdge(int32 SwapIndex, int32 LastIn
 
 void UMetaSoundSourceBuilder::OnRemovingAssetReference(TScriptInterface<IMetaSoundDocumentInterface> DocInterface)
 {
+	using namespace Metasound::Frontend;
+
 	check(DocInterface.GetObject());
-	GetMetaSoundSource().ReferencedAssetClassObjects.Remove(DocInterface.GetObject());
+	UMetaSoundSource& Source = GetMetaSoundSource();
+	FNodeRegistryKey RegistryKey(Source.GetConstDocument().RootGraph.Metadata);
+	Source.ReferencedAssetClassObjects.Remove(DocInterface.GetObject());
+	Source.ReferencedAssetClassKeys.Remove(RegistryKey.ToString());
 }
 
 void UMetaSoundSourceBuilder::OnRemovingInput(int32 InputIndex)
diff --git a/Engine/Plugins/Runtime/Metasound/Source/MetasoundEngine/Private/MetasoundDocumentBuilderRegistry.cpp b/Engine/Plugins/Runtime/Metasound/Source/MetasoundEngine/Private/MetasoundDocumentBuilderRegistry.cpp
--- a/Engine/Plugins/Runtime/Metasound/Source/MetasoundEngine/Private/MetasoundDocumentBuilderRegistry.cpp
+++ b/Engine/Plugins/Runtime/Metasound/Source/MetasoundEngine/Private/MetasoundDocumentBuilderRegistry.cpp
@@ -80,7 +80,7 @@ namespace Metasound::Engine
 		return Builders.FindRef(ClassName).Get();
 	}
 
-	bool FDocumentBuilderRegistry::FinishBuilding(const FMetasoundFrontendClassName& InClassName) const
+	bool FDocumentBuilderRegistry::FinishBuilding(const FMetasoundFrontendClassName& InClassName, bool bForceUnregister) const
 	{
 		using namespace Metasound;
 		using namespace Metasound::Engine;
@@ -91,17 +91,24 @@ namespace Metasound::Engine
 			// If the builder has applied transactions to its document object that are not mirrored in the frontend registry,
 			// unregister version in registry. This will ensure that future requests for the builder's associated asset will
 			// register a fresh version from the object as the transaction history is intrinsically lost once this builder
-			// is destroyed.
-			if (Builder->GetLastTransactionRegistered() != Builder->GetConstBuilder().GetTransactionCount())
+			// is destroyed. It is also possible that the DocBuilder's underlying object can be invalid if object was force
+			// deleted, so validity check is necessary.
+			FMetaSoundFrontendDocumentBuilder& DocBuilder = Builder->GetBuilder();
+			if (DocBuilder.IsValid())
 			{
-				UObject& MetaSound = Builder->GetConstBuilder().CastDocumentObjectChecked<UObject>();
-				if (FMetasoundAssetBase* MetaSoundAsset = IMetasoundUObjectRegistry::Get().GetObjectAsAssetBase(&MetaSound))
+				const int32 TransactionCount = DocBuilder.GetTransactionCount();
+				const int32 LastTransactionRegistered = Builder->GetLastTransactionRegistered();
+				if (bForceUnregister || LastTransactionRegistered != TransactionCount)
 				{
-					MetaSoundAsset->UnregisterGraphWithFrontend();
+					UObject& MetaSound = DocBuilder.CastDocumentObjectChecked<UObject>();
+					if (FMetasoundAssetBase* MetaSoundAsset = IMetasoundUObjectRegistry::Get().GetObjectAsAssetBase(&MetaSound))
+					{
+						MetaSoundAsset->UnregisterGraphWithFrontend();
+					}
 				}
 			}
 
-			Builder->GetBuilder().FinishBuilding();
+			DocBuilder.FinishBuilding();
 			ensureAlways(Builders.Remove(InClassName));
 			return true;
 		}
diff --git a/Engine/Plugins/Runtime/Metasound/Source/MetasoundEngine/Public/MetasoundBuilderBase.cpp b/Engine/Plugins/Runtime/Metasound/Source/MetasoundEngine/Public/MetasoundBuilderBase.cpp
--- a/Engine/Plugins/Runtime/Metasound/Source/MetasoundEngine/Public/MetasoundBuilderBase.cpp
+++ b/Engine/Plugins/Runtime/Metasound/Source/MetasoundEngine/Public/MetasoundBuilderBase.cpp
@@ -36,12 +36,11 @@ void UMetaSoundBuilderBase::BeginDestroy()
 {
 	using namespace Metasound::Engine;
 
-	FDocumentBuilderRegistry::GetChecked().FinishBuilding(ClassName);
-
-	// Need to detach before destroying UPROPERTYs as the Builder 
-	// often holds a TScriptInterface<IMetaSoundDocumentInterface> of
-	// a UPROPERTY that lives on this or derived objects. 
-	Builder.FinishBuilding();
+	// Need to finish building before destroying UPROPERTYs (Super::BeginDestroy)
+	// as the Builder often holds a TScriptInterface<IMetaSoundDocumentInterface>
+	// of a UPROPERTY that lives on this or derived objects. 
+	constexpr bool bForceUnregister = true;
+	FDocumentBuilderRegistry::GetChecked().FinishBuilding(ClassName, bForceUnregister);
 	Super::BeginDestroy();
 }
 
diff --git a/Engine/Plugins/Runtime/Metasound/Source/MetasoundEngine/Public/MetasoundDocumentBuilderRegistry.h b/Engine/Plugins/Runtime/Metasound/Source/MetasoundEngine/Public/MetasoundDocumentBuilderRegistry.h
--- a/Engine/Plugins/Runtime/Metasound/Source/MetasoundEngine/Public/MetasoundDocumentBuilderRegistry.h
+++ b/Engine/Plugins/Runtime/Metasound/Source/MetasoundEngine/Public/MetasoundDocumentBuilderRegistry.h
@@ -88,7 +88,7 @@ namespace Metasound::Engine
 
 		virtual FMetaSoundFrontendDocumentBuilder* FindBuilder(TScriptInterface<IMetaSoundDocumentInterface> MetaSound) const override;
 		virtual FMetaSoundFrontendDocumentBuilder* FindBuilder(const FMetasoundFrontendClassName& InClassName) const override;
-		virtual bool FinishBuilding(const FMetasoundFrontendClassName& InClassName) const override;
+		virtual bool FinishBuilding(const FMetasoundFrontendClassName& InClassName, bool bForceUnregister = false) const override;
 
 		UMetaSoundBuilderBase* FindBuilderObject(TScriptInterface<const IMetaSoundDocumentInterface> MetaSound) const;
 		UMetaSoundBuilderBase* FindBuilderObject(const FMetasoundFrontendClassName& ClassName) const;
diff --git a/Engine/Plugins/Runtime/Metasound/Source/MetasoundFrontend/Public/MetasoundAssetManager.h b/Engine/Plugins/Runtime/Metasound/Source/MetasoundFrontend/Public/MetasoundAssetManager.h
--- a/Engine/Plugins/Runtime/Metasound/Source/MetasoundFrontend/Public/MetasoundAssetManager.h
+++ b/Engine/Plugins/Runtime/Metasound/Source/MetasoundFrontend/Public/MetasoundAssetManager.h
@@ -206,7 +206,7 @@ namespace Metasound::Frontend
 		virtual void RemoveAsset(const FAssetData& InAssetData) = 0;
 
 		// Updates the given MetaSound's asset record with the new name and optionally reregisters it with the Frontend Node Class Registry.
-		virtual void RenameAsset(const FAssetData& InAssetData, bool bInReregisterWithFrontend = true) = 0;
+		virtual void RenameAsset(const FAssetData& InAssetData, const FString& InOldObjectPath) = 0;
 
 		// Waits until all async load requests related to this asset are complete.
  		virtual void WaitUntilAsyncLoadReferencedAssetsComplete(FMetasoundAssetBase& InAssetBase) = 0;
diff --git a/Engine/Plugins/Runtime/Metasound/Source/MetasoundFrontend/Public/MetasoundDocumentInterface.h b/Engine/Plugins/Runtime/Metasound/Source/MetasoundFrontend/Public/MetasoundDocumentInterface.h
--- a/Engine/Plugins/Runtime/Metasound/Source/MetasoundFrontend/Public/MetasoundDocumentInterface.h
+++ b/Engine/Plugins/Runtime/Metasound/Source/MetasoundFrontend/Public/MetasoundDocumentInterface.h
@@ -90,7 +90,7 @@ namespace Metasound::Frontend
 		virtual FMetaSoundFrontendDocumentBuilder& FindOrBeginBuilding(TScriptInterface<IMetaSoundDocumentInterface> MetaSound) = 0;
 #endif // WITH_EDITORONLY_DATA
 
-		virtual bool FinishBuilding(const FMetasoundFrontendClassName& InClassName) const = 0;
+		virtual bool FinishBuilding(const FMetasoundFrontendClassName& InClassName, bool bForceUnregister = false) const = 0;
 
 		UE_DEPRECATED(5.5, "Document cache can now be invalidated by retrieving an asset builder and calling 'Reload'")
 		virtual void InvalidateDocumentCache(const FMetasoundFrontendClassName& InClassName) const { }
