Add missing logging spot from prior CL (And fix for last-minute tweak that added missing ref)
#rnx

diff --git a/Engine/Plugins/Runtime/Metasound/Source/MetasoundEngine/Private/MetasoundAssetSubsystem.cpp b/Engine/Plugins/Runtime/Metasound/Source/MetasoundEngine/Private/MetasoundAssetSubsystem.cpp
--- a/Engine/Plugins/Runtime/Metasound/Source/MetasoundEngine/Private/MetasoundAssetSubsystem.cpp
+++ b/Engine/Plugins/Runtime/Metasound/Source/MetasoundEngine/Private/MetasoundAssetSubsystem.cpp
@@ -106,6 +106,23 @@ namespace Metasound::Engine
 
 			return false;
 		}
+
+		void AddPath(TMap<Frontend::FAssetKey, TArray<FTopLevelAssetPath>>& InMap, const Frontend::FAssetKey& AssetKey, const FTopLevelAssetPath& AssetPath)
+		{
+			TArray<FTopLevelAssetPath>& Paths = InMap.FindOrAdd(AssetKey);
+			Paths.AddUnique(AssetPath);
+#if !NO_LOGGING
+			if (Paths.Num() > 1)
+			{
+				TArray<FString> PathStrings;
+				Algo::Transform(Paths, PathStrings, [](const FTopLevelAssetPath& Path) { return Path.ToString(); });
+				UE_LOG(LogMetaSound, Warning,
+					TEXT("MetaSoundAssetManager has registered multiple assets with key '%s':\n%s\n"),
+					*AssetKey.ToString(),
+					*FString::Join(PathStrings, TEXT("\n")));
+			}
+#endif // !NO_LOGGING
+		}
 	}
 
 	class FMetaSoundAssetManager :
@@ -274,19 +291,7 @@ namespace Metasound::Engine
 
 		if (AssetKey.IsValid())
 		{
-			TArray<FTopLevelAssetPath> Paths = PathMap.FindOrAdd(AssetKey);
-			Paths.AddUnique(FTopLevelAssetPath(&InObject));
-#if !NO_LOGGING
-			if (Paths.Num() > 1)
-			{
-				TArray<FString> PathStrings;
-				Algo::Transform(Paths, PathStrings, [](const FTopLevelAssetPath& Path) { return Path.ToString(); });
-				UE_LOG(LogMetaSound, Warning,
-					TEXT("Registering multiple assets with key '%s' with AssetManager:\n%s\n"),
-				*AssetKey.ToString(),
-				*FString::Join(PathStrings, TEXT("\n")));
-			}
-#endif // !NO_LOGGING
+			AssetSubsystemPrivate::AddPath(PathMap, AssetKey, FTopLevelAssetPath(&InObject));
 		}
 
 		return AssetKey;
@@ -327,13 +332,13 @@ namespace Metasound::Engine
 
 		if (ClassInfo.AssetClassID.IsValid())
 		{
-			const FAssetKey RegistryKey = FAssetKey(ClassInfo.ClassName, ClassInfo.Version);
-			if (RegistryKey.IsValid())
+			const FAssetKey AssetKey = FAssetKey(ClassInfo.ClassName, ClassInfo.Version);
+			if (AssetKey.IsValid())
 			{
-				PathMap.FindOrAdd(RegistryKey).AddUnique(ClassInfo.AssetPath);
+				AssetSubsystemPrivate::AddPath(PathMap, AssetKey, ClassInfo.AssetPath);
 			}
 
-			return RegistryKey;
+			return AssetKey;
 		}
 
 		// Invalid ClassID means the node could not be registered.
