[Metasound] Remove tick from UMetaSoundOutputSubsystem, only drain queues when there is new data.

#tests automated tests

#virtualized

diff --git a/Engine/Plugins/Runtime/Metasound/Source/MetasoundEngine/Private/MetasoundGeneratorHandle.cpp b/Engine/Plugins/Runtime/Metasound/Source/MetasoundEngine/Private/MetasoundGeneratorHandle.cpp
--- a/Engine/Plugins/Runtime/Metasound/Source/MetasoundEngine/Private/MetasoundGeneratorHandle.cpp
+++ b/Engine/Plugins/Runtime/Metasound/Source/MetasoundEngine/Private/MetasoundGeneratorHandle.cpp
@@ -177,35 +177,6 @@ namespace Metasound
 		return UnwatchOutputInternal(OutputName, FWatchOutputUnifiedDelegate(OnOutputValueChanged), AnalyzerName, AnalyzerOutputName);
 	}
 
-	void FMetasoundGeneratorHandle::UpdateOutputWatchers()
-	{
-		METASOUND_LLM_SCOPE;
-		METASOUND_TRACE_CPUPROFILER_EVENT_SCOPE(FMetasoundGeneratorHandle::UpdateOutputWatchers);
-
-		check(IsInGameThread());
-
-		int32 NumDequeued = 0;
-
-		while (TOptional<FOutputPayload> ChangedOutput = ChangedOutputs.Dequeue())
-		{
-			const FOutputWatcherKey WatcherKey
-			{
-				ChangedOutput->OutputName,
-				ChangedOutput->AnalyzerName,
-				ChangedOutput->OutputValue.Name
-			};
-
-			if (const FOutputWatcher* Watcher = OutputWatchers.Find(WatcherKey))
-			{
-				Watcher->OnOutputValueChanged.Broadcast(ChangedOutput->OutputName, ChangedOutput->OutputValue);
-			}
-
-			++NumDequeued;
-		}
-
-		ChangedOutputsQueueCount.store(FMath::Max(0, ChangedOutputsQueueCount.load() - NumDequeued));
-	}
-
 	void FMetasoundGeneratorHandle::RegisterPassthroughAnalyzerForType(
 		const FName TypeName,
 		const FName AnalyzerName,
@@ -423,6 +394,38 @@ namespace Metasound
 		return true;
 	}
 
+	void FMetasoundGeneratorHandle::UpdateOutputWatchersInternal()
+	{
+		METASOUND_LLM_SCOPE;
+		METASOUND_TRACE_CPUPROFILER_EVENT_SCOPE(FMetasoundGeneratorHandle::UpdateOutputWatchersInternal);
+
+		check(IsInGameThread());
+
+		// Clear the flag *before* we drain the queue, so we don't leave any output updates behind.
+		OutputWatcherUpdateScheduled.clear();
+
+		int32 NumDequeued = 0;
+
+		while (TOptional<FOutputPayload> ChangedOutput = ChangedOutputs.Dequeue())
+		{
+			const FOutputWatcherKey WatcherKey
+			{
+				ChangedOutput->OutputName,
+				ChangedOutput->AnalyzerName,
+				ChangedOutput->OutputValue.Name
+			};
+
+			if (const FOutputWatcher* Watcher = OutputWatchers.Find(WatcherKey))
+			{
+				Watcher->OnOutputValueChanged.Broadcast(ChangedOutput->OutputName, ChangedOutput->OutputValue);
+			}
+
+			++NumDequeued;
+		}
+
+		ChangedOutputsQueueCount.store(FMath::Max(0, ChangedOutputsQueueCount.load() - NumDequeued));
+	}
+
 	bool FMetasoundGeneratorHandle::TryCreateAnalyzerAddress(
 		const FName OutputName,
 		const FName AnalyzerName,
@@ -690,6 +693,19 @@ namespace Metasound
 		
 		ChangedOutputs.Enqueue(AnalyzerName, OutputName, AnalyzerOutputName, OutputData);
 		ChangedOutputsQueueCount.fetch_add(1);
+
+		// Drain the queue on the game thread, but don't bother if it's already been scheduled
+		if (!OutputWatcherUpdateScheduled.test_and_set())
+		{
+			// Defer to the game thread. We grab a weak pointer in case this gets destroyed while we wait.
+			AsyncTask(ENamedThreads::GameThread, [WeakThis = AsWeak()]()
+			{
+				if (const TSharedPtr<FMetasoundGeneratorHandle> PinnedThis = WeakThis.Pin())
+				{
+					PinnedThis->UpdateOutputWatchersInternal();
+				}
+			});
+		}
 	}
 }
 
@@ -802,13 +818,7 @@ void UMetasoundGeneratorHandle::RegisterPassthroughAnalyzerForType(
 
 void UMetasoundGeneratorHandle::UpdateWatchers() const
 {
-	METASOUND_LLM_SCOPE;
-	METASOUND_TRACE_CPUPROFILER_EVENT_SCOPE(UMetasoundGeneratorHandle::UpdateWatchers);
-
-	if (IsValid())
-	{
-		GeneratorHandle->UpdateOutputWatchers();
-	}
+	// Do nothing. No longer necessary.
 }
 
 void UMetasoundGeneratorHandle::EnableRuntimeRenderTiming(const bool Enable) const
diff --git a/Engine/Plugins/Runtime/Metasound/Source/MetasoundEngine/Private/MetasoundOutputSubsystem.cpp b/Engine/Plugins/Runtime/Metasound/Source/MetasoundEngine/Private/MetasoundOutputSubsystem.cpp
--- a/Engine/Plugins/Runtime/Metasound/Source/MetasoundEngine/Private/MetasoundOutputSubsystem.cpp
+++ b/Engine/Plugins/Runtime/Metasound/Source/MetasoundEngine/Private/MetasoundOutputSubsystem.cpp
@@ -2,12 +2,13 @@
 
 #include "MetasoundOutputSubsystem.h"
 
-#include "MetasoundGenerator.h"
-#include "MetasoundGeneratorHandle.h"
-#include "MetasoundSource.h"
 #include "MetasoundTrace.h"
 #include "Components/AudioComponent.h"
-#include "ProfilingDebugging/CpuProfilerTrace.h"
+
+bool HandleIsValid(const TSharedPtr<Metasound::FMetasoundGeneratorHandle>& Handle)
+{
+	return Handle.IsValid() && Handle->IsValid();
+}
 
 bool UMetaSoundOutputSubsystem::WatchOutput(
 	UAudioComponent* AudioComponent,
@@ -18,9 +19,9 @@ bool UMetaSoundOutputSubsystem::WatchOutput(
 {
 	METASOUND_TRACE_CPUPROFILER_EVENT_SCOPE(UMetasoundOutputSubsystem::WatchOutput_Dynamic);
 
-	UMetasoundGeneratorHandle* Handle = GetOrCreateGeneratorHandle(AudioComponent);
+	const TSharedPtr<Metasound::FMetasoundGeneratorHandle> Handle = GetOrCreateGeneratorHandle(AudioComponent);
 
-	if (nullptr == Handle)
+	if (!HandleIsValid(Handle))
 	{
 		return false;
 	}
@@ -37,9 +38,9 @@ bool UMetaSoundOutputSubsystem::WatchOutput(
 {
 	METASOUND_TRACE_CPUPROFILER_EVENT_SCOPE(UMetasoundOutputSubsystem::WatchOutput_Native);
 
-	UMetasoundGeneratorHandle* Handle = GetOrCreateGeneratorHandle(AudioComponent);
+	const TSharedPtr<Metasound::FMetasoundGeneratorHandle> Handle = GetOrCreateGeneratorHandle(AudioComponent);
 
-	if (nullptr == Handle)
+	if (!HandleIsValid(Handle))
 	{
 		return false;
 	}
@@ -47,62 +48,51 @@ bool UMetaSoundOutputSubsystem::WatchOutput(
 	return Handle->WatchOutput(OutputName, OnOutputValueChanged, AnalyzerName, AnalyzerOutputName);
 }
 
-bool UMetaSoundOutputSubsystem::IsTickable() const
-{
-	return TrackedGenerators.Num() > 0;
-}
-
-void UMetaSoundOutputSubsystem::Tick(float DeltaTime)
-{
-	METASOUND_TRACE_CPUPROFILER_EVENT_SCOPE(UMetasoundOutputSubsystem::Tick);
-
-	for (auto It = TrackedGenerators.CreateIterator(); It; ++It)
-	{
-		if (!(*It)->IsValid())
-		{
-			It.RemoveCurrent();
-		}
-		else
-		{
-			(*It)->UpdateWatchers();
-		}
-	}
-}
-
-TStatId UMetaSoundOutputSubsystem::GetStatId() const
-{
-	RETURN_QUICK_DECLARE_CYCLE_STAT(UMetasoundGeneratorAccessSubsystem, STATGROUP_Tickables);
-}
-
-UMetasoundGeneratorHandle* UMetaSoundOutputSubsystem::GetOrCreateGeneratorHandle(UAudioComponent* AudioComponent)
+TSharedPtr<Metasound::FMetasoundGeneratorHandle> UMetaSoundOutputSubsystem::GetOrCreateGeneratorHandle(UAudioComponent* AudioComponent)
 {
+	METASOUND_TRACE_CPUPROFILER_EVENT_SCOPE(UMetasoundOutputSubsystem::GetOrCreateGeneratorHandle);
+	
+	check(IsInGameThread());
+	
+	CleanUpInvalidGeneratorHandles();
+	
 	if (nullptr == AudioComponent)
 	{
 		return nullptr;
 	}
 
-	UMetasoundGeneratorHandle* Handle = nullptr;
-
 	// Try to find an existing handle
 	const uint64 AudioComponentId = AudioComponent->GetAudioComponentID();
-	if (const TObjectPtr<UMetasoundGeneratorHandle>* FoundHandle = TrackedGenerators.FindByPredicate(
-		[AudioComponentId](const TObjectPtr<UMetasoundGeneratorHandle> ExistingHandle)
+	
+	if (const TSharedPtr<Metasound::FMetasoundGeneratorHandle>* FoundHandle = TrackedGenerators.FindByPredicate(
+		[AudioComponentId](const TSharedPtr<Metasound::FMetasoundGeneratorHandle>& ExistingHandle)
 		{
-			return ExistingHandle->IsValid() && ExistingHandle->GetAudioComponentId() == AudioComponentId;
+			return HandleIsValid(ExistingHandle) && ExistingHandle->GetAudioComponentId() == AudioComponentId;
 		}))
 	{
-		Handle = *FoundHandle;
+		return *FoundHandle;
 	}
+	
 	// Create a new one
-	else
 	{
-		Handle = UMetasoundGeneratorHandle::CreateMetaSoundGeneratorHandle(AudioComponent);
+		const TSharedPtr<Metasound::FMetasoundGeneratorHandle> Handle = Metasound::FMetasoundGeneratorHandle::Create(AudioComponent);
 
-		if (nullptr != Handle && Handle->IsValid())
+		if (HandleIsValid(Handle))
 		{
 			TrackedGenerators.Add(Handle);
+			return Handle;
 		}
 	}
 
-	return Handle;
-}
\ No newline at end of file
+	return nullptr;
+}
+
+void UMetaSoundOutputSubsystem::CleanUpInvalidGeneratorHandles()
+{
+	METASOUND_TRACE_CPUPROFILER_EVENT_SCOPE(UMetasoundOutputSubsystem::CleanUpInvalidGeneratorHandles);
+
+	TrackedGenerators.RemoveAll([](const TSharedPtr<Metasound::FMetasoundGeneratorHandle>& Handle)
+	{
+		return !HandleIsValid(Handle);
+	});
+}
diff --git a/Engine/Plugins/Runtime/Metasound/Source/MetasoundEngine/Public/MetasoundGeneratorHandle.h b/Engine/Plugins/Runtime/Metasound/Source/MetasoundEngine/Public/MetasoundGeneratorHandle.h
--- a/Engine/Plugins/Runtime/Metasound/Source/MetasoundEngine/Public/MetasoundGeneratorHandle.h
+++ b/Engine/Plugins/Runtime/Metasound/Source/MetasoundEngine/Public/MetasoundGeneratorHandle.h
@@ -161,8 +161,9 @@ namespace Metasound
 		/**
 		 * Update any watched outputs
 		 */
-		void UpdateOutputWatchers();
-
+		UE_DEPRECATED(5.5, "Directly calling UpdateOutputWatchers() is no longer necessary.")
+		void UpdateOutputWatchers() const {}
+		
 		/**
 		 * Map a type name to a passthrough analyzer name to use as a default for UMetasoundOutputSubsystem::WatchOutput()
 		 *
@@ -221,6 +222,8 @@ namespace Metasound
 			FName AnalyzerName = NAME_None,
 			FName AnalyzerOutputName = NAME_None);
 
+		void UpdateOutputWatchersInternal();
+
 		bool TryCreateAnalyzerAddress(
 			const FName OutputName,
 			const FName AnalyzerName,
@@ -395,6 +398,8 @@ namespace Metasound
 		std::atomic<int32> ChangedOutputsQueueCount{ 0 };
 		std::atomic<bool> ChangedOutputsQueueShouldLogIfFull{ true };
 
+		std::atomic_flag OutputWatcherUpdateScheduled = ATOMIC_FLAG_INIT;
+
 		bool bRuntimeRenderTimingShouldBeEnabled{ false };
 	};
 }
@@ -499,6 +504,7 @@ public:
 	/**
 	 * Update any watched outputs
 	 */
+	UE_DEPRECATED(5.5, "Directly calling UpdateWatchers() is no longer necessary.")
 	UFUNCTION(BlueprintCallable, Category="MetaSoundOutput")
 	void UpdateWatchers() const;
 
diff --git a/Engine/Plugins/Runtime/Metasound/Source/MetasoundEngine/Public/MetasoundOutputSubsystem.h b/Engine/Plugins/Runtime/Metasound/Source/MetasoundEngine/Public/MetasoundOutputSubsystem.h
--- a/Engine/Plugins/Runtime/Metasound/Source/MetasoundEngine/Public/MetasoundOutputSubsystem.h
+++ b/Engine/Plugins/Runtime/Metasound/Source/MetasoundEngine/Public/MetasoundOutputSubsystem.h
@@ -8,13 +8,12 @@
 #include "MetasoundOutputSubsystem.generated.h"
 
 class UAudioComponent;
-class UMetasoundGeneratorHandle;
 
 /**
  * Provides access to a playing Metasound generator's outputs
  */
 UCLASS()
-class METASOUNDENGINE_API UMetaSoundOutputSubsystem : public UTickableWorldSubsystem
+class METASOUNDENGINE_API UMetaSoundOutputSubsystem : public UWorldSubsystem
 {
 	GENERATED_BODY()
 
@@ -45,16 +44,10 @@ public:
 		const FOnMetasoundOutputValueChangedNative& OnOutputValueChanged,
 		FName AnalyzerName = NAME_None,
 		FName AnalyzerOutputName = NAME_None);
-	
-	/** Begin UTickableWorldSubsystem */
-	virtual bool IsTickable() const override;
-	virtual void Tick(float DeltaTime) override;
-	virtual TStatId GetStatId() const override;
-	/** End UTickableWorldSubsystem */
 
 private:
-	UMetasoundGeneratorHandle* GetOrCreateGeneratorHandle(UAudioComponent* AudioComponent);
+	TSharedPtr<Metasound::FMetasoundGeneratorHandle> GetOrCreateGeneratorHandle(UAudioComponent* AudioComponent);
+	void CleanUpInvalidGeneratorHandles();
 
-	UPROPERTY(Transient)
-	TArray<TObjectPtr<UMetasoundGeneratorHandle>> TrackedGenerators;
+	TArray<TSharedPtr<Metasound::FMetasoundGeneratorHandle>> TrackedGenerators;
 };
