[Backout] - CL33551869
Original CL Desc
-----------------------------------------------------------------
- Fix for rename ensure regression from AssetManager refactor
- Add logging and track case when multiple assets are registered with the AssetManager with the same key for better debugging
- Add ability to force unregistration when request to finish building & utilize on destruction of MetaSound assets
#rnx

diff --git a/Engine/Plugins/Runtime/Metasound/Source/MetasoundEditor/Private/MetasoundEditorModule.cpp b/Engine/Plugins/Runtime/Metasound/Source/MetasoundEditor/Private/MetasoundEditorModule.cpp
--- a/Engine/Plugins/Runtime/Metasound/Source/MetasoundEditor/Private/MetasoundEditorModule.cpp
+++ b/Engine/Plugins/Runtime/Metasound/Source/MetasoundEditor/Private/MetasoundEditorModule.cpp
@@ -439,15 +439,11 @@ namespace Metasound
 					if (bIsRegisteredClass)
 					{
 						// Use the editor version of UnregisterWithFrontend so it refreshes any open MetaSound editors
-						// Doesn't use AssetData::GetAsset() as this can result in attempting to reload the object.
-						// If this call is hit after the asset is removed, the assumption is unregistration already
-						// occurred on object destroy.
-						if (UObject* AssetObject = InAssetData.GetSoftObjectPath().ResolveObject())
+						IMetaSoundAssetManager::GetChecked().RemoveAsset(InAssetData);
+						if (UObject* AssetObject = InAssetData.GetAsset())
 						{
 							FGraphBuilder::UnregisterGraphWithFrontend(*AssetObject);
 						}
-
-						IMetaSoundAssetManager::GetChecked().RemoveAsset(InAssetData);
 					}
 				}
 			}
@@ -461,29 +457,14 @@ namespace Metasound
 					const bool bIsRegisteredClass = IMetasoundUObjectRegistry::Get().IsRegisteredClass(*AssetClass);
 					if (bIsRegisteredClass)
 					{
-						IMetaSoundAssetManager& AssetManager = IMetaSoundAssetManager::GetChecked();
+						constexpr bool bReregisterWithFrontend = false;
+						IMetaSoundAssetManager::GetChecked().RenameAsset(InAssetData, bReregisterWithFrontend);
 
-						// Unregister using the new asset data even though the old object was last to be registered
-						// as the old asset is no longer accessible by the time rename is called. The asset at this
-						// point is identical however to its prior counterpart.
+						// Use the FGraphBuilder Register call instead of registering via the
+						// MetaSoundAssetSubsystem so as to properly refresh respective open editors.
+						constexpr bool bForceViewSynchronization = true;
 						UObject* AssetObject = InAssetData.GetAsset();
-						check(AssetObject);
-
-						FMetasoundAssetBase* AssetBase = AssetManager.GetAsAsset(*AssetObject);
-						check(AssetBase);
-						bool bIsRegistered = AssetBase->IsRegistered();
-						if (bIsRegistered)
-						{
-							FGraphBuilder::UnregisterGraphWithFrontend(*AssetObject);
-						}
-
-						IMetaSoundAssetManager::GetChecked().RenameAsset(InAssetData, InOldObjectPath);
-
-						if (bIsRegistered)
-						{
-							constexpr bool bForceViewSynchronization = true;
-							FGraphBuilder::RegisterGraphWithFrontend(*AssetObject, bForceViewSynchronization);
-						}
+						FGraphBuilder::RegisterGraphWithFrontend(*AssetObject, bForceViewSynchronization);
 					}
 				}
 			}
diff --git a/Engine/Plugins/Runtime/Metasound/Source/MetasoundEngine/Private/MetasoundAssetSubsystem.cpp b/Engine/Plugins/Runtime/Metasound/Source/MetasoundEngine/Private/MetasoundAssetSubsystem.cpp
--- a/Engine/Plugins/Runtime/Metasound/Source/MetasoundEngine/Private/MetasoundAssetSubsystem.cpp
+++ b/Engine/Plugins/Runtime/Metasound/Source/MetasoundEngine/Private/MetasoundAssetSubsystem.cpp
@@ -89,39 +89,24 @@ namespace Metasound::Engine
 		//
 		// This protects against scenarios where a metasound is renamed or moved 
 		// and the new entry was being erroneously removed from the PathMap
-		bool RemovePath(TMap<Frontend::FAssetKey, TArray<FTopLevelAssetPath>>& InMap, const Frontend::FAssetKey& AssetKey, const FTopLevelAssetPath& AssetPath)
+		bool RemoveAndWarnIfNotExactMatch(TMap<Frontend::FAssetKey, FTopLevelAssetPath>& InMap, const Frontend::FAssetKey& AssetKey, const FTopLevelAssetPath& AssetPath)
 		{
+#if !NO_LOGGING
+			// Can be invalid if document on object being destroyed is being removed and MetaSound was no longer registered,
+			// in which case ignore logging as the assumption is it was a valid object to be removed.
 			if (AssetPath.IsValid())
 			{
-				if (TArray<FTopLevelAssetPath>* MapAssetPaths = InMap.Find(AssetKey))
+				if (const FTopLevelAssetPath* MapAssetPath = InMap.Find(AssetKey))
 				{
-					const int32 NumRemoved = MapAssetPaths->RemoveAllSwap([&AssetPath](const FTopLevelAssetPath& Path) { return Path == AssetPath; }, EAllowShrinking::No);
-					if (MapAssetPaths->IsEmpty())
+					if (*MapAssetPath != AssetPath)
 					{
-						InMap.Remove(AssetKey);
+						UE_LOG(LogMetaSound, Warning, TEXT("AssetManager removing object where path (%s) does not match registry copy (%s)"), *MapAssetPath->ToString(), *AssetPath.ToString());
 					}
-					return NumRemoved > 0;
 				}
 			}
-
-			return false;
-		}
-
-		void AddPath(TMap<Frontend::FAssetKey, TArray<FTopLevelAssetPath>>& InMap, const Frontend::FAssetKey& AssetKey, const FTopLevelAssetPath& AssetPath)
-		{
-			TArray<FTopLevelAssetPath>& Paths = InMap.FindOrAdd(AssetKey);
-			Paths.AddUnique(AssetPath);
-#if !NO_LOGGING
-			if (Paths.Num() > 1)
-			{
-				TArray<FString> PathStrings;
-				Algo::Transform(Paths, PathStrings, [](const FTopLevelAssetPath& Path) { return Path.ToString(); });
-				UE_LOG(LogMetaSound, Warning,
-					TEXT("MetaSoundAssetManager has registered multiple assets with key '%s':\n%s\n"),
-					*AssetKey.ToString(),
-					*FString::Join(PathStrings, TEXT("\n")));
-			}
 #endif // !NO_LOGGING
+
+			return InMap.Remove(AssetKey) > 0;
 		}
 	}
 
@@ -168,7 +153,7 @@ namespace Metasound::Engine
 #endif // WITH_EDITOR
 		virtual void RemoveAsset(const UObject& InObject) override;
 		virtual void RemoveAsset(const FAssetData& InAssetData) override;
-		virtual void RenameAsset(const FAssetData& InAssetData, const FString& InOldObjectPath) override;
+		virtual void RenameAsset(const FAssetData& InAssetData, bool bInReregisterWithFrontend = true) override;
 		virtual FMetasoundAssetBase* TryLoadAssetFromKey(const FAssetKey& InKey) const override;
 		virtual bool TryGetAssetIDFromClassName(const FMetasoundFrontendClassName& InClassName, FGuid& OutGuid) const override;
 		virtual bool TryLoadReferencedAssets(const FMetasoundAssetBase& InAssetBase, TArray<FMetasoundAssetBase*>& OutReferencedAssets) const override;
@@ -192,14 +177,14 @@ namespace Metasound::Engine
 		int32 AutoUpdateDenyListChangeID = INDEX_NONE;
 		TSet<FName> AutoUpdateDenyListCache;
 		std::atomic<bool> bIsInitialAssetScanComplete = false;
-		TMap<FAssetKey, TArray<FTopLevelAssetPath>> PathMap;
+		TMap<FAssetKey, FTopLevelAssetPath> PathMap;
 	};
 
 	FMetaSoundAssetManager::~FMetaSoundAssetManager()
 	{
 		if (!PathMap.IsEmpty())
 		{
-			UE_LOG(LogMetaSound, Display, TEXT("Destroying MetaSoundAssetManager while %i asset keys are still registered."), PathMap.Num());
+			UE_LOG(LogMetaSound, Display, TEXT("Destroying MetaSoundAssetManager while %i asset paths are still available."), PathMap.Num());
 		}
 	}
 
@@ -291,7 +276,7 @@ namespace Metasound::Engine
 
 		if (AssetKey.IsValid())
 		{
-			AssetSubsystemPrivate::AddPath(PathMap, AssetKey, FTopLevelAssetPath(&InObject));
+			PathMap.FindOrAdd(AssetKey) = FTopLevelAssetPath(&InObject);
 		}
 
 		return AssetKey;
@@ -332,13 +317,13 @@ namespace Metasound::Engine
 
 		if (ClassInfo.AssetClassID.IsValid())
 		{
-			const FAssetKey AssetKey = FAssetKey(ClassInfo.ClassName, ClassInfo.Version);
-			if (AssetKey.IsValid())
+			const FAssetKey RegistryKey = FAssetKey(ClassInfo.ClassName, ClassInfo.Version);
+			if (RegistryKey.IsValid())
 			{
-				AssetSubsystemPrivate::AddPath(PathMap, AssetKey, ClassInfo.AssetPath);
+				PathMap.FindOrAdd(RegistryKey) = ClassInfo.AssetPath;
 			}
 
-			return AssetKey;
+			return RegistryKey;
 		}
 
 		// Invalid ClassID means the node could not be registered.
@@ -384,7 +369,7 @@ namespace Metasound::Engine
 
 	FMetasoundAssetBase* FMetaSoundAssetManager::FindAsset(const Metasound::Frontend::FAssetKey& InKey) const
 	{
-		if (const FTopLevelAssetPath* AssetPath = FindAssetPath(InKey))
+		if (const FTopLevelAssetPath* AssetPath = PathMap.Find(InKey))
 		{
 			if (UObject* Object = FSoftObjectPath(*AssetPath, { }).ResolveObject())
 			{
@@ -397,7 +382,7 @@ namespace Metasound::Engine
 
 	TScriptInterface<IMetaSoundDocumentInterface> FMetaSoundAssetManager::FindAssetAsDocumentInterface(const Metasound::Frontend::FAssetKey& InKey) const
 	{
-		if (const FTopLevelAssetPath* AssetPath = FindAssetPath(InKey))
+		if (const FTopLevelAssetPath* AssetPath = PathMap.Find(InKey))
 		{
 			if (UObject* Object = FSoftObjectPath(*AssetPath, { }).ResolveObject())
 			{
@@ -410,15 +395,7 @@ namespace Metasound::Engine
 
 	const FTopLevelAssetPath* FMetaSoundAssetManager::FindAssetPath(const Metasound::Frontend::FAssetKey& InKey) const
 	{
-		if (const TArray<FTopLevelAssetPath>* Paths = PathMap.Find(InKey))
-		{
-			if (!Paths->IsEmpty())
-			{
-				return &Paths->Last();
-			}
-		}
-
-		return nullptr;
+		return PathMap.Find(InKey);
 	}
 
 	FMetasoundAssetBase* FMetaSoundAssetManager::GetAsAsset(UObject& InObject) const
@@ -494,7 +471,7 @@ namespace Metasound::Engine
 			}
 
 			const FAssetKey AssetKey(Class.Metadata);
-			if (const FTopLevelAssetPath* ObjectPath = FindAssetPath(AssetKey))
+			if (const FTopLevelAssetPath* ObjectPath = PathMap.Find(AssetKey))
 			{
 				FAssetInfo AssetInfo { FNodeRegistryKey(Class.Metadata), FSoftObjectPath(*ObjectPath) };
 				OutAssetInfos.Add(MoveTemp(AssetInfo));
@@ -614,11 +591,10 @@ namespace Metasound::Engine
 
 			if (IDocumentBuilderRegistry* BuilderRegistry = IDocumentBuilderRegistry::Get())
 			{
-				constexpr bool bForceUnregister = true;
-				BuilderRegistry->FinishBuilding(Metadata.GetClassName(), bForceUnregister);
+				BuilderRegistry->FinishBuilding(Metadata.GetClassName());
 			}
 
-			AssetSubsystemPrivate::RemovePath(PathMap, AssetKey, GraphKey.AssetPath);
+			AssetSubsystemPrivate::RemoveAndWarnIfNotExactMatch(PathMap, AssetKey, GraphKey.AssetPath);
 		}
 	}
 
@@ -631,13 +607,12 @@ namespace Metasound::Engine
 		{
 			if (IDocumentBuilderRegistry* BuilderRegistry = IDocumentBuilderRegistry::Get())
 			{
-				constexpr bool bForceUnregister = true;
-				BuilderRegistry->FinishBuilding(ClassInfo.ClassName, bForceUnregister);
+				BuilderRegistry->FinishBuilding(ClassInfo.ClassName);
 			}
 
 			const FAssetKey AssetKey(ClassInfo.ClassName, ClassInfo.Version);
 			const FTopLevelAssetPath AssetPath(InAssetData.PackageName, InAssetData.AssetName);
-			AssetSubsystemPrivate::RemovePath(PathMap, AssetKey, AssetPath);
+			AssetSubsystemPrivate::RemoveAndWarnIfNotExactMatch(PathMap, AssetKey, AssetPath);
 		}
 	}
 
@@ -650,7 +625,7 @@ namespace Metasound::Engine
 		LoadingDependencies.RemoveAllSwap(IsEqualID);
 	}
 
-	void FMetaSoundAssetManager::RenameAsset(const FAssetData& InAssetData, const FString& InOldObjectPath)
+	void FMetaSoundAssetManager::RenameAsset(const FAssetData& InAssetData, bool bInReregisterWithFrontend)
 	{
 		using namespace Frontend;
 
@@ -660,12 +635,16 @@ namespace Metasound::Engine
 		FNodeClassInfo ClassInfo;
 		if (ensureAlways(AssetSubsystemPrivate::GetAssetClassInfo(InAssetData, ClassInfo)))
 		{
+			IDocumentBuilderRegistry* BuilderRegistry = IDocumentBuilderRegistry::Get();
+			if (BuilderRegistry)
+			{
+				BuilderRegistry->FinishBuilding(ClassInfo.ClassName);
+			}
+
 			const FAssetKey AssetKey(ClassInfo.ClassName, ClassInfo.Version);
-			const FTopLevelAssetPath OldPath(InOldObjectPath);
-			AssetSubsystemPrivate::RemovePath(PathMap, AssetKey, OldPath);
+			PathMap.Remove(AssetKey); // Don't warn if names do not match as by this point the incoming asset data has the newly assigned asset path assigned
 
 #if WITH_EDITOR
-			IDocumentBuilderRegistry* BuilderRegistry = IDocumentBuilderRegistry::Get();
 			if (BuilderRegistry)
 			{
 				FMetaSoundFrontendDocumentBuilder& Builder = BuilderRegistry->FindOrBeginBuilding(MetaSoundAsset->GetOwningAsset());
@@ -675,12 +654,18 @@ namespace Metasound::Engine
 
 			if (ClassInfo.AssetClassID.IsValid())
 			{
-				if (AssetKey.IsValid())
+				const FAssetKey NewAssetKey(ClassInfo.ClassName, ClassInfo.Version);
+				if (NewAssetKey.IsValid())
 				{
-					PathMap.FindOrAdd(AssetKey).AddUnique(ClassInfo.AssetPath);
+					PathMap.FindOrAdd(NewAssetKey) = ClassInfo.AssetPath;
 				}
 			}
 		}
+
+		if (bInReregisterWithFrontend)
+		{
+			MetaSoundAsset->RegisterGraphWithFrontend();
+		}
 	}
 
 	void FMetaSoundAssetManager::RequestAsyncLoadReferencedAssets(FMetasoundAssetBase& InAssetBase)
@@ -752,9 +737,9 @@ namespace Metasound::Engine
 		});
 	}
 
-	FMetasoundAssetBase* FMetaSoundAssetManager::TryLoadAssetFromKey(const Metasound::Frontend::FAssetKey& InAssetKey) const
+	FMetasoundAssetBase* FMetaSoundAssetManager::TryLoadAssetFromKey(const Metasound::Frontend::FAssetKey& RegistryKey) const
 	{
-		if (const FTopLevelAssetPath* ObjectPath = FindAssetPath(InAssetKey))
+		if (const FTopLevelAssetPath* ObjectPath = FindAssetPath(RegistryKey))
 		{
 			const FSoftObjectPath SoftPath(*ObjectPath);
 			return TryLoadAsset(SoftPath);
@@ -830,7 +815,7 @@ namespace Metasound::Engine
 						FMetasoundFrontendRegistryContainer::Get()->UnregisterNode(RegistryKey);
 						const FTopLevelAssetPath AssetPath(AssetData.PackageName, AssetData.AssetName);
 						const FAssetKey AssetKey(AssetClassInfo.ClassName, AssetClassInfo.Version);
-						AssetSubsystemPrivate::RemovePath(PathMap, AssetKey, AssetPath);
+						AssetSubsystemPrivate::RemoveAndWarnIfNotExactMatch(PathMap, AssetKey, AssetPath);
 					}
 				}
 			}
@@ -1027,10 +1012,10 @@ void UMetaSoundAssetSubsystem::RemoveAsset(const FAssetData& InAssetData)
 	IMetaSoundAssetManager::GetChecked().RemoveAsset(InAssetData);
 }
 
-void UMetaSoundAssetSubsystem::RenameAsset(const FAssetData& InAssetData, bool /* bInReregisterWithFrontend */)
+void UMetaSoundAssetSubsystem::RenameAsset(const FAssetData& InAssetData, bool bInReregisterWithFrontend)
 {
 	using namespace Metasound::Frontend;
-	IMetaSoundAssetManager::GetChecked().RenameAsset(InAssetData, { });
+	IMetaSoundAssetManager::GetChecked().RenameAsset(InAssetData, bInReregisterWithFrontend);
 }
 
 UMetaSoundAssetSubsystem& UMetaSoundAssetSubsystem::GetChecked()
diff --git a/Engine/Plugins/Runtime/Metasound/Source/MetasoundEngine/Private/MetasoundBuilderSubsystem.cpp b/Engine/Plugins/Runtime/Metasound/Source/MetasoundEngine/Private/MetasoundBuilderSubsystem.cpp
--- a/Engine/Plugins/Runtime/Metasound/Source/MetasoundEngine/Private/MetasoundBuilderSubsystem.cpp
+++ b/Engine/Plugins/Runtime/Metasound/Source/MetasoundEngine/Private/MetasoundBuilderSubsystem.cpp
@@ -74,24 +74,14 @@ const UClass& UMetaSoundPatchBuilder::GetBaseMetaSoundUClass() const
 
 void UMetaSoundPatchBuilder::OnAssetReferenceAdded(TScriptInterface<IMetaSoundDocumentInterface> DocInterface)
 {
-	using namespace Metasound::Frontend;
-
 	check(DocInterface.GetObject());
-	UMetaSoundPatch& Patch = Builder.CastDocumentObjectChecked<UMetaSoundPatch>();
-	FNodeRegistryKey RegistryKey(Patch.GetConstDocument().RootGraph.Metadata);
-	Patch.ReferencedAssetClassObjects.Add(DocInterface.GetObject());
-	Patch.ReferencedAssetClassKeys.Add(RegistryKey.ToString());
+	Builder.CastDocumentObjectChecked<UMetaSoundPatch>().ReferencedAssetClassObjects.Add(DocInterface.GetObject());
 }
 
 void UMetaSoundPatchBuilder::OnRemovingAssetReference(TScriptInterface<IMetaSoundDocumentInterface> DocInterface)
 {
-	using namespace Metasound::Frontend;
-
 	check(DocInterface.GetObject());
-	UMetaSoundPatch& Patch = Builder.CastDocumentObjectChecked<UMetaSoundPatch>();
-	FNodeRegistryKey RegistryKey(Patch.GetConstDocument().RootGraph.Metadata);
-	Patch.ReferencedAssetClassObjects.Remove(DocInterface.GetObject());
-	Patch.ReferencedAssetClassKeys.Remove(RegistryKey.ToString());
+	Builder.CastDocumentObjectChecked<UMetaSoundPatch>().ReferencedAssetClassObjects.Remove(DocInterface.GetObject());
 }
 
 UMetaSoundBuilderBase& UMetaSoundBuilderSubsystem::AttachBuilderToAssetChecked(UObject& InObject) const
@@ -272,13 +262,8 @@ void UMetaSoundSourceBuilder::InitDelegates(Metasound::Frontend::FDocumentModify
 
 void UMetaSoundSourceBuilder::OnAssetReferenceAdded(TScriptInterface<IMetaSoundDocumentInterface> DocInterface)
 {
-	using namespace Metasound::Frontend;
-
 	check(DocInterface.GetObject());
-	UMetaSoundSource& Source = GetMetaSoundSource();
-	FNodeRegistryKey RegistryKey(Source.GetConstDocument().RootGraph.Metadata);
-	Source.ReferencedAssetClassObjects.Add(DocInterface.GetObject());
-	Source.ReferencedAssetClassKeys.Add(RegistryKey.ToString());
+	GetMetaSoundSource().ReferencedAssetClassObjects.Add(DocInterface.GetObject());
 }
 
 void UMetaSoundSourceBuilder::OnEdgeAdded(int32 EdgeIndex) const
@@ -544,13 +529,8 @@ void UMetaSoundSourceBuilder::OnRemoveSwappingEdge(int32 SwapIndex, int32 LastIn
 
 void UMetaSoundSourceBuilder::OnRemovingAssetReference(TScriptInterface<IMetaSoundDocumentInterface> DocInterface)
 {
-	using namespace Metasound::Frontend;
-
 	check(DocInterface.GetObject());
-	UMetaSoundSource& Source = GetMetaSoundSource();
-	FNodeRegistryKey RegistryKey(Source.GetConstDocument().RootGraph.Metadata);
-	Source.ReferencedAssetClassObjects.Remove(DocInterface.GetObject());
-	Source.ReferencedAssetClassKeys.Remove(RegistryKey.ToString());
+	GetMetaSoundSource().ReferencedAssetClassObjects.Remove(DocInterface.GetObject());
 }
 
 void UMetaSoundSourceBuilder::OnRemovingInput(int32 InputIndex)
diff --git a/Engine/Plugins/Runtime/Metasound/Source/MetasoundEngine/Private/MetasoundDocumentBuilderRegistry.cpp b/Engine/Plugins/Runtime/Metasound/Source/MetasoundEngine/Private/MetasoundDocumentBuilderRegistry.cpp
--- a/Engine/Plugins/Runtime/Metasound/Source/MetasoundEngine/Private/MetasoundDocumentBuilderRegistry.cpp
+++ b/Engine/Plugins/Runtime/Metasound/Source/MetasoundEngine/Private/MetasoundDocumentBuilderRegistry.cpp
@@ -80,7 +80,7 @@ namespace Metasound::Engine
 		return Builders.FindRef(ClassName).Get();
 	}
 
-	bool FDocumentBuilderRegistry::FinishBuilding(const FMetasoundFrontendClassName& InClassName, bool bForceUnregister) const
+	bool FDocumentBuilderRegistry::FinishBuilding(const FMetasoundFrontendClassName& InClassName) const
 	{
 		using namespace Metasound;
 		using namespace Metasound::Engine;
@@ -91,24 +91,17 @@ namespace Metasound::Engine
 			// If the builder has applied transactions to its document object that are not mirrored in the frontend registry,
 			// unregister version in registry. This will ensure that future requests for the builder's associated asset will
 			// register a fresh version from the object as the transaction history is intrinsically lost once this builder
-			// is destroyed. It is also possible that the DocBuilder's underlying object can be invalid if object was force
-			// deleted, so validity check is necessary.
-			FMetaSoundFrontendDocumentBuilder& DocBuilder = Builder->GetBuilder();
-			if (DocBuilder.IsValid())
+			// is destroyed.
+			if (Builder->GetLastTransactionRegistered() != Builder->GetConstBuilder().GetTransactionCount())
 			{
-				const int32 TransactionCount = DocBuilder.GetTransactionCount();
-				const int32 LastTransactionRegistered = Builder->GetLastTransactionRegistered();
-				if (bForceUnregister || LastTransactionRegistered != TransactionCount)
+				UObject& MetaSound = Builder->GetConstBuilder().CastDocumentObjectChecked<UObject>();
+				if (FMetasoundAssetBase* MetaSoundAsset = IMetasoundUObjectRegistry::Get().GetObjectAsAssetBase(&MetaSound))
 				{
-					UObject& MetaSound = DocBuilder.CastDocumentObjectChecked<UObject>();
-					if (FMetasoundAssetBase* MetaSoundAsset = IMetasoundUObjectRegistry::Get().GetObjectAsAssetBase(&MetaSound))
-					{
-						MetaSoundAsset->UnregisterGraphWithFrontend();
-					}
+					MetaSoundAsset->UnregisterGraphWithFrontend();
 				}
 			}
 
-			DocBuilder.FinishBuilding();
+			Builder->GetBuilder().FinishBuilding();
 			ensureAlways(Builders.Remove(InClassName));
 			return true;
 		}
diff --git a/Engine/Plugins/Runtime/Metasound/Source/MetasoundEngine/Public/MetasoundBuilderBase.cpp b/Engine/Plugins/Runtime/Metasound/Source/MetasoundEngine/Public/MetasoundBuilderBase.cpp
--- a/Engine/Plugins/Runtime/Metasound/Source/MetasoundEngine/Public/MetasoundBuilderBase.cpp
+++ b/Engine/Plugins/Runtime/Metasound/Source/MetasoundEngine/Public/MetasoundBuilderBase.cpp
@@ -36,11 +36,12 @@ void UMetaSoundBuilderBase::BeginDestroy()
 {
 	using namespace Metasound::Engine;
 
-	// Need to finish building before destroying UPROPERTYs (Super::BeginDestroy)
-	// as the Builder often holds a TScriptInterface<IMetaSoundDocumentInterface>
-	// of a UPROPERTY that lives on this or derived objects. 
-	constexpr bool bForceUnregister = true;
-	FDocumentBuilderRegistry::GetChecked().FinishBuilding(ClassName, bForceUnregister);
+	FDocumentBuilderRegistry::GetChecked().FinishBuilding(ClassName);
+
+	// Need to detach before destroying UPROPERTYs as the Builder 
+	// often holds a TScriptInterface<IMetaSoundDocumentInterface> of
+	// a UPROPERTY that lives on this or derived objects. 
+	Builder.FinishBuilding();
 	Super::BeginDestroy();
 }
 
diff --git a/Engine/Plugins/Runtime/Metasound/Source/MetasoundEngine/Public/MetasoundDocumentBuilderRegistry.h b/Engine/Plugins/Runtime/Metasound/Source/MetasoundEngine/Public/MetasoundDocumentBuilderRegistry.h
--- a/Engine/Plugins/Runtime/Metasound/Source/MetasoundEngine/Public/MetasoundDocumentBuilderRegistry.h
+++ b/Engine/Plugins/Runtime/Metasound/Source/MetasoundEngine/Public/MetasoundDocumentBuilderRegistry.h
@@ -88,7 +88,7 @@ namespace Metasound::Engine
 
 		virtual FMetaSoundFrontendDocumentBuilder* FindBuilder(TScriptInterface<IMetaSoundDocumentInterface> MetaSound) const override;
 		virtual FMetaSoundFrontendDocumentBuilder* FindBuilder(const FMetasoundFrontendClassName& InClassName) const override;
-		virtual bool FinishBuilding(const FMetasoundFrontendClassName& InClassName, bool bForceUnregister = false) const override;
+		virtual bool FinishBuilding(const FMetasoundFrontendClassName& InClassName) const override;
 
 		UMetaSoundBuilderBase* FindBuilderObject(TScriptInterface<const IMetaSoundDocumentInterface> MetaSound) const;
 		UMetaSoundBuilderBase* FindBuilderObject(const FMetasoundFrontendClassName& ClassName) const;
diff --git a/Engine/Plugins/Runtime/Metasound/Source/MetasoundFrontend/Public/MetasoundAssetManager.h b/Engine/Plugins/Runtime/Metasound/Source/MetasoundFrontend/Public/MetasoundAssetManager.h
--- a/Engine/Plugins/Runtime/Metasound/Source/MetasoundFrontend/Public/MetasoundAssetManager.h
+++ b/Engine/Plugins/Runtime/Metasound/Source/MetasoundFrontend/Public/MetasoundAssetManager.h
@@ -206,7 +206,7 @@ namespace Metasound::Frontend
 		virtual void RemoveAsset(const FAssetData& InAssetData) = 0;
 
 		// Updates the given MetaSound's asset record with the new name and optionally reregisters it with the Frontend Node Class Registry.
-		virtual void RenameAsset(const FAssetData& InAssetData, const FString& InOldObjectPath) = 0;
+		virtual void RenameAsset(const FAssetData& InAssetData, bool bInReregisterWithFrontend = true) = 0;
 
 		// Waits until all async load requests related to this asset are complete.
  		virtual void WaitUntilAsyncLoadReferencedAssetsComplete(FMetasoundAssetBase& InAssetBase) = 0;
diff --git a/Engine/Plugins/Runtime/Metasound/Source/MetasoundFrontend/Public/MetasoundDocumentInterface.h b/Engine/Plugins/Runtime/Metasound/Source/MetasoundFrontend/Public/MetasoundDocumentInterface.h
--- a/Engine/Plugins/Runtime/Metasound/Source/MetasoundFrontend/Public/MetasoundDocumentInterface.h
+++ b/Engine/Plugins/Runtime/Metasound/Source/MetasoundFrontend/Public/MetasoundDocumentInterface.h
@@ -90,7 +90,7 @@ namespace Metasound::Frontend
 		virtual FMetaSoundFrontendDocumentBuilder& FindOrBeginBuilding(TScriptInterface<IMetaSoundDocumentInterface> MetaSound) = 0;
 #endif // WITH_EDITORONLY_DATA
 
-		virtual bool FinishBuilding(const FMetasoundFrontendClassName& InClassName, bool bForceUnregister = false) const = 0;
+		virtual bool FinishBuilding(const FMetasoundFrontendClassName& InClassName) const = 0;
 
 		UE_DEPRECATED(5.5, "Document cache can now be invalidated by retrieving an asset builder and calling 'Reload'")
 		virtual void InvalidateDocumentCache(const FMetasoundFrontendClassName& InClassName) const { }
