Unhide Advanced category in MetaSound source details panel

diff --git a/Engine/Plugins/Runtime/Metasound/Source/MetasoundEditor/Private/MetasoundDetailCustomization.cpp b/Engine/Plugins/Runtime/Metasound/Source/MetasoundEditor/Private/MetasoundDetailCustomization.cpp
--- a/Engine/Plugins/Runtime/Metasound/Source/MetasoundEditor/Private/MetasoundDetailCustomization.cpp
+++ b/Engine/Plugins/Runtime/Metasound/Source/MetasoundEditor/Private/MetasoundDetailCustomization.cpp
@@ -307,7 +307,6 @@ namespace Metasound
 			// Hack to hide parent structs for nested metadata properties
 			DetailLayout.HideCategory("CustomView");
 
-			DetailLayout.HideCategory("Advanced");
 			DetailLayout.HideCategory("Analysis");
 			DetailLayout.HideCategory("Curves");
 			DetailLayout.HideCategory("File Path");
@@ -316,7 +315,7 @@ namespace Metasound
 			DetailLayout.HideCategory("Loading");
 			DetailLayout.HideCategory("Playback");
 			DetailLayout.HideCategory("Subtitles");
-			DetailLayout.HideCategory("Waveform Processing");;
+			DetailLayout.HideCategory("Waveform Processing");
 		}
 
 		FMetasoundInterfacesDetailCustomization::FMetasoundInterfacesDetailCustomization()
