Only create named events for MetaSound trace macro when stat namedevents is enabled

diff --git a/Engine/Plugins/Runtime/Metasound/Source/MetasoundGraphCore/Public/MetasoundTrace.h b/Engine/Plugins/Runtime/Metasound/Source/MetasoundGraphCore/Public/MetasoundTrace.h
--- a/Engine/Plugins/Runtime/Metasound/Source/MetasoundGraphCore/Public/MetasoundTrace.h
+++ b/Engine/Plugins/Runtime/Metasound/Source/MetasoundGraphCore/Public/MetasoundTrace.h
@@ -19,15 +19,16 @@
 // Copied from SCOPED_NAMED_EVENT but modified
 // to accommodate event names containing ::
 #define METASOUND_TRACE_CPUPROFILER_EVENT_SCOPE(Name)\
-	FScopedNamedEventStatic PREPROCESSOR_JOIN(MetaSound_NamedEvent_,__LINE__)(FColor::Green, NAMED_EVENT_STR(#Name));\
+	FScopedNamedEventConditionalStatic PREPROCESSOR_JOIN(MetaSound_NamedEvent_,__LINE__)(FColor::Green, NAMED_EVENT_STR(#Name), GCycleStatsShouldEmitNamedEvents > 0);\
 	TRACE_CPUPROFILER_EVENT_SCOPE(Name);
 
 #define METASOUND_TRACE_CPUPROFILER_EVENT_SCOPE_TEXT(Name) \
-    SCOPED_NAMED_EVENT_TCHAR(Name, FColor::Green)
+	FScopedNamedEventConditional ANONYMOUS_VARIABLE(NamedEvent_)(FColor::Green, Name, GCycleStatsShouldEmitNamedEvents > 0);\
+	TRACE_CPUPROFILER_EVENT_SCOPE_TEXT(Name);
 
 // Uses cached Insights SpecId to avoid string lookup
 #define METASOUND_TRACE_CPUPROFILER_EVENT_SCOPE_FAST(TraceSpecId, Name) \
-	FScopedNamedEvent ANONYMOUS_VARIABLE(NamedEvent_)(FColor::Green, Name);\
+	FScopedNamedEventConditional ANONYMOUS_VARIABLE(NamedEvent_)(FColor::Green, Name, GCycleStatsShouldEmitNamedEvents > 0);\
 	TRACE_CPUPROFILER_EVENT_SCOPE_USE(TraceSpecId, Name, PREPROCESSOR_JOIN(MetaSound_NamedEventScope_, __LINE__), true)
 
 #else
diff --git a/Engine/Source/Runtime/Core/Public/HAL/PlatformMisc.h b/Engine/Source/Runtime/Core/Public/HAL/PlatformMisc.h
--- a/Engine/Source/Runtime/Core/Public/HAL/PlatformMisc.h
+++ b/Engine/Source/Runtime/Core/Public/HAL/PlatformMisc.h
@@ -129,6 +129,49 @@ public:
 	}
 };
 
+// Conditional version of FScopedNamedEventStatic
+class FScopedNamedEventConditionalStatic
+{
+public:
+
+	FScopedNamedEventConditionalStatic(const struct FColor& Color, const TCHAR* Text, bool bCondition)
+		: bStarted(bCondition)
+	{
+		if (bCondition)
+		{
+#if PLATFORM_IMPLEMENTS_BeginNamedEventStatic
+			FPlatformMisc::BeginNamedEventStatic(Color, Text);
+#else
+			FPlatformMisc::BeginNamedEvent(Color, Text);
+#endif
+		}
+	}
+
+	FScopedNamedEventConditionalStatic(const struct FColor& Color, const ANSICHAR* Text, bool bCondition)
+		: bStarted(bCondition)
+	{
+		if (bCondition)
+		{
+#if PLATFORM_IMPLEMENTS_BeginNamedEventStatic
+			FPlatformMisc::BeginNamedEventStatic(Color, Text);
+#else
+			FPlatformMisc::BeginNamedEvent(Color, Text);
+#endif
+		}
+	}
+
+	~FScopedNamedEventConditionalStatic()
+	{
+		if (bStarted)
+		{
+			FPlatformMisc::EndNamedEvent();
+		}
+	}
+
+private:
+	bool bStarted;
+};
+
 
 // Lightweight scoped named event separate from stats system.  Will be available in test builds.  
 // Events cost profiling overhead so use them judiciously in final code.
