Missing string arg to fix build, I can has program I swear (...perforce lately, apparently not so much)!
#rnx

diff --git a/Engine/Plugins/Runtime/Metasound/Source/MetasoundEngine/Public/MetasoundBuilderBase.cpp b/Engine/Plugins/Runtime/Metasound/Source/MetasoundEngine/Public/MetasoundBuilderBase.cpp
--- a/Engine/Plugins/Runtime/Metasound/Source/MetasoundEngine/Public/MetasoundBuilderBase.cpp
+++ b/Engine/Plugins/Runtime/Metasound/Source/MetasoundEngine/Public/MetasoundBuilderBase.cpp
@@ -53,7 +53,8 @@ void UMetaSoundBuilderBase::BeginDestroy()
 			{
 				UE_LOG(LogMetaSound, Error, TEXT("Failed to destroy MetaSound builder of asset with class name '%s'. \n"
 					"Registry is likely corrupt (i.e. a different builder instance was registered with the same ClassName). \n"
-					"Duplicate assets with the same ClassName/Guid should be removed"));
+					"Duplicate assets with the same ClassName/Guid should be removed"),
+					*MetaSoundClassName.ToString());
 				Builder.FinishBuilding();
 			}
 		}
