Remove spamming error to fix cook.
#rnx

diff --git a/Engine/Plugins/Runtime/Metasound/Source/MetasoundEngine/Public/MetasoundBuilderBase.cpp b/Engine/Plugins/Runtime/Metasound/Source/MetasoundEngine/Public/MetasoundBuilderBase.cpp
--- a/Engine/Plugins/Runtime/Metasound/Source/MetasoundEngine/Public/MetasoundBuilderBase.cpp
+++ b/Engine/Plugins/Runtime/Metasound/Source/MetasoundEngine/Public/MetasoundBuilderBase.cpp
@@ -48,20 +48,9 @@ void UMetaSoundBuilderBase::BeginDestroy()
 			constexpr bool bForceUnregister = true;
 			const FMetasoundFrontendClassName& MetaSoundClassName = Builder.GetConstDocument().RootGraph.Metadata.GetClassName();
 			BuilderRegistry->FinishBuilding(MetaSoundClassName, bForceUnregister);
-
-			if(Builder.IsValid())
-			{
-				UE_LOG(LogMetaSound, Error, TEXT("Failed to destroy MetaSound builder of asset with class name '%s'. \n"
-					"Registry is likely corrupt (i.e. a different builder instance was registered with the same ClassName). \n"
-					"Duplicate assets with the same ClassName/Guid should be removed"),
-					*MetaSoundClassName.ToString());
-				Builder.FinishBuilding();
-			}
-		}
-		else
-		{
-			Builder.FinishBuilding();
 		}
+
+		Builder.FinishBuilding();
 	}
 	Super::BeginDestroy();
 }
