[Metasound] Disallow array nodes for parameter packs

#tests metasound editor

diff --git a/Engine/Plugins/Runtime/Metasound/Source/MetasoundFrontend/Private/MetasoundParameterPack.cpp b/Engine/Plugins/Runtime/Metasound/Source/MetasoundFrontend/Private/MetasoundParameterPack.cpp
--- a/Engine/Plugins/Runtime/Metasound/Source/MetasoundFrontend/Private/MetasoundParameterPack.cpp
+++ b/Engine/Plugins/Runtime/Metasound/Source/MetasoundFrontend/Private/MetasoundParameterPack.cpp
@@ -11,6 +11,12 @@ DEFINE_LOG_CATEGORY_STATIC(LogMetasoundParamPack, Log, All);
 
 REGISTER_METASOUND_DATATYPE(FMetasoundParameterStorageWrapper, "MetasoundParameterPack", Metasound::ELiteralType::UObjectProxy, UMetasoundParameterPack);
 
+template<>
+struct Metasound::TEnableArrayNodes<FMetasoundParameterStorageWrapper>
+{
+	static constexpr bool Value = false;
+};
+
 namespace MetasoundParameterPack
 {
 	FName RoutingName("ParamPack");
