Protect against invalid proxies in MetaSound proxy type hash functions

diff --git a/Engine/Plugins/Runtime/Metasound/Source/MetasoundEngine/Public/MetasoundAudioBus.h b/Engine/Plugins/Runtime/Metasound/Source/MetasoundEngine/Public/MetasoundAudioBus.h
--- a/Engine/Plugins/Runtime/Metasound/Source/MetasoundEngine/Public/MetasoundAudioBus.h
+++ b/Engine/Plugins/Runtime/Metasound/Source/MetasoundEngine/Public/MetasoundAudioBus.h
@@ -32,7 +32,12 @@ namespace Metasound
 
 		friend FORCEINLINE uint32 GetTypeHash(const Metasound::FAudioBusAsset& InAudioBusAsset)
 		{
-			return GetTypeHash(*InAudioBusAsset.GetAudioBusProxy());
+			const FAudioBusProxyPtr& Proxy = InAudioBusAsset.GetAudioBusProxy();
+			if (Proxy.IsValid())
+			{
+				return GetTypeHash(*Proxy);
+			}
+			return INDEX_NONE;
 		}
 	};
 
diff --git a/Engine/Plugins/Runtime/Metasound/Source/MetasoundEngine/Public/MetasoundWave.h b/Engine/Plugins/Runtime/Metasound/Source/MetasoundEngine/Public/MetasoundWave.h
--- a/Engine/Plugins/Runtime/Metasound/Source/MetasoundEngine/Public/MetasoundWave.h
+++ b/Engine/Plugins/Runtime/Metasound/Source/MetasoundEngine/Public/MetasoundWave.h
@@ -53,7 +53,11 @@ namespace Metasound
 
 		friend FORCEINLINE uint32 GetTypeHash(const Metasound::FWaveAsset& InWaveAsset)
 		{
-			return GetTypeHash(*InWaveAsset.GetSoundWaveProxy());
+			if (InWaveAsset.IsSoundWaveValid())
+			{
+				return GetTypeHash(*InWaveAsset.GetSoundWaveProxy());
+			}
+			return INDEX_NONE;
 		}
 	};
 
diff --git a/Engine/Plugins/Runtime/Metasound/Source/MetasoundEngine/Public/MetasoundWaveTable.h b/Engine/Plugins/Runtime/Metasound/Source/MetasoundEngine/Public/MetasoundWaveTable.h
--- a/Engine/Plugins/Runtime/Metasound/Source/MetasoundEngine/Public/MetasoundWaveTable.h
+++ b/Engine/Plugins/Runtime/Metasound/Source/MetasoundEngine/Public/MetasoundWaveTable.h
@@ -79,7 +79,11 @@ namespace Metasound
 		
 		friend FORCEINLINE uint32 GetTypeHash(const Metasound::FWaveTableBankAsset& InWaveTableBankAsset)
 		{
-			return GetTypeHash(*InWaveTableBankAsset.GetProxy());
+			if (InWaveTableBankAsset.IsValid())
+			{
+				return GetTypeHash(*InWaveTableBankAsset.GetProxy());
+			}
+			return INDEX_NONE;
 		}
 	};
 
