UE-213171
Changing the Active Pins for the MetaSound Patch and Source in the editor menu does not update the settings until the Editor is closed and opened
#rnx

diff --git a/Engine/Plugins/Runtime/Metasound/Source/MetasoundEditor/Private/MetasoundEditorSettings.cpp b/Engine/Plugins/Runtime/Metasound/Source/MetasoundEditor/Private/MetasoundEditorSettings.cpp
--- a/Engine/Plugins/Runtime/Metasound/Source/MetasoundEditor/Private/MetasoundEditorSettings.cpp
+++ b/Engine/Plugins/Runtime/Metasound/Source/MetasoundEditor/Private/MetasoundEditorSettings.cpp
@@ -3,12 +3,16 @@
 
 #include "AudioMaterialSlate/AudioMaterialSlateTypes.h"
 #include "AudioWidgetsStyle.h"
+#include "Framework/Notifications/NotificationManager.h"
+#include "Internationalization/Text.h"
 #include "Styling/AppStyle.h"
 #include "Styling/SlateWidgetStyleAsset.h"
+#include "Widgets/Notifications/SNotificationList.h"
 #include "UObject/UnrealType.h"
 
 #include UE_INLINE_GENERATED_CPP_BY_NAME(MetasoundEditorSettings)
 
+#define LOCTEXT_NAMESPACE "MetasoundEditorSettings"
 
 UMetasoundEditorSettings::UMetasoundEditorSettings(const FObjectInitializer& ObjectInitializer)
 	: Super(ObjectInitializer)
@@ -33,6 +37,22 @@ UMetasoundEditorSettings::UMetasoundEditorSettings(const FObjectInitializer& Obj
 	VariableNodeTitleColor = FLinearColor(0.211f, 0.513f, 0.035f);				// copper
 }
 
+#if WITH_EDITOR
+void UMetasoundEditorSettings::PostEditChangeProperty(FPropertyChangedEvent& PropertyChangedEvent)
+{
+	FName PropertyName = PropertyChangedEvent.Property->GetFName();
+	if (PropertyName == GET_MEMBER_NAME_CHECKED(UMetasoundEditorSettings, bPinMetaSoundPatchInAssetMenu)
+		|| PropertyName == GET_MEMBER_NAME_CHECKED(UMetasoundEditorSettings, bPinMetaSoundSourceInAssetMenu))
+	{
+		FNotificationInfo Info(LOCTEXT("MetasoundEditorSettings_ChangeRequiresEditorRestart", "Change to Asset Menu Settings requires editor restart in order for changes to take effect."));
+		Info.bFireAndForget = true;
+		Info.ExpireDuration = 2.0f;
+		Info.bUseThrobber = true;
+		FSlateNotificationManager::Get().AddNotification(Info);
+	}
+}
+#endif // WITH_EDITOR
+
 const FAudioMaterialKnobStyle* UMetasoundEditorSettings::GetKnobStyle() const
 {
 	if (const UObject* Style = KnobStyleOverride.TryLoad())
@@ -84,3 +104,4 @@ const FAudioMaterialMeterStyle* UMetasoundEditorSettings::GetMeterStyle() const
 
 	return &FAudioWidgetsStyle::Get().GetWidgetStyle<FAudioMaterialMeterStyle>("AudioMaterialMeter.Style");
 }
+#undef LOCTEXT_NAMESPACE // "MetaSoundEditor"
diff --git a/Engine/Plugins/Runtime/Metasound/Source/MetasoundEditor/Private/MetasoundEditorSettings.h b/Engine/Plugins/Runtime/Metasound/Source/MetasoundEditor/Private/MetasoundEditorSettings.h
--- a/Engine/Plugins/Runtime/Metasound/Source/MetasoundEditor/Private/MetasoundEditorSettings.h
+++ b/Engine/Plugins/Runtime/Metasound/Source/MetasoundEditor/Private/MetasoundEditorSettings.h
@@ -86,11 +86,13 @@ class METASOUNDEDITOR_API UMetasoundEditorSettings : public UObject
 	GENERATED_UCLASS_BODY()
 
 public:
-	/** Whether to pin the MetaSound Patch asset type when creating new assets. */
+	/** Whether to pin the MetaSound Patch asset type when creating new assets.
+	  * Requires editor restart for change to take effect.*/
 	UPROPERTY(EditAnywhere, config, DisplayName = "Pin MetaSound Patch in Asset Menu", Category = AssetMenu)
 	bool bPinMetaSoundPatchInAssetMenu = false;
 
-	/** Whether to pin the MetaSound Source asset type when creating new assets. */
+	/** Whether to pin the MetaSound Source asset type when creating new assets. 
+	  * Requires editor restart for change to take effect.*/
 	UPROPERTY(EditAnywhere, config, DisplayName = "Pin MetaSound Source in Asset Menu", Category = AssetMenu)
 	bool bPinMetaSoundSourceInAssetMenu = true;
 
@@ -198,6 +200,12 @@ public:
 	/**Override the Meter Style used in the Metasound Editor.*/
 	UPROPERTY(EditAnywhere, config, Category = MetasoundStyling, meta = (AllowedClasses = "/Script/SlateCore.SlateWidgetStyleAsset", EditCondition = "bUseAudioMaterialWidgets", DisplayName = "Meter Style"))
 	FSoftObjectPath MeterStyleOverride;	
+
+	//UObject
+#if WITH_EDITOR
+	virtual void PostEditChangeProperty(FPropertyChangedEvent& PropertyChangedEvent) override;
+#endif // WITH_EDITOR
+	//~UObject
 		
 	/** Get the AudioMaterialKnob Style. If KnobStyleOverride is not set, returns default style.*/
 	const FAudioMaterialKnobStyle* GetKnobStyle() const;	
