Added Audio Analyzer template node, which is an input-only debug node. Presently displays a spectrum analyzer in the metasound editor. Also: Bug fix for FGraphAnalyzer active vertex analyzer replication.

diff --git a/Engine/Plugins/Runtime/Metasound/Source/MetasoundEditor/Private/MetasoundEditorGraphConnectionManager.cpp b/Engine/Plugins/Runtime/Metasound/Source/MetasoundEditor/Private/MetasoundEditorGraphConnectionManager.cpp
--- a/Engine/Plugins/Runtime/Metasound/Source/MetasoundEditor/Private/MetasoundEditorGraphConnectionManager.cpp
+++ b/Engine/Plugins/Runtime/Metasound/Source/MetasoundEditor/Private/MetasoundEditorGraphConnectionManager.cpp
@@ -99,6 +99,19 @@ namespace Metasound
 			}
 		}
 
+		bool FGraphConnectionManager::HasAudioBusWriter(const FGuid& InAnalyzerInstanceID) const
+		{
+			bool bResult = false;
+
+			if (GraphAnalyzerView.IsValid())
+			{
+				const FName AnalyzerName = Engine::FVertexAnalyzerAudioBusWriter::GetAnalyzerName();
+				bResult = GraphAnalyzerView->HasAnalyzerInstance(AnalyzerName, InAnalyzerInstanceID);
+			}
+
+			return bResult;
+		}
+
 		void FGraphConnectionManager::Update(float InDeltaTime)
 		{
 			using namespace Frontend;
diff --git a/Engine/Plugins/Runtime/Metasound/Source/MetasoundEditor/Private/MetasoundEditorGraphConnectionManager.h b/Engine/Plugins/Runtime/Metasound/Source/MetasoundEditor/Private/MetasoundEditorGraphConnectionManager.h
--- a/Engine/Plugins/Runtime/Metasound/Source/MetasoundEditor/Private/MetasoundEditorGraphConnectionManager.h
+++ b/Engine/Plugins/Runtime/Metasound/Source/MetasoundEditor/Private/MetasoundEditorGraphConnectionManager.h
@@ -105,6 +105,9 @@ namespace Metasound
 			// Remove the specified analyzer. InAnalyzerInstanceID is an FGuid that was previously returned from a call to AddAudioBusWriter(...).
 			void RemoveAudioBusWriter(const FGuid& InAnalyzerInstanceID);
 
+			// Does this connection manager have the specified analyzer? InAnalyzerInstanceID is an FGuid that was previously returned from a call to AddAudioBusWriter(...).
+			bool HasAudioBusWriter(const FGuid& InAnalyzerInstanceID) const;
+
 			void Update(float InDeltaTime);
 
 			template <typename TAnalyzerName>
diff --git a/Engine/Plugins/Runtime/Metasound/Source/MetasoundEditor/Private/MetasoundEditorGraphNodeFactory.h b/Engine/Plugins/Runtime/Metasound/Source/MetasoundEditor/Private/MetasoundEditorGraphNodeFactory.h
--- a/Engine/Plugins/Runtime/Metasound/Source/MetasoundEditor/Private/MetasoundEditorGraphNodeFactory.h
+++ b/Engine/Plugins/Runtime/Metasound/Source/MetasoundEditor/Private/MetasoundEditorGraphNodeFactory.h
@@ -6,10 +6,12 @@
 #include "MetasoundEditorGraphCommentNode.h"
 #include "MetasoundEditorGraphNode.h"
 #include "MetasoundEditorGraphSchema.h"
+#include "NodeTemplates/MetasoundFrontendNodeTemplateAudioAnalyzer.h"
 #include "NodeTemplates/MetasoundFrontendNodeTemplateReroute.h"
 #include "SGraphNode.h"
 #include "SMetasoundGraphNode.h"
 #include "SMetasoundGraphNodeComment.h"
+#include "SMetasoundSpectrumAnalyzerGraphNode.h"
 #include "Templates/SharedPointer.h"
 #include "Widgets/DeclarativeSyntaxSupport.h"
 
@@ -31,6 +33,10 @@ class FMetasoundGraphNodeFactory : public FGraphPanelNodeFactory
 				{
 					return SNew(SMetaSoundGraphNodeKnot, InNode);
 				}
+				else if (ClassName == FAudioAnalyzerNodeTemplate::ClassName)
+				{
+					return SNew(SMetaSoundSpectrumAnalyzerGraphNode, InNode);
+				}
 			}
 			return SNew(SMetaSoundGraphNode, InNode);
 		}
diff --git a/Engine/Plugins/Runtime/Metasound/Source/MetasoundEditor/Private/MetasoundEditorGraphSchema.cpp b/Engine/Plugins/Runtime/Metasound/Source/MetasoundEditor/Private/MetasoundEditorGraphSchema.cpp
--- a/Engine/Plugins/Runtime/Metasound/Source/MetasoundEditor/Private/MetasoundEditorGraphSchema.cpp
+++ b/Engine/Plugins/Runtime/Metasound/Source/MetasoundEditor/Private/MetasoundEditorGraphSchema.cpp
@@ -43,6 +43,7 @@
 #include "MetasoundUObjectRegistry.h"
 #include "MetasoundVariableNodes.h"
 #include "MetasoundVertex.h"
+#include "NodeTemplates/MetasoundFrontendNodeTemplateAudioAnalyzer.h"
 #include "NodeTemplates/MetasoundFrontendNodeTemplateReroute.h"
 #include "ScopedTransaction.h"
 #include "Settings/EditorStyleSettings.h"
@@ -903,6 +904,84 @@ UEdGraphNode* FMetasoundGraphSchemaAction_NewFromSelected::PerformAction(UEdGrap
 	return nullptr;
 }
 
+FMetasoundGraphSchemaAction_NewAudioAnalyzer::FMetasoundGraphSchemaAction_NewAudioAnalyzer()
+	: FMetasoundGraphSchemaAction(
+		FText(),
+		LOCTEXT("AddAudioAnalyzerName", "Add Audio Analyzer Node..."),
+		LOCTEXT("AddAudioAnalyzerTooltip", "Analyze an audio signal (editor only)"),
+		Metasound::Editor::EPrimaryContextGroup::Common)
+{
+	//
+}
+
+UEdGraphNode* FMetasoundGraphSchemaAction_NewAudioAnalyzer::PerformAction(UEdGraph* ParentGraph, UEdGraphPin* FromPin, const FVector2D Location, bool bSelectNewNode/* = true*/)
+{
+	using namespace Metasound::Editor;
+	using namespace Metasound::Frontend;
+
+	check(ParentGraph);
+	UMetasoundEditorGraph* MetaSoundGraph = CastChecked<UMetasoundEditorGraph>(ParentGraph);
+	UObject& ParentMetasound = MetaSoundGraph->GetMetasoundChecked();
+
+	const FScopedTransaction Transaction(LOCTEXT("AddNewAudioAnalyzerNode", "Add Audio Analyzer Node"));
+	ParentMetasound.Modify();
+	ParentGraph->Modify();
+
+	FMetasoundAssetBase* MetaSoundAsset = Metasound::IMetasoundUObjectRegistry::Get().GetObjectAsAssetBase(&ParentMetasound);
+	check(MetaSoundAsset);
+
+	const INodeTemplate* AudioAnalyzerTemplate = INodeTemplateRegistry::Get().FindTemplate(FAudioAnalyzerNodeTemplate::ClassName);
+	if (!AudioAnalyzerTemplate)
+	{
+		UE_LOG(LogMetasoundEditor, Error, TEXT("Failed to find template for class \"%s\""), *FAudioAnalyzerNodeTemplate::ClassName.ToString());
+		return nullptr;
+	}
+
+	FMetaSoundFrontendDocumentBuilder& DocBuilder = IDocumentBuilderRegistry::GetChecked().FindOrBeginBuilding(&ParentMetasound);
+
+	const FName FromVertexDataType = FGraphBuilder::GetPinDataType(FromPin);
+	FNodeTemplateGenerateInterfaceParams Params;
+	Params.InputsToConnect.Add(FromVertexDataType);
+
+	const FMetasoundFrontendNode* TemplateNode = DocBuilder.AddNodeByTemplate(*AudioAnalyzerTemplate, MoveTemp(Params));
+	DocBuilder.SetNodeLocation(TemplateNode->GetID(), Location);
+
+	const FMetasoundFrontendVertexHandle FromVertexHandle = FGraphBuilder::GetPinVertexHandle(DocBuilder, FromPin);
+	const auto VertexIsMatchingDataType = [FromVertexDataType](const FMetasoundFrontendVertex& Vertex) { return Vertex.TypeName == FromVertexDataType; };
+	const FMetasoundFrontendVertex* ToVertex = TemplateNode->Interface.Inputs.FindByPredicate(VertexIsMatchingDataType);
+	if (FromVertexHandle.IsSet() && ToVertex)
+	{
+		DocBuilder.AddEdge(FMetasoundFrontendEdge
+			{
+				FromVertexHandle.NodeID,
+				FromVertexHandle.VertexID,
+				TemplateNode->GetID(),
+				ToVertex->VertexID,
+			});
+	}	
+
+	const FMetasoundFrontendClass& FrontendClass = AudioAnalyzerTemplate->GetFrontendClass();
+
+	// Proactively create the corresponding EdGraphNode so that we have something to return:
+	if (UMetasoundEditorGraphExternalNode* NewGraphNode = FGraphBuilder::AddTemplateNode(ParentMetasound, TemplateNode->GetID(), FrontendClass.Metadata, bSelectNewNode))
+	{
+		TSharedPtr<FEditor> ParentEditor = FGraphBuilder::GetEditorForMetasound(ParentMetasound);
+		if (ParentEditor.IsValid() && bSelectNewNode)
+		{
+			ParentEditor->ClearSelectionAndSelectNode(NewGraphNode);
+		}
+
+		return NewGraphNode;
+	}
+
+	return nullptr;
+}
+
+const FLinearColor& FMetasoundGraphSchemaAction_NewAudioAnalyzer::GetIconColor() const
+{
+	return GetDefault<UMetasoundEditorSettings>()->AudioPinTypeColor;
+}
+
 FMetasoundGraphSchemaAction_NewReroute::FMetasoundGraphSchemaAction_NewReroute(const FLinearColor* InIconColor, bool bInShouldTransact /* = true */)
 	: FMetasoundGraphSchemaAction(
 		FText(),
@@ -1172,6 +1251,12 @@ void UMetasoundEditorGraphSchema::GetGraphContextActions(FGraphContextMenuBuilde
 			});
 
 			FGraphActionMenuBuilder& ActionMenuBuilder = static_cast<FGraphActionMenuBuilder&>(ContextMenuBuilder);
+
+			if (OutputHandle->GetDataType() == GetMetasoundDataTypeName<FAudioBuffer>())
+			{
+				ActionMenuBuilder.AddAction(MakeShared<FMetasoundGraphSchemaAction_NewAudioAnalyzer>());
+			}
+
 			ActionMenuBuilder.AddAction(MakeShared<FMetasoundGraphSchemaAction_PromoteToOutput>());
 			ActionMenuBuilder.AddAction(MakeShared<FMetasoundGraphSchemaAction_PromoteToVariable_MutatorNode>());
 
diff --git a/Engine/Plugins/Runtime/Metasound/Source/MetasoundEditor/Private/SMetasoundSpectrumAnalyzerGraphNode.cpp b/Engine/Plugins/Runtime/Metasound/Source/MetasoundEditor/Private/SMetasoundSpectrumAnalyzerGraphNode.cpp
new file mode 100644
--- /dev/null
+++ b/Engine/Plugins/Runtime/Metasound/Source/MetasoundEditor/Private/SMetasoundSpectrumAnalyzerGraphNode.cpp
@@ -0,0 +1,81 @@
+// Copyright Epic Games, Inc. All Rights Reserved.
+#include "SMetasoundSpectrumAnalyzerGraphNode.h"
+
+#include "AudioSpectrumAnalyzer.h"
+#include "Editor.h"
+#include "MetasoundEditorGraphBuilder.h"
+#include "MetasoundEditorGraphNode.h"
+#include "Widgets/Layout/SBox.h"
+
+namespace Metasound::Editor
+{
+	SMetaSoundSpectrumAnalyzerGraphNode::~SMetaSoundSpectrumAnalyzerGraphNode()
+	{
+		if (AnalyzerInstanceID.IsValid())
+		{
+			if (TSharedPtr<FEditor> Editor = FGraphBuilder::GetEditorForNode(GetMetaSoundNode()))
+			{
+				Editor->GetConnectionManager().RemoveAudioBusWriter(AnalyzerInstanceID);
+			}
+
+			AnalyzerInstanceID.Invalidate();
+		}
+	}
+
+	void SMetaSoundSpectrumAnalyzerGraphNode::Tick(const FGeometry& AllottedGeometry, const double InCurrentTime, const float InDeltaTime)
+	{
+		SMetaSoundGraphNode::Tick(AllottedGeometry, InCurrentTime, InDeltaTime);
+
+		const UMetasoundEditorGraphNode& MetaSoundNode = GetMetaSoundNode();
+		if (SpectrumAnalyzer.IsValid() && ensure(MetaSoundNode.Pins.Num() == 1))
+		{
+			const UEdGraphPin* Pin = MetaSoundNode.Pins.Last();
+			if (!Pin->LinkedTo.IsEmpty() && ensure(Pin->PinType.PinCategory == FGraphBuilder::PinCategoryAudio && Pin->Direction == EGPD_Input))
+			{
+				// Find connected output for the input (only ever one):
+				UEdGraphPin* SourcePin = Pin->LinkedTo.Last();
+				ensure(SourcePin->Direction == EGPD_Output);
+
+				TSharedPtr<FEditor> Editor = FGraphBuilder::GetEditorForNode(MetaSoundNode);
+				if (Editor.IsValid() && ensure(GEditor))
+				{
+					// Analyze audio with AudioBusWriter, if not already doing so:
+					FGraphConnectionManager& ConnectionManager = Editor->GetConnectionManager();
+					if (!AnalyzerInstanceID.IsValid() || !ConnectionManager.HasAudioBusWriter(AnalyzerInstanceID))
+					{
+						const Frontend::FConstOutputHandle OutputHandle = FGraphBuilder::FindReroutedOutputHandleFromPin(SourcePin);
+						const FGuid NodeID = OutputHandle->GetOwningNodeID();
+						const FName OutputName = OutputHandle->GetName();
+						const Audio::FDeviceId DeviceId = GEditor->GetMainAudioDeviceID();
+						AnalyzerInstanceID = ConnectionManager.AddAudioBusWriter(NodeID, OutputName, DeviceId, SpectrumAnalyzer->GetAudioBus());
+					}
+				}
+			}
+		}
+	}
+
+	void SMetaSoundSpectrumAnalyzerGraphNode::CreateBelowPinControls(TSharedPtr<SVerticalBox> MainBox)
+	{
+		if (ensure(GEditor))
+		{
+			if (!SpectrumAnalyzer.IsValid())
+			{
+				constexpr int32 NumChannels = 1;
+				const Audio::FDeviceId AudioDeviceId = GEditor->GetMainAudioDeviceID();
+				SpectrumAnalyzer = MakeShared<AudioWidgets::FAudioSpectrumAnalyzer>(NumChannels, AudioDeviceId);
+			}
+
+			MainBox->AddSlot()
+				.AutoHeight()
+				.Padding(1.0f, 0.0f)
+				[
+					SNew(SBox)
+						.MinDesiredWidth(250.0f)
+						.MinDesiredHeight(250.0f)
+						[
+							SpectrumAnalyzer->GetWidget()
+						]
+				];
+		}
+	}
+}
diff --git a/Engine/Plugins/Runtime/Metasound/Source/MetasoundEditor/Private/SMetasoundSpectrumAnalyzerGraphNode.h b/Engine/Plugins/Runtime/Metasound/Source/MetasoundEditor/Private/SMetasoundSpectrumAnalyzerGraphNode.h
new file mode 100644
--- /dev/null
+++ b/Engine/Plugins/Runtime/Metasound/Source/MetasoundEditor/Private/SMetasoundSpectrumAnalyzerGraphNode.h
@@ -0,0 +1,30 @@
+// Copyright Epic Games, Inc. All Rights Reserved.
+#pragma once
+
+#include "SMetasoundGraphNode.h"
+
+// Forward Declarations
+namespace AudioWidgets
+{
+	class FAudioSpectrumAnalyzer;
+}
+
+namespace Metasound::Editor
+{
+	class SMetaSoundSpectrumAnalyzerGraphNode : public SMetaSoundGraphNode
+	{
+	public:
+		virtual ~SMetaSoundSpectrumAnalyzerGraphNode();
+
+		// Begin SWidget overrides.
+		virtual void Tick(const FGeometry& AllottedGeometry, const double InCurrentTime, const float InDeltaTime) override;
+		// End SWidget overrides.
+
+	protected:
+		virtual void CreateBelowPinControls(TSharedPtr<SVerticalBox> MainBox) override;
+
+	private:
+		TSharedPtr<AudioWidgets::FAudioSpectrumAnalyzer> SpectrumAnalyzer;
+		FGuid AnalyzerInstanceID;
+	};
+}
diff --git a/Engine/Plugins/Runtime/Metasound/Source/MetasoundEditor/Public/MetasoundEditorGraphSchema.h b/Engine/Plugins/Runtime/Metasound/Source/MetasoundEditor/Public/MetasoundEditorGraphSchema.h
--- a/Engine/Plugins/Runtime/Metasound/Source/MetasoundEditor/Public/MetasoundEditorGraphSchema.h
+++ b/Engine/Plugins/Runtime/Metasound/Source/MetasoundEditor/Public/MetasoundEditorGraphSchema.h
@@ -382,6 +382,23 @@ struct METASOUNDEDITOR_API FMetasoundGraphSchemaAction_NewFromSelected : public
 	//~ End FEdGraphSchemaAction Interface
 };
 
+/** Action to create new Audio Analyzer node */
+USTRUCT()
+struct METASOUNDEDITOR_API FMetasoundGraphSchemaAction_NewAudioAnalyzer : public FMetasoundGraphSchemaAction
+{
+	GENERATED_USTRUCT_BODY();
+
+	FMetasoundGraphSchemaAction_NewAudioAnalyzer();
+
+	//~ Begin FMetasoundGraphSchemaAction Interface
+	virtual const FLinearColor& GetIconColor() const override;
+	//~ End FMetasoundGraphSchemaAction Interface
+
+	//~ Begin FEdGraphSchemaAction Interface
+	virtual UEdGraphNode* PerformAction(UEdGraph* ParentGraph, UEdGraphPin* FromPin, const FVector2D Location, bool bSelectNewNode = true) override;
+	//~ End FEdGraphSchemaAction Interface
+};
+
 /** Action to create new reroute node */
 USTRUCT()
 struct METASOUNDEDITOR_API FMetasoundGraphSchemaAction_NewReroute : public FMetasoundGraphSchemaAction
diff --git a/Engine/Plugins/Runtime/Metasound/Source/MetasoundFrontend/Private/Analysis/MetasoundFrontendGraphAnalyzer.cpp b/Engine/Plugins/Runtime/Metasound/Source/MetasoundFrontend/Private/Analysis/MetasoundFrontendGraphAnalyzer.cpp
--- a/Engine/Plugins/Runtime/Metasound/Source/MetasoundFrontend/Private/Analysis/MetasoundFrontendGraphAnalyzer.cpp
+++ b/Engine/Plugins/Runtime/Metasound/Source/MetasoundFrontend/Private/Analysis/MetasoundFrontendGraphAnalyzer.cpp
@@ -49,7 +49,11 @@ namespace Metasound
 					const FAnalyzerAddress& AnalyzerAddress = Analyzer->GetAnalyzerAddress();
 					if (ReceiverAddresses.Contains(AnalyzerAddress))
 					{
+						// Found the matching active analyzer list entry for this vertex analyzer.
 						ReceiverAddresses.Remove(AnalyzerAddress);
+
+						// Keep this vertex analyzer alive:
+						continue;
 					}
 
 					Analyzers.RemoveAtSwap(i, EAllowShrinking::No);
diff --git a/Engine/Plugins/Runtime/Metasound/Source/MetasoundFrontend/Private/Analysis/MetasoundFrontendGraphAnalyzerView.cpp b/Engine/Plugins/Runtime/Metasound/Source/MetasoundFrontend/Private/Analysis/MetasoundFrontendGraphAnalyzerView.cpp
--- a/Engine/Plugins/Runtime/Metasound/Source/MetasoundFrontend/Private/Analysis/MetasoundFrontendGraphAnalyzerView.cpp
+++ b/Engine/Plugins/Runtime/Metasound/Source/MetasoundFrontend/Private/Analysis/MetasoundFrontendGraphAnalyzerView.cpp
@@ -293,6 +293,19 @@ namespace Metasound
 			SendActiveAnalyzers();
 		}
 
+		bool FMetasoundGraphAnalyzerView::HasAnalyzerInstance(FName InAnalyzerName, const FGuid& InAnalyzerInstanceID) const
+		{
+			for (const FAnalyzerAddress& AnalyzerAddress : ActiveAnalyzers)
+			{
+				if (AnalyzerAddress.AnalyzerName == InAnalyzerName && AnalyzerAddress.AnalyzerInstanceID == InAnalyzerInstanceID)
+				{
+					return true;
+				}
+			}
+
+			return false;
+		}
+
 		const FMetasoundAssetBase& FMetasoundGraphAnalyzerView::GetMetaSoundAssetChecked() const
 		{
 			check(MetaSoundAsset);
diff --git a/Engine/Plugins/Runtime/Metasound/Source/MetasoundFrontend/Private/MetasoundFrontend.cpp b/Engine/Plugins/Runtime/Metasound/Source/MetasoundFrontend/Private/MetasoundFrontend.cpp
--- a/Engine/Plugins/Runtime/Metasound/Source/MetasoundFrontend/Private/MetasoundFrontend.cpp
+++ b/Engine/Plugins/Runtime/Metasound/Source/MetasoundFrontend/Private/MetasoundFrontend.cpp
@@ -18,6 +18,7 @@
 #include "MetasoundRouter.h"
 #include "MetasoundVertex.h"
 #include "Modules/ModuleManager.h"
+#include "NodeTemplates/MetasoundFrontendNodeTemplateAudioAnalyzer.h"
 #include "NodeTemplates/MetasoundFrontendNodeTemplateInput.h"
 #include "NodeTemplates/MetasoundFrontendNodeTemplateReroute.h"
 #include "StructDeserializer.h"
@@ -95,6 +96,7 @@ class FMetasoundFrontendModule : public IModuleInterface
 	{
 		using namespace Metasound::Frontend;
 
+		RegisterNodeTemplate(MakeUnique<FAudioAnalyzerNodeTemplate>());
 		RegisterNodeTemplate(MakeUnique<FInputNodeTemplate>());
 		RegisterNodeTemplate(MakeUnique<FRerouteNodeTemplate>());
 
@@ -109,6 +111,7 @@ class FMetasoundFrontendModule : public IModuleInterface
 	{
 		using namespace Metasound::Frontend;
 
+		UnregisterNodeTemplate(FAudioAnalyzerNodeTemplate::ClassName, FAudioAnalyzerNodeTemplate::VersionNumber);
 		UnregisterNodeTemplate(FInputNodeTemplate::ClassName, FInputNodeTemplate::VersionNumber);
 		UnregisterNodeTemplate(FRerouteNodeTemplate::ClassName, FRerouteNodeTemplate::VersionNumber);
 	}
diff --git a/Engine/Plugins/Runtime/Metasound/Source/MetasoundFrontend/Private/NodeTemplates/MetasoundFrontendNodeTemplateAudioAnalyzer.cpp b/Engine/Plugins/Runtime/Metasound/Source/MetasoundFrontend/Private/NodeTemplates/MetasoundFrontendNodeTemplateAudioAnalyzer.cpp
new file mode 100644
--- /dev/null
+++ b/Engine/Plugins/Runtime/Metasound/Source/MetasoundFrontend/Private/NodeTemplates/MetasoundFrontendNodeTemplateAudioAnalyzer.cpp
@@ -0,0 +1,125 @@
+// Copyright Epic Games, Inc. All Rights Reserved.
+
+#include "NodeTemplates/MetasoundFrontendNodeTemplateAudioAnalyzer.h"
+
+#include "MetasoundFrontendDocumentBuilder.h"
+
+namespace Metasound::Frontend
+{
+	namespace AudioAnalyzerNodeTemplatePrivate
+	{
+		class FAudioAnalyzerNodeTemplateTransform : public INodeTemplateTransform
+		{
+		public:
+			FAudioAnalyzerNodeTemplateTransform() = default;
+			virtual ~FAudioAnalyzerNodeTemplateTransform() = default;
+
+			virtual bool Transform(const FGuid& InNodeID, FMetaSoundFrontendDocumentBuilder& OutBuilder) const override;
+		};
+
+		bool FAudioAnalyzerNodeTemplateTransform::Transform(const FGuid& InNodeID, FMetaSoundFrontendDocumentBuilder& OutBuilder) const
+		{
+			// Strip it out:
+			return OutBuilder.RemoveNode(InNodeID);
+		}
+	} // namespace AudioAnalyzerNodeTemplatePrivate
+
+	const FMetasoundFrontendClassName FAudioAnalyzerNodeTemplate::ClassName{ "UE", "Audio Analyzer", "" };
+
+	const FMetasoundFrontendVersionNumber FAudioAnalyzerNodeTemplate::VersionNumber{ 1, 0 };
+	
+	FMetasoundFrontendNodeInterface FAudioAnalyzerNodeTemplate::GenerateNodeInterface(FNodeTemplateGenerateInterfaceParams InParams) const
+	{
+		static const FName VertexName = "Value";
+		const FName DataType = GetMetasoundDataTypeName<FAudioBuffer>();
+
+		FMetasoundFrontendNodeInterface NewInterface;
+		NewInterface.Inputs.Add(FMetasoundFrontendVertex{ VertexName, DataType, FGuid::NewGuid() });
+		return NewInterface;
+	}
+
+	const FMetasoundFrontendClassName& FAudioAnalyzerNodeTemplate::GetClassName() const
+	{
+		return ClassName;
+	}
+
+	TUniquePtr<INodeTemplateTransform> FAudioAnalyzerNodeTemplate::GenerateNodeTransform() const
+	{
+		return TUniquePtr<INodeTemplateTransform>(new AudioAnalyzerNodeTemplatePrivate::FAudioAnalyzerNodeTemplateTransform());
+	}
+
+	const FMetasoundFrontendClass& FAudioAnalyzerNodeTemplate::GetFrontendClass() const
+	{
+		auto CreateFrontendClass = []()
+		{
+			FMetasoundFrontendClass Class;
+			Class.Metadata.SetClassName(ClassName);
+#if WITH_EDITOR
+			Class.Metadata.SetSerializeText(false);
+			Class.Metadata.SetAuthor(Metasound::PluginAuthor);
+#endif // WITH_EDITOR
+			Class.Metadata.SetType(EMetasoundFrontendClassType::Template);
+			Class.Metadata.SetVersion(VersionNumber);
+
+			return Class;
+		};
+
+		static const FMetasoundFrontendClass FrontendClass = CreateFrontendClass();
+		return FrontendClass;
+	}
+	
+	EMetasoundFrontendVertexAccessType FAudioAnalyzerNodeTemplate::GetNodeInputAccessType(const FMetaSoundFrontendDocumentBuilder& InBuilder, const FGuid& InNodeID, const FGuid& InVertexID) const
+	{
+		return EMetasoundFrontendVertexAccessType::Unset;
+	}
+
+	EMetasoundFrontendVertexAccessType FAudioAnalyzerNodeTemplate::GetNodeOutputAccessType(const FMetaSoundFrontendDocumentBuilder& InBuilder, const FGuid& InNodeID, const FGuid& InVertexID) const
+	{
+		return EMetasoundFrontendVertexAccessType::Unset;
+	}
+
+	const FMetasoundFrontendVersionNumber& FAudioAnalyzerNodeTemplate::GetVersionNumber() const
+	{
+		return VersionNumber;
+	}
+
+	bool FAudioAnalyzerNodeTemplate::IsInputAccessTypeDynamic() const
+	{
+		return false;
+	}
+
+	bool FAudioAnalyzerNodeTemplate::IsInputConnectionUserModifiable() const
+	{
+		return true;
+	}
+
+	bool FAudioAnalyzerNodeTemplate::IsOutputAccessTypeDynamic() const
+	{
+		return false;
+	}
+
+	bool FAudioAnalyzerNodeTemplate::IsOutputConnectionUserModifiable() const
+	{
+		return false;
+	}
+
+	bool FAudioAnalyzerNodeTemplate::IsValidNodeInterface(const FMetasoundFrontendNodeInterface& InNodeInterface) const
+	{
+		if (InNodeInterface.Inputs.Num() != 1)
+		{
+			return false;
+		}
+
+		if (InNodeInterface.Outputs.Num() != 0)
+		{
+			return false;
+		}
+
+		if (InNodeInterface.Inputs.Last().TypeName != GetMetasoundDataTypeName<FAudioBuffer>())
+		{
+			return false;
+		}
+
+		return true;
+	}
+}
diff --git a/Engine/Plugins/Runtime/Metasound/Source/MetasoundFrontend/Public/Analysis/MetasoundFrontendGraphAnalyzerView.h b/Engine/Plugins/Runtime/Metasound/Source/MetasoundFrontend/Public/Analysis/MetasoundFrontendGraphAnalyzerView.h
--- a/Engine/Plugins/Runtime/Metasound/Source/MetasoundFrontend/Public/Analysis/MetasoundFrontendGraphAnalyzerView.h
+++ b/Engine/Plugins/Runtime/Metasound/Source/MetasoundFrontend/Public/Analysis/MetasoundFrontendGraphAnalyzerView.h
@@ -52,6 +52,7 @@ namespace Metasound
 			FGuid AddAnalyzerForSpecifiedOutput(const FGuid& InNodeID, FVertexName InOutputName, FName InAnalyzerName, FName InAnalyzerMemberName);
 			void RemoveAnalyzerForAllSupportedOutputs(FName InAnalyzerName);
 			void RemoveAnalyzerInstance(FName InAnalyzerName, const FGuid& InAnalyzerInstanceID);
+			bool HasAnalyzerInstance(FName InAnalyzerName, const FGuid& InAnalyzerInstanceID) const;
 			TArray<FMetasoundAnalyzerView*> GetAnalyzerViews(FName InAnalyzerName);
 			TArray<const FMetasoundAnalyzerView*> GetAnalyzerViews(FName InAnalyzerName) const;
 			TArray<FMetasoundAnalyzerView*> GetAnalyzerViewsForOutput(const FGuid& InNodeID, FName InOutputName, FName InAnalyzerName);
diff --git a/Engine/Plugins/Runtime/Metasound/Source/MetasoundFrontend/Public/NodeTemplates/MetasoundFrontendNodeTemplateAudioAnalyzer.h b/Engine/Plugins/Runtime/Metasound/Source/MetasoundFrontend/Public/NodeTemplates/MetasoundFrontendNodeTemplateAudioAnalyzer.h
new file mode 100644
--- /dev/null
+++ b/Engine/Plugins/Runtime/Metasound/Source/MetasoundFrontend/Public/NodeTemplates/MetasoundFrontendNodeTemplateAudioAnalyzer.h
@@ -0,0 +1,27 @@
+// Copyright Epic Games, Inc. All Rights Reserved.
+#pragma once
+
+#include "MetasoundFrontendNodeTemplateRegistry.h"
+
+namespace Metasound::Frontend
+{
+	class METASOUNDFRONTEND_API FAudioAnalyzerNodeTemplate : public FNodeTemplateBase
+	{
+	public:
+		static const FMetasoundFrontendClassName ClassName;
+		static const FMetasoundFrontendVersionNumber VersionNumber;
+
+		virtual FMetasoundFrontendNodeInterface GenerateNodeInterface(FNodeTemplateGenerateInterfaceParams InParams) const override;
+		virtual const FMetasoundFrontendClassName& GetClassName() const override;
+		virtual TUniquePtr<INodeTemplateTransform> GenerateNodeTransform() const override;
+		virtual const FMetasoundFrontendClass& GetFrontendClass() const override;
+		virtual EMetasoundFrontendVertexAccessType GetNodeInputAccessType(const FMetaSoundFrontendDocumentBuilder& InBuilder, const FGuid& InNodeID, const FGuid& InVertexID) const override;
+		virtual EMetasoundFrontendVertexAccessType GetNodeOutputAccessType(const FMetaSoundFrontendDocumentBuilder& InBuilder, const FGuid& InNodeID, const FGuid& InVertexID) const override;
+		virtual const FMetasoundFrontendVersionNumber& GetVersionNumber() const override;
+		virtual bool IsInputAccessTypeDynamic() const override;
+		virtual bool IsInputConnectionUserModifiable() const override;
+		virtual bool IsOutputAccessTypeDynamic() const override;
+		virtual bool IsOutputConnectionUserModifiable() const override;
+		virtual bool IsValidNodeInterface(const FMetasoundFrontendNodeInterface& InNodeInterface) const override;
+	};
+} // namespace Metasound::Frontend
