Fix for assert due to using wrong flavor of NodeRegistryKey constructor
#rnx

#tests EngineTest

diff --git a/Engine/Plugins/Runtime/Metasound/Source/MetasoundEngine/Private/MetasoundBuilderSubsystem.cpp b/Engine/Plugins/Runtime/Metasound/Source/MetasoundEngine/Private/MetasoundBuilderSubsystem.cpp
--- a/Engine/Plugins/Runtime/Metasound/Source/MetasoundEngine/Private/MetasoundBuilderSubsystem.cpp
+++ b/Engine/Plugins/Runtime/Metasound/Source/MetasoundEngine/Private/MetasoundBuilderSubsystem.cpp
@@ -78,7 +78,7 @@ void UMetaSoundPatchBuilder::OnAssetReferenceAdded(TScriptInterface<IMetaSoundDo
 
 	check(DocInterface.GetObject());
 	UMetaSoundPatch& Patch = Builder.CastDocumentObjectChecked<UMetaSoundPatch>();
-	FNodeRegistryKey RegistryKey(Patch.GetConstDocument().RootGraph.Metadata);
+	FNodeRegistryKey RegistryKey(Patch.GetConstDocument().RootGraph);
 	Patch.ReferencedAssetClassObjects.Add(DocInterface.GetObject());
 	Patch.ReferencedAssetClassKeys.Add(RegistryKey.ToString());
 }
@@ -89,7 +89,7 @@ void UMetaSoundPatchBuilder::OnRemovingAssetReference(TScriptInterface<IMetaSoun
 
 	check(DocInterface.GetObject());
 	UMetaSoundPatch& Patch = Builder.CastDocumentObjectChecked<UMetaSoundPatch>();
-	FNodeRegistryKey RegistryKey(Patch.GetConstDocument().RootGraph.Metadata);
+	FNodeRegistryKey RegistryKey(Patch.GetConstDocument().RootGraph);
 	Patch.ReferencedAssetClassObjects.Remove(DocInterface.GetObject());
 	Patch.ReferencedAssetClassKeys.Remove(RegistryKey.ToString());
 }
@@ -276,7 +276,7 @@ void UMetaSoundSourceBuilder::OnAssetReferenceAdded(TScriptInterface<IMetaSoundD
 
 	check(DocInterface.GetObject());
 	UMetaSoundSource& Source = GetMetaSoundSource();
-	FNodeRegistryKey RegistryKey(Source.GetConstDocument().RootGraph.Metadata);
+	FNodeRegistryKey RegistryKey(Source.GetConstDocument().RootGraph);
 	Source.ReferencedAssetClassObjects.Add(DocInterface.GetObject());
 	Source.ReferencedAssetClassKeys.Add(RegistryKey.ToString());
 }
@@ -548,7 +548,7 @@ void UMetaSoundSourceBuilder::OnRemovingAssetReference(TScriptInterface<IMetaSou
 
 	check(DocInterface.GetObject());
 	UMetaSoundSource& Source = GetMetaSoundSource();
-	FNodeRegistryKey RegistryKey(Source.GetConstDocument().RootGraph.Metadata);
+	FNodeRegistryKey RegistryKey(Source.GetConstDocument().RootGraph);
 	Source.ReferencedAssetClassObjects.Remove(DocInterface.GetObject());
 	Source.ReferencedAssetClassKeys.Remove(RegistryKey.ToString());
 }
