[Metasound] Re-bind vertex analyzers when their output data ref changes (e.g. dynamic graphs re-connecting things)
#tests cooked client/server

diff --git a/Engine/Plugins/Runtime/Metasound/Source/MetasoundFrontend/Public/Analysis/MetasoundFrontendVertexAnalyzer.h b/Engine/Plugins/Runtime/Metasound/Source/MetasoundFrontend/Public/Analysis/MetasoundFrontendVertexAnalyzer.h
--- a/Engine/Plugins/Runtime/Metasound/Source/MetasoundFrontend/Public/Analysis/MetasoundFrontendVertexAnalyzer.h
+++ b/Engine/Plugins/Runtime/Metasound/Source/MetasoundFrontend/Public/Analysis/MetasoundFrontendVertexAnalyzer.h
@@ -48,6 +48,12 @@ namespace Metasound
 			// Returns analyzer address
 			virtual const FAnalyzerAddress& GetAnalyzerAddress() const = 0;
 
+			// Returns the data reference id for the analyzer's input
+			virtual FDataReferenceID GetDataReferenceId() const = 0;
+
+			// Re-bind the analyzer's input
+			virtual void SetDataReference(const FAnyDataReference& NewDataRef) = 0;
+
 			// Executes analysis
 			virtual void Execute() = 0;
 
@@ -141,6 +147,16 @@ namespace Metasound
 				return AnalyzerAddress;
 			}
 
+			virtual FDataReferenceID GetDataReferenceId() const override final
+			{
+				return Metasound::GetDataReferenceID(VertexDataReference);
+			}
+
+			virtual void SetDataReference(const FAnyDataReference& NewDataRef) override final
+			{
+				VertexDataReference = NewDataRef;
+			}
+
 		protected:
 			using FBoundOutputDataPtr = TSharedPtr<IBoundAnalyzerOutput>;
 
diff --git a/Engine/Plugins/Runtime/Metasound/Source/MetasoundGenerator/Private/MetasoundGenerator.cpp b/Engine/Plugins/Runtime/Metasound/Source/MetasoundGenerator/Private/MetasoundGenerator.cpp
--- a/Engine/Plugins/Runtime/Metasound/Source/MetasoundGenerator/Private/MetasoundGenerator.cpp
+++ b/Engine/Plugins/Runtime/Metasound/Source/MetasoundGenerator/Private/MetasoundGenerator.cpp
@@ -505,18 +505,25 @@ PRAGMA_ENABLE_DEPRECATION_WARNINGS
 	{
 		OutputAnalyzerModificationQueue.Enqueue([this, AnalyzerAddress]
 		{
-			if (OutputAnalyzers.ContainsByPredicate([AnalyzerAddress](const TUniquePtr<Frontend::IVertexAnalyzer>& Analyzer)
-			{
-				return AnalyzerAddressesReferToSameGeneratorOutput(AnalyzerAddress, Analyzer->GetAnalyzerAddress());
-			}))
+			const FAnyDataReference* OutputReference = VertexInterfaceData.GetOutputs().FindDataReference(AnalyzerAddress.OutputName);
+
+			if (nullptr == OutputReference)
 			{
 				return;
 			}
-			
-			const FAnyDataReference* OutputReference = VertexInterfaceData.GetOutputs().FindDataReference(AnalyzerAddress.OutputName);
-			
-			if (nullptr == OutputReference)
+
+			TUniquePtr<Frontend::IVertexAnalyzer>* ExistingAnalyzer = OutputAnalyzers.FindByPredicate([AnalyzerAddress](const TUniquePtr<Frontend::IVertexAnalyzer>& Analyzer)
+				{
+					return AnalyzerAddressesReferToSameGeneratorOutput(AnalyzerAddress, Analyzer->GetAnalyzerAddress());
+				});
+
+			if (ExistingAnalyzer && ExistingAnalyzer->IsValid())
 			{
+				if (ExistingAnalyzer->Get()->GetDataReferenceId() != GetDataReferenceID(*OutputReference))
+				{
+					ExistingAnalyzer->Get()->SetDataReference(*OutputReference);
+				}
+
 				return;
 			}
 			
