Fix for ensures due to invalid attempt to reconnect undesired vertices when removing a connected variable getter
#tests AudioQA Tests & MetaSound Asset Editor
#rnx

diff --git a/Engine/Plugins/Runtime/Metasound/Source/MetasoundFrontend/Private/MetasoundFrontendDocumentBuilder.cpp b/Engine/Plugins/Runtime/Metasound/Source/MetasoundFrontend/Private/MetasoundFrontendDocumentBuilder.cpp
--- a/Engine/Plugins/Runtime/Metasound/Source/MetasoundFrontend/Private/MetasoundFrontendDocumentBuilder.cpp
+++ b/Engine/Plugins/Runtime/Metasound/Source/MetasoundFrontend/Private/MetasoundFrontendDocumentBuilder.cpp
@@ -2271,15 +2271,6 @@ bool FMetaSoundFrontendDocumentBuilder::IsPreset() const
 	return GetConstDocument().RootGraph.PresetOptions.bIsPreset;
 }
 
-bool FMetaSoundFrontendDocumentBuilder::IsVariableClass(EMetasoundFrontendClassType ClassType) const
-{
-	const bool bIsVariableNode = (ClassType == EMetasoundFrontendClassType::Variable)
-		|| (ClassType == EMetasoundFrontendClassType::VariableAccessor)
-		|| (ClassType == EMetasoundFrontendClassType::VariableDeferredAccessor)
-		|| (ClassType == EMetasoundFrontendClassType::VariableMutator);
-	return bIsVariableNode;
-}
-
 Metasound::Frontend::EInvalidEdgeReason FMetaSoundFrontendDocumentBuilder::IsValidEdge(const FMetasoundFrontendEdge& InEdge) const
 {
 	using namespace Metasound::Frontend;
@@ -2523,10 +2514,14 @@ bool FMetaSoundFrontendDocumentBuilder::RemoveEdge(const FMetasoundFrontendEdge&
 	if (const int32* IndexPtr = EdgeCache.FindEdgeIndexToNodeInput(EdgeToRemove.ToNodeID, EdgeToRemove.ToVertexID))
 	{
 		const int32 Index = *IndexPtr;
-		const int32 LastIndex = Edges.Num() - 1;
-		DocumentDelegates->EdgeDelegates.OnRemoveSwappingEdge.Broadcast(Index, LastIndex);
-		Edges.RemoveAtSwap(Index, EAllowShrinking::No);
-		return true;
+		FMetasoundFrontendEdge& FoundEdge = Edges[Index];
+		if (EdgeToRemove.FromNodeID == FoundEdge.FromNodeID && EdgeToRemove.FromVertexID == FoundEdge.FromVertexID)
+		{
+			const int32 LastIndex = Edges.Num() - 1;
+			DocumentDelegates->EdgeDelegates.OnRemoveSwappingEdge.Broadcast(Index, LastIndex);
+			Edges.RemoveAtSwap(Index, EAllowShrinking::No);
+			return true;
+		}
 	}
 
 	return false;
@@ -2645,6 +2640,33 @@ bool FMetaSoundFrontendDocumentBuilder::RemoveNodeInputDefault(const FGuid& InNo
 	return false;
 }
 
+bool FMetaSoundFrontendDocumentBuilder::RemoveEdges(const FGuid& InNodeID)
+{
+	using namespace Metasound::Frontend;
+	using namespace Metasound::VariableNames;
+
+	const IDocumentGraphNodeCache& NodeCache = DocumentCache->GetNodeCache();
+	if (const FMetasoundFrontendNode* Node = NodeCache.FindNode(InNodeID))
+	{
+		const IDocumentGraphEdgeCache& EdgeCache = DocumentCache->GetEdgeCache();
+
+		for (const FMetasoundFrontendVertex& Vertex : Node->Interface.Inputs)
+		{
+			RemoveEdgeToNodeInput(InNodeID, Vertex.VertexID);
+		}
+
+		TArray<FMetasoundFrontendVertexHandle> ToVertexHandles;
+		for (const FMetasoundFrontendVertex& Vertex : Node->Interface.Outputs)
+		{
+			RemoveEdgesFromNodeOutput(InNodeID, Vertex.VertexID);
+		}
+
+		return true;
+	}
+
+	return false;
+}
+
 bool FMetaSoundFrontendDocumentBuilder::RemoveEdgesByNodeClassInterfaceBindings(const FGuid& InFromNodeID, const FGuid& InToNodeID)
 {
 	using namespace Metasound;
@@ -2934,7 +2956,6 @@ bool FMetaSoundFrontendDocumentBuilder::RemoveNode(const FGuid& InNodeID)
 	using namespace Metasound::Frontend;
 
 	const IDocumentGraphNodeCache& NodeCache = DocumentCache->GetNodeCache();
-	const IDocumentGraphEdgeCache& EdgeCache = DocumentCache->GetEdgeCache();
 
 	if (const int32* IndexPtr = NodeCache.FindNodeIndex(InNodeID))
 	{
@@ -2948,53 +2969,21 @@ bool FMetaSoundFrontendDocumentBuilder::RemoveNode(const FGuid& InNodeID)
 
 		const FMetasoundFrontendClass* NodeClass = DocumentCache->FindDependency(Node.ClassID);
 		check(NodeClass);
-
-		const bool bIsVariableNode = IsVariableClass(NodeClass->Metadata.GetType());
-		if (bIsVariableNode)
-		{
-			const bool bUnlinkedVariable = UnlinkVariableNode(Node.GetID());
-			ensureMsgf(bUnlinkedVariable, TEXT("Removed variable node that was not linked to variable"));
-		}
-
-		FMetasoundFrontendVertexHandle FromVertexHandle;
-		for (const FMetasoundFrontendVertex& Vertex : Node.Interface.Inputs)
+		const EMetasoundFrontendClassType ClassType = NodeClass->Metadata.GetType();
+		switch (ClassType)
 		{
-			if (bIsVariableNode)
+			case EMetasoundFrontendClassType::Variable:
+			case EMetasoundFrontendClassType::VariableDeferredAccessor:
+			case EMetasoundFrontendClassType::VariableAccessor:
+			case EMetasoundFrontendClassType::VariableMutator:
 			{
-				if (const int32* InputEdgeIndex = EdgeCache.FindEdgeIndexToNodeInput(InNodeID, Vertex.VertexID))
-				{
-					FromVertexHandle = Graph.Edges[*InputEdgeIndex].GetFromVertexHandle();
-				}
-			}
-
-			RemoveEdgeToNodeInput(InNodeID, Vertex.VertexID);
-		}
-
-		TArray<FMetasoundFrontendVertexHandle> ToVertexHandles;
-		for (const FMetasoundFrontendVertex& Vertex : Node.Interface.Outputs)
-		{
-			if (bIsVariableNode)
-			{
-				ToVertexHandles.Reset();
-				const TArrayView<const int32> OutputEdgeIndices = EdgeCache.FindEdgeIndicesFromNodeOutput(InNodeID, Vertex.VertexID);
-				Algo::Transform(OutputEdgeIndices, ToVertexHandles, [&Graph](const int32& VertIndex) { return Graph.Edges[VertIndex].GetToVertexHandle(); });
-			}
-			RemoveEdgesFromNodeOutput(InNodeID, Vertex.VertexID);
-			if (bIsVariableNode)
-			{
-				for (const FMetasoundFrontendVertexHandle& ToHandle : ToVertexHandles)
-				{
-					AddEdge(FMetasoundFrontendEdge
-					{
-						FromVertexHandle.NodeID,
-						FromVertexHandle.VertexID,
-						ToHandle.NodeID,
-						ToHandle.VertexID
-					});
-				}
+				const bool bVariableNodeUnlinked = UnlinkVariableNode(NodeID);
+				ensureAlwaysMsgf(bVariableNodeUnlinked, TEXT("Failed to unlink %s node with ID '%s"), LexToString(ClassType), *InNodeID.ToString());
 			}
+			break;
 		}
 
+		RemoveEdges(NodeID);
 		const int32 LastIndex = Nodes.Num() - 1;
 		DocumentDelegates->NodeDelegates.OnRemoveSwappingNode.Broadcast(Index, LastIndex);
 		Nodes.RemoveAtSwap(Index, EAllowShrinking::No);
@@ -3462,6 +3451,55 @@ void FMetaSoundFrontendDocumentBuilder::SetVersionNumber(const FMetasoundFronten
 	GetDocument().Metadata.Version.Number = InDocumentVersionNumber;
 }
 
+bool FMetaSoundFrontendDocumentBuilder::SpliceVariableNodeFromStack(const FGuid& InNodeID)
+{
+	using namespace Metasound;
+	using namespace Metasound::Frontend;
+
+	FMetasoundFrontendGraphClass& GraphClass = GetDocument().RootGraph;
+	FMetasoundFrontendGraph& Graph = GraphClass.Graph;
+
+	const IDocumentGraphEdgeCache& EdgeCache = DocumentCache->GetEdgeCache();
+
+	FMetasoundFrontendVertexHandle FromVariableVertexHandle;
+	{
+		const FMetasoundFrontendVertex* InputVertex = FindNodeInput(InNodeID, VariableNames::InputVariableName);
+		check(InputVertex);
+		if (const int32* InputEdgeIndex = EdgeCache.FindEdgeIndexToNodeInput(InNodeID, InputVertex->VertexID))
+		{
+			FromVariableVertexHandle = Graph.Edges[*InputEdgeIndex].GetFromVertexHandle();
+			const bool bRemovedEdge = RemoveEdgeToNodeInput(InNodeID, InputVertex->VertexID);
+			check(bRemovedEdge);
+		}
+	}
+
+	if (FromVariableVertexHandle.IsSet())
+	{
+		if (const FMetasoundFrontendVertex* OutputVertex = FindNodeOutput(InNodeID, VariableNames::OutputVariableName))
+		{
+			TArray<FMetasoundFrontendVertexHandle> ToVertexHandles;
+			const TArrayView<const int32> OutputEdgeIndices = EdgeCache.FindEdgeIndicesFromNodeOutput(InNodeID, OutputVertex->VertexID);
+			Algo::Transform(OutputEdgeIndices, ToVertexHandles, [&Graph](const int32& VertIndex) { return Graph.Edges[VertIndex].GetToVertexHandle(); });
+
+			RemoveEdgesFromNodeOutput(InNodeID, OutputVertex->VertexID);
+
+			for (const FMetasoundFrontendVertexHandle& ToHandle : ToVertexHandles)
+			{
+				AddEdge(FMetasoundFrontendEdge
+				{
+					FromVariableVertexHandle.NodeID,
+					FromVariableVertexHandle.VertexID,
+					ToHandle.NodeID,
+					ToHandle.VertexID
+				});
+			}
+			return true;
+		}
+	}
+
+	return false;
+}
+
 bool FMetaSoundFrontendDocumentBuilder::SwapGraphInput(const FMetasoundFrontendClassVertex& InExistingInputVertex, const FMetasoundFrontendClassVertex& InNewInputVertex)
 {
 	using namespace Metasound::Frontend;
@@ -3619,30 +3657,37 @@ bool FMetaSoundFrontendDocumentBuilder::SwapGraphOutput(const FMetasoundFrontend
 
 bool FMetaSoundFrontendDocumentBuilder::UnlinkVariableNode(const FGuid& InNodeID)
 {
-	using namespace Metasound::VariableNames;
+	auto IsNodeID = [&InNodeID](const FGuid& TestID) { return TestID == InNodeID; };
 
 	FMetasoundFrontendGraphClass& GraphClass = GetDocument().RootGraph;
-	for (FMetasoundFrontendVariable& Variable : GraphClass.Graph.Variables)
+	FMetasoundFrontendGraph& Graph = GraphClass.Graph;
+	for (FMetasoundFrontendVariable& Variable : Graph.Variables)
 	{
-		if (InNodeID == Variable.VariableNodeID)
+		if (Variable.MutatorNodeID == InNodeID)
 		{
-			Variable.VariableNodeID = FGuid();
+			Variable.MutatorNodeID = FGuid();
+			SpliceVariableNodeFromStack(InNodeID);
 			return true;
 		}
 
-		if (InNodeID == Variable.MutatorNodeID)
+		if (Variable.VariableNodeID == InNodeID)
 		{
-			Variable.MutatorNodeID = FGuid();
+			Variable.VariableNodeID = FGuid();
+			SpliceVariableNodeFromStack(InNodeID);
 			return true;
 		}
 
-		if (Variable.AccessorNodeIDs.Remove(InNodeID) > 0)
+		const bool bRemovedDeferredNode = Variable.DeferredAccessorNodeIDs.RemoveAllSwap(IsNodeID, EAllowShrinking::No) > 0;
+		if (bRemovedDeferredNode)
 		{
+			SpliceVariableNodeFromStack(InNodeID);
 			return true;
 		}
 
-		if (Variable.DeferredAccessorNodeIDs.Remove(InNodeID) > 0)
+		const bool bRemovedAccessorNode = Variable.AccessorNodeIDs.RemoveAllSwap(IsNodeID, EAllowShrinking::No) > 0;
+		if (bRemovedAccessorNode)
 		{
+			SpliceVariableNodeFromStack(InNodeID);
 			return true;
 		}
 	}
diff --git a/Engine/Plugins/Runtime/Metasound/Source/MetasoundFrontend/Public/MetasoundFrontendDocumentBuilder.h b/Engine/Plugins/Runtime/Metasound/Source/MetasoundFrontend/Public/MetasoundFrontendDocumentBuilder.h
--- a/Engine/Plugins/Runtime/Metasound/Source/MetasoundFrontend/Public/MetasoundFrontendDocumentBuilder.h
+++ b/Engine/Plugins/Runtime/Metasound/Source/MetasoundFrontend/Public/MetasoundFrontendDocumentBuilder.h
@@ -276,7 +276,6 @@ public:
 
 	// Returns whether or not builder is attached to a DocumentInterface and is valid to build or act on a document.
 	bool IsValid() const;
-	bool IsVariableClass(EMetasoundFrontendClassType ClassType) const;
 
 	// Returns whether or not the given edge is valid (i.e. represents an input and output that equate in data and access types) or malformed.
 	// Note that this does not return whether or not the given edge exists, but rather if it could be legally applied to the given edge vertices.
@@ -292,6 +291,10 @@ public:
 	bool RemoveDependency(const FGuid& InClassID);
 	bool RemoveDependency(EMetasoundFrontendClassType ClassType, const FMetasoundFrontendClassName& InClassName, const FMetasoundFrontendVersionNumber& InClassVersionNumber);
 	bool RemoveEdge(const FMetasoundFrontendEdge& EdgeToRemove);
+
+	// Removes all edges connected to an input or output vertex associated with the node of the given ID.
+	bool RemoveEdges(const FGuid& InNodeID);
+
 	bool RemoveEdgesByNodeClassInterfaceBindings(const FGuid& InOutputNodeID, const FGuid& InInputNodeID);
 	bool RemoveEdgesFromNodeOutput(const FGuid& InNodeID, const FGuid& InVertexID);
 	bool RemoveEdgeToNodeInput(const FGuid& InNodeID, const FGuid& InVertexID);
@@ -367,6 +370,7 @@ private:
 	using FFinalizeNodeFunctionRef = TFunctionRef<void(FMetasoundFrontendNode&, const Metasound::Frontend::FNodeRegistryKey&)>;
 
 	FMetasoundFrontendNode* AddNodeInternal(const FMetasoundFrontendClassMetadata& InClassMetadata, Metasound::Frontend::FFinalizeNodeFunctionRef FinalizeNode, FGuid InNodeID = FGuid::NewGuid(), int32* NewNodeIndex = nullptr);
+	void BeginBuilding(TSharedPtr<Metasound::Frontend::FDocumentModifyDelegates> Delegates = {}, bool bPrimeCache = false);
 
 	// Conforms GraphOutput node's ClassID, Access & Data Type with the GraphOutput.
 	// creating and removing dependencies as necessary within the document dependency array. Does *NOT*
@@ -389,9 +393,9 @@ private:
 	FMetasoundFrontendDocument& GetDocument();
 
 	bool SetGraphInputInheritsDefault(FName InName, bool bInputInheritsDefault);
-	bool UnlinkVariableNode(const FGuid& InNodeID);
 
-	void BeginBuilding(TSharedPtr<Metasound::Frontend::FDocumentModifyDelegates> Delegates = {}, bool bPrimeCache = false);
+	bool SpliceVariableNodeFromStack(const FGuid& InNodeID);
+	bool UnlinkVariableNode(const FGuid& InNodeID);
 
 	UPROPERTY(Transient)
 	TScriptInterface<IMetaSoundDocumentInterface> DocumentInterface;
