More MetaSound Asset/Builder registry clean-up
- Distinguish executing finish call from removal of register's weak pointer entry to avoid leaking stale builder registry entries
- Fix for removal calls to asset registry entries leaking during cook (due to cook package data being stripped when BeginDestroy is called)
- Fix for removing an asset manager entry that hasn't been registered with the node registry not getting removed due to the RegistryGraphKey referenced not being set on AssetBase
- Remove AssetRegistry MetaSoundEditor delegate handlers during cook that aren't used
- Add thread checks for safety to MetaSoundAssetManager
- Add better logging around AssetManager/BuilderRegistry destruction if applied prior to underlying respective entry/entries destruction

#rnx
#tests local AudioQA cook, loading all assets in editor, PIE, -game, automated tests, loading all assets in editor, PIE, -game

diff --git a/Engine/Plugins/Runtime/Metasound/Source/MetasoundEditor/Private/MetasoundEditorModule.cpp b/Engine/Plugins/Runtime/Metasound/Source/MetasoundEditor/Private/MetasoundEditorModule.cpp
--- a/Engine/Plugins/Runtime/Metasound/Source/MetasoundEditor/Private/MetasoundEditorModule.cpp
+++ b/Engine/Plugins/Runtime/Metasound/Source/MetasoundEditor/Private/MetasoundEditorModule.cpp
@@ -412,6 +412,11 @@ namespace Metasound
 
 			void OnAssetScanFinished()
 			{
+				if (IsRunningCookCommandlet())
+				{
+					return;
+				}
+
 				AssetScanStatus = EAssetScanStatus::Complete;
 
 				if (AssetPrimeStatus == EAssetPrimeStatus::Requested)
diff --git a/Engine/Plugins/Runtime/Metasound/Source/MetasoundEngine/Private/MetasoundAssetSubsystem.cpp b/Engine/Plugins/Runtime/Metasound/Source/MetasoundEngine/Private/MetasoundAssetSubsystem.cpp
--- a/Engine/Plugins/Runtime/Metasound/Source/MetasoundEngine/Private/MetasoundAssetSubsystem.cpp
+++ b/Engine/Plugins/Runtime/Metasound/Source/MetasoundEngine/Private/MetasoundAssetSubsystem.cpp
@@ -85,22 +85,25 @@ namespace Metasound::Engine
 			return bSuccess;
 		}
 
-		// Remove the Map entry only if the key and value are equal.
-		//
-		// This protects against scenarios where a metasound is renamed or moved 
-		// and the new entry was being erroneously removed from the PathMap
 		bool RemovePath(TMap<Frontend::FAssetKey, TArray<FTopLevelAssetPath>>& InMap, const Frontend::FAssetKey& AssetKey, const FTopLevelAssetPath& AssetPath)
 		{
-			if (AssetPath.IsValid())
+			check(IsInGameThread());
+			if (TArray<FTopLevelAssetPath>* MapAssetPaths = InMap.Find(AssetKey))
 			{
-				if (TArray<FTopLevelAssetPath>* MapAssetPaths = InMap.Find(AssetKey))
+				auto ComparePaths = [&AssetPath](const FTopLevelAssetPath& Path)
+				{
+					// Cook can strip package names on destruction, so only asset name is reliable
+					return IsRunningCookCommandlet()
+						? Path.GetAssetName() == AssetPath.GetAssetName()
+						: Path == AssetPath;
+				};
+				if (MapAssetPaths->RemoveAllSwap(ComparePaths, EAllowShrinking::No) > 0)
 				{
-					const int32 NumRemoved = MapAssetPaths->RemoveAllSwap([&AssetPath](const FTopLevelAssetPath& Path) { return Path == AssetPath; }, EAllowShrinking::No);
 					if (MapAssetPaths->IsEmpty())
 					{
 						InMap.Remove(AssetKey);
 					}
-					return NumRemoved > 0;
+					return true;
 				}
 			}
 
@@ -109,6 +112,7 @@ namespace Metasound::Engine
 
 		void AddPath(TMap<Frontend::FAssetKey, TArray<FTopLevelAssetPath>>& InMap, const Frontend::FAssetKey& AssetKey, const FTopLevelAssetPath& AssetPath)
 		{
+			check(IsInGameThread());
 			TArray<FTopLevelAssetPath>& Paths = InMap.FindOrAdd(AssetKey);
 			Paths.AddUnique(AssetPath);
 #if !NO_LOGGING
@@ -197,10 +201,24 @@ namespace Metasound::Engine
 
 	FMetaSoundAssetManager::~FMetaSoundAssetManager()
 	{
+#if !NO_LOGGING
 		if (!PathMap.IsEmpty())
 		{
-			UE_LOG(LogMetaSound, Display, TEXT("Destroying MetaSoundAssetManager while %i asset keys are still registered."), PathMap.Num());
+			TSet<FAssetKey> Keys;
+			if (int32 NumKeys = PathMap.GetKeys(Keys); NumKeys > 0)
+			{
+				UE_LOG(LogMetaSound, Warning, TEXT("AssetManager is shutting down with the following '%i' assets active:"), NumKeys);
+				for (const TPair<FAssetKey, TArray<FTopLevelAssetPath>>& Pair : PathMap)
+				{
+					for (const FTopLevelAssetPath& Path : Pair.Value)
+					{
+						UE_LOG(LogMetaSound, Warning, TEXT("- %s"), *Path.ToString());
+					}
+				}
+			}
 		}
+#endif // !NO_LOGGING
+
 	}
 
 	void FMetaSoundAssetManager::AddReferencedObjects(FReferenceCollector& Collector)
@@ -359,6 +377,7 @@ namespace Metasound::Engine
 
 	bool FMetaSoundAssetManager::ContainsKey(const Metasound::Frontend::FAssetKey& InKey) const
 	{
+		check(IsInGameThread());
 		return PathMap.Contains(InKey);
 	}
 
@@ -410,6 +429,7 @@ namespace Metasound::Engine
 
 	const FTopLevelAssetPath* FMetaSoundAssetManager::FindAssetPath(const Metasound::Frontend::FAssetKey& InKey) const
 	{
+		check(IsInGameThread());
 		if (const TArray<FTopLevelAssetPath>* Paths = PathMap.Find(InKey))
 		{
 			if (!Paths->IsEmpty())
@@ -604,22 +624,20 @@ namespace Metasound::Engine
 	{
 		using namespace Frontend;
 
-		if (const FMetasoundAssetBase* MetaSoundAsset = Metasound::IMetasoundUObjectRegistry::Get().GetObjectAsAssetBase(&InObject))
-		{
-			const FMetasoundFrontendDocument& Document = MetaSoundAsset->GetConstDocumentChecked();
-			const FMetasoundFrontendClassMetadata& Metadata = Document.RootGraph.Metadata;
-
-			const FAssetKey AssetKey(Metadata.GetClassName(), Metadata.GetVersion());
-			const FGraphRegistryKey GraphKey = MetaSoundAsset->GetGraphRegistryKey();
-
-			if (IDocumentBuilderRegistry* BuilderRegistry = IDocumentBuilderRegistry::Get())
-			{
-				constexpr bool bForceUnregister = true;
-				BuilderRegistry->FinishBuilding(Metadata.GetClassName(), bForceUnregister);
-			}
+		TScriptInterface<const IMetaSoundDocumentInterface> DocInterface(&InObject);
+		check(DocInterface.GetObject());
+		const FMetasoundFrontendDocument& Document = DocInterface->GetConstDocument();
+		const FMetasoundFrontendClassMetadata& Metadata = Document.RootGraph.Metadata;
 
-			AssetSubsystemPrivate::RemovePath(PathMap, AssetKey, GraphKey.AssetPath);
+		if (IDocumentBuilderRegistry* BuilderRegistry = IDocumentBuilderRegistry::Get())
+		{
+			constexpr bool bForceUnregister = true;
+			BuilderRegistry->FinishBuilding(Metadata.GetClassName(), bForceUnregister);
 		}
+
+		const FAssetKey AssetKey(Metadata.GetClassName(), Metadata.GetVersion());
+		const FTopLevelAssetPath AssetPath(&InObject);
+		AssetSubsystemPrivate::RemovePath(PathMap, AssetKey, AssetPath);
 	}
 
 	void FMetaSoundAssetManager::RemoveAsset(const FAssetData& InAssetData)
@@ -677,7 +695,7 @@ namespace Metasound::Engine
 			{
 				if (AssetKey.IsValid())
 				{
-					PathMap.FindOrAdd(AssetKey).AddUnique(ClassInfo.AssetPath);
+					AssetSubsystemPrivate::AddPath(PathMap, AssetKey, ClassInfo.AssetPath);
 				}
 			}
 		}
diff --git a/Engine/Plugins/Runtime/Metasound/Source/MetasoundEngine/Private/MetasoundDocumentBuilderRegistry.cpp b/Engine/Plugins/Runtime/Metasound/Source/MetasoundEngine/Private/MetasoundDocumentBuilderRegistry.cpp
--- a/Engine/Plugins/Runtime/Metasound/Source/MetasoundEngine/Private/MetasoundDocumentBuilderRegistry.cpp
+++ b/Engine/Plugins/Runtime/Metasound/Source/MetasoundEngine/Private/MetasoundDocumentBuilderRegistry.cpp
@@ -10,26 +10,32 @@ namespace Metasound::Engine
 {
 	FDocumentBuilderRegistry::~FDocumentBuilderRegistry()
 	{
-		for(const TPair<FMetasoundFrontendClassName, TWeakObjectPtr<UMetaSoundBuilderBase>>& Pair : Builders)
+#if !NO_LOGGING
+		if (!Builders.IsEmpty())
 		{
-			TWeakObjectPtr<UMetaSoundBuilderBase> Builder = Pair.Value;
-			if (Builder.IsValid())
+			TArray<TPair<FMetasoundFrontendClassName, TWeakObjectPtr<UMetaSoundBuilderBase>>> PairsToFinish;
+			int32 NumStale = 0;
+			for (const TPair<FMetasoundFrontendClassName, TWeakObjectPtr<UMetaSoundBuilderBase>>& Pair : Builders)
 			{
-				// If the builder has applied transactions to its document object that are not mirrored in the frontend registry,
-				// unregister version in registry. This will ensure that future requests for the builder's associated asset will
-				// register a fresh version from the object as the transaction history is intrinsically lost once this builder
-				// is destroyed.
-				if (Builder->GetLastTransactionRegistered() != Builder->GetConstBuilder().GetTransactionCount())
+				if (Pair.Value.IsValid())
 				{
-					UObject& MetaSound = Builder->GetConstBuilder().CastDocumentObjectChecked<UObject>();
-					if (FMetasoundAssetBase* MetaSoundAsset = IMetasoundUObjectRegistry::Get().GetObjectAsAssetBase(&MetaSound))
-					{
-						MetaSoundAsset->UnregisterGraphWithFrontend();
-					}
+					PairsToFinish.Add(Pair);
+				}
+				else
+				{
+					NumStale++;
 				}
 			}
+
+			UE_LOG(LogMetaSound, Warning, TEXT("BuilderRegistry is shutting down with existing records: %i stale entries never removed, the following %i are still active and will be forcefully shutdown:"), NumStale, PairsToFinish.Num());
+			for (const TPair<FMetasoundFrontendClassName, TWeakObjectPtr<UMetaSoundBuilderBase>>& Pair : PairsToFinish)
+			{
+				constexpr bool bForceUnregister = true;
+				FinishBuilding(Pair.Key, bForceUnregister);
+				UE_LOG(LogMetaSound, Warning, TEXT("- %s"), *Pair.Value->GetFullName());
+			}
 		}
-		Builders.Reset();
+#endif // !NO_LOGGING
 	}
 
 #if WITH_EDITORONLY_DATA
@@ -85,33 +91,37 @@ namespace Metasound::Engine
 		using namespace Metasound;
 		using namespace Metasound::Engine;
 
-		TWeakObjectPtr<UMetaSoundBuilderBase> Builder = Builders.FindRef(InClassName);
-		if (Builder.IsValid())
+		if (TWeakObjectPtr<UMetaSoundBuilderBase>* BuilderPtr = Builders.Find(InClassName))
 		{
-			// If the builder has applied transactions to its document object that are not mirrored in the frontend registry,
-			// unregister version in registry. This will ensure that future requests for the builder's associated asset will
-			// register a fresh version from the object as the transaction history is intrinsically lost once this builder
-			// is destroyed. It is also possible that the DocBuilder's underlying object can be invalid if object was force
-			// deleted, so validity check is necessary.
-			FMetaSoundFrontendDocumentBuilder& DocBuilder = Builder->GetBuilder();
-			if (DocBuilder.IsValid())
+			if (UMetaSoundBuilderBase* Builder = BuilderPtr->Get())
 			{
-				if (!IsRunningCookCommandlet())
+				// If the builder has applied transactions to its document object that are not mirrored in the frontend registry,
+				// unregister version in registry. This will ensure that future requests for the builder's associated asset will
+				// register a fresh version from the object as the transaction history is intrinsically lost once this builder
+				// is destroyed. It is also possible that the DocBuilder's underlying object can be invalid if object was force
+				// deleted, so validity check is necessary.
+				FMetaSoundFrontendDocumentBuilder& DocBuilder = Builder->GetBuilder();
+				if (DocBuilder.IsValid())
 				{
-					const int32 TransactionCount = DocBuilder.GetTransactionCount();
-					const int32 LastTransactionRegistered = Builder->GetLastTransactionRegistered();
-					if (bForceUnregister || LastTransactionRegistered != TransactionCount)
+					if (!IsRunningCookCommandlet())
 					{
-						UObject& MetaSound = DocBuilder.CastDocumentObjectChecked<UObject>();
-						if (FMetasoundAssetBase* MetaSoundAsset = IMetasoundUObjectRegistry::Get().GetObjectAsAssetBase(&MetaSound))
+						const int32 TransactionCount = DocBuilder.GetTransactionCount();
+						const int32 LastTransactionRegistered = Builder->GetLastTransactionRegistered();
+						if (bForceUnregister || LastTransactionRegistered != TransactionCount)
 						{
-							MetaSoundAsset->UnregisterGraphWithFrontend();
+							UObject& MetaSound = DocBuilder.CastDocumentObjectChecked<UObject>();
+							if (FMetasoundAssetBase* MetaSoundAsset = IMetasoundUObjectRegistry::Get().GetObjectAsAssetBase(&MetaSound))
+							{
+								MetaSoundAsset->UnregisterGraphWithFrontend();
+							}
 						}
 					}
+					DocBuilder.FinishBuilding();
 				}
-				DocBuilder.FinishBuilding();
-			}
+			}	
 
+			// Still return true in this case as the builder likely has become inaccessible and may be in the "beginning" of destruction,
+			// so entries is still reporting that it was successfully removed.
 			ensureAlways(Builders.Remove(InClassName));
 			return true;
 		}
diff --git a/Engine/Plugins/Runtime/Metasound/Source/MetasoundEngine/Public/MetasoundBuilderBase.cpp b/Engine/Plugins/Runtime/Metasound/Source/MetasoundEngine/Public/MetasoundBuilderBase.cpp
--- a/Engine/Plugins/Runtime/Metasound/Source/MetasoundEngine/Public/MetasoundBuilderBase.cpp
+++ b/Engine/Plugins/Runtime/Metasound/Source/MetasoundEngine/Public/MetasoundBuilderBase.cpp
@@ -45,11 +45,12 @@ void UMetaSoundBuilderBase::BeginDestroy()
 	{
 		if (IDocumentBuilderRegistry* BuilderRegistry = IDocumentBuilderRegistry::Get())
 		{
-			constexpr bool bForceUnregister = true;
 			const FMetasoundFrontendClassName& MetaSoundClassName = Builder.GetConstDocument().RootGraph.Metadata.GetClassName();
-			BuilderRegistry->FinishBuilding(MetaSoundClassName, bForceUnregister);
+			BuilderRegistry->FinishBuilding(MetaSoundClassName);
 		}
 
+		// The registry may have not been active or it was and the internal weak pointer record of this object no longer accessible.
+		// In either of these cases, we could still have a local, valid builder, so call finish directly here just in case.
 		Builder.FinishBuilding();
 	}
 	Super::BeginDestroy();
diff --git a/Engine/Plugins/Runtime/Metasound/Source/MetasoundFrontend/Private/MetasoundAssetBase.cpp b/Engine/Plugins/Runtime/Metasound/Source/MetasoundFrontend/Private/MetasoundAssetBase.cpp
--- a/Engine/Plugins/Runtime/Metasound/Source/MetasoundFrontend/Private/MetasoundAssetBase.cpp
+++ b/Engine/Plugins/Runtime/Metasound/Source/MetasoundFrontend/Private/MetasoundAssetBase.cpp
@@ -427,6 +427,9 @@ void FMetasoundAssetBase::OnNotifyBeginDestroy()
 {
 	using namespace Metasound::Frontend;
 
+	UObject* OwningAsset = GetOwningAsset();
+	check(OwningAsset);
+
 	// Unregistration of graph is not necessary when cooking as deserialized objects are not mutable and, should they be reloaded,
 	// omitting unregistration avoids potentially kicking off an invalid asynchronous task to unregister a non-existent runtime graph.
 	if (IsRunningCookCommandlet())
@@ -439,21 +442,13 @@ void FMetasoundAssetBase::OnNotifyBeginDestroy()
 	}
 	else
 	{
-		// The editor module removal and addition of assets with the MetaSoundAssetManager directly,
-		// so it can update open editors and load possible assets to reference at will, as opposed
-		// to managing by direct load requests and associated asset references.
-		// (Owning "asset" is a misnomer as we support dynamically generated MetaSounds at runtime now.)
-#if WITH_EDITOR
-		UObject* OwningAsset = GetOwningAsset();
-		check(OwningAsset);
-		if (IMetaSoundAssetManager* AssetManager = IMetaSoundAssetManager::Get())
-		{
-			AssetManager->RemoveAsset(*OwningAsset);
-		}
-#endif // !WITH_EDITOR
-
 		UnregisterGraphWithFrontend();
 	}
+
+	if (IMetaSoundAssetManager* AssetManager = IMetaSoundAssetManager::Get())
+	{
+		AssetManager->RemoveAsset(*OwningAsset);
+	};
 }
 
 void FMetasoundAssetBase::UnregisterGraphWithFrontend()
