Fix potential PIE crash from FMetasoundGeneratorHandle::UpdateOutputWatchersInternal triggering Blueprint callbacks during PostLoad. Have FMetasoundGeneratorHandle::HandleOutputChanged call back to the game thread using ExecuteOnGameThread rather than AsyncTask(ENamedThreads::GameThread). And fix up several other locations that might trigger BP events.

diff --git a/Engine/Plugins/Runtime/Metasound/Source/MetasoundEngine/Private/MetasoundGeneratorHandle.cpp b/Engine/Plugins/Runtime/Metasound/Source/MetasoundEngine/Private/MetasoundGeneratorHandle.cpp
--- a/Engine/Plugins/Runtime/Metasound/Source/MetasoundEngine/Private/MetasoundGeneratorHandle.cpp
+++ b/Engine/Plugins/Runtime/Metasound/Source/MetasoundEngine/Private/MetasoundGeneratorHandle.cpp
@@ -589,7 +589,7 @@ namespace Metasound
 		if (InAudioComponentId == GetAudioComponentId())
 		{
 			// Set the generator on the game thread. We grab a weak pointer in case this gets destroyed while we wait.
-			AsyncTask(ENamedThreads::GameThread, [WeakThis = AsWeak(), WeakGenerator = InGenerator.ToWeakPtr()]()
+			ExecuteOnGameThread(UE_SOURCE_LOCATION, [WeakThis = AsWeak(), WeakGenerator = InGenerator.ToWeakPtr()]()
 			{
 				if (const TSharedPtr<FMetasoundGeneratorHandle> PinnedThis = WeakThis.Pin())
 				{
@@ -611,7 +611,7 @@ namespace Metasound
 		if (InAudioComponentId == GetAudioComponentId())
 		{
 			// Unset the generator on the game thread. We grab a weak pointer in case this gets destroyed while we wait.
-			AsyncTask(ENamedThreads::GameThread, [WeakThis = AsWeak()]()
+			ExecuteOnGameThread(UE_SOURCE_LOCATION, [WeakThis = AsWeak()]()
 			{
 				if (const TSharedPtr<FMetasoundGeneratorHandle> PinnedThis = WeakThis.Pin())
 				{
@@ -627,7 +627,7 @@ namespace Metasound
 		METASOUND_TRACE_CPUPROFILER_EVENT_SCOPE(FMetasoundGeneratorHandle::HandleGeneratorGraphSet);
 
 		// Defer to the game thread. We grab a weak pointer in case this gets destroyed while we wait.
-		AsyncTask(ENamedThreads::GameThread, [WeakThis = AsWeak()]()
+		ExecuteOnGameThread(UE_SOURCE_LOCATION, [WeakThis = AsWeak()]()
 		{
 			if (const TSharedPtr<FMetasoundGeneratorHandle> PinnedThis = WeakThis.Pin())
 			{
@@ -648,7 +648,7 @@ namespace Metasound
 		METASOUND_TRACE_CPUPROFILER_EVENT_SCOPE(FMetasoundGeneratorHandle::HandleGeneratorVertexInterfaceChanged);
 
 		// Defer to the game thread. We grab a weak pointer in case this gets destroyed while we wait.
-		AsyncTask(ENamedThreads::GameThread, [WeakThis = AsWeak(), VertexInterfaceChanges]()
+		ExecuteOnGameThread(UE_SOURCE_LOCATION, [WeakThis = AsWeak(), VertexInterfaceChanges]()
 		{
 			if (const TSharedPtr<FMetasoundGeneratorHandle> PinnedThis = WeakThis.Pin())
 			{
@@ -698,7 +698,7 @@ namespace Metasound
 		if (!OutputWatcherUpdateScheduled.test_and_set())
 		{
 			// Defer to the game thread. We grab a weak pointer in case this gets destroyed while we wait.
-			AsyncTask(ENamedThreads::GameThread, [WeakThis = AsWeak()]()
+			ExecuteOnGameThread(UE_SOURCE_LOCATION, [WeakThis = AsWeak()]()
 			{
 				if (const TSharedPtr<FMetasoundGeneratorHandle> PinnedThis = WeakThis.Pin())
 				{
