[MetaSound]Find References to the selected Node in the current Graph

diff --git a/Engine/Plugins/Runtime/Metasound/Source/MetasoundEditor/Private/FindInMetasound.cpp b/Engine/Plugins/Runtime/Metasound/Source/MetasoundEditor/Private/FindInMetasound.cpp
--- a/Engine/Plugins/Runtime/Metasound/Source/MetasoundEditor/Private/FindInMetasound.cpp
+++ b/Engine/Plugins/Runtime/Metasound/Source/MetasoundEditor/Private/FindInMetasound.cpp
@@ -372,6 +372,7 @@ namespace Metasound::Editor
 					.HintText(LOCTEXT("FindMetasound_GraphSearchHint", "Search"))
 					.OnTextChanged(this, &SFindInMetasound::OnSearchTextChanged)
 					.OnTextCommitted(this, &SFindInMetasound::OnSearchTextCommitted)
+					.DelayChangeNotificationsWhileTyping(false)
 				]
 			]
 			+ SVerticalBox::Slot()
@@ -403,6 +404,17 @@ namespace Metasound::Editor
 		FSlateApplication::Get().SetKeyboardFocus(FilterTextBoxWidgetPath, EFocusCause::SetDirectly);
 	}
 
+	void SFindInMetasound::FocusForUse(const FString& NewSearchTerms)
+	{
+		FocusForUse();
+
+		if (!NewSearchTerms.IsEmpty())
+		{
+			SearchTextField->SetText(FText::FromString(NewSearchTerms));
+			InitiateSearch();
+		}
+	}
+
 	void SFindInMetasound::OnSearchTextChanged(const FText& Text)
 	{
 		SearchValue = Text.ToString();
diff --git a/Engine/Plugins/Runtime/Metasound/Source/MetasoundEditor/Private/MetasoundEditor.cpp b/Engine/Plugins/Runtime/Metasound/Source/MetasoundEditor/Private/MetasoundEditor.cpp
--- a/Engine/Plugins/Runtime/Metasound/Source/MetasoundEditor/Private/MetasoundEditor.cpp
+++ b/Engine/Plugins/Runtime/Metasound/Source/MetasoundEditor/Private/MetasoundEditor.cpp
@@ -2149,6 +2149,9 @@ namespace Metasound
 				GraphEditorCommands->MapAction(FGraphEditorCommands::Get().CreateComment,
 					FExecuteAction::CreateSP(this, &FEditor::OnCreateComment));
 
+				GraphEditorCommands->MapAction(FGraphEditorCommands::Get().FindReferences,
+					FExecuteAction::CreateSP(this, &FEditor::FindSelectedNodeInGraph));
+
 				GraphEditorCommands->MapAction(FEditorCommands::Get().UpdateNodeClass,
 					FExecuteAction::CreateSP(this, &FEditor::UpdateSelectedNodeClasses));
 			}
@@ -3647,6 +3650,23 @@ namespace Metasound
 				FindWidget->FocusForUse();
 			}
 		}
+
+		void FEditor::FindSelectedNodeInGraph()
+		{
+			TabManager->TryInvokeTab(TabFactory::Names::Find);
+			if (FindWidget.IsValid())
+			{		
+				const FGraphPanelSelectionSet& SelectedNodes = MetasoundGraphEditor->GetSelectedNodes();
+				for (UObject* Object : SelectedNodes)
+				{
+					if (UEdGraphNode* SelectedNode = Cast<UEdGraphNode>(Object))
+					{
+						FString SearchTerms = SelectedNode->GetFindReferenceSearchString(EGetFindReferenceSearchStringFlags::UseSearchSyntax);
+						FindWidget->FocusForUse(SearchTerms);
+					}
+				}
+			}
+		}
 	}
 }
 #undef LOCTEXT_NAMESPACE
diff --git a/Engine/Plugins/Runtime/Metasound/Source/MetasoundEditor/Private/MetasoundEditor.h b/Engine/Plugins/Runtime/Metasound/Source/MetasoundEditor/Private/MetasoundEditor.h
--- a/Engine/Plugins/Runtime/Metasound/Source/MetasoundEditor/Private/MetasoundEditor.h
+++ b/Engine/Plugins/Runtime/Metasound/Source/MetasoundEditor/Private/MetasoundEditor.h
@@ -339,6 +339,9 @@ namespace Metasound
 			/** Show and focus the Find in MetaSound tab. */
 			void ShowFindInMetaSound();
 
+			/** Find selected node from Graph*/
+			void FindSelectedNodeInGraph();
+
 		private:
 			void RemoveInvalidSelection();
 
diff --git a/Engine/Plugins/Runtime/Metasound/Source/MetasoundEditor/Private/MetasoundEditorGraphNode.cpp b/Engine/Plugins/Runtime/Metasound/Source/MetasoundEditor/Private/MetasoundEditorGraphNode.cpp
--- a/Engine/Plugins/Runtime/Metasound/Source/MetasoundEditor/Private/MetasoundEditorGraphNode.cpp
+++ b/Engine/Plugins/Runtime/Metasound/Source/MetasoundEditor/Private/MetasoundEditorGraphNode.cpp
@@ -764,6 +764,29 @@ bool UMetasoundEditorGraphMemberNode::ClampFloatLiteral(const UMetasoundEditorGr
 	return bClampedFloatLiteral;
 }
 
+FString UMetasoundEditorGraphMemberNode::GetFindReferenceSearchString_Impl(EGetFindReferenceSearchStringFlags InFlags) const
+{
+	using namespace Metasound::Frontend;
+
+	if (UMetasoundEditorGraphMember* GraphMember = GetMember())
+	{
+		FString NameToSearch;
+	
+		if (!GraphMember->GetDisplayName().IsEmpty())
+		{
+			NameToSearch = GraphMember->GetDisplayName().ToString();		
+		}
+		else
+		{
+			NameToSearch = GraphMember->GetMemberName().ToString();
+		}
+
+		return FString::Printf(TEXT("\"%s\" \"%s\""), *NameToSearch, *GraphMember->GetDataType().ToString());
+	}
+
+	return FString();
+}
+
 bool UMetasoundEditorGraphOutputNode::CanDuplicateNode() const
 {
 	return false;
diff --git a/Engine/Plugins/Runtime/Metasound/Source/MetasoundEditor/Private/MetasoundEditorGraphSchema.cpp b/Engine/Plugins/Runtime/Metasound/Source/MetasoundEditor/Private/MetasoundEditorGraphSchema.cpp
--- a/Engine/Plugins/Runtime/Metasound/Source/MetasoundEditor/Private/MetasoundEditorGraphSchema.cpp
+++ b/Engine/Plugins/Runtime/Metasound/Source/MetasoundEditor/Private/MetasoundEditorGraphSchema.cpp
@@ -1309,6 +1309,9 @@ void UMetasoundEditorGraphSchema::GetContextMenuActions(class UToolMenu* Menu, c
 		Section.AddMenuEntry(FGenericCommands::Get().Duplicate);
 		Section.AddMenuEntry(FGenericCommands::Get().Rename);
 		Section.AddMenuEntry(FGraphEditorCommands::Get().BreakNodeLinks);
+		Section.AddMenuEntry(FGraphEditorCommands::Get().FindReferences, LOCTEXT("FindReferencesInGraph", "Find References In Graph"),
+			LOCTEXT("FindReferencesInGraph_Tooltip", "Find References to the selected Node in the current Graph"),
+			FSlateIcon());
 
 		// Only display update ability if node is of type external
 		// and node registry is reporting a major update is available.
diff --git a/Engine/Plugins/Runtime/Metasound/Source/MetasoundEditor/Public/FindInMetasound.h b/Engine/Plugins/Runtime/Metasound/Source/MetasoundEditor/Public/FindInMetasound.h
--- a/Engine/Plugins/Runtime/Metasound/Source/MetasoundEditor/Public/FindInMetasound.h
+++ b/Engine/Plugins/Runtime/Metasound/Source/MetasoundEditor/Public/FindInMetasound.h
@@ -92,6 +92,8 @@ namespace Metasound::Editor
 
 		/** Focuses this widget's search box */
 		void FocusForUse();
+		/** Focuses this widget's search box, and Initiates a search with given SearchTerms */
+		void FocusForUse(const FString& NewSearchTerms);
 
 	protected:
 		typedef TSharedPtr<FFindInMetasoundResult> FSearchResult;
diff --git a/Engine/Plugins/Runtime/Metasound/Source/MetasoundEditor/Public/MetasoundEditorGraphNode.h b/Engine/Plugins/Runtime/Metasound/Source/MetasoundEditor/Public/MetasoundEditorGraphNode.h
--- a/Engine/Plugins/Runtime/Metasound/Source/MetasoundEditor/Public/MetasoundEditorGraphNode.h
+++ b/Engine/Plugins/Runtime/Metasound/Source/MetasoundEditor/Public/MetasoundEditorGraphNode.h
@@ -207,6 +207,10 @@ public:
 	// Clamp float literal value based on the given default float literal. 
 	// Returns whether the literal was clamped. 
 	static bool ClampFloatLiteral(const UMetasoundEditorGraphMemberDefaultFloat* DefaultFloatLiteral, FMetasoundFrontendLiteral& LiteralValue);
+
+protected:
+	virtual FString GetFindReferenceSearchString_Impl(EGetFindReferenceSearchStringFlags InFlags) const override;
+
 };
 
 /** Node that represents a graph output */
