Fix error C4702: unreachable code in editor builds

diff --git a/Engine/Plugins/Runtime/Metasound/Source/MetasoundEditor/Private/MetasoundNodeDetailCustomization.cpp b/Engine/Plugins/Runtime/Metasound/Source/MetasoundEditor/Private/MetasoundNodeDetailCustomization.cpp
--- a/Engine/Plugins/Runtime/Metasound/Source/MetasoundEditor/Private/MetasoundNodeDetailCustomization.cpp
+++ b/Engine/Plugins/Runtime/Metasound/Source/MetasoundEditor/Private/MetasoundNodeDetailCustomization.cpp
@@ -252,7 +252,6 @@ namespace Metasound
 					{
 						return bIsGraphEditable;
 					}
-					return false;
 				});
 
 				PasteAction.ExecuteAction = CreateDefaultValueObjectPasteExecuteAction(PropertyHandle, bIsArray);
