Move log to display instead of warning to avoid build spam
#rnx

diff --git a/Engine/Plugins/Runtime/Metasound/Source/MetasoundEngine/Private/MetasoundAssetSubsystem.cpp b/Engine/Plugins/Runtime/Metasound/Source/MetasoundEngine/Private/MetasoundAssetSubsystem.cpp
--- a/Engine/Plugins/Runtime/Metasound/Source/MetasoundEngine/Private/MetasoundAssetSubsystem.cpp
+++ b/Engine/Plugins/Runtime/Metasound/Source/MetasoundEngine/Private/MetasoundAssetSubsystem.cpp
@@ -207,12 +207,12 @@ namespace Metasound::Engine
 			TSet<FAssetKey> Keys;
 			if (int32 NumKeys = PathMap.GetKeys(Keys); NumKeys > 0)
 			{
-				UE_LOG(LogMetaSound, Warning, TEXT("AssetManager is shutting down with the following '%i' assets active:"), NumKeys);
+				UE_LOG(LogMetaSound, Display, TEXT("AssetManager is shutting down with the following %i assets active:"), NumKeys);
 				for (const TPair<FAssetKey, TArray<FTopLevelAssetPath>>& Pair : PathMap)
 				{
 					for (const FTopLevelAssetPath& Path : Pair.Value)
 					{
-						UE_LOG(LogMetaSound, Warning, TEXT("- %s"), *Path.ToString());
+						UE_LOG(LogMetaSound, Display, TEXT("- %s"), *Path.ToString());
 					}
 				}
 			}
