UE-206945 - [MetaSounds] Make MetaSound editor registration notification last until registration is complete (2)

diff --git a/Engine/Plugins/Runtime/Metasound/Source/MetasoundEditor/Private/MetasoundEditor.cpp b/Engine/Plugins/Runtime/Metasound/Source/MetasoundEditor/Private/MetasoundEditor.cpp
--- a/Engine/Plugins/Runtime/Metasound/Source/MetasoundEditor/Private/MetasoundEditor.cpp
+++ b/Engine/Plugins/Runtime/Metasound/Source/MetasoundEditor/Private/MetasoundEditor.cpp
@@ -1105,15 +1105,31 @@ namespace Metasound
 		{
 			if (MetasoundGraphEditor.IsValid())
 			{
+				const FText CloseNotificationText = LOCTEXT("MetaSoundScanInProgressNotificationButtonText", "Close");
+
+				FSimpleDelegate OnCloseNotification = FSimpleDelegate::CreateLambda([this]()
+				{
+					if (NotificationPtr)
+					{
+						NotificationPtr->Fadeout();
+						NotificationPtr.Reset();
+					}
+				});
+
 				FNotificationInfo Info(LOCTEXT("MetaSoundScanInProgressNotificationText", "Registering MetaSound Assets..."));
 				Info.SubText = LOCTEXT("MetaSoundScanInProgressNotificationSubText", "Class selector results may be incomplete");
 				Info.bUseThrobber = true;
-				Info.bFireAndForget = true;
+				Info.bFireAndForget = false;
 				Info.bUseSuccessFailIcons = false;
-				Info.ExpireDuration = 3.0f;
 				Info.FadeOutDuration = 1.0f;
+				Info.ButtonDetails.Add(FNotificationButtonInfo(CloseNotificationText, FText(), OnCloseNotification));
 
-				MetasoundGraphEditor->AddNotification(Info, false /* bSuccess */);
+				NotificationPtr = MetasoundGraphEditor->AddNotification(Info);
+				if (NotificationPtr.IsValid())
+				{
+					NotificationPtr->SetVisibility(EVisibility::Visible);
+					NotificationPtr->SetCompletionState(SNotificationItem::CS_Pending);
+				}
 			}
 		}
 
@@ -1121,6 +1137,12 @@ namespace Metasound
 		{
 			if (MetasoundGraphEditor.IsValid())
 			{
+				if (NotificationPtr.IsValid())
+				{
+					NotificationPtr->Fadeout();
+					NotificationPtr.Reset();
+				}
+
 				FNotificationInfo Info(LOCTEXT("MetaSoundScanInProgressNotification", "MetaSound Asset Registration Complete"));
 				Info.bFireAndForget = true;
 				Info.bUseSuccessFailIcons = true;
diff --git a/Engine/Plugins/Runtime/Metasound/Source/MetasoundEditor/Private/MetasoundEditor.h b/Engine/Plugins/Runtime/Metasound/Source/MetasoundEditor/Private/MetasoundEditor.h
--- a/Engine/Plugins/Runtime/Metasound/Source/MetasoundEditor/Private/MetasoundEditor.h
+++ b/Engine/Plugins/Runtime/Metasound/Source/MetasoundEditor/Private/MetasoundEditor.h
@@ -484,6 +484,8 @@ namespace Metasound
 			bool bPrimingRegistry = false;
 
 			int32 HighestMessageSeverity = EMessageSeverity::Info;
+
+			TSharedPtr<SNotificationItem> NotificationPtr;
 		};
 	} // namespace Editor
 } // namespace Metasound
diff --git a/Engine/Source/Editor/GraphEditor/Private/SGraphEditorImpl.cpp b/Engine/Source/Editor/GraphEditor/Private/SGraphEditorImpl.cpp
--- a/Engine/Source/Editor/GraphEditor/Private/SGraphEditorImpl.cpp
+++ b/Engine/Source/Editor/GraphEditor/Private/SGraphEditorImpl.cpp
@@ -850,7 +850,7 @@ void SGraphEditorImpl::Construct( const FArguments& InArgs )
 		.HAlign(HAlign_Right)
 		[
 			SAssignNew(NotificationListPtr, SNotificationList)
-			.Visibility(EVisibility::HitTestInvisible)
+			.Visibility(EVisibility::Visible)
 		]
 	];
 
@@ -1463,10 +1463,25 @@ void SGraphEditorImpl::AddNotification( FNotificationInfo& Info, bool bSuccess )
 	TSharedPtr<SNotificationItem> Notification = NotificationListPtr->AddNotification(Info);
 	if ( Notification.IsValid() )
 	{
+		Notification->SetVisibility(EVisibility::HitTestInvisible);
 		Notification->SetCompletionState( bSuccess ? SNotificationItem::CS_Success : SNotificationItem::CS_Fail );
 	}
 }
 
+TSharedPtr<SNotificationItem> SGraphEditorImpl::AddNotification(FNotificationInfo& Info)
+{
+	// set up common notification properties
+	Info.bUseLargeFont = true;
+
+	TSharedPtr<SNotificationItem> Notification = NotificationListPtr->AddNotification(Info);
+	if (Notification.IsValid())
+	{
+		Notification->SetVisibility(EVisibility::HitTestInvisible);
+		return Notification;
+	}
+	return nullptr;
+}
+
 EActiveTimerReturnType SGraphEditorImpl::HandleFocusEditorDeferred(double InCurrentTime, float InDeltaTime)
 {
 	
diff --git a/Engine/Source/Editor/GraphEditor/Private/SGraphEditorImpl.h b/Engine/Source/Editor/GraphEditor/Private/SGraphEditorImpl.h
--- a/Engine/Source/Editor/GraphEditor/Private/SGraphEditorImpl.h
+++ b/Engine/Source/Editor/GraphEditor/Private/SGraphEditorImpl.h
@@ -334,6 +334,7 @@ public:
 	virtual void LockToGraphEditor(TWeakPtr<SGraphEditor> Other) override;
 	virtual void UnlockFromGraphEditor(TWeakPtr<SGraphEditor> Other) override;
 	virtual void AddNotification ( FNotificationInfo& Info, bool bSuccess ) override;
+	virtual TSharedPtr<SNotificationItem> AddNotification(FNotificationInfo& Info) override;
 	virtual void SetPinVisibility(SGraphEditor::EPinVisibility Visibility) override;
 	virtual TSharedRef<FActiveTimerHandle> RegisterActiveTimer(float TickPeriod, FWidgetActiveTimerDelegate TickFunction) override;
 	virtual void StraightenConnections() override;
diff --git a/Engine/Source/Editor/UnrealEd/Public/GraphEditor.h b/Engine/Source/Editor/UnrealEd/Public/GraphEditor.h
--- a/Engine/Source/Editor/UnrealEd/Public/GraphEditor.h
+++ b/Engine/Source/Editor/UnrealEd/Public/GraphEditor.h
@@ -32,6 +32,7 @@ class FAssetEditorToolkit;
 class FMenuBuilder;
 class FReply;
 class SGraphPanel;
+class SNotificationItem;
 class SWidget;
 class UEdGraph;
 class UEdGraphNode;
@@ -490,6 +491,20 @@ public:
 		}
 	}
 
+	virtual TSharedPtr<SNotificationItem> AddNotification(FNotificationInfo& Info)
+	{
+		if (Implementation.IsValid())
+		{
+			TSharedPtr<SNotificationItem> Notification = Implementation->AddNotification(Info);
+			
+			if (Notification.IsValid())
+			{
+				return Notification;
+			}
+		}
+		return nullptr;
+	}
+
 	/** Capture keyboard */
 	virtual void CaptureKeyboard()
 	{
