[MetaSound]Dropping a variable to a pin with matching type does not connect the nodes automatically in MetaSound graph

diff --git a/Engine/Plugins/Runtime/Metasound/Source/MetasoundEditor/Private/MetasoundEditor.cpp b/Engine/Plugins/Runtime/Metasound/Source/MetasoundEditor/Private/MetasoundEditor.cpp
--- a/Engine/Plugins/Runtime/Metasound/Source/MetasoundEditor/Private/MetasoundEditor.cpp
+++ b/Engine/Plugins/Runtime/Metasound/Source/MetasoundEditor/Private/MetasoundEditor.cpp
@@ -15,6 +15,7 @@
 #include "EdGraph/EdGraphNode.h"
 #include "EdGraph/EdGraphPin.h"
 #include "EdGraphUtilities.h"
+#include "EdGraphHandleTypes.h"
 #include "Editor.h"
 #include "Engine/World.h"
 #include "Framework/Application/SlateApplication.h"
@@ -197,15 +198,26 @@ namespace Metasound
 			DRAG_DROP_OPERATOR_TYPE(FMetaSoundDragDropMemberAction, FGraphSchemaActionDragDropAction)
 
 			virtual FReply DroppedOnPanel(const TSharedRef<SWidget>& InPanel, FVector2D InScreenPosition, FVector2D InGraphPosition, UEdGraph& InGraph) override
+			{			
+				if (!GraphMember.IsValid() || &InGraph != GraphMember->GetOwningGraph())
+				{
+					return FReply::Unhandled();
+				}
+
+				return DroppedOnPin(InScreenPosition, InGraphPosition);
+			}
+			virtual FReply DroppedOnNode(FVector2D ScreenPosition, FVector2D GraphPosition) override { return FReply::Unhandled(); }
+			virtual FReply DroppedOnPin(FVector2D InScreenPosition, FVector2D InGraphPosition) override
 			{
 				using namespace Frontend;
 
-				if (!GraphMember.IsValid() || &InGraph != GraphMember->GetOwningGraph())
+				if (!GraphMember.IsValid())
 				{
 					return FReply::Unhandled();
 				}
 
-				UMetasoundEditorGraph* MetasoundGraph = CastChecked<UMetasoundEditorGraph>(&InGraph);
+				UMetasoundEditorGraph* MetasoundGraph = GraphMember->GetOwningGraph();
+				check(MetasoundGraph);
 				UObject& ParentMetasound = MetasoundGraph->GetMetasoundChecked();
 
 				if (UMetasoundEditorGraphInput* Input = Cast<UMetasoundEditorGraphInput>(GraphMember.Get()))
@@ -223,6 +235,8 @@ namespace Metasound
 						NewGraphNode->UpdateFrontendNodeLocation(InGraphPosition);
 						NewGraphNode->SyncLocationFromFrontendNode();
 
+						TryConnectToHoveredPin(NewGraphNode);
+
 						FGraphBuilder::RegisterGraphWithFrontend(ParentMetasound);
 						TSharedPtr<FEditor> MetasoundEditor = FGraphBuilder::GetEditorForGraph(*MetasoundGraph);
 						if (MetasoundEditor.IsValid())
@@ -249,6 +263,9 @@ namespace Metasound
 							NewGraphNode->Modify();
 							NewGraphNode->UpdateFrontendNodeLocation(InGraphPosition);
 							NewGraphNode->SyncLocationFromFrontendNode();
+
+							TryConnectToHoveredPin(NewGraphNode);
+
 							FGraphBuilder::RegisterGraphWithFrontend(ParentMetasound);
 							TSharedPtr<FEditor> MetasoundEditor = FGraphBuilder::GetEditorForGraph(*MetasoundGraph);
 							if (MetasoundEditor.IsValid())
@@ -344,6 +361,9 @@ namespace Metasound
 						NewGraphNode->Modify();
 						NewGraphNode->UpdateFrontendNodeLocation(InGraphPosition);
 						NewGraphNode->SyncLocationFromFrontendNode();
+
+						TryConnectToHoveredPin(NewGraphNode);
+
 						FGraphBuilder::RegisterGraphWithFrontend(ParentMetasound);
 						TSharedPtr<FEditor> MetasoundEditor = FGraphBuilder::GetEditorForGraph(*MetasoundGraph);
 						if (MetasoundEditor.IsValid())
@@ -356,10 +376,33 @@ namespace Metasound
 
 				return FReply::Unhandled();
 			}
-			virtual FReply DroppedOnNode(FVector2D ScreenPosition, FVector2D GraphPosition) override { return FReply::Unhandled(); }
-			virtual FReply DroppedOnPin(FVector2D ScreenPosition, FVector2D GraphPosition) override { return FReply::Unhandled(); }
 			virtual FReply DroppedOnAction(TSharedRef<FEdGraphSchemaAction> Action) { return FReply::Unhandled(); }
 			virtual FReply DroppedOnCategory(FText Category) override { return FReply::Unhandled(); }
+
+			bool TryConnectToHoveredPin(UMetasoundEditorGraphNode* InNewGraphNode)
+			{
+				if (!GetHoveredPin())
+				{
+					return false;
+				}
+
+				FEdGraphPinHandle FromPin = InNewGraphNode->GetPinAt(0);
+				FEdGraphPinHandle ToPin = GetHoveredPin();
+
+				if ((FromPin.GetPin() != nullptr) && (ToPin.GetPin() != nullptr))
+				{
+					const UEdGraph* MyGraphObj = FromPin.GetGraph();
+
+					// the pin may change during the creation of the link
+					if (const UEdGraphSchema* GraphSchema = MyGraphObj->GetSchema())
+					{
+						return GraphSchema->TryCreateConnection(FromPin.GetPin(), ToPin.GetPin());
+					}
+				}
+
+				return false;
+			}
+
 			virtual void HoverTargetChanged() override
 			{
 				using namespace Frontend;
@@ -480,7 +523,6 @@ namespace Metasound
 			}
 		};
 
-
 		class SMetaSoundGraphPaletteItem : public SGraphPaletteItem
 		{
 		private:
