[Audio Material Widgets] - Audio Styling Header captialization on MetaSounds is inconsistent with the rest of the Editor Preference settings
#rnx

diff --git a/Engine/Plugins/Runtime/Metasound/Source/MetasoundEditor/Private/MetasoundEditorSettings.h b/Engine/Plugins/Runtime/Metasound/Source/MetasoundEditor/Private/MetasoundEditorSettings.h
--- a/Engine/Plugins/Runtime/Metasound/Source/MetasoundEditor/Private/MetasoundEditorSettings.h
+++ b/Engine/Plugins/Runtime/Metasound/Source/MetasoundEditor/Private/MetasoundEditorSettings.h
@@ -182,23 +182,23 @@ public:
 	EMetasoundActiveDetailView DetailView = EMetasoundActiveDetailView::General;
 
 	/** Whether the AudioMaterialWidgets are used when possible in Metasound Editor*/
-	UPROPERTY(EditAnywhere, config, DisplayName = "Use Audio Material Widgets", Category = MetasoundStyling)
+	UPROPERTY(EditAnywhere, config, DisplayName = "Use Audio Material Widgets", Category = WidgetStyling)
 	bool bUseAudioMaterialWidgets = false;
 	
 	/**Override the Knob Style used in the Metasound Editor.*/
-	UPROPERTY(EditAnywhere, config, Category = MetasoundStyling, meta = (AllowedClasses = "/Script/SlateCore.SlateWidgetStyleAsset", EditCondition = "bUseAudioMaterialWidgets", DisplayName = "Knob Style"))
+	UPROPERTY(EditAnywhere, config, Category = WidgetStyling, meta = (AllowedClasses = "/Script/SlateCore.SlateWidgetStyleAsset", EditCondition = "bUseAudioMaterialWidgets", DisplayName = "Knob Style"))
 	FSoftObjectPath KnobStyleOverride;	
 	
 	/**Override the Slider Style used in the Metasound Editor.*/
-	UPROPERTY(EditAnywhere, config, Category = MetasoundStyling, meta = (AllowedClasses = "/Script/SlateCore.SlateWidgetStyleAsset", EditCondition = "bUseAudioMaterialWidgets", DisplayName = "Slider Style"))
+	UPROPERTY(EditAnywhere, config, Category = WidgetStyling, meta = (AllowedClasses = "/Script/SlateCore.SlateWidgetStyleAsset", EditCondition = "bUseAudioMaterialWidgets", DisplayName = "Slider Style"))
 	FSoftObjectPath SliderStyleOverride;
 
 	/**Override the Button Style used in the Metasound Editor.*/
-	UPROPERTY(EditAnywhere, config, Category = MetasoundStyling, meta = (AllowedClasses = "/Script/SlateCore.SlateWidgetStyleAsset", EditCondition = "bUseAudioMaterialWidgets", DisplayName = "Button Style"))
+	UPROPERTY(EditAnywhere, config, Category = WidgetStyling, meta = (AllowedClasses = "/Script/SlateCore.SlateWidgetStyleAsset", EditCondition = "bUseAudioMaterialWidgets", DisplayName = "Button Style"))
 	FSoftObjectPath ButtonStyleOverride;	
 	
 	/**Override the Meter Style used in the Metasound Editor.*/
-	UPROPERTY(EditAnywhere, config, Category = MetasoundStyling, meta = (AllowedClasses = "/Script/SlateCore.SlateWidgetStyleAsset", EditCondition = "bUseAudioMaterialWidgets", DisplayName = "Meter Style"))
+	UPROPERTY(EditAnywhere, config, Category = WidgetStyling, meta = (AllowedClasses = "/Script/SlateCore.SlateWidgetStyleAsset", EditCondition = "bUseAudioMaterialWidgets", DisplayName = "Meter Style"))
 	FSoftObjectPath MeterStyleOverride;	
 
 	//UObject
