Fix node offset bug introduced in prior submission when creating a new MetaSound causing them to overlap
#tests Repro as provided in bug
#rnx

diff --git a/Engine/Plugins/Runtime/Metasound/Source/MetasoundFrontend/Private/MetasoundFrontendDocumentBuilder.cpp b/Engine/Plugins/Runtime/Metasound/Source/MetasoundFrontend/Private/MetasoundFrontendDocumentBuilder.cpp
--- a/Engine/Plugins/Runtime/Metasound/Source/MetasoundFrontend/Private/MetasoundFrontendDocumentBuilder.cpp
+++ b/Engine/Plugins/Runtime/Metasound/Source/MetasoundFrontend/Private/MetasoundFrontendDocumentBuilder.cpp
@@ -391,6 +391,7 @@ namespace Metasound::Frontend
 				const FMetaSoundFrontendDocumentBuilder& InBuilder,
 				TSet<FName>& AddedNames,
 				TFunctionRef<int32(const FVertexName&)> InGetSortOrder,
+				const FVector2D& InitOffset,
 				TArrayView<FMetasoundFrontendNode> OutNodes)
 			{
 				// Add graph member nodes by sort order
@@ -409,7 +410,7 @@ namespace Metasound::Frontend
 
 				// Prime the first location as an offset prior to an existing location (as provided by a swapped member)
 				//  to avoid placing away from user's active area if possible.
-				FVector2D NextLocation = { 0.0f, 0.0f };
+				FVector2D NextLocation = InitOffset;
 				{
 					int32 NumBeforeDefined = 1;
 					for (const TPair<int32, FMetasoundFrontendNode*>& Pair : SortOrderToNode)
@@ -513,7 +514,7 @@ namespace Metasound::Frontend
 							checkf(Input, TEXT("Input must exist by this point of modifying the document's interfaces and respective members"));
 							return Input->Metadata.SortOrderIndex;
 						};
-						UpdateAddedVertexNodePositions(EMetasoundFrontendClassType::Input, OutBuilder, NamesToSort, GetInputSortOrder, Nodes);
+						UpdateAddedVertexNodePositions(EMetasoundFrontendClassType::Input, OutBuilder, NamesToSort, GetInputSortOrder, FVector2D::Zero(), Nodes);
 					}
 
 					// Sort/Place Outputs
@@ -526,7 +527,7 @@ namespace Metasound::Frontend
 							checkf(Output, TEXT("Output must exist by this point of modifying the document's interfaces and respective members"));
 							return Output->Metadata.SortOrderIndex;
 						};
-						UpdateAddedVertexNodePositions(EMetasoundFrontendClassType::Output, OutBuilder, NamesToSort, GetOutputSortOrder, Nodes);
+						UpdateAddedVertexNodePositions(EMetasoundFrontendClassType::Output, OutBuilder, NamesToSort, GetOutputSortOrder, 3 * DisplayStyle::NodeLayout::DefaultOffsetX, Nodes);
 					}
 				}
 #endif // WITH_EDITORONLY_DATA
