Fix incorrect asset being used by [REDACTED_PERSON] to generate cached dependency asset registry keys
#tests
#rnx

diff --git a/Engine/Plugins/Runtime/Metasound/Source/MetasoundEngine/Private/MetasoundBuilderSubsystem.cpp b/Engine/Plugins/Runtime/Metasound/Source/MetasoundEngine/Private/MetasoundBuilderSubsystem.cpp
--- a/Engine/Plugins/Runtime/Metasound/Source/MetasoundEngine/Private/MetasoundBuilderSubsystem.cpp
+++ b/Engine/Plugins/Runtime/Metasound/Source/MetasoundEngine/Private/MetasoundBuilderSubsystem.cpp
@@ -78,8 +78,9 @@ void UMetaSoundPatchBuilder::OnAssetReferenceAdded(TScriptInterface<IMetaSoundDo
 
 	check(DocInterface.GetObject());
 	UMetaSoundPatch& Patch = Builder.CastDocumentObjectChecked<UMetaSoundPatch>();
-	FNodeRegistryKey RegistryKey(Patch.GetConstDocument().RootGraph);
 	Patch.ReferencedAssetClassObjects.Add(DocInterface.GetObject());
+
+	const FNodeRegistryKey RegistryKey(DocInterface->GetConstDocument().RootGraph);
 	Patch.ReferencedAssetClassKeys.Add(RegistryKey.ToString());
 }
 
@@ -89,8 +90,9 @@ void UMetaSoundPatchBuilder::OnRemovingAssetReference(TScriptInterface<IMetaSoun
 
 	check(DocInterface.GetObject());
 	UMetaSoundPatch& Patch = Builder.CastDocumentObjectChecked<UMetaSoundPatch>();
-	FNodeRegistryKey RegistryKey(Patch.GetConstDocument().RootGraph);
 	Patch.ReferencedAssetClassObjects.Remove(DocInterface.GetObject());
+
+	const FNodeRegistryKey RegistryKey(DocInterface->GetConstDocument().RootGraph);
 	Patch.ReferencedAssetClassKeys.Remove(RegistryKey.ToString());
 }
 
@@ -276,8 +278,9 @@ void UMetaSoundSourceBuilder::OnAssetReferenceAdded(TScriptInterface<IMetaSoundD
 
 	check(DocInterface.GetObject());
 	UMetaSoundSource& Source = GetMetaSoundSource();
-	FNodeRegistryKey RegistryKey(Source.GetConstDocument().RootGraph);
 	Source.ReferencedAssetClassObjects.Add(DocInterface.GetObject());
+
+	const FNodeRegistryKey RegistryKey(DocInterface->GetConstDocument().RootGraph);
 	Source.ReferencedAssetClassKeys.Add(RegistryKey.ToString());
 }
 
@@ -289,9 +292,9 @@ void UMetaSoundSourceBuilder::OnEdgeAdded(int32 EdgeIndex) const
 	const FMetasoundFrontendEdge& NewEdge = Doc.RootGraph.Graph.Edges[EdgeIndex];
 	ExecuteAuditionableTransaction([this, &NewEdge](Metasound::DynamicGraph::FDynamicOperatorTransactor& Transactor)
 	{
-		const FMetaSoundFrontendDocumentBuilder& Builder = GetConstBuilder();
-		const FMetasoundFrontendVertex* FromNodeOutput = Builder.FindNodeOutput(NewEdge.FromNodeID, NewEdge.FromVertexID);
-		const FMetasoundFrontendVertex* ToNodeInput = Builder.FindNodeInput(NewEdge.ToNodeID, NewEdge.ToVertexID);
+		const FMetaSoundFrontendDocumentBuilder& DocBuilder = GetConstBuilder();
+		const FMetasoundFrontendVertex* FromNodeOutput = DocBuilder.FindNodeOutput(NewEdge.FromNodeID, NewEdge.FromVertexID);
+		const FMetasoundFrontendVertex* ToNodeInput = DocBuilder.FindNodeInput(NewEdge.ToNodeID, NewEdge.ToVertexID);
 		if (FromNodeOutput && ToNodeInput)
 		{
 			Transactor.AddDataEdge(NewEdge.FromNodeID, FromNodeOutput->Name, NewEdge.ToNodeID, ToNodeInput->Name);
@@ -548,8 +551,9 @@ void UMetaSoundSourceBuilder::OnRemovingAssetReference(TScriptInterface<IMetaSou
 
 	check(DocInterface.GetObject());
 	UMetaSoundSource& Source = GetMetaSoundSource();
-	FNodeRegistryKey RegistryKey(Source.GetConstDocument().RootGraph);
 	Source.ReferencedAssetClassObjects.Remove(DocInterface.GetObject());
+
+	const FNodeRegistryKey RegistryKey(DocInterface->GetConstDocument().RootGraph);
 	Source.ReferencedAssetClassKeys.Remove(RegistryKey.ToString());
 }
 
