[MetaSounds] Maintain object parameters set in blueprint after un-virtualization.

diff --git a/Engine/Plugins/Runtime/Metasound/Source/MetasoundEngine/Private/MetasoundSource.cpp b/Engine/Plugins/Runtime/Metasound/Source/MetasoundEngine/Private/MetasoundSource.cpp
--- a/Engine/Plugins/Runtime/Metasound/Source/MetasoundEngine/Private/MetasoundSource.cpp
+++ b/Engine/Plugins/Runtime/Metasound/Source/MetasoundEngine/Private/MetasoundSource.cpp
@@ -1176,12 +1176,22 @@ void UMetaSoundSource::InitParametersInternal(const Metasound::TSortedVertexName
 		FAudioParameter& Parameter = ParametersToInit[i];
 		if (const FRuntimeInput* Input = InInputMap.Find(Parameter.ParamName))
 		{
+			bIsParameterValid = true;
+			
 			if (IsParameterValidInternal(Parameter, Input->TypeName, DataTypeRegistry))
 			{
-				Sanitize(Parameter);
-				constexpr bool bClearUObjectPointers = true; // protect against leaking UObject ptrs to the audio thread 
-				SourcePrivate::CreateUObjectProxies(DataTypeRegistry, Input->TypeName, bClearUObjectPointers, Parameter);
-				bIsParameterValid = true;
+				// note: this was originally called 2x, the 1st call would create the proxy
+				// and clear out the ObjectParam.  The 2nd call would clear out the proxy,
+				// and fail to re-create it because we cleared out the ObjectParam. 
+				
+				// if there is no object param we should assume we have already created a proxy
+				// (there is no point in attempting to create a proxy w/o a UObject.)
+				if(Parameter.ObjectParam) 
+				{
+					Sanitize(Parameter);
+					constexpr bool bClearUObjectPointers = true; // protect against leaking UObject ptrs to the audio thread 
+					SourcePrivate::CreateUObjectProxies(DataTypeRegistry, Input->TypeName, bClearUObjectPointers, Parameter);
+				}
 			}
 		}
 
diff --git a/Engine/Source/Runtime/Engine/Private/AudioVirtualLoop.cpp b/Engine/Source/Runtime/Engine/Private/AudioVirtualLoop.cpp
--- a/Engine/Source/Runtime/Engine/Private/AudioVirtualLoop.cpp
+++ b/Engine/Source/Runtime/Engine/Private/AudioVirtualLoop.cpp
@@ -322,7 +322,12 @@ bool FAudioVirtualLoop::Update(float DeltaTime, bool bForceUpdate)
 			if (UAudioComponent* AudioComponent = UAudioComponent::GetAudioComponentFromID(AudioComponentID))
 			{
 				// Use instance parameters which will not contain transient parameters
+				check(AudioComponent->Sound);
 				TArray<FAudioParameter> InstanceParams = AudioComponent->GetInstanceParameters();
+
+				// create proxies from the UObject references in InstanceParams
+				AudioComponent->Sound->InitParameters(InstanceParams);
+				
 				ParameterTransmitter->OnRealizeVirtualizedActiveSound(MoveTemp(InstanceParams));
 			}
 		}
