Fix copy pasta when swapping inputs via MetaSoundBuilder API
#rnx

diff --git a/Engine/Plugins/Runtime/Metasound/Source/MetasoundFrontend/Private/MetasoundFrontendDocumentBuilder.cpp b/Engine/Plugins/Runtime/Metasound/Source/MetasoundFrontend/Private/MetasoundFrontendDocumentBuilder.cpp
--- a/Engine/Plugins/Runtime/Metasound/Source/MetasoundFrontend/Private/MetasoundFrontendDocumentBuilder.cpp
+++ b/Engine/Plugins/Runtime/Metasound/Source/MetasoundFrontend/Private/MetasoundFrontendDocumentBuilder.cpp
@@ -3542,7 +3542,7 @@ bool FMetaSoundFrontendDocumentBuilder::SwapGraphInput(const FMetasoundFrontendC
 
 	// 3. Remove existing graph vertex
 	{
-		const bool bRemovedVertex = RemoveGraphOutput(InExistingInputVertex.Name);
+		const bool bRemovedVertex = RemoveGraphInput(InExistingInputVertex.Name);
 		checkf(bRemovedVertex, TEXT("Failed to swap MetaSound input expected to exist"));
 	}
 
