UE-213567
[MetaSounds][Get Wave Info] - Node name change impedes discoverability of Get Wave Duration
#rnx

diff --git a/Engine/Plugins/Runtime/Metasound/Source/MetasoundEngine/Private/MetasoundWaveInfoNode.cpp b/Engine/Plugins/Runtime/Metasound/Source/MetasoundEngine/Private/MetasoundWaveInfoNode.cpp
--- a/Engine/Plugins/Runtime/Metasound/Source/MetasoundEngine/Private/MetasoundWaveInfoNode.cpp
+++ b/Engine/Plugins/Runtime/Metasound/Source/MetasoundEngine/Private/MetasoundWaveInfoNode.cpp
@@ -90,7 +90,8 @@ namespace Metasound
 			Info.PromptIfMissing = PluginNodeMissingPrompt;
 			Info.DefaultInterface = DeclareVertexInterface();
 			Info.CategoryHierarchy.Emplace(NodeCategories::Debug);
-
+			Info.Keywords = { METASOUND_LOCTEXT("WaveInfoNode_DurationKeyword", "Duration"), METASOUND_LOCTEXT("WaveInfoNode_NameKeyword", "Name"), METASOUND_LOCTEXT("WaveInfoNode_PathKeyword", "Path") };
+			
 			return Info;
 		};
 
