- Move MetaSound input/output editor validation, GetMemberName/Description, IsInterfaceMember to builder API
- Add warnings on register for interface vertex members to ensure they are following the expected naming convention (i.e. namespace matches that of owning interface)
    - Note this is editor only to avoid spamming cook and failing builds
- Move versioning to use shared builder now that IDocumentBuilderRegistry is available prior to all asset serialize load calls. This ensures that multiple builders are not accessing the same asset when loading and versioning/migrating old editor data
#rnx

diff --git a/Engine/Plugins/Runtime/Metasound/Source/MetasoundEditor/Private/MetasoundEditorGraph.cpp b/Engine/Plugins/Runtime/Metasound/Source/MetasoundEditor/Private/MetasoundEditorGraph.cpp
--- a/Engine/Plugins/Runtime/Metasound/Source/MetasoundEditor/Private/MetasoundEditorGraph.cpp
+++ b/Engine/Plugins/Runtime/Metasound/Source/MetasoundEditor/Private/MetasoundEditorGraph.cpp
@@ -18,6 +18,7 @@
 #include "MetasoundEditorSubsystem.h"
 #include "MetasoundFrontendDocumentBuilder.h"
 #include "MetasoundFrontendNodeTemplateRegistry.h"
+#include "MetasoundFrontendSearchEngine.h"
 #include "MetasoundLog.h"
 #include "MetasoundUObjectRegistry.h"
 #include "MetasoundVariableNodes.h"
@@ -99,6 +100,15 @@ namespace Metasound::Editor
 	} // namespace GraphPrivate
 } // namespace Metasound::Editor
 
+FMetaSoundFrontendDocumentBuilder& UMetasoundEditorGraphMember::GetFrontendBuilderChecked() const
+{
+	using namespace Metasound::Frontend;
+
+	const UMetasoundEditorGraph* Graph = GetOwningGraph();
+	check(Graph);
+	UObject& MetaSound = Graph->GetMetasoundChecked();
+	return IDocumentBuilderRegistry::GetChecked().FindOrBeginBuilding(&MetaSound);
+}
 
 UMetasoundEditorGraph* UMetasoundEditorGraphMember::GetOwningGraph()
 {
@@ -135,10 +145,7 @@ void UMetasoundEditorGraphMember::InitializeLiteral()
 
 	if (!Literal || Literal->GetClass() != LiteralClass)
 	{
-		UMetasoundEditorGraph* Graph = GetOwningGraph();
-		check(Graph);
-		UObject& MetaSound = Graph->GetMetasoundChecked();
-		FMetaSoundFrontendDocumentBuilder& Builder = IDocumentBuilderRegistry::GetChecked().FindOrBeginBuilding(&MetaSound);
+		FMetaSoundFrontendDocumentBuilder& Builder = GetFrontendBuilderChecked();
 		const bool bIsNew = UMetaSoundEditorSubsystem::GetChecked().BindMemberMetadata(Builder, *this, LiteralClass);
 		if (bIsNew)
 		{
@@ -204,6 +211,19 @@ void UMetasoundEditorGraphVertex::InitMember(FName InDataType, const FMetasoundF
 	}
 }
 
+const FMetasoundFrontendNode* UMetasoundEditorGraphVertex::GetFrontendNode() const
+{
+	using namespace Metasound::Engine;
+
+	if (const UMetasoundEditorGraph* Graph = GetOwningGraph())
+	{
+		const FMetaSoundFrontendDocumentBuilder& Builder = FDocumentBuilderRegistry::GetChecked().FindOrBeginBuilding(&Graph->GetMetasoundChecked());
+		return Builder.FindNode(NodeID);
+	}
+
+	return nullptr;
+}
+
 TArray<UMetasoundEditorGraphMemberNode*> UMetasoundEditorGraphVertex::GetNodes() const
 {
 	TArray<UMetasoundEditorGraphMemberNode*> Nodes;
@@ -225,13 +245,6 @@ TArray<UMetasoundEditorGraphMemberNode*> UMetasoundEditorGraphVertex::GetNodes()
 	return Nodes;
 }
 
-FText UMetasoundEditorGraphVertex::GetDescription() const
-{
-	// TODO: should be getting description directly from vertex instead of from
-	// node handle.
-	return GetConstNodeHandle()->GetDescription();
-}
-
 void UMetasoundEditorGraphVertex::SetDescription(const FText& InDescription, bool bPostTransaction)
 {
 	using namespace Metasound::Editor;
@@ -265,8 +278,12 @@ FName UMetasoundEditorGraphVertex::GetMemberName() const
 {
 	using namespace Metasound::Frontend;
 
-	FConstNodeHandle NodeHandle = GetConstNodeHandle();
-	return NodeHandle->GetNodeName();
+	if (const FMetasoundFrontendNode* FrontendNode = GetFrontendNode())
+	{
+		return FrontendNode->Name;
+	}
+
+	return FName();
 }
 
 void UMetasoundEditorGraphVertex::SetMemberName(const FName& InNewName, bool bPostTransaction)
@@ -392,11 +409,31 @@ const FMetasoundFrontendVersion& UMetasoundEditorGraphVertex::GetInterfaceVersio
 	return GetConstNodeHandle()->GetInterfaceVersion();
 }
 
-bool UMetasoundEditorGraphVertex::IsInterfaceMember() const
+bool UMetasoundEditorGraphVertex::IsInterfaceMember(FMetasoundFrontendInterface* OutInterface) const
 {
-	using namespace Metasound;
+	using namespace Metasound::Frontend;
+
+	const FName MemberName = GetMemberName();
+	FName InterfaceNamespace;
+	FName ParamName;
+	Audio::FParameterPath::SplitName(MemberName, InterfaceNamespace, ParamName);
+
+	FMetasoundFrontendInterface InterfaceToValidate;
+	if (InterfaceNamespace.IsNone() && ISearchEngine::Get().FindInterfaceWithHighestVersion(InterfaceNamespace, InterfaceToValidate))
+	{
+		if (OutInterface)
+		{
+			*OutInterface = MoveTemp(InterfaceToValidate);
+		}
+
+		return true;
+	}
 
-	return GetConstNodeHandle()->IsInterfaceMember();
+	if (OutInterface)
+	{
+		*OutInterface = { };
+	}
+	return false;
 }
 
 bool UMetasoundEditorGraphVertex::CanRename() const
@@ -580,6 +617,17 @@ const FText& UMetasoundEditorGraphInput::GetGraphMemberLabel() const
 	return Label;
 }
 
+FText UMetasoundEditorGraphInput::GetDescription() const
+{
+	const FMetaSoundFrontendDocumentBuilder& Builder = GetFrontendBuilderChecked();
+	if (const FMetasoundFrontendClassInput* Input = Builder.FindGraphInput(GetMemberName()))
+	{
+		return Input->Metadata.GetDescription();
+	}
+
+	return { };
+}
+
 int32 UMetasoundEditorGraphInput::GetSortOrderIndex() const
 {
 	using namespace Metasound::Editor;
@@ -819,6 +867,17 @@ UMetasoundEditorGraphNode* UMetasoundEditorGraphOutput::AddNode(Metasound::Front
 	return FGraphBuilder::AddOutputNode(Graph->GetMetasoundChecked(), InNodeHandle, bInSelectNewNode);
 }
 
+FText UMetasoundEditorGraphOutput::GetDescription() const
+{
+	const FMetaSoundFrontendDocumentBuilder& Builder = GetFrontendBuilderChecked();
+	if (const FMetasoundFrontendClassOutput* Output = Builder.FindGraphOutput(GetMemberName()))
+	{
+		return Output->Metadata.GetDescription();
+	}
+
+	return { };
+}
+
 int32 UMetasoundEditorGraphOutput::GetSortOrderIndex() const
 {
 	using namespace Metasound::Editor;
@@ -1512,10 +1571,10 @@ void UMetasoundEditorGraph::PreSave(FObjectPreSaveContext InSaveContext)
 	Super::PreSave(InSaveContext);
 }
 
-UMetaSoundBuilderBase& UMetasoundEditorGraph::GetBuilderChecked()
+UMetaSoundBuilderBase& UMetasoundEditorGraph::GetBuilderChecked() const
 {
 	using namespace Metasound::Engine;
-	return FDocumentBuilderRegistry::GetChecked().FindOrBeginBuilding(*GetOutermostObject());
+	return FDocumentBuilderRegistry::GetChecked().FindOrBeginBuilding(GetMetasoundChecked());
 }
 
 UObject* UMetasoundEditorGraph::GetMetasound() const
diff --git a/Engine/Plugins/Runtime/Metasound/Source/MetasoundEditor/Private/MetasoundEditorGraphInputNode.cpp b/Engine/Plugins/Runtime/Metasound/Source/MetasoundEditor/Private/MetasoundEditorGraphInputNode.cpp
--- a/Engine/Plugins/Runtime/Metasound/Source/MetasoundEditor/Private/MetasoundEditorGraphInputNode.cpp
+++ b/Engine/Plugins/Runtime/Metasound/Source/MetasoundEditor/Private/MetasoundEditorGraphInputNode.cpp
@@ -3,6 +3,7 @@
 
 #include "Algo/Count.h"
 #include "Algo/Transform.h"
+#include "AudioParameter.h"
 #include "AudioParameterControllerInterface.h"
 #include "EdGraph/EdGraphNode.h"
 #include "GraphEditorSettings.h"
@@ -158,35 +159,32 @@ void UMetasoundEditorGraphInputNode::ReconstructNode()
 
 void UMetasoundEditorGraphInputNode::Validate(Metasound::Editor::FGraphNodeValidationResult& OutResult)
 {
-#if WITH_EDITOR
 	using namespace Metasound::Editor;
 	using namespace Metasound::Frontend;
 
 	Super::Validate(OutResult);
 
-	FConstNodeHandle NodeHandle = GetConstNodeHandle();
-	const FMetasoundFrontendClassMetadata& Metadata = NodeHandle->GetClassMetadata();
-
-	const FMetasoundFrontendVersion& MetasoundFrontendVersion = NodeHandle->GetInterfaceVersion();
-
-	FName InterfaceNameToValidate = MetasoundFrontendVersion.Name;
-	FMetasoundFrontendInterface InterfaceToValidate;
-	if (ISearchEngine::Get().FindInterfaceWithHighestVersion(InterfaceNameToValidate, InterfaceToValidate))
+	if (const UMetasoundEditorGraphVertex* Vertex = Cast<UMetasoundEditorGraphVertex>(GetMember()))
 	{
-		const FName& NodeName = NodeHandle->GetNodeName();
-		FText RequiredText;
-		if (InterfaceToValidate.IsMemberInputRequired(NodeName, RequiredText))
+		FMetasoundFrontendInterface InterfaceToValidate;
+		if (Vertex->IsInterfaceMember(&InterfaceToValidate))
 		{
-			TArray<FConstOutputHandle> OutputHandles = NodeHandle->GetConstOutputs();
-			if (ensure(!OutputHandles.IsEmpty()))
+			FText RequiredText;
+			if (InterfaceToValidate.IsMemberOutputRequired(Vertex->GetMemberName(), RequiredText))
 			{
-				const FConstOutputHandle& OutputHandle = OutputHandles.Last();
-				if (!OutputHandle->IsConnected())
+				if (const FMetasoundFrontendNode* Node = GetFrontendNode())
 				{
-					OutResult.SetMessage(EMessageSeverity::Warning, *RequiredText.ToString());
+					const TArray<FMetasoundFrontendVertex>& Outputs = Node->Interface.Outputs;
+					if (ensure(!Outputs.IsEmpty()))
+					{
+						const FMetaSoundFrontendDocumentBuilder& Builder = GetBuilderChecked().GetConstBuilder();
+						if (!Builder.IsNodeOutputConnected(Node->GetID(), Outputs.Last().VertexID))
+						{
+							OutResult.SetMessage(EMessageSeverity::Warning, *RequiredText.ToString());
+						}
+					}
 				}
 			}
 		}
 	}
-#endif // #if WITH_EDITOR
 }
diff --git a/Engine/Plugins/Runtime/Metasound/Source/MetasoundEditor/Private/MetasoundEditorGraphNode.cpp b/Engine/Plugins/Runtime/Metasound/Source/MetasoundEditor/Private/MetasoundEditorGraphNode.cpp
--- a/Engine/Plugins/Runtime/Metasound/Source/MetasoundEditor/Private/MetasoundEditorGraphNode.cpp
+++ b/Engine/Plugins/Runtime/Metasound/Source/MetasoundEditor/Private/MetasoundEditorGraphNode.cpp
@@ -903,31 +903,32 @@ void UMetasoundEditorGraphOutputNode::Validate(Metasound::Editor::FGraphNodeVali
 
 	Super::Validate(OutResult);
 
-	// 2. Check if node is invalid, version is missing and cache if interface changes exist between the document's records and the registry
-	FNodeHandle NodeHandle = GetNodeHandle();
-	const FMetasoundFrontendVersion& MetasoundFrontendVersion = NodeHandle->GetInterfaceVersion();
-
-	FName InterfaceNameToValidate = MetasoundFrontendVersion.Name;
-	FMetasoundFrontendInterface InterfaceToValidate;
-	if (ISearchEngine::Get().FindInterfaceWithHighestVersion(InterfaceNameToValidate, InterfaceToValidate))
+	if (const UMetasoundEditorGraphVertex* Vertex = Cast<UMetasoundEditorGraphVertex>(GetMember()))
 	{
-		const FName& NodeName = NodeHandle->GetNodeName();
-		FText RequiredText;
-		if (InterfaceToValidate.IsMemberOutputRequired(NodeName, RequiredText))
+		FMetasoundFrontendInterface InterfaceToValidate;
+		if (Vertex->IsInterfaceMember(&InterfaceToValidate))
 		{
-			TArray<FConstInputHandle> InputHandles = NodeHandle->GetConstInputs();
-			if (ensure(!InputHandles.IsEmpty()))
+			FText RequiredText;
+			if (InterfaceToValidate.IsMemberOutputRequired(Vertex->GetMemberName(), RequiredText))
 			{
-				bool bIsConnected = InputHandles.Last()->IsConnected();
-				if (!bIsConnected)
+				if (const FMetasoundFrontendNode* Node = GetFrontendNode())
 				{
-					OutResult.SetMessage(EMessageSeverity::Warning, *RequiredText.ToString());
+					const TArray<FMetasoundFrontendVertex>& Inputs = Node->Interface.Inputs;
+					if (ensure(!Inputs.IsEmpty()))
+					{
+						const FMetaSoundFrontendDocumentBuilder& Builder = GetBuilderChecked().GetConstBuilder();
+						if (!Builder.IsNodeInputConnected(Node->GetID(), Inputs.Last().VertexID))
+						{
+							OutResult.SetMessage(EMessageSeverity::Warning, *RequiredText.ToString());
+						}
+					}
 				}
 			}
 		}
 	}
 }
 
+
 const FMetasoundEditorGraphNodeBreadcrumb& UMetasoundEditorGraphOutputNode::GetBreadcrumb() const
 {
 	return Breadcrumb;
@@ -976,7 +977,13 @@ FGuid UMetasoundEditorGraphOutputNode::GetNodeID() const
 
 bool UMetasoundEditorGraphOutputNode::CanUserDeleteNode() const
 {
-	return !GetNodeHandle()->IsInterfaceMember();
+	if (const UMetasoundEditorGraphVertex* Vertex = Cast<UMetasoundEditorGraphVertex>(GetMember()))
+	{
+		return Vertex->IsInterfaceMember();
+	}
+	
+	return true;	
+
 }
 
 void UMetasoundEditorGraphOutputNode::SetNodeID(FGuid InNodeID)
diff --git a/Engine/Plugins/Runtime/Metasound/Source/MetasoundEditor/Public/MetasoundEditorGraph.h b/Engine/Plugins/Runtime/Metasound/Source/MetasoundEditor/Public/MetasoundEditorGraph.h
--- a/Engine/Plugins/Runtime/Metasound/Source/MetasoundEditor/Public/MetasoundEditorGraph.h
+++ b/Engine/Plugins/Runtime/Metasound/Source/MetasoundEditor/Public/MetasoundEditorGraph.h
@@ -134,6 +134,8 @@ public:
 	 */
 	virtual void UpdateFrontendDefaultLiteral(bool bPostTransaction) PURE_VIRTUAL(UMetasoundEditorGraphMember::UpdateFrontendDefaultLiteral, );
 
+	FMetaSoundFrontendDocumentBuilder& GetFrontendBuilderChecked() const;
+
 	/** Returns the parent MetaSound Graph. If the Outer object of the member is non
 	 * a UMetasoundEditorGraph, returns a nullptr. */
 	UMetasoundEditorGraph* GetOwningGraph();
@@ -188,6 +190,8 @@ protected:
 	/** Adds the node handle for a newly created vertex. */
 	virtual Metasound::Frontend::FNodeHandle AddNodeHandle(const FName& InNodeName, const Metasound::Editor::FCreateNodeVertexParams& InParams) PURE_VIRTUAL(UMetasoundEditorGraphVertex::AddNodeHandle, return Metasound::Frontend::INodeController::GetInvalidHandle(); )
 
+	const FMetasoundFrontendNode* GetFrontendNode() const;
+
 public:
 	/** Initializes all properties with the given parameters required to identify the frontend member from this editor graph member. */
 	void InitMember(FName InDataType, const FMetasoundFrontendLiteral& InDefaultLiteral, FGuid InNodeID, FMetasoundFrontendClassName&& InClassName);
@@ -203,7 +207,6 @@ public:
 	/* ~Begin UMetasoundEditorGraphMember interface */
 	virtual FGuid GetMemberID() const override;
 	virtual FName GetMemberName() const override;
-	virtual FText GetDescription() const override;
 	virtual FText GetDisplayName() const override;
 
 	virtual bool CanRename(const FText& InNewName, FText& OutError) const override;
@@ -216,8 +219,8 @@ public:
 	/** Version of interface membership, or invalid version if not an interface member. */
 	virtual const FMetasoundFrontendVersion& GetInterfaceVersion() const;
 
-	/** Returns true if member is part of an interface. */
-	virtual bool IsInterfaceMember() const;
+	/** Returns true if member is part of an interface. If supplied interface pointer, sets pointer's data to the interface vertex is member of. */
+	virtual bool IsInterfaceMember(FMetasoundFrontendInterface* OutInterface = nullptr) const;
 
 	/** Returns the Metasound class type of the associated node */
 	virtual EMetasoundFrontendClassType GetClassType() const PURE_VIRTUAL(UMetasoundEditorGraphVertex::GetClassType, return EMetasoundFrontendClassType::Invalid; )
@@ -253,6 +256,7 @@ class METASOUNDEDITOR_API UMetasoundEditorGraphInput : public UMetasoundEditorGr
 	GENERATED_BODY()
 
 public:
+	virtual FText GetDescription() const override;
 	virtual int32 GetSortOrderIndex() const override;
 	virtual void SetSortOrderIndex(int32 InSortOrderIndex) override;
 	virtual TArray<UMetasoundEditorGraphMemberNode*> GetNodes() const override;
@@ -280,6 +284,7 @@ class METASOUNDEDITOR_API UMetasoundEditorGraphOutput : public UMetasoundEditorG
 	GENERATED_BODY()
 
 public:
+	virtual FText GetDescription() const override;
 	virtual int32 GetSortOrderIndex() const override;
 	virtual void SetSortOrderIndex(int32 InSortOrderIndex) override;
 	virtual const FText& GetGraphMemberLabel() const override;
@@ -377,7 +382,7 @@ public:
 
 	virtual void PreSave(FObjectPreSaveContext InSaveContext) override;
 
-	UMetaSoundBuilderBase& GetBuilderChecked();
+	UMetaSoundBuilderBase& GetBuilderChecked() const;
 	UObject* GetMetasound() const;
 	UObject& GetMetasoundChecked() const;
 
diff --git a/Engine/Plugins/Runtime/Metasound/Source/MetasoundEngine/Private/Interfaces/MetasoundInterface.cpp b/Engine/Plugins/Runtime/Metasound/Source/MetasoundEngine/Private/Interfaces/MetasoundInterface.cpp
--- a/Engine/Plugins/Runtime/Metasound/Source/MetasoundEngine/Private/Interfaces/MetasoundInterface.cpp
+++ b/Engine/Plugins/Runtime/Metasound/Source/MetasoundEngine/Private/Interfaces/MetasoundInterface.cpp
@@ -55,6 +55,11 @@ namespace Metasound::Engine
 		return Interface;
 	}
 
+	bool FInterfaceRegistryEntry::IsDeprecated() const
+	{
+		return UpdateTransform.IsValid();
+	}
+
 	bool FInterfaceRegistryEntry::UpdateRootGraphInterface(Frontend::FDocumentHandle InDocument) const
 	{
 		if (UpdateTransform.IsValid())
diff --git a/Engine/Plugins/Runtime/Metasound/Source/MetasoundEngine/Private/Interfaces/MetasoundInterface.h b/Engine/Plugins/Runtime/Metasound/Source/MetasoundEngine/Private/Interfaces/MetasoundInterface.h
--- a/Engine/Plugins/Runtime/Metasound/Source/MetasoundEngine/Private/Interfaces/MetasoundInterface.h
+++ b/Engine/Plugins/Runtime/Metasound/Source/MetasoundEngine/Private/Interfaces/MetasoundInterface.h
@@ -19,6 +19,7 @@ namespace Metasound::Engine
 
 		virtual FName GetRouterName() const override;
 		virtual const FMetasoundFrontendInterface& GetInterface() const override;
+		virtual bool IsDeprecated() const override;
 		virtual bool UpdateRootGraphInterface(Frontend::FDocumentHandle InDocument) const override;
 
 	private:
diff --git a/Engine/Plugins/Runtime/Metasound/Source/MetasoundFrontend/Private/Interfaces/MetasoundFrontendInterfaceRegistry.cpp b/Engine/Plugins/Runtime/Metasound/Source/MetasoundFrontend/Private/Interfaces/MetasoundFrontendInterfaceRegistry.cpp
--- a/Engine/Plugins/Runtime/Metasound/Source/MetasoundFrontend/Private/Interfaces/MetasoundFrontendInterfaceRegistry.cpp
+++ b/Engine/Plugins/Runtime/Metasound/Source/MetasoundFrontend/Private/Interfaces/MetasoundFrontendInterfaceRegistry.cpp
@@ -1,6 +1,7 @@
 // Copyright Epic Games, Inc. All Rights Reserved.
 #include "Interfaces/MetasoundFrontendInterfaceRegistry.h"
 
+#include "AudioParameter.h"
 #include "HAL/PlatformTime.h"
 #include "Interfaces/MetasoundFrontendInterfaceBindingRegistry.h"
 #include "MetasoundFrontendInterfaceRegistryPrivate.h"
@@ -22,7 +23,7 @@ namespace Metasound::Frontend
 		FInterfaceRegistryTransaction::FTimeType TransactionTime = FPlatformTime::Cycles64();
 		if (InEntry.IsValid())
 		{
-			FInterfaceRegistryKey Key = GetInterfaceRegistryKey(InEntry->GetInterface());
+			const FInterfaceRegistryKey Key = GetInterfaceRegistryKey(InEntry->GetInterface());
 			if (IsValidInterfaceRegistryKey(Key))
 			{
 				if (const IInterfaceRegistryEntry* Entry = FindInterfaceRegistryEntry(Key))
@@ -32,7 +33,36 @@ namespace Metasound::Frontend
 					FInterfaceRegistryTransaction Transaction{FInterfaceRegistryTransaction::ETransactionType::InterfaceUnregistration, Key, Entry->GetInterface().Version, TransactionTime};
 					TransactionBuffer->AddTransaction(MoveTemp(Transaction));
 				}
-					
+
+#if WITH_EDITOR
+			// Don't run vertex name validation warning for deprecated registry entries. Some of these may be
+			// versioning schema that have a subsequent version or versions fixing the very problem this log is reporting.
+			if (!InEntry->IsDeprecated())
+			{
+				const FName InterfaceNamespace = InEntry->GetInterface().Version.Name;
+				auto LogIfMismatch = [this, &InterfaceNamespace](FName VertexName)
+				{
+					FName VertexNamespace;
+					if (!IsInterfaceVertexNameValid(InterfaceNamespace, VertexName, &VertexNamespace))
+					{
+						UE_LOG(LogMetaSound, Warning, TEXT("Interface '%s' contains vertex '%s' with mismatched namespace '%s': "
+							"All interface-defined vertices' must start with matching interface namespace (See AUDIO_PARAMETER_INTERFACE_MEMBER_DEFINE/AUDIO_PARAMETER_INTERFACE_NAMESPACE macro to ensure convention is followed). "
+							"Failing to fix relationship via interface versioning will fail validation/cook in future builds."), *InterfaceNamespace.ToString(), *VertexName.ToString(), *VertexNamespace.ToString());
+					}
+				};
+
+				for (const FMetasoundFrontendClassVertex& Vertex : InEntry->GetInterface().Inputs)
+				{
+					LogIfMismatch(Vertex.Name);
+				}
+
+				for (const FMetasoundFrontendClassVertex& Vertex : InEntry->GetInterface().Outputs)
+				{
+					LogIfMismatch(Vertex.Name);
+				}
+			}
+#endif // WITH_EDITOR
+
 				FInterfaceRegistryTransaction Transaction{FInterfaceRegistryTransaction::ETransactionType::InterfaceRegistration, Key, InEntry->GetInterface().Version, TransactionTime};
 				TransactionBuffer->AddTransaction(MoveTemp(Transaction));
 				Entries.Add(Key, MoveTemp(InEntry));
@@ -118,6 +148,18 @@ namespace Metasound::Frontend
 		return Registry;
 	}
 
+	bool FInterfaceRegistry::IsInterfaceVertexNameValid(FName InterfaceNamespace, FName FullVertexName, FName* VertexNamespace) const
+	{
+		FName Namespace;
+		FName Name;
+		Audio::FParameterPath::SplitName(FullVertexName, Namespace, Name);
+		if (VertexNamespace)
+		{
+			*VertexNamespace = Namespace;
+		}
+		return InterfaceNamespace == Namespace;
+	}
+
 	IInterfaceRegistry& IInterfaceRegistry::Get()
 	{
 		return FInterfaceRegistry::Get();
diff --git a/Engine/Plugins/Runtime/Metasound/Source/MetasoundFrontend/Private/Interfaces/MetasoundFrontendInterfaceRegistryPrivate.h b/Engine/Plugins/Runtime/Metasound/Source/MetasoundFrontend/Private/Interfaces/MetasoundFrontendInterfaceRegistryPrivate.h
--- a/Engine/Plugins/Runtime/Metasound/Source/MetasoundFrontend/Private/Interfaces/MetasoundFrontendInterfaceRegistryPrivate.h
+++ b/Engine/Plugins/Runtime/Metasound/Source/MetasoundFrontend/Private/Interfaces/MetasoundFrontendInterfaceRegistryPrivate.h
@@ -29,6 +29,8 @@ namespace Metasound::Frontend
 		TUniquePtr<FInterfaceTransactionStream> CreateTransactionStream();
 
 	private:
+		bool IsInterfaceVertexNameValid(FName InterfaceNamespace, FName FullVertexName, FName* VertexNamespace) const;
+
 		using FInterfaceTransactionBuffer = TTransactionBuffer<FInterfaceRegistryTransaction>;
 
 		TMap<FInterfaceRegistryKey, TUniquePtr<IInterfaceRegistryEntry>> Entries;
diff --git a/Engine/Plugins/Runtime/Metasound/Source/MetasoundFrontend/Private/MetasoundFrontendDocumentVersioning.cpp b/Engine/Plugins/Runtime/Metasound/Source/MetasoundFrontend/Private/MetasoundFrontendDocumentVersioning.cpp
--- a/Engine/Plugins/Runtime/Metasound/Source/MetasoundFrontend/Private/MetasoundFrontendDocumentVersioning.cpp
+++ b/Engine/Plugins/Runtime/Metasound/Source/MetasoundFrontend/Private/MetasoundFrontendDocumentVersioning.cpp
@@ -581,6 +581,21 @@ namespace Metasound::Frontend
 		}
 	};
 
+	namespace VersioningPrivate
+	{
+		bool VersionBuilderDocument(FMetasoundAssetBase& InAssetBase, FMetaSoundFrontendDocumentBuilder& Builder)
+		{
+			UObject& DocObject = Builder.CastDocumentObjectChecked<UObject>();
+			const FName Name = DocObject.GetFName();
+			const FString Path = DocObject.GetPathName();
+
+			bool bWasUpdated = false;
+			bWasUpdated |= FVersionDocument_1_12(InAssetBase, Name, Path).Transform(Builder);
+
+			return bWasUpdated;
+		}
+	} // namespace VersioningPrivate
+
 	bool VersionDocument(FMetasoundAssetBase& InAssetBase)
 	{
 		bool bWasUpdated = false;
@@ -588,14 +603,13 @@ namespace Metasound::Frontend
 		UObject* OwningAsset = InAssetBase.GetOwningAsset();
 		TScriptInterface<IMetaSoundDocumentInterface> DocumentInterface(OwningAsset);
 		const FMetasoundFrontendDocument& Document = DocumentInterface->GetConstDocument();
+		const FName Name = FName(OwningAsset->GetName());
+		const FString Path = OwningAsset->GetPathName();
 
 		// Copied as value will be mutated with each applicable transform below
 		const FMetasoundFrontendVersionNumber InitVersionNumber = Document.Metadata.Version.Number;
 		if (InitVersionNumber < GetMaxDocumentVersion())
 		{
-			FName Name = FName(OwningAsset->GetName());
-			FString Path = OwningAsset->GetPathName();
-
 			// Controller (Soft Deprecated) Transforms
 			if (InitVersionNumber.Major == 1 && InitVersionNumber.Minor < 12)
 			{
@@ -613,16 +627,16 @@ namespace Metasound::Frontend
 				bWasUpdated |= FVersionDocument_1_9(Name, Path).Transform(DocHandle);
 				bWasUpdated |= FVersionDocument_1_10().Transform(DocHandle);
 				bWasUpdated |= FVersionDocument_1_11().Transform(DocHandle);
+				// No longer supported, new versions should go in VersioningPrivate::VersionBuilderDocument
 			}
 
-			// Builder Transforms
-			{
-				// Mutation of document via both controller & builder systems simultaneously is forbidden as:
-				// 1. Builders were implemented post document version 1.11, so earlier versions are not supported.
-				// 2. Controller mutations are not tracked by analogous builder, which can cause internal cache corruption.
-				FMetaSoundFrontendDocumentBuilder Builder(DocumentInterface);
-				bWasUpdated |= FVersionDocument_1_12(InAssetBase, Name, Path).Transform(Builder);
-			}
+#if WITH_EDITORONLY_DATA
+			FMetaSoundFrontendDocumentBuilder& Builder = IDocumentBuilderRegistry::GetChecked().FindOrBeginBuilding(DocumentInterface);
+			bWasUpdated |= VersioningPrivate::VersionBuilderDocument(InAssetBase, Builder);
+#else
+			FMetaSoundFrontendDocumentBuilder Builder(DocumentInterface);
+			bWasUpdated |= VersioningPrivate::VersionBuilderDocument(InAssetBase, Builder);
+#endif // WITH_EDITORONLY_DATA
 
 			if (bWasUpdated)
 			{
diff --git a/Engine/Plugins/Runtime/Metasound/Source/MetasoundFrontend/Public/Interfaces/MetasoundFrontendInterfaceRegistry.h b/Engine/Plugins/Runtime/Metasound/Source/MetasoundFrontend/Public/Interfaces/MetasoundFrontendInterfaceRegistry.h
--- a/Engine/Plugins/Runtime/Metasound/Source/MetasoundFrontend/Public/Interfaces/MetasoundFrontendInterfaceRegistry.h
+++ b/Engine/Plugins/Runtime/Metasound/Source/MetasoundFrontend/Public/Interfaces/MetasoundFrontendInterfaceRegistry.h
@@ -23,11 +23,8 @@ namespace Metasound::Frontend
 		// MetaSound Interface definition
 		virtual const FMetasoundFrontendInterface& GetInterface() const = 0;
 
-		UE_DEPRECATED(5.3, "UClass name must now be provided to determine if interface is to be apply by default.")
-		virtual bool IsDefault() const { return false; }
-
-		UE_DEPRECATED(5.3, "Option (renamed 'IsModifiable') is now data driven and accessibly from the interface (via GetInterface) per UClass option. If no class options found, addition/removal is now considered to be implicitly true")
-		virtual bool EditorCanAddOrRemove() const { return false; }
+		//  Whether or not entry is deprecated or not
+		virtual bool IsDeprecated() const = 0;
 
 		// Name of routing system used to update interface inputs (ex. ParameterInterface or DataReference).
 		virtual FName GetRouterName() const = 0;
diff --git a/Engine/Plugins/Runtime/Metasound/Source/MetasoundFrontend/Public/MetasoundFrontendSearchEngine.h b/Engine/Plugins/Runtime/Metasound/Source/MetasoundFrontend/Public/MetasoundFrontendSearchEngine.h
--- a/Engine/Plugins/Runtime/Metasound/Source/MetasoundFrontend/Public/MetasoundFrontendSearchEngine.h
+++ b/Engine/Plugins/Runtime/Metasound/Source/MetasoundFrontend/Public/MetasoundFrontendSearchEngine.h
@@ -27,8 +27,10 @@ namespace Metasound::Frontend
 			/** Find the class with the given ClassName & Major Version. Returns false if not found, true if found. */
 			virtual bool FindClassWithHighestMinorVersion(const FMetasoundFrontendClassName& InName, int32 InMajorVersion, FMetasoundFrontendClass& OutClass) = 0;
 
+			/** Finds all registered interfaces with the given name. */
 			virtual TArray<FMetasoundFrontendVersion> FindAllRegisteredInterfacesWithName(FName InInterfaceName) = 0;
 
+			/** Finds the registered interface with the highest version of the given name. Returns true if found, false if not. */
 			virtual bool FindInterfaceWithHighestVersion(FName InInterfaceName, FMetasoundFrontendInterface& OutInterface) = 0;
 
 			UE_DEPRECATED(5.3, "Use ISearchEngine::FindUClassDefaultInterfaceVersions using TopLevelAssetPath instead.")
