[Metasound] Plate reverb node + optimized stereo, non-interleaved render

#tests automated tests
[RN] New MetaSound node: Plate Reverb (Stereo)

#virtualized

diff --git a/Engine/Plugins/Runtime/Metasound/Source/MetasoundStandardNodes/Private/MetasoundPlateReverbNode.cpp b/Engine/Plugins/Runtime/Metasound/Source/MetasoundStandardNodes/Private/MetasoundPlateReverbNode.cpp
new file mode 100644
--- /dev/null
+++ b/Engine/Plugins/Runtime/Metasound/Source/MetasoundStandardNodes/Private/MetasoundPlateReverbNode.cpp
@@ -0,0 +1,208 @@
+// Copyright Epic Games, Inc. All Rights Reserved.
+
+#include "MetasoundDataTypeRegistrationMacro.h"
+#include "MetasoundExecutableOperator.h"
+#include "MetasoundFacade.h"
+#include "MetasoundNodeRegistrationMacro.h"
+#include "MetasoundParamHelper.h"
+#include "MetasoundStandardNodesCategories.h"
+#include "MetasoundStandardNodesNames.h"
+#include "DSP/FloatArrayMath.h"
+#include "DSP/ReverbFast.h"
+#include "Internationalization/Text.h"
+
+#define LOCTEXT_NAMESPACE "MetasoundStandardNodes_PlateReverb"
+
+namespace Metasound
+{
+	namespace PlateReverb
+	{
+		namespace Inputs
+		{
+			METASOUND_PARAM(Bypass, "Bypass", "Toggle to bypass the effect and send audio through unaltered.")
+			METASOUND_PARAM(AudioLeft, "In Left", "Left channel audio input.")
+			METASOUND_PARAM(AudioRight, "In Right", "Right channel audio input.")
+			METASOUND_PARAM(DryLevel, "Dry Level", "The level of the dry signal (linear).")
+			METASOUND_PARAM(WetLevel, "Wet Level", "The level of the wet signal (linear).")
+		}
+
+		namespace Outputs
+		{
+			METASOUND_PARAM(AudioLeft, "Out Left", "Left channel audio output.")
+			METASOUND_PARAM(AudioRight, "Out Right", "Right channel audio output.")
+		}
+	}
+
+	class FPlateReverbOperator final : public TExecutableOperator<FPlateReverbOperator>
+	{
+	public:
+		struct FInputs
+		{
+			FBoolReadRef Bypass;
+			FAudioBufferReadRef AudioLeft;
+			FAudioBufferReadRef AudioRight;
+			FFloatReadRef DryLevel;
+			FFloatReadRef WetLevel;
+		};
+		
+		struct FOutputs
+		{
+			FAudioBufferWriteRef AudioLeft;
+			FAudioBufferWriteRef AudioRight;
+		};
+		
+		static const FNodeClassMetadata& GetNodeInfo()
+		{
+			auto CreateNodeClassMetadata = []() -> FNodeClassMetadata
+			{
+				FNodeClassMetadata Metadata
+				{
+					FNodeClassName { StandardNodes::Namespace, "Plate Reverb", "Stereo" },
+					1, // Major Version
+					0, // Minor Version
+					METASOUND_LOCTEXT("PlateReverbStereoDisplayName", "Plate Reverb (Stereo)"),
+					METASOUND_LOCTEXT("PlateReverbDesc", "Plate reverb with configurable early and late reflections."),
+					PluginAuthor,
+					PluginNodeMissingPrompt,
+					GetVertexInterface(),
+					{ NodeCategories::Reverbs },
+					{},
+					FNodeDisplayStyle()
+				};
+
+				return Metadata;
+			};
+
+			static const FNodeClassMetadata Metadata = CreateNodeClassMetadata();
+			return Metadata;
+		}
+		
+		static const FVertexInterface& GetVertexInterface()
+		{
+			static const FVertexInterface Interface
+			{
+				FInputVertexInterface
+				{
+					TInputDataVertex<bool>(METASOUND_GET_PARAM_NAME_AND_METADATA(PlateReverb::Inputs::Bypass), false),
+					TInputDataVertex<FAudioBuffer>(METASOUND_GET_PARAM_NAME_AND_METADATA(PlateReverb::Inputs::AudioLeft)),
+					TInputDataVertex<FAudioBuffer>(METASOUND_GET_PARAM_NAME_AND_METADATA(PlateReverb::Inputs::AudioRight)),
+					TInputDataVertex<float>(METASOUND_GET_PARAM_NAME_AND_METADATA(PlateReverb::Inputs::DryLevel), 1.0f),
+					TInputDataVertex<float>(METASOUND_GET_PARAM_NAME_AND_METADATA(PlateReverb::Inputs::WetLevel), 1.0f)
+				},
+				FOutputVertexInterface
+				{
+					TOutputDataVertex<FAudioBuffer>(METASOUND_GET_PARAM_NAME_AND_METADATA(PlateReverb::Outputs::AudioLeft)),
+					TOutputDataVertex<FAudioBuffer>(METASOUND_GET_PARAM_NAME_AND_METADATA(PlateReverb::Outputs::AudioRight))
+				}
+			};
+
+			return Interface;
+		}
+		
+		static TUniquePtr<IOperator> CreateOperator(const FBuildOperatorParams& InParams, FBuildResults& OutResults)
+		{
+			FInputs Inputs
+			{
+				InParams.InputData.GetOrCreateDefaultDataReadReference<bool>(PlateReverb::Inputs::BypassName, InParams.OperatorSettings),
+				InParams.InputData.GetOrConstructDataReadReference<FAudioBuffer>(PlateReverb::Inputs::AudioLeftName, InParams.OperatorSettings),
+				InParams.InputData.GetOrConstructDataReadReference<FAudioBuffer>(PlateReverb::Inputs::AudioRightName, InParams.OperatorSettings),
+				InParams.InputData.GetOrCreateDefaultDataReadReference<float>(PlateReverb::Inputs::DryLevelName, InParams.OperatorSettings),
+				InParams.InputData.GetOrCreateDefaultDataReadReference<float>(PlateReverb::Inputs::WetLevelName, InParams.OperatorSettings)
+			};
+
+			return MakeUnique<FPlateReverbOperator>(InParams, MoveTemp(Inputs));
+		}
+
+		FPlateReverbOperator(const FBuildOperatorParams& BuildParams, FInputs&& Inputs)
+			: Inputs(MoveTemp(Inputs))
+			, Outputs({ FAudioBufferWriteRef::CreateNew(BuildParams.OperatorSettings), FAudioBufferWriteRef::CreateNew(BuildParams.OperatorSettings) })
+			, Reverb(BuildParams.OperatorSettings.GetSampleRate())
+		{
+			Reset(BuildParams);
+		}
+
+		virtual void BindInputs(FInputVertexInterfaceData& InOutVertexData) override
+		{
+			InOutVertexData.BindReadVertex(PlateReverb::Inputs::BypassName, Inputs.Bypass);
+			InOutVertexData.BindReadVertex(PlateReverb::Inputs::AudioLeftName, Inputs.AudioLeft);
+			InOutVertexData.BindReadVertex(PlateReverb::Inputs::AudioRightName, Inputs.AudioRight);
+			InOutVertexData.BindReadVertex(PlateReverb::Inputs::DryLevelName, Inputs.DryLevel);
+			InOutVertexData.BindReadVertex(PlateReverb::Inputs::WetLevelName, Inputs.WetLevel);
+		}
+		
+		virtual void BindOutputs(FOutputVertexInterfaceData& InOutVertexData) override
+		{
+			InOutVertexData.BindReadVertex(PlateReverb::Outputs::AudioLeftName, Outputs.AudioLeft);
+			InOutVertexData.BindReadVertex(PlateReverb::Outputs::AudioRightName, Outputs.AudioRight);
+		}
+
+		void Reset(const FResetParams&)
+		{
+		}
+		
+		void Execute()
+		{
+			const int32 NumFrames = Inputs.AudioLeft->Num();
+			check(Inputs.AudioRight->Num() == NumFrames);
+			
+			if (*Inputs.Bypass)
+			{
+				FMemory::Memcpy(Outputs.AudioLeft->GetData(), Inputs.AudioLeft->GetData(), NumFrames * sizeof(float));
+				FMemory::Memcpy(Outputs.AudioRight->GetData(), Inputs.AudioRight->GetData(), NumFrames * sizeof(float));
+				return;
+			}
+
+			// Copy the inputs to the work buffers
+			WorkBufferLeft.SetNumUninitialized(NumFrames);
+			WorkBufferRight.SetNumUninitialized(NumFrames);
+			FMemory::Memcpy(WorkBufferLeft.GetData(), Inputs.AudioLeft->GetData(), NumFrames * sizeof(float));
+			FMemory::Memcpy(WorkBufferRight.GetData(), Inputs.AudioRight->GetData(), NumFrames * sizeof(float));
+			
+			const float CurrentWetLevel = FMath::Clamp(*Inputs.WetLevel, 0.0f, 1.0f);
+
+			// Apply the wet gain to the input to preserve the reverb tail
+			if (LastWetLevel >= 0.0f && !FMath::IsNearlyEqual(CurrentWetLevel, LastWetLevel))
+			{
+				Audio::ArrayFade(WorkBufferLeft, LastWetLevel, CurrentWetLevel);
+				Audio::ArrayFade(WorkBufferRight, LastWetLevel, CurrentWetLevel);
+			}
+			else
+			{
+				Audio::ArrayMultiplyByConstantInPlace(WorkBufferLeft, CurrentWetLevel);
+				Audio::ArrayMultiplyByConstantInPlace(WorkBufferRight, CurrentWetLevel);
+			}
+			
+			LastWetLevel = CurrentWetLevel;
+
+			// Process
+			Reverb.ProcessAudioStereoNonInterleaved(WorkBufferLeft, WorkBufferRight, *Outputs.AudioLeft, *Outputs.AudioRight);
+
+			// Mix in the dry signal
+			const float DryLevel = FMath::Clamp(*Inputs.DryLevel, 0.0f, 1.0f);
+			Audio::ArrayMixIn(*Inputs.AudioLeft, *Outputs.AudioLeft, DryLevel);
+			Audio::ArrayMixIn(*Inputs.AudioRight, *Outputs.AudioRight, DryLevel);
+		}
+
+	private:
+		FInputs Inputs;
+		FOutputs Outputs;
+		
+		Audio::FPlateReverbFast Reverb;
+		Audio::FAlignedFloatBuffer WorkBufferLeft;
+		Audio::FAlignedFloatBuffer WorkBufferRight;
+		float LastWetLevel{ -1 };
+	};
+
+	class FPlateReverbNode final : public FNodeFacade
+	{
+	public:
+		explicit FPlateReverbNode(const FNodeInitData& InitData)
+			: FNodeFacade(InitData.InstanceName, InitData.InstanceID, TFacadeOperatorClass<FPlateReverbOperator>())
+		{
+		}
+	};
+
+	METASOUND_REGISTER_NODE(FPlateReverbNode)
+}
+
+#undef LOCTEXT_NAMESPACE
diff --git a/Engine/Plugins/Runtime/Metasound/Source/MetasoundStandardNodes/Private/MetasoundStandardNodesCategories.cpp b/Engine/Plugins/Runtime/Metasound/Source/MetasoundStandardNodes/Private/MetasoundStandardNodesCategories.cpp
--- a/Engine/Plugins/Runtime/Metasound/Source/MetasoundStandardNodes/Private/MetasoundStandardNodesCategories.cpp
+++ b/Engine/Plugins/Runtime/Metasound/Source/MetasoundStandardNodes/Private/MetasoundStandardNodesCategories.cpp
@@ -22,6 +22,7 @@ namespace Metasound::NodeCategories
 	const FText Spatialization = { METASOUND_LOCTEXT("Metasound_SpatializationCategory", "Spatialization") };
 	const FText Trigger = { METASOUND_LOCTEXT("Metasound_TriggerCategory", "Triggers") };
 	const FText WaveTables = { METASOUND_LOCTEXT("Metasound_WaveTableCategory", "WaveTables") };
+	const FText Reverbs = { METASOUND_LOCTEXT("Metasound_ReverbCategory", "Reverbs") };
 } // Metasound::NodeCategories
 
 #undef LOCTEXT_NAMESPACE
\ No newline at end of file
diff --git a/Engine/Plugins/Runtime/Metasound/Source/MetasoundStandardNodes/Public/MetasoundStandardNodesCategories.h b/Engine/Plugins/Runtime/Metasound/Source/MetasoundStandardNodes/Public/MetasoundStandardNodesCategories.h
--- a/Engine/Plugins/Runtime/Metasound/Source/MetasoundStandardNodes/Public/MetasoundStandardNodesCategories.h
+++ b/Engine/Plugins/Runtime/Metasound/Source/MetasoundStandardNodes/Public/MetasoundStandardNodesCategories.h
@@ -21,4 +21,5 @@ namespace Metasound::NodeCategories
 	extern const FText METASOUNDSTANDARDNODES_API Spatialization;
 	extern const FText METASOUNDSTANDARDNODES_API Trigger;
 	extern const FText METASOUNDSTANDARDNODES_API WaveTables;
+	extern const FText METASOUNDSTANDARDNODES_API Reverbs;
 } // Metasound::NodeCategories
diff --git a/Engine/Source/Runtime/SignalProcessing/Private/EarlyReflectionsFast.cpp b/Engine/Source/Runtime/SignalProcessing/Private/EarlyReflectionsFast.cpp
--- a/Engine/Source/Runtime/SignalProcessing/Private/EarlyReflectionsFast.cpp
+++ b/Engine/Source/Runtime/SignalProcessing/Private/EarlyReflectionsFast.cpp
@@ -221,6 +221,40 @@ namespace Audio
 		ArrayMultiplyByConstantInPlace(OutRightSamples, Settings.Gain);
 	}
 
+	void FEarlyReflectionsFast::ProcessAudioStereoNonInterleaved(
+		const TArrayView<const float>& InLeft,
+		const TArrayView<const float>& InRight,
+		FAlignedFloatBuffer& OutLeft,
+		FAlignedFloatBuffer& OutRight)
+	{
+		const int32 InNumFrames = InLeft.Num();
+		check(InRight.Num() == InNumFrames);
+		check(OutLeft.Num() == InNumFrames);
+		check(OutRight.Num() == InNumFrames);
+
+		// Resize internal buffers
+		LeftWorkBufferA.SetNumUninitialized(InNumFrames);
+		LeftWorkBufferB.SetNumUninitialized(InNumFrames);
+		RightWorkBufferA.SetNumUninitialized(InNumFrames);
+		RightWorkBufferB.SetNumUninitialized(InNumFrames);
+
+		// predelay
+		LeftPreDelay.ProcessAudio(InLeft, LeftWorkBufferB);
+		RightPreDelay.ProcessAudio(InRight, RightWorkBufferB);
+
+		// lpf
+		LeftInputLPF.ProcessAudio(LeftWorkBufferB, LeftWorkBufferA);
+		RightInputLPF.ProcessAudio(RightWorkBufferB, RightWorkBufferA);
+
+		// feedback delay network
+		LeftFDN.ProcessAudio(LeftWorkBufferA, OutLeft);
+		RightFDN.ProcessAudio(RightWorkBufferA, OutRight);
+
+		// Apply Gain
+		ArrayMultiplyByConstantInPlace(OutLeft, Settings.Gain);
+		ArrayMultiplyByConstantInPlace(OutRight, Settings.Gain);
+	}
+
 	void FEarlyReflectionsFast::FlushAudio()
 	{
 		// predelay
diff --git a/Engine/Source/Runtime/SignalProcessing/Private/ReverbFast.cpp b/Engine/Source/Runtime/SignalProcessing/Private/ReverbFast.cpp
--- a/Engine/Source/Runtime/SignalProcessing/Private/ReverbFast.cpp
+++ b/Engine/Source/Runtime/SignalProcessing/Private/ReverbFast.cpp
@@ -148,6 +148,68 @@ namespace Audio {
 		InterleaveAndMixOutput(FrontLeftReverbSamples, FrontRightReverbSamples, OutSamples, OutNumChannels);
 	}
 
+	void FPlateReverbFast::ProcessAudioStereoNonInterleaved(
+		const TArrayView<const float>& InLeft,
+		const TArrayView<const float>& InRight,
+		FAlignedFloatBuffer& OutLeft,
+		FAlignedFloatBuffer& OutRight)
+	{
+		const int32 NumFrames = InLeft.Num();
+		check(NumFrames == InRight.Num());
+		check(NumFrames == OutLeft.Num());
+		check(NumFrames == OutRight.Num());
+
+		if (NumFrames == 0)
+		{
+			return;
+		}
+
+		// Reverb is disabled, so just zero the output buffer
+		if (!Settings.bEnableEarlyReflections && !Settings.bEnableLateReflections)
+		{
+			FMemory::Memzero(OutLeft.GetData(), NumFrames * sizeof(float));
+			FMemory::Memzero(OutRight.GetData(), NumFrames * sizeof(float));
+			
+			return;
+		}
+
+		// Both early and late reflections
+		if (Settings.bEnableEarlyReflections && Settings.bEnableLateReflections)
+		{
+			// Process the early reflections directly to the outputs
+			EarlyReflections.ProcessAudioStereoNonInterleaved(InLeft, InRight, OutLeft, OutRight);
+
+			// Process the late reflections and hang onto the result in the work buffers
+			// Sum to mono and scale by half because the late reflection algorithm does that internally. This saves some copies.
+			FrontLeftReverbSamples.SetNumUninitialized(NumFrames);
+			ArraySum(InLeft, InRight, FrontLeftReverbSamples);
+			ArrayMultiplyByConstantInPlace(FrontLeftReverbSamples, 0.5f);
+			LateReflections.ProcessAudio(FrontLeftReverbSamples, 1, FrontLeftLateReflectionsSamples, FrontRightLateReflectionsSamples);
+
+			// Mix in the early reflections
+			ArrayAddInPlace(FrontLeftLateReflectionsSamples, OutLeft);
+			ArrayAddInPlace(FrontRightLateReflectionsSamples, OutRight);
+		}
+		// Only early reflections
+		else if (Settings.bEnableEarlyReflections)
+		{
+			EarlyReflections.ProcessAudioStereoNonInterleaved(InLeft, InRight, OutLeft, OutRight);
+		}
+		// Only late reflections
+		else if (Settings.bEnableLateReflections)
+		{
+			// Sum to mono and scale by half because the late reflection algorithm does that internally. This saves some copies.
+			FrontLeftReverbSamples.SetNumUninitialized(NumFrames);
+			ArraySum(InLeft, InRight, FrontLeftReverbSamples);
+			ArrayMultiplyByConstantInPlace(FrontLeftReverbSamples, 0.5f);
+			LateReflections.ProcessAudio(FrontLeftReverbSamples, 1, OutLeft, OutRight);
+		}
+		else
+		{
+			checkNoEntry();
+		}
+	}
+
 	void FPlateReverbFast::ClampSettings(FPlateReverbFastSettings& InOutSettings)
 	{
 		// Clamp settings for this object and member objects.
diff --git a/Engine/Source/Runtime/SignalProcessing/Public/DSP/EarlyReflectionsFast.h b/Engine/Source/Runtime/SignalProcessing/Public/DSP/EarlyReflectionsFast.h
--- a/Engine/Source/Runtime/SignalProcessing/Public/DSP/EarlyReflectionsFast.h
+++ b/Engine/Source/Runtime/SignalProcessing/Public/DSP/EarlyReflectionsFast.h
@@ -66,6 +66,13 @@ namespace Audio
 		// Process the single audio frame
 		SIGNALPROCESSING_API void ProcessAudio(const FAlignedFloatBuffer& InSamples, const int32 InNumChannels, FAlignedFloatBuffer& OutLeftSamples, FAlignedFloatBuffer& OutRightSamples);
 
+		// Process a block of audio. Optimized for stereo, non-interleaved input.
+		SIGNALPROCESSING_API void ProcessAudioStereoNonInterleaved(
+			const TArrayView<const float>& InLeft,
+			const TArrayView<const float>& InRight,
+			FAlignedFloatBuffer& OutLeft,
+			FAlignedFloatBuffer& OutRight);
+		
 		// Silence internal audio.
 		SIGNALPROCESSING_API void FlushAudio();
 
diff --git a/Engine/Source/Runtime/SignalProcessing/Public/DSP/ReverbFast.h b/Engine/Source/Runtime/SignalProcessing/Public/DSP/ReverbFast.h
--- a/Engine/Source/Runtime/SignalProcessing/Public/DSP/ReverbFast.h
+++ b/Engine/Source/Runtime/SignalProcessing/Public/DSP/ReverbFast.h
@@ -72,7 +72,14 @@ namespace Audio
 			// InNumChannels can be 1 or 2 channels.
 			// OutSamples must be greater or equal to 2.
 			SIGNALPROCESSING_API void ProcessAudio(const FAlignedFloatBuffer& InSamples, const int32 InNumChannels, FAlignedFloatBuffer& OutSamples, const int32 OutNumChannels);
-
+		
+			// Process a block of audio. Optimized for stereo, non-interleaved input and output.
+			SIGNALPROCESSING_API void ProcessAudioStereoNonInterleaved(
+				const TArrayView<const float>& InLeft,
+				const TArrayView<const float>& InRight,
+				FAlignedFloatBuffer& OutLeft,
+				FAlignedFloatBuffer& OutRight);
+		
 			SIGNALPROCESSING_API void FlushAudio();
 
 			// Clamp individual settings to values supported by this class.
