[Metasounds]  setting object arrays from bp not working (Fix fallout from UE-213913)
#rnx
#tests Metasound Blueprint Engine tests, AQ_Metasounds test map

diff --git a/Engine/Plugins/Runtime/Metasound/Source/MetasoundEngine/Private/MetasoundSource.cpp b/Engine/Plugins/Runtime/Metasound/Source/MetasoundEngine/Private/MetasoundSource.cpp
--- a/Engine/Plugins/Runtime/Metasound/Source/MetasoundEngine/Private/MetasoundSource.cpp
+++ b/Engine/Plugins/Runtime/Metasound/Source/MetasoundEngine/Private/MetasoundSource.cpp
@@ -1175,18 +1175,17 @@ void UMetaSoundSource::InitParametersInternal(const Metasound::TSortedVertexName
 
 		FAudioParameter& Parameter = ParametersToInit[i];
 		if (const FRuntimeInput* Input = InInputMap.Find(Parameter.ParamName))
-		{
-			bIsParameterValid = true;
-			
+		{	
 			if (IsParameterValidInternal(Parameter, Input->TypeName, DataTypeRegistry))
 			{
+				bIsParameterValid = true;
+				
 				// note: this was originally called 2x, the 1st call would create the proxy
 				// and clear out the ObjectParam.  The 2nd call would clear out the proxy,
 				// and fail to re-create it because we cleared out the ObjectParam. 
 				
-				// if there is no object param we should assume we have already created a proxy
-				// (there is no point in attempting to create a proxy w/o a UObject.)
-				if(Parameter.ObjectParam) 
+				// if there is already a proxy, don't attempt to recreate them.
+				if (Parameter.ObjectProxies.IsEmpty())
 				{
 					Sanitize(Parameter);
 					constexpr bool bClearUObjectPointers = true; // protect against leaking UObject ptrs to the audio thread 
