MetaSound Array Random Get node fixes for passthrough inputs and node id for Same Node shared state behavior
- Make trigger inputs passthrough
- Use both node id and containing graph id to avoid shared state in duplicated assets when using Same Node shared state behavior
- Use asset class id instead of a newly generated id for FFrontendGraph graph id to guarantee ids are shared for a given asset, and clean up now unneeded id map in the Operator Cache

diff --git a/Engine/Plugins/Runtime/Metasound/Source/MetasoundEngine/Private/MetasoundOperatorCacheSubsystem.cpp b/Engine/Plugins/Runtime/Metasound/Source/MetasoundEngine/Private/MetasoundOperatorCacheSubsystem.cpp
--- a/Engine/Plugins/Runtime/Metasound/Source/MetasoundEngine/Private/MetasoundOperatorCacheSubsystem.cpp
+++ b/Engine/Plugins/Runtime/Metasound/Source/MetasoundEngine/Private/MetasoundOperatorCacheSubsystem.cpp
@@ -157,6 +157,6 @@ void UMetaSoundCacheSubsystem::RemoveCachedOperatorsForMetaSound(UMetaSoundSourc
 
 	if (TSharedPtr<FOperatorPool> OperatorPool = Module->GetOperatorPool())
 	{
-		OperatorPool->RemoveOperatorsWithAssetClassID(InMetaSound->AssetClassID);
+		OperatorPool->RemoveOperatorsWithID(InMetaSound->AssetClassID);
 	}
 }
\ No newline at end of file
diff --git a/Engine/Plugins/Runtime/Metasound/Source/MetasoundFrontend/Private/MetasoundFrontendGraph.cpp b/Engine/Plugins/Runtime/Metasound/Source/MetasoundFrontend/Private/MetasoundFrontendGraph.cpp
--- a/Engine/Plugins/Runtime/Metasound/Source/MetasoundFrontend/Private/MetasoundFrontendGraph.cpp
+++ b/Engine/Plugins/Runtime/Metasound/Source/MetasoundFrontend/Private/MetasoundFrontendGraph.cpp
@@ -9,7 +9,6 @@
 #include "CoreMinimal.h"
 #include "MetasoundFrontend.h"
 #include "MetasoundFrontendDataTypeRegistry.h"
-#include "MetasoundFrontendDocumentIdGenerator.h"
 #include "MetasoundFrontendNodeTemplateRegistry.h"
 #include "MetasoundFrontendProxyDataCache.h"
 #include "MetasoundFrontendRegistries.h"
@@ -744,7 +743,7 @@ namespace Metasound
 
 		FBuildGraphContext BuildGraphContext
 		{
-			MakeUnique<FFrontendGraph>(GraphName, Frontend::CreateLocallyUniqueId()),
+			MakeUnique<FFrontendGraph>(GraphName, InContext.GraphId),
 			InGraphClass,
 			InContext
 		};
@@ -776,7 +775,7 @@ namespace Metasound
 		return CreateGraph(InGraph, InSubgraphs, InDependencies, InDebugAssetName);
 	}
 
-	TUniquePtr<FFrontendGraph> FFrontendGraphBuilder::CreateGraph(const FMetasoundFrontendGraphClass& InGraph, const TArray<FMetasoundFrontendGraphClass>& InSubgraphs, const TArray<FMetasoundFrontendClass>& InDependencies, const Frontend::FProxyDataCache& InProxyDataCache, const FString& InDebugAssetName)
+	TUniquePtr<FFrontendGraph> FFrontendGraphBuilder::CreateGraph(const FMetasoundFrontendGraphClass& InGraph, const TArray<FMetasoundFrontendGraphClass>& InSubgraphs, const TArray<FMetasoundFrontendClass>& InDependencies, const Frontend::FProxyDataCache& InProxyDataCache, const FString& InDebugAssetName, const FGuid InGraphId)
 	{
 		FBuildContext Context
 		{
@@ -784,7 +783,8 @@ namespace Metasound
 			{}, 								// FrontendClasses
 			{}, 								// Graphs
 			Frontend::IDataTypeRegistry::Get(), // DataTypeRegistry
-			InProxyDataCache 					// ProxyDataCache
+			InProxyDataCache, 					// ProxyDataCache
+			InGraphId							// GraphId
 		};
 
 		// Gather all references to node classes from external dependencies and subgraphs.
@@ -859,8 +859,8 @@ namespace Metasound
 		return CreateGraph(InDocument, InDebugAssetName);
 	}
 
-	TUniquePtr<FFrontendGraph> FFrontendGraphBuilder::CreateGraph(const FMetasoundFrontendDocument& InDocument, const Frontend::FProxyDataCache& InProxyDataCache, const FString& InDebugAssetName)
+	TUniquePtr<FFrontendGraph> FFrontendGraphBuilder::CreateGraph(const FMetasoundFrontendDocument& InDocument, const Frontend::FProxyDataCache& InProxyDataCache, const FString& InDebugAssetName, const FGuid InGraphId)
 	{
-		return CreateGraph(InDocument.RootGraph, InDocument.Subgraphs, InDocument.Dependencies, InProxyDataCache, InDebugAssetName);
+		return CreateGraph(InDocument.RootGraph, InDocument.Subgraphs, InDocument.Dependencies, InProxyDataCache, InDebugAssetName, InGraphId);
 	}
 }
diff --git a/Engine/Plugins/Runtime/Metasound/Source/MetasoundFrontend/Private/MetasoundFrontendRegistryContainerImpl.cpp b/Engine/Plugins/Runtime/Metasound/Source/MetasoundFrontend/Private/MetasoundFrontendRegistryContainerImpl.cpp
--- a/Engine/Plugins/Runtime/Metasound/Source/MetasoundFrontend/Private/MetasoundFrontendRegistryContainerImpl.cpp
+++ b/Engine/Plugins/Runtime/Metasound/Source/MetasoundFrontend/Private/MetasoundFrontendRegistryContainerImpl.cpp
@@ -282,7 +282,8 @@ namespace Metasound::Frontend
 		METASOUND_TRACE_CPUPROFILER_EVENT_SCOPE_TEXT(*FString::Printf(TEXT("Metasound::FRegistryContainerImpl::BuildAndRegisterGraphFromDocument asset %s"), *InNodeClassInfo.AssetPath.ToString()));
 
 		const FMetasoundFrontendDocument& Document = DocumentInterface->GetConstDocument();
-		TUniquePtr<FFrontendGraph> FrontendGraph = FFrontendGraphBuilder::CreateGraph(Document, InProxyDataCache, InNodeClassInfo.AssetPath.ToString());
+		// Use the asset class id for the graph id because it should be locally unique. Unlike other ids, it is regenerated on asset duplicate. 
+		TUniquePtr<FFrontendGraph> FrontendGraph = FFrontendGraphBuilder::CreateGraph(Document, InProxyDataCache, InNodeClassInfo.AssetPath.ToString(), /*GraphId=*/InNodeClassInfo.AssetClassID);
 		if (!FrontendGraph.IsValid())
 		{
 			UE_LOG(LogMetaSound, Error, TEXT("Failed to build MetaSound graph in asset '%s'"), *InNodeClassInfo.AssetPath.ToString());
diff --git a/Engine/Plugins/Runtime/Metasound/Source/MetasoundFrontend/Public/MetasoundArrayRandomNode.h b/Engine/Plugins/Runtime/Metasound/Source/MetasoundFrontend/Public/MetasoundArrayRandomNode.h
--- a/Engine/Plugins/Runtime/Metasound/Source/MetasoundFrontend/Public/MetasoundArrayRandomNode.h
+++ b/Engine/Plugins/Runtime/Metasound/Source/MetasoundFrontend/Public/MetasoundArrayRandomNode.h
@@ -324,27 +324,30 @@ namespace Metasound
 			TriggerOnReset->Reset();
 
 			// Cache shared state id for shared state behavior types that cannot be changed after node init
-			if (SharedStateBehavior == ESharedStateBehaviorType::SameNode)
+			if (InParams.Environment.Contains<TArray<FGuid>>(OperatorBuilder::Environment::GraphHierarchy))
 			{
-				SharedStateId = NodeId;
-			}
-			else if (SharedStateBehavior == ESharedStateBehaviorType::SameNodeInComposition)
-			{
-				if (InParams.Environment.Contains<TArray<FGuid>>(OperatorBuilder::Environment::GraphHierarchy))
+				const TArray<FGuid>& GraphHierarchy = InParams.Environment.GetValue<TArray<FGuid>>(OperatorBuilder::Environment::GraphHierarchy);
+				if (SharedStateBehavior == ESharedStateBehaviorType::SameNode)
 				{
-					const TArray<FGuid>& GraphHierarchy = InParams.Environment.GetValue<TArray<FGuid>>(OperatorBuilder::Environment::GraphHierarchy);
-					SharedStateId = GetSameNodeInCompositionId(NodeId, GraphHierarchy);
+					check(GraphHierarchy.Num() > 0);
+					// Hash node id with this node's graph id because node ids are not guaranteed to be unique 
+					// (they are not regenerated when duplicating assets)
+					SharedStateId = GetSameNodeSharedStateId(NodeId, GraphHierarchy.Last());
 				}
-				else
+				else if (SharedStateBehavior == ESharedStateBehaviorType::SameNodeInComposition)
 				{
+					SharedStateId = GetSameNodeInCompositionId(NodeId, GraphHierarchy);
+				}
+			}
+			else
+			{
 #if WITH_METASOUND_DEBUG_ENVIRONMENT
-					if (!bHasLoggedMissingGraphHierarchyWarning)
-					{
-						UE_LOG(LogMetaSound, Warning, TEXT("Array Random Get: Graph Hierarchy environment variable needed for Same Node in Composition shared state id not found (Graph '%s')"), *GraphName);
-						bHasLoggedMissingGraphHierarchyWarning = true;
-					}
-#endif // WITH_METASOUND_DEBUG_ENVIRONMENT
+				if (!bHasLoggedMissingGraphHierarchyWarning)
+				{
+					UE_LOG(LogMetaSound, Warning, TEXT("Array Random Get: Graph Hierarchy environment variable needed for Same Node or Same Node in Composition shared state id not found (Graph '%s')"), *GraphName);
+					bHasLoggedMissingGraphHierarchyWarning = true;
 				}
+#endif // WITH_METASOUND_DEBUG_ENVIRONMENT
 			}
 		}
 
@@ -363,6 +366,27 @@ namespace Metasound
 					bHasLoggedEmptyArrayWarning = true;
 				}
 #endif // WITH_METASOUND_DEBUG_ENVIRONMENT
+				// Pass through triggers
+				TriggerReset->ExecuteBlock(
+					[&](int32 StartFrame, int32 EndFrame)
+					{
+					},
+					[this](int32 StartFrame, int32 EndFrame)
+					{
+						TriggerOnReset->TriggerFrame(StartFrame);
+					}
+				);
+
+				TriggerNext->ExecuteBlock(
+					[&](int32 StartFrame, int32 EndFrame)
+					{
+					},
+					[this](int32 StartFrame, int32 EndFrame)
+					{
+						TriggerOnNext->TriggerFrame(StartFrame);
+					}
+				);
+
 				return;
 			}
 
@@ -479,6 +503,17 @@ namespace Metasound
 			TriggerOnNext->TriggerFrame(StartFrame);
 		}
 
+		// Hash combine the current node id with another id
+		FGuid GetSameNodeSharedStateId(const FGuid& InNodeId, const FGuid& InOtherId) const
+		{
+			return FGuid(
+				HashCombineFast(InNodeId.A, InOtherId.A),
+				HashCombineFast(InNodeId.B, InOtherId.B),
+				HashCombineFast(InNodeId.C, InOtherId.C),
+				HashCombineFast(InNodeId.D, InOtherId.D)
+			);
+		}
+
 		// Hash combine the current node id with the graph hierarchy ids
 		FGuid GetSameNodeInCompositionId(const FGuid& InNodeId, const TArray<FGuid>& InGraphHierarchy) const
 		{
diff --git a/Engine/Plugins/Runtime/Metasound/Source/MetasoundFrontend/Public/MetasoundFrontendGraph.h b/Engine/Plugins/Runtime/Metasound/Source/MetasoundFrontend/Public/MetasoundFrontendGraph.h
--- a/Engine/Plugins/Runtime/Metasound/Source/MetasoundFrontend/Public/MetasoundFrontendGraph.h
+++ b/Engine/Plugins/Runtime/Metasound/Source/MetasoundFrontend/Public/MetasoundFrontendGraph.h
@@ -7,6 +7,7 @@
 #include "Containers/Set.h"
 #include "Containers/UnrealString.h"
 #include "MetasoundFrontendDocument.h"
+#include "MetasoundFrontendDocumentIdGenerator.h"
 #include "MetasoundGraph.h"
 #include "MetasoundNodeConstructorParams.h"
 #include "MetasoundNodeInterface.h"
@@ -105,13 +106,13 @@ namespace Metasound
 		static TUniquePtr<FFrontendGraph> CreateGraph(const FMetasoundFrontendDocument& InDocument, const FString& InDebugAssetName);
 
 		/* Create a FFrontendGraph from a FMetasoundFrontendDocument retrieving proxies from a FProxyDataCache.*/
-		static TUniquePtr<FFrontendGraph> CreateGraph(const FMetasoundFrontendDocument& InDocument, const Frontend::FProxyDataCache& InProxies, const FString& InDebugAssetName);
+		static TUniquePtr<FFrontendGraph> CreateGraph(const FMetasoundFrontendDocument& InDocument, const Frontend::FProxyDataCache& InProxies, const FString& InDebugAssetName, const FGuid InGraphId = Frontend::CreateLocallyUniqueId());
 
-		/* Create a FFrontendGraph from a FMetasoundFrontendDocument subobjects retrieving proxies from a FProxyDataCache.*/
+		/* Create a FFrontendGraph from a FMetasoundFrontendDocument and subobjects.*/
 		static TUniquePtr<FFrontendGraph> CreateGraph(const FMetasoundFrontendGraphClass& InGraph, const TArray<FMetasoundFrontendGraphClass>& InSubgraphs, const TArray<FMetasoundFrontendClass>& InDependencies, const FString& InDebugAssetName);
 
-		/* Create a FFrontendGraph from a FMetasoundFrontendDocument subobjects retrieving proxies from a FProxyDataCache.*/
-		static TUniquePtr<FFrontendGraph> CreateGraph(const FMetasoundFrontendGraphClass& InGraph, const TArray<FMetasoundFrontendGraphClass>& InSubgraphs, const TArray<FMetasoundFrontendClass>& InDependencies, const Frontend::FProxyDataCache& InProxyDataCache, const FString& InDebugAssetName);
+		/* Create a FFrontendGraph from a FMetasoundFrontendDocument and subobjects, retrieving proxies from a FProxyDataCache.*/
+		static TUniquePtr<FFrontendGraph> CreateGraph(const FMetasoundFrontendGraphClass& InGraph, const TArray<FMetasoundFrontendGraphClass>& InSubgraphs, const TArray<FMetasoundFrontendClass>& InDependencies, const Frontend::FProxyDataCache& InProxyDataCache, const FString& InDebugAssetName, const FGuid InGraphId = Frontend::CreateLocallyUniqueId());
 
 	private:
 		struct FDefaultLiteralData
@@ -138,6 +139,7 @@ namespace Metasound
 			FSharedNodeByIDMap Graphs;
 			const Frontend::IDataTypeRegistry& DataTypeRegistry;
 			const Frontend::FProxyDataCache& ProxyDataCache;
+			FGuid GraphId;
 		};
 
 		// Transient context used for building a specific graph
diff --git a/Engine/Plugins/Runtime/Metasound/Source/MetasoundGenerator/Private/MetasoundOperatorCache.cpp b/Engine/Plugins/Runtime/Metasound/Source/MetasoundGenerator/Private/MetasoundOperatorCache.cpp
--- a/Engine/Plugins/Runtime/Metasound/Source/MetasoundGenerator/Private/MetasoundOperatorCache.cpp
+++ b/Engine/Plugins/Runtime/Metasound/Source/MetasoundGenerator/Private/MetasoundOperatorCache.cpp
@@ -408,8 +408,8 @@ namespace Metasound
 				if (OperatorPool.IsValid())
 				{
 					// Get the number of instances already in the cache & move pre-existing to the top of the cache
-					const int32 NumInCache = OperatorPool->GetNumCachedOperatorsWithAssetClassID(PreCacheData->AssetClassID);
-					OperatorPool->TouchOperatorsViaAssetClassID(PreCacheData->AssetClassID, FMath::Min(NumInCache, NumToBuild));
+					const int32 NumInCache = OperatorPool->GetNumCachedOperatorsWithID(PreCacheData->AssetClassID);
+					OperatorPool->TouchOperators(PreCacheData->AssetClassID, FMath::Min(NumInCache, NumToBuild));
 					NumToBuild -= NumInCache;
 				}
 			}
@@ -432,7 +432,6 @@ namespace Metasound
 
 				const FGuid& GraphID = PreCacheData->InitParams.Graph->GetInstanceID();
 				OperatorPool->AddOperatorInternal(GraphID, MoveTemp(OperatorAndInputs));
-				OperatorPool->AddAssetIdToGraphIdLookUpInternal(PreCacheData->AssetClassID, GraphID);
 			}
 		});
 	}
@@ -465,15 +464,7 @@ namespace Metasound
 
 	void FOperatorPool::TouchOperatorsViaAssetClassID(const FGuid& InAssetClassID, int32 NumToTouch)
 	{
-		if (!IsStopping())
-		{
-			FScopeLock Lock(&CriticalSection);
-			FGuid* GraphIdPtr = AssetIdToGraphIdLookUp.Find(InAssetClassID);
-			if (GraphIdPtr)
-			{
-				TouchOperators(*GraphIdPtr, NumToTouch);
-			}
-		}
+		TouchOperators(InAssetClassID, NumToTouch);
 	}
 
 	void FOperatorPool::RemoveOperatorsWithID(const FGuid& InOperatorID)
@@ -502,13 +493,7 @@ namespace Metasound
 
 	void FOperatorPool::RemoveOperatorsWithAssetClassID(const FGuid& InAssetClassID)
 	{
-		FScopeLock Lock(&CriticalSection);
-		FGuid* GraphIdPtr = AssetIdToGraphIdLookUp.Find(InAssetClassID);
-		if (GraphIdPtr)
-		{
-			RemoveOperatorsWithID(*GraphIdPtr);
-			AssetIdToGraphIdLookUp.Remove(InAssetClassID);
-		}
+		return RemoveOperatorsWithID(InAssetClassID);
 	}
 		
 	void FOperatorPool::SetMaxNumOperators(uint32 InMaxNumOperators)
@@ -534,19 +519,7 @@ namespace Metasound
 
 	int32 FOperatorPool::GetNumCachedOperatorsWithAssetClassID(const FGuid& InAssetClassID) const
 	{
-		FScopeLock Lock(&CriticalSection);
-		if (const FGuid* GraphIdPtr = AssetIdToGraphIdLookUp.Find(InAssetClassID))
-		{
-			return GetNumCachedOperatorsWithID(*GraphIdPtr);
-		}
-
-		return 0;
-	}
-
-	void FOperatorPool::AddAssetIdToGraphIdLookUpInternal(const FGuid& InAssetClassID, const FGuid& InOperatorID)
-	{
-		FScopeLock Lock(&CriticalSection);
-		AssetIdToGraphIdLookUp.Add(InAssetClassID, InOperatorID);
+		return GetNumCachedOperatorsWithID(InAssetClassID);
 	}
 
 #if METASOUND_OPERATORCACHEPROFILER_ENABLED
@@ -584,7 +557,5 @@ namespace Metasound
 			TRACE_COUNTER_DECREMENT(MetaSound_OperatorPool_NumOperators);
 #endif // #if METASOUND_OPERATORCACHEPROFILER_ENABLED
 		}
-
-		// todo: prune AssetIdToGraphIdLookUp?
 	}
 } // namespace Metasound
diff --git a/Engine/Plugins/Runtime/Metasound/Source/MetasoundGenerator/Public/MetasoundOperatorCache.h b/Engine/Plugins/Runtime/Metasound/Source/MetasoundGenerator/Public/MetasoundOperatorCache.h
--- a/Engine/Plugins/Runtime/Metasound/Source/MetasoundGenerator/Public/MetasoundOperatorCache.h
+++ b/Engine/Plugins/Runtime/Metasound/Source/MetasoundGenerator/Public/MetasoundOperatorCache.h
@@ -101,18 +101,21 @@ namespace Metasound
 		void BuildAndAddOperator(TUniquePtr<FOperatorBuildData> InBuildData);
 
 		void TouchOperators(const FGuid& InOperatorID, int32 NumToTouch = 1);
+		UE_DEPRECATED(5.5, "Operators can now be directly identified by asset class id.")
 		void TouchOperatorsViaAssetClassID(const FGuid& InAssetClassID, int32 NumToTouch = 1);
 
 		bool IsStopping() const;
 
 		void RemoveOperatorsWithID(const FGuid& InOperatorID);
+		UE_DEPRECATED(5.5, "Operators can now be directly identified by asset class id.")
 		void RemoveOperatorsWithAssetClassID(const FGuid& InAssetClassID);
 
 		int32 GetNumCachedOperatorsWithID(const FGuid& InOperatorID) const;
+		UE_DEPRECATED(5.5, "Operators can now be directly identified by asset class id.")
 		int32 GetNumCachedOperatorsWithAssetClassID(const FGuid& InAssetClassID) const;
 
-		UE_DEPRECATED(5.5, "Adding id to look-up is now private implementation")
-		void AddAssetIdToGraphIdLookUp(const FGuid& InAssetClassID, const FGuid& InOperatorID) { }
+		UE_DEPRECATED(5.5, "Operators can now be directly identified by asset class id, so lookup is not needed anymore.")
+		void AddAssetIdToGraphIdLookUp(const FGuid& InAssetClassID, const FGuid& InOperatorID) {}
 
 		void SetMaxNumOperators(uint32 InMaxNumOperators);
 #if METASOUND_OPERATORCACHEPROFILER_ENABLED
@@ -130,7 +133,6 @@ namespace Metasound
 	private:
 		FTaskId LastTaskId = 0;
 
-		void AddAssetIdToGraphIdLookUpInternal(const FGuid& InAssetClassID, const FGuid& InOperatorID);
 		void AddOperatorInternal(const FGuid& InOperatorID, FOperatorAndInputs&& OperatorAndInputs);
 		bool ExecuteTaskAsync(FTaskFunction&& InFunction);
 		void Trim();
@@ -150,7 +152,6 @@ namespace Metasound
 		UE::Tasks::FPipe AsyncBuildPipe;
 
 		TMap<FGuid, TArray<FOperatorAndInputs>> Operators;
-		TMap<FGuid, FGuid> AssetIdToGraphIdLookUp;
 		TArray<FGuid> Stack;
 	};
 } // namespace Metasound
