UE-188318 - [MetaSounds] - Promoting to graph input or variable should default to 'rename' mode for name

diff --git a/Engine/Plugins/Runtime/Metasound/Source/MetasoundEditor/Private/MetasoundEditor.cpp b/Engine/Plugins/Runtime/Metasound/Source/MetasoundEditor/Private/MetasoundEditor.cpp
--- a/Engine/Plugins/Runtime/Metasound/Source/MetasoundEditor/Private/MetasoundEditor.cpp
+++ b/Engine/Plugins/Runtime/Metasound/Source/MetasoundEditor/Private/MetasoundEditor.cpp
@@ -3388,6 +3388,8 @@ namespace Metasound
 			FMetasoundAssetBase* MetasoundAsset = IMetasoundUObjectRegistry::Get().GetObjectAsAssetBase(Metasound);
 			check(MetasoundAsset);
 
+			bool bNodesSelected = false;
+
 			// Capture after synchronizing as the modification state may be modified therein
 			const FMetasoundFrontendDocumentModifyContext& ModifyContext = MetasoundAsset->GetConstModifyContext();
 			const bool bForceRefreshViews = ModifyContext.GetForceRefreshViews();
@@ -3484,13 +3486,21 @@ namespace Metasound
 				if (!Selection.IsEmpty())
 				{
 					SetSelection(Selection);
-
+					bNodesSelected = true;
 				}
 
 				// Avoids details panel displaying
 				// removed members in certain cases.
 				RemoveInvalidSelection();
 			}
+
+			// Wait for GraphMembersMenu to get updated with selected nodes
+			if (bMemberRenameRequested && bNodesSelected)
+			{
+				GraphMembersMenu->RefreshAllActions(/*bPreserveExpansion=*/ true, /*bHandleOnSelectionEvent=*/ true);
+				GraphMembersMenu->OnRequestRenameOnActionNode();
+				bMemberRenameRequested = false;
+			}
 		}
 
 		TStatId FEditor::GetStatId() const
@@ -3731,6 +3741,11 @@ namespace Metasound
 				}
 			}
 		}
+
+		void FEditor::SetDelayedRename()
+		{
+			bMemberRenameRequested = true;
+		}
 	}
 }
 #undef LOCTEXT_NAMESPACE
diff --git a/Engine/Plugins/Runtime/Metasound/Source/MetasoundEditor/Private/MetasoundEditor.h b/Engine/Plugins/Runtime/Metasound/Source/MetasoundEditor/Private/MetasoundEditor.h
--- a/Engine/Plugins/Runtime/Metasound/Source/MetasoundEditor/Private/MetasoundEditor.h
+++ b/Engine/Plugins/Runtime/Metasound/Source/MetasoundEditor/Private/MetasoundEditor.h
@@ -238,6 +238,8 @@ namespace Metasound
 				}
 			}
 
+			void SetDelayedRename();
+
 		protected:
 			// Callbacks for action tree
 			bool CanRenameOnActionNode(TWeakPtr<FGraphActionNode> InSelectedNode) const;
@@ -486,6 +488,8 @@ namespace Metasound
 			int32 HighestMessageSeverity = EMessageSeverity::Info;
 
 			TSharedPtr<SNotificationItem> NotificationPtr;
+
+			bool bMemberRenameRequested = false;
 		};
 	} // namespace Editor
 } // namespace Metasound
diff --git a/Engine/Plugins/Runtime/Metasound/Source/MetasoundEditor/Private/MetasoundEditorGraphSchema.cpp b/Engine/Plugins/Runtime/Metasound/Source/MetasoundEditor/Private/MetasoundEditorGraphSchema.cpp
--- a/Engine/Plugins/Runtime/Metasound/Source/MetasoundEditor/Private/MetasoundEditorGraphSchema.cpp
+++ b/Engine/Plugins/Runtime/Metasound/Source/MetasoundEditor/Private/MetasoundEditorGraphSchema.cpp
@@ -241,7 +241,7 @@ PRAGMA_ENABLE_DEPRECATION_WARNINGS
 				}
 			}
 
-			void SelecteNodeInEditor(UMetasoundEditorGraph& InMetaSoundGraph, UMetasoundEditorGraphNode& InNode)
+			void SelectNodeInEditor(UMetasoundEditorGraph& InMetaSoundGraph, UMetasoundEditorGraphNode& InNode)
 			{
 				TSharedPtr<FEditor> MetasoundEditor = FGraphBuilder::GetEditorForGraph(InMetaSoundGraph);
 				if (MetasoundEditor.IsValid())
@@ -250,6 +250,16 @@ PRAGMA_ENABLE_DEPRECATION_WARNINGS
 				}
 			}
 
+			void SelectNodeInEditorForRename(UMetasoundEditorGraph& InMetaSoundGraph, UMetasoundEditorGraphNode& InNode)
+			{
+				TSharedPtr<FEditor> MetasoundEditor = FGraphBuilder::GetEditorForGraph(InMetaSoundGraph);
+				if (MetasoundEditor.IsValid())
+				{
+					MetasoundEditor->ClearSelectionAndSelectNode(&InNode);
+					MetasoundEditor->SetDelayedRename();
+				}
+			}
+
 			UEdGraphNode* PromoteToVariable(const FName BaseName, UEdGraphPin& FromPin, const FName DataType, const FMetasoundFrontendClass& InVariableClass, const FVector2D& InLocation, bool bSelectNode)
 			{
 				using namespace Frontend;
@@ -315,7 +325,7 @@ PRAGMA_ENABLE_DEPRECATION_WARNINGS
 							if (ensure(SchemaPrivate::TryConnectNewNodeToMatchingDataTypePin(*EdGraphNode, &FromPin)))
 							{
 								FGraphBuilder::RegisterGraphWithFrontend(ParentMetaSound);
-								SelecteNodeInEditor(*MetaSoundGraph, *NewGraphNode);
+								SelectNodeInEditorForRename(*MetaSoundGraph, *NewGraphNode);
 								return EdGraphNode;
 							}
 						}
@@ -465,7 +475,7 @@ UEdGraphNode* FMetasoundGraphSchemaAction_NewNode::PerformAction(UEdGraph* Paren
 		NewGraphNode->UpdateFrontendNodeLocation(Location);
 		NewGraphNode->SyncLocationFromFrontendNode();
 		SchemaPrivate::TryConnectNewNodeToMatchingDataTypePin(*NewGraphNode, FromPin);
-		SchemaPrivate::SelecteNodeInEditor(*MetaSoundGraph, *NewGraphNode);
+		SchemaPrivate::SelectNodeInEditor(*MetaSoundGraph, *NewGraphNode);
 		return NewGraphNode;
 	}
 
@@ -582,7 +592,7 @@ UEdGraphNode* FMetasoundGraphSchemaAction_PromoteToInput::PerformAction(UEdGraph
 				if (ensure(SchemaPrivate::TryConnectNewNodeToMatchingDataTypePin(*EdGraphNode, FromPin)))
 				{
 					FGraphBuilder::RegisterGraphWithFrontend(ParentMetasound);
-					SchemaPrivate::SelecteNodeInEditor(*MetasoundGraph, *NewGraphNode);
+					SchemaPrivate::SelectNodeInEditorForRename(*MetasoundGraph, *NewGraphNode);
 					return EdGraphNode;
 				}
 			}
@@ -792,7 +802,7 @@ UEdGraphNode* FMetasoundGraphSchemaAction_PromoteToOutput::PerformAction(UEdGrap
 				if (ensure(SchemaPrivate::TryConnectNewNodeToMatchingDataTypePin(*EdGraphNode, FromPin)))
 				{
 					FGraphBuilder::RegisterGraphWithFrontend(ParentMetasound);
-					SchemaPrivate::SelecteNodeInEditor(*MetasoundGraph, *NewGraphNode);
+					SchemaPrivate::SelectNodeInEditorForRename(*MetasoundGraph, *NewGraphNode);
 					return EdGraphNode;
 				}
 			}
