UE-171614
[MetaSounds] - Show input widgets for constructor pins while not playing
#rnx

diff --git a/Engine/Plugins/Runtime/Metasound/Source/MetasoundEditor/Private/MetasoundEditorGraphInputNode.cpp b/Engine/Plugins/Runtime/Metasound/Source/MetasoundEditor/Private/MetasoundEditorGraphInputNode.cpp
--- a/Engine/Plugins/Runtime/Metasound/Source/MetasoundEditor/Private/MetasoundEditorGraphInputNode.cpp
+++ b/Engine/Plugins/Runtime/Metasound/Source/MetasoundEditor/Private/MetasoundEditorGraphInputNode.cpp
@@ -25,6 +25,8 @@
 
 #include UE_INLINE_GENERATED_CPP_BY_NAME(MetasoundEditorGraphInputNode)
 
+#define LOCTEXT_NAMESPACE "MetaSoundEditor"
+
 void UMetasoundEditorGraphInputNode::CacheTitle()
 {
 	using namespace Metasound::Frontend;
@@ -190,3 +192,34 @@ void UMetasoundEditorGraphInputNode::Validate(Metasound::Editor::FGraphNodeValid
 	}
 #endif // #if WITH_EDITOR
 }
+
+FText UMetasoundEditorGraphInputNode::GetTooltipText() const
+{
+	//If Constructor input
+	if (Input && Input->GetVertexAccessType() == EMetasoundFrontendVertexAccessType::Value)
+	{
+		UMetasoundEditorGraph* Graph = CastChecked<UMetasoundEditorGraph>(GetGraph());
+		if (Graph->IsPreviewing())
+		{
+			FText ToolTip = LOCTEXT("Metasound_ConstructorInputNodeDescription", "Editing constructor values is disabled while previewing.");
+			return ToolTip;
+		}
+	}
+
+	return Super::GetTooltipText();
+}
+
+bool UMetasoundEditorGraphInputNode::EnableInteractWidgets() const
+{
+	using namespace Metasound::Frontend;
+	
+	//If Constructor input
+	if (Input && Input->GetVertexAccessType() == EMetasoundFrontendVertexAccessType::Value)
+	{
+		UMetasoundEditorGraph* Graph = CastChecked<UMetasoundEditorGraph>(GetGraph());
+		return !Graph->IsPreviewing();
+	}
+
+	return true;
+}
+#undef LOCTEXT_NAMESPACE
diff --git a/Engine/Plugins/Runtime/Metasound/Source/MetasoundEditor/Private/MetasoundNodeDetailCustomization.cpp b/Engine/Plugins/Runtime/Metasound/Source/MetasoundEditor/Private/MetasoundNodeDetailCustomization.cpp
--- a/Engine/Plugins/Runtime/Metasound/Source/MetasoundEditor/Private/MetasoundNodeDetailCustomization.cpp
+++ b/Engine/Plugins/Runtime/Metasound/Source/MetasoundEditor/Private/MetasoundNodeDetailCustomization.cpp
@@ -343,8 +343,7 @@ namespace Metasound
 			{
 				if (const UMetasoundEditorGraph* OwningGraph = ParentMember->GetOwningGraph())
 				{
-					// Disable widget options for constructor inputs for now to prevent changing default value via widget while playing 
-					bShowWidgetOptions = OwningGraph->IsEditable() && ParentMember->GetVertexAccessType() == EMetasoundFrontendVertexAccessType::Reference;
+					bShowWidgetOptions = OwningGraph->IsEditable();
 				}
 			}
 
@@ -398,8 +397,7 @@ namespace Metasound
 			{
 				if (const UMetasoundEditorGraph* OwningGraph = ParentMember->GetOwningGraph())
 				{
-					// Disable widget options for constructor inputs for now to prevent changing default value via widget while playing 
-					bShowWidgetOptions = OwningGraph->IsEditable() && ParentMember->GetVertexAccessType() == EMetasoundFrontendVertexAccessType::Reference;
+					bShowWidgetOptions = OwningGraph->IsEditable();
 				}
 			}
 
diff --git a/Engine/Plugins/Runtime/Metasound/Source/MetasoundEditor/Private/SMetasoundGraphNode.cpp b/Engine/Plugins/Runtime/Metasound/Source/MetasoundEditor/Private/SMetasoundGraphNode.cpp
--- a/Engine/Plugins/Runtime/Metasound/Source/MetasoundEditor/Private/SMetasoundGraphNode.cpp
+++ b/Engine/Plugins/Runtime/Metasound/Source/MetasoundEditor/Private/SMetasoundGraphNode.cpp
@@ -712,12 +712,12 @@ namespace Metasound
 				bUseAudioMaterialWidgets = EditorSettings->bUseAudioMaterialWidgets;
 			}
 
-			// If editable float input node and not constructor input, check if custom widget required
+			// If input node, check if custom widget required
 			bool bShowContentWidget = false;
 			if (UMetasoundEditorGraphInput* GraphMember = Cast<UMetasoundEditorGraphInput>(GetMetaSoundMember()))
 			{
 				const UMetasoundEditorGraph* OwningGraph = GraphMember->GetOwningGraph();
-				if (OwningGraph && OwningGraph->IsEditable() && GraphMember->GetVertexAccessType() == EMetasoundFrontendVertexAccessType::Reference)
+				if (OwningGraph && OwningGraph->IsEditable() && GraphMember->GetVertexAccessType() != EMetasoundFrontendVertexAccessType::Unset)
 				{
 					UMetasoundEditorGraphMemberDefaultFloat* DefaultFloat = Cast<UMetasoundEditorGraphMemberDefaultFloat>(GraphMember->GetLiteral());
 					if (DefaultFloat && DefaultFloat->WidgetType != EMetasoundMemberDefaultWidget::None)
@@ -969,13 +969,13 @@ namespace Metasound
 						InputWidget->SetOutputRange(DefaultFloat->GetRange());
 						InputWidget->SetUnitsTextReadOnly(true);
 						InputWidget->SetSliderValue(InputWidget->GetSliderValue(DefaultFloat->GetDefault()));
-						InputWidget->SetVisibility(TAttribute<EVisibility>::Create([this]()
+						InputWidget->SetEnabled(TAttribute<bool>::Create([this]()
 						{
 							if (UMetasoundEditorGraphMemberNode* Node = GetMetaSoundMemberNode())
 							{
-								return Node->EnableInteractWidgets() ? EVisibility::Visible : EVisibility::Collapsed;
+								return Node->EnableInteractWidgets();
 							}
-							return EVisibility::Collapsed;
+							return false;
 						}));
 
 						// Setup & clear delegate if necessary (ex. if was just saved)
@@ -1078,7 +1078,14 @@ namespace Metasound
 									];
 
 								MaterialButtonWidget->SetDesiredSizeOverride(ButtonDesiredSize);
-
+								MaterialButtonWidget->SetEnabled(TAttribute<bool>::Create([this]()
+									{
+										if (UMetasoundEditorGraphMemberNode* Node = GetMetaSoundMemberNode())
+										{
+											return Node->EnableInteractWidgets();
+										}
+										return false;
+									}));
 
 								// Setup & clear delegate if necessary (ex. if was just saved)
 								if (InputButtonOnStateChangedDelegateHandle.IsValid())
diff --git a/Engine/Plugins/Runtime/Metasound/Source/MetasoundEditor/Public/MetasoundEditorGraphInputNode.h b/Engine/Plugins/Runtime/Metasound/Source/MetasoundEditor/Public/MetasoundEditorGraphInputNode.h
--- a/Engine/Plugins/Runtime/Metasound/Source/MetasoundEditor/Public/MetasoundEditorGraphInputNode.h
+++ b/Engine/Plugins/Runtime/Metasound/Source/MetasoundEditor/Public/MetasoundEditorGraphInputNode.h
@@ -45,6 +45,8 @@ public:
 	virtual FSlateIcon GetNodeTitleIcon() const override;
 	virtual void GetPinHoverText(const UEdGraphPin& Pin, FString& OutHoverText) const override;
 	virtual void Validate(Metasound::Editor::FGraphNodeValidationResult& OutResult) override;
+	virtual bool EnableInteractWidgets() const override;
+	virtual FText GetTooltipText() const override;
 
 protected:
 	// Breadcrumb used if associated FrontendNode cannot be found or has been unlinked
