Plot frequency responses of biquad, one-pole, ladder, and state variable filters.

diff --git a/Engine/Plugins/Runtime/AudioWidgets/Source/AudioWidgets/Private/SAudioSpectrumPlot.cpp b/Engine/Plugins/Runtime/AudioWidgets/Source/AudioWidgets/Private/SAudioSpectrumPlot.cpp
--- a/Engine/Plugins/Runtime/AudioWidgets/Source/AudioWidgets/Private/SAudioSpectrumPlot.cpp
+++ b/Engine/Plugins/Runtime/AudioWidgets/Source/AudioWidgets/Private/SAudioSpectrumPlot.cpp
@@ -281,6 +281,12 @@ FReply SAudioSpectrumPlot::OnMouseButtonUp(const FGeometry& InMyGeometry, const
 	return SCompoundWidget::OnMouseButtonUp(InMyGeometry, InMouseEvent);
 }
 
+FAudioSpectrumPlotScaleInfo SAudioSpectrumPlot::GetScaleInfo() const
+{
+	const FGeometry& AllottedGeometry = GetPaintSpaceGeometry();
+	return FAudioSpectrumPlotScaleInfo(AllottedGeometry.GetLocalSize(), FrequencyAxisScale.Get(), ViewMinFrequency.Get(), ViewMaxFrequency.Get(), ViewMinSoundLevel.Get(), ViewMaxSoundLevel.Get());
+}
+
 int32 SAudioSpectrumPlot::OnPaint(const FPaintArgs& Args, const FGeometry& AllottedGeometry, const FSlateRect& MyCullingRect, FSlateWindowElementList& OutDrawElements, int32 LayerId, const FWidgetStyle& InWidgetStyle, bool bParentEnabled) const
 {
 	const FAudioSpectrumPlotScaleInfo ScaleInfo(AllottedGeometry.GetLocalSize(), FrequencyAxisScale.Get(), ViewMinFrequency.Get(), ViewMaxFrequency.Get(), ViewMinSoundLevel.Get(), ViewMaxSoundLevel.Get());
diff --git a/Engine/Plugins/Runtime/AudioWidgets/Source/AudioWidgets/Public/SAudioSpectrumPlot.h b/Engine/Plugins/Runtime/AudioWidgets/Source/AudioWidgets/Public/SAudioSpectrumPlot.h
--- a/Engine/Plugins/Runtime/AudioWidgets/Source/AudioWidgets/Public/SAudioSpectrumPlot.h
+++ b/Engine/Plugins/Runtime/AudioWidgets/Source/AudioWidgets/Public/SAudioSpectrumPlot.h
@@ -176,6 +176,8 @@ public:
 
 	void UnbindOnGetAudioSpectrumData() { OnGetAudioSpectrumData.Unbind(); }
 
+	FAudioSpectrumPlotScaleInfo GetScaleInfo() const;
+
 private:
 	// Begin SWidget overrides.
 	virtual int32 OnPaint(const FPaintArgs& Args, const FGeometry& AllottedGeometry, const FSlateRect& MyCullingRect, FSlateWindowElementList& OutDrawElements, int32 LayerId, const FWidgetStyle& InWidgetStyle, bool bParentEnabled) const override;
diff --git a/Engine/Plugins/Runtime/Metasound/Source/MetasoundEditor/MetasoundEditor.Build.cs b/Engine/Plugins/Runtime/Metasound/Source/MetasoundEditor/MetasoundEditor.Build.cs
--- a/Engine/Plugins/Runtime/Metasound/Source/MetasoundEditor/MetasoundEditor.Build.cs
+++ b/Engine/Plugins/Runtime/Metasound/Source/MetasoundEditor/MetasoundEditor.Build.cs
@@ -30,6 +30,7 @@ namespace UnrealBuildTool.Rules
 					"MetasoundFrontend",
 					"MetasoundGenerator",
 					"MetasoundGraphCore",
+					"MetasoundStandardNodes",
 					"SignalProcessing",
 					"SubobjectEditor",
 					"ToolMenus",
diff --git a/Engine/Plugins/Runtime/Metasound/Source/MetasoundEditor/Private/MetasoundEditorGraphNodeFactory.h b/Engine/Plugins/Runtime/Metasound/Source/MetasoundEditor/Private/MetasoundEditorGraphNodeFactory.h
--- a/Engine/Plugins/Runtime/Metasound/Source/MetasoundEditor/Private/MetasoundEditorGraphNodeFactory.h
+++ b/Engine/Plugins/Runtime/Metasound/Source/MetasoundEditor/Private/MetasoundEditorGraphNodeFactory.h
@@ -6,9 +6,11 @@
 #include "MetasoundEditorGraphCommentNode.h"
 #include "MetasoundEditorGraphNode.h"
 #include "MetasoundEditorGraphSchema.h"
+#include "MetasoundStandardNodesNames.h"
 #include "NodeTemplates/MetasoundFrontendNodeTemplateAudioAnalyzer.h"
 #include "NodeTemplates/MetasoundFrontendNodeTemplateReroute.h"
 #include "SGraphNode.h"
+#include "SMetasoundFilterGraphNodes.h"
 #include "SMetasoundGraphNode.h"
 #include "SMetasoundGraphNodeComment.h"
 #include "SMetasoundSpectrumAnalyzerGraphNode.h"
@@ -37,6 +39,26 @@ class FMetasoundGraphNodeFactory : public FGraphPanelNodeFactory
 				{
 					return SNew(SMetaSoundSpectrumAnalyzerGraphNode, InNode);
 				}
+				else if (ClassName == FMetasoundFrontendClassName{ Metasound::StandardNodes::Namespace, TEXT("Biquad Filter"), Metasound::StandardNodes::AudioVariant })
+				{
+					return SNew(SMetaSoundBiquadFilterGraphNode, InNode);
+				}
+				else if (ClassName == FMetasoundFrontendClassName{ Metasound::StandardNodes::Namespace, TEXT("Ladder Filter"), Metasound::StandardNodes::AudioVariant })
+				{
+					return SNew(SMetaSoundLadderFilterGraphNode, InNode);
+				}
+				else if (ClassName == FMetasoundFrontendClassName{ Metasound::StandardNodes::Namespace, TEXT("One-Pole High Pass Filter"), Metasound::StandardNodes::AudioVariant })
+				{
+					return SNew(SMetaSoundOnePoleHighPassFilterGraphNode, InNode);
+				}
+				else if (ClassName == FMetasoundFrontendClassName{ Metasound::StandardNodes::Namespace, TEXT("One-Pole Low Pass Filter"), Metasound::StandardNodes::AudioVariant })
+				{
+					return SNew(SMetaSoundOnePoleLowPassFilterGraphNode, InNode);
+				}
+				else if (ClassName == FMetasoundFrontendClassName{ Metasound::StandardNodes::Namespace, TEXT("State Variable Filter"), Metasound::StandardNodes::AudioVariant })
+				{
+					return SNew(SMetaSoundStateVariableFilterGraphNode, InNode);
+				}
 			}
 			return SNew(SMetaSoundGraphNode, InNode);
 		}
diff --git a/Engine/Plugins/Runtime/Metasound/Source/MetasoundEditor/Private/SMetasoundFilterGraphNodes.cpp b/Engine/Plugins/Runtime/Metasound/Source/MetasoundEditor/Private/SMetasoundFilterGraphNodes.cpp
new file mode 100644
--- /dev/null
+++ b/Engine/Plugins/Runtime/Metasound/Source/MetasoundEditor/Private/SMetasoundFilterGraphNodes.cpp
@@ -0,0 +1,482 @@
+// Copyright Epic Games, Inc. All Rights Reserved.
+#include "SMetasoundFilterGraphNodes.h"
+
+#include "Algo/Count.h"
+#include "DSP/FloatArrayMath.h"
+#include "MetasoundEditorGraphBuilder.h"
+#include "MetasoundEditorGraphNode.h"
+#include "Widgets/Layout/SBox.h"
+
+#define LOCTEXT_NAMESPACE "SMetasoundFilterGraphNodes"
+
+namespace Metasound::Editor
+{
+	namespace SMetasoundFilterGraphNodesPrivate
+	{
+		inline bool IsConnectedAudioOutputPin(const UEdGraphPin* Pin)
+		{
+			const bool bIsPinConnected = !Pin->LinkedTo.IsEmpty();
+			const bool bIsAudioOutputPin = Pin->Direction == EGPD_Output && Pin->PinType.PinCategory == FGraphBuilder::PinCategoryAudio;
+			return (bIsAudioOutputPin && bIsPinConnected);
+		}
+
+		template<class T>
+		inline TOptional<T> GetPinValue(UMetasoundEditorGraphNode& MetaSoundNode, FName PinName)
+		{
+			UEdGraphPin** FoundNamedInputPin = MetaSoundNode.Pins.FindByPredicate([&PinName](const UEdGraphPin* InPin)
+				{
+					return (InPin->Direction == EGPD_Input && InPin->PinType.PinCategory != FGraphBuilder::PinCategoryAudio && InPin->PinName == PinName);
+				});
+
+			if (!FoundNamedInputPin)
+			{
+				return NullOpt;
+			}
+
+			T PinValue;
+
+			UEdGraphPin* Pin = *FoundNamedInputPin;
+			if (Pin->LinkedTo.IsEmpty())
+			{
+				FMetasoundFrontendLiteral DefaultLiteral;
+				if (FGraphBuilder::GetPinLiteral(*Pin, DefaultLiteral))
+				{
+					if (DefaultLiteral.TryGet(PinValue))
+					{
+						return PinValue;
+					}
+				}
+			}
+			else
+			{
+				// Find connected output for the input (only ever one):
+				UEdGraphPin* SourcePin = Pin->LinkedTo.Last();
+				ensure(SourcePin->Direction == EGPD_Output);
+
+				const UEdGraphPin* ReroutedOutputPin = FGraphBuilder::FindReroutedOutputPin(SourcePin);
+
+				if (const UMetasoundEditorGraphNode* Node = Cast<UMetasoundEditorGraphNode>(ReroutedOutputPin->GetOwningNode()))
+				{
+					TSharedPtr<FEditor> Editor = FGraphBuilder::GetEditorForNode(MetaSoundNode);
+					if (Editor.IsValid())
+					{
+						const FGuid NodeID = Node->GetNodeID();
+						const FName OutputName = ReroutedOutputPin->GetFName();
+						if (Editor->GetConnectionManager().GetValue(NodeID, OutputName, PinValue))
+						{
+							return PinValue;
+						}
+					}
+
+					if (const UMetasoundEditorGraphMemberNode* MemberNode = Cast<UMetasoundEditorGraphMemberNode>(Node))
+					{
+						if (const UMetasoundEditorGraphMember* Member = MemberNode->GetMember())
+						{
+							if (Editor.IsValid())
+							{
+								// For an input member, we must use this NodeID/VertexName pair with the GraphConnectionManager:
+								const FGuid MemberID = Member->GetMemberID();
+								const FName MemberName = Member->GetMemberName();
+								if (Editor->GetConnectionManager().GetValue(MemberID, MemberName, PinValue))
+								{
+									return PinValue;
+								}
+							}
+
+							if (const UMetasoundEditorGraphMemberDefaultLiteral* MemberDefaultLiteral = Member->GetLiteral())
+							{
+								const FMetasoundFrontendLiteral DefaultLiteral = MemberDefaultLiteral->GetDefault();
+								if (DefaultLiteral.TryGet(PinValue))
+								{
+									return PinValue;
+								}
+							}
+						}
+					}
+				}
+			}
+
+			return NullOpt;
+		}
+	} // namespace SMetasoundFilterGraphNodesPrivate
+
+	const float SMetaSoundFilterGraphNode::SampleRate = 48000.0f;
+
+	SMetaSoundFilterGraphNode::~SMetaSoundFilterGraphNode()
+	{
+		if (ContextMenuExtension.IsValid())
+		{
+			FrequencyResponsePlot->RemoveContextMenuExtension(ContextMenuExtension.ToSharedRef());
+		}
+	}
+
+	void SMetaSoundFilterGraphNode::Tick(const FGeometry& AllottedGeometry, const double InCurrentTime, const float InDeltaTime)
+	{
+		bHasFilterParams = UpdateFilterParams();
+
+		SMetaSoundGraphNode::Tick(AllottedGeometry, InCurrentTime, InDeltaTime);
+	}
+
+	void SMetaSoundFilterGraphNode::CreateBelowPinControls(TSharedPtr<SVerticalBox> MainBox)
+	{
+		if (!FrequencyResponsePlot.IsValid())
+		{
+			FrequencyResponsePlot = SNew(SAudioSpectrumPlot)
+				.Clipping(EWidgetClipping::ClipToBounds)
+				.ViewMinSoundLevel(-24.0f)
+				.FrequencyAxisPixelBucketMode_Lambda([]() { return EAudioSpectrumPlotFrequencyAxisPixelBucketMode::Sample; }) // Binding this property has the effect of hiding its context menu entry (PixelBucketMode is not much use here).
+				.OnGetAudioSpectrumData(this, &SMetaSoundFilterGraphNode::GetAudioSpectrumData);
+
+			ContextMenuExtension = FrequencyResponsePlot->AddContextMenuExtension(EExtensionHook::Before, nullptr, FMenuExtensionDelegate::CreateSP(this, &SMetaSoundFilterGraphNode::ExtendSpectrumPlotContextMenu));
+		}
+
+		MainBox->AddSlot()
+			.Padding(1.0f, 0.0f)
+			[
+				SNew(SBox)
+					.MinDesiredHeight(125.0f)
+					[
+						FrequencyResponsePlot.ToSharedRef()
+					]
+			];
+	}
+
+	FAudioPowerSpectrumData SMetaSoundFilterGraphNode::GetAudioSpectrumData()
+	{
+		if (!bHasFilterParams || !FrequencyResponsePlot.IsValid())
+		{
+			return FAudioPowerSpectrumData();
+		}
+
+		// Fill array of frequencies (in Hz) to be plotted (these may be log spaced or linear, as defined by the FrequencyResponsePlot ScaleInfo transform):
+		CenterFrequencies.Reset();
+		float LocalX = 0.0f;
+		const float LocalXEnd = FrequencyResponsePlot->GetPaintSpaceGeometry().GetLocalSize().X;
+		const FAudioSpectrumPlotScaleInfo ScaleInfo = FrequencyResponsePlot->GetScaleInfo();
+		while (LocalX <= LocalXEnd)
+		{
+			const float Frequency = ScaleInfo.LocalXToFrequency(LocalX);
+			CenterFrequencies.Add(Frequency);
+			LocalX += 1.0f;
+		}
+
+		// Create array of complex z values for all desired frequencies (interleaved real and imaginary):
+		const int32 NumFrequencies = CenterFrequencies.Num();
+		TArray<float> ComplexValues;
+		ComplexValues.SetNumUninitialized(2 * NumFrequencies);
+		const float HzToOmega = UE_TWO_PI / SampleRate;
+		for (int32 Index = 0; Index < NumFrequencies; Index++)
+		{
+			const float Omega = HzToOmega * CenterFrequencies[Index];
+			FMath::SinCos(&ComplexValues[2 * Index + 1], &ComplexValues[2 * Index + 0], Omega);
+		}
+
+		// Get the frequency response for all of the frequencies:
+		ArrayCalculateFilterResponseInPlace(ComplexValues);
+
+		// Store frequency response as squared magnitudes:
+		SquaredMagnitudes.SetNumUninitialized(NumFrequencies);
+		Audio::ArrayComplexToPower(ComplexValues, SquaredMagnitudes);
+
+		return FAudioPowerSpectrumData
+		{
+			.CenterFrequencies = CenterFrequencies,
+			.SquaredMagnitudes = SquaredMagnitudes,
+		};
+	}
+
+#pragma region Biquad Filter
+
+	bool SMetaSoundBiquadFilterGraphNode::UpdateFilterParams()
+	{
+		using namespace SMetasoundFilterGraphNodesPrivate;
+
+		UMetasoundEditorGraphNode& MetaSoundNode = GetMetaSoundNode();
+
+		// Try and get all required filter params:
+		const TOptional<int32> Type = GetPinValue<int32>(MetaSoundNode, TEXT("Type"));
+		const TOptional<float> Frequency = GetPinValue<float>(MetaSoundNode, TEXT("Cutoff Frequency"));
+		const TOptional<float> Bandwidth = GetPinValue<float>(MetaSoundNode, TEXT("Bandwidth"));
+		const TOptional<float> FilterGainDb = GetPinValue<float>(MetaSoundNode, TEXT("Gain"));
+
+		if (!Type.IsSet() || !Frequency.IsSet() || !Bandwidth.IsSet() || !FilterGainDb.IsSet())
+		{
+			return false;
+		}
+
+		const float MaxCutoffFrequency = 0.5f * SampleRate;
+		const float CurrentFrequency = FMath::Clamp(*Frequency, 0.f, MaxCutoffFrequency);
+		const float CurrentBandwidth = FMath::Max(*Bandwidth, 0.f);
+		const float CurrentFilterGainDb = FMath::Clamp(*FilterGainDb, -90.0f, 20.0f);
+
+		// Update the filter with the filter params:
+		if (Filter.GetNumChannels() == 0)
+		{
+			constexpr int32 NumChannels = 1;
+			Filter.Init(SampleRate, NumChannels, static_cast<Audio::EBiquadFilter::Type>(Type.GetValue()), CurrentFrequency, CurrentBandwidth, CurrentFilterGainDb);
+		}
+		else
+		{
+			Filter.SetParams(static_cast<Audio::EBiquadFilter::Type>(Type.GetValue()), CurrentFrequency, CurrentBandwidth, CurrentFilterGainDb);
+		}
+
+		return true;
+	}
+
+	void SMetaSoundBiquadFilterGraphNode::ArrayCalculateFilterResponseInPlace(TArrayView<float> InOutComplexValues) const
+	{
+		Filter.ArrayCalculateResponseInPlace(InOutComplexValues);
+	}
+
+#pragma endregion
+
+#pragma region Ladder Filter
+
+	SMetaSoundLadderFilterGraphNode::SMetaSoundLadderFilterGraphNode()
+	{
+		constexpr int32 NumChannels = 1;
+		Filter.Init(SampleRate, NumChannels);
+	}
+
+	bool SMetaSoundLadderFilterGraphNode::UpdateFilterParams()
+	{
+		using namespace SMetasoundFilterGraphNodesPrivate;
+
+		UMetasoundEditorGraphNode& MetaSoundNode = GetMetaSoundNode();
+
+		// Try and get all required filter params:
+		const TOptional<float> Frequency = GetPinValue<float>(MetaSoundNode, TEXT("Cutoff Frequency"));
+		const TOptional<float> Resonance = GetPinValue<float>(MetaSoundNode, TEXT("Resonance"));
+
+		if (!Frequency.IsSet() || !Resonance.IsSet())
+		{
+			return false;
+		}
+
+		const float MaxCutoffFrequency = 0.5f * SampleRate;
+		const float CurrentFrequency = FMath::Clamp(*Frequency, 0.f, MaxCutoffFrequency);
+		const float CurrentResonance = FMath::Clamp(*Resonance, 1.0f, 10.0f);
+
+		// Update the filter with the filter params:
+		Filter.SetQ(CurrentResonance);
+		Filter.SetFrequency(CurrentFrequency);
+
+		Filter.Update();
+
+		return true;
+	}
+
+	void SMetaSoundLadderFilterGraphNode::ArrayCalculateFilterResponseInPlace(TArrayView<float> InOutComplexValues) const
+	{
+		Filter.ArrayCalculateResponseInPlace(InOutComplexValues);
+	}
+
+#pragma endregion
+
+#pragma region OnePoleHighPass Filter
+
+	SMetaSoundOnePoleHighPassFilterGraphNode::SMetaSoundOnePoleHighPassFilterGraphNode()
+	{
+		constexpr int32 NumChannels = 1;
+		Filter.Init(SampleRate, NumChannels);
+	}
+
+	bool SMetaSoundOnePoleHighPassFilterGraphNode::UpdateFilterParams()
+	{
+		using namespace SMetasoundFilterGraphNodesPrivate;
+
+		UMetasoundEditorGraphNode& MetaSoundNode = GetMetaSoundNode();
+
+		// Try and get all required filter params:
+		const TOptional<float> Frequency = GetPinValue<float>(MetaSoundNode, TEXT("Cutoff Frequency"));
+
+		if (!Frequency.IsSet())
+		{
+			return false;
+		}
+
+		const float ClampedFreq = FMath::Clamp(0.0f, *Frequency, SampleRate);
+
+		// Update the filter with the filter params:
+		Filter.StartFrequencyInterpolation(ClampedFreq);
+
+		return true;
+	}
+
+	void SMetaSoundOnePoleHighPassFilterGraphNode::ArrayCalculateFilterResponseInPlace(TArrayView<float> InOutComplexValues) const
+	{
+		Filter.ArrayCalculateResponseInPlace(InOutComplexValues);
+	}
+
+#pragma endregion
+
+#pragma region OnePoleLowPass Filter
+
+	SMetaSoundOnePoleLowPassFilterGraphNode::SMetaSoundOnePoleLowPassFilterGraphNode()
+	{
+		constexpr int32 NumChannels = 1;
+		Filter.Init(SampleRate, NumChannels);
+	}
+
+	bool SMetaSoundOnePoleLowPassFilterGraphNode::UpdateFilterParams()
+	{
+		using namespace SMetasoundFilterGraphNodesPrivate;
+
+		UMetasoundEditorGraphNode& MetaSoundNode = GetMetaSoundNode();
+
+		// Try and get all required filter params:
+		const TOptional<float> Frequency = GetPinValue<float>(MetaSoundNode, TEXT("Cutoff Frequency"));
+		if (!Frequency.IsSet())
+		{
+			return false;
+		}
+
+		const float ClampedFreq = FMath::Clamp(0.0f, *Frequency, SampleRate);
+
+		// Update the filter with the filter params:
+		Filter.StartFrequencyInterpolation(ClampedFreq);
+
+		return true;
+	}
+
+	void SMetaSoundOnePoleLowPassFilterGraphNode::ArrayCalculateFilterResponseInPlace(TArrayView<float> InOutComplexValues) const
+	{
+		Filter.ArrayCalculateResponseInPlace(InOutComplexValues);
+	}
+
+#pragma endregion
+
+#pragma region State Variable Filter
+
+	const FName SMetaSoundStateVariableFilterGraphNode::LowPassFilter(TEXT("Low Pass Filter"));
+	const FName SMetaSoundStateVariableFilterGraphNode::HighPassFilter(TEXT("High Pass Filter"));
+	const FName SMetaSoundStateVariableFilterGraphNode::BandPass(TEXT("Band Pass"));
+	const FName SMetaSoundStateVariableFilterGraphNode::BandStop(TEXT("Band Stop"));
+
+	SMetaSoundStateVariableFilterGraphNode::SMetaSoundStateVariableFilterGraphNode()
+		: DisplayedFilterResponse(LowPassFilter)
+	{
+		constexpr int32 NumChannels = 1;
+		Filter.Init(SampleRate, NumChannels);
+	}
+
+	bool SMetaSoundStateVariableFilterGraphNode::UpdateFilterParams()
+	{
+		using namespace SMetasoundFilterGraphNodesPrivate;
+
+		UMetasoundEditorGraphNode& MetaSoundNode = GetMetaSoundNode();
+
+		// If there are connected audio outputs but the current selection is not connected, auto-select a new filter type:
+		UEdGraphPin** FirstConnectedAudioOutputPin = MetaSoundNode.Pins.FindByPredicate(IsConnectedAudioOutputPin);
+		if (FirstConnectedAudioOutputPin != nullptr)
+		{
+			for (UEdGraphPin* Pin : MetaSoundNode.Pins)
+			{
+				const bool bIsPinConnected = !Pin->LinkedTo.IsEmpty();
+				if (Pin->PinName == DisplayedFilterResponse && !bIsPinConnected)
+				{
+					DisplayedFilterResponse = (*FirstConnectedAudioOutputPin)->PinName;
+				}
+			}
+		}
+
+		// Set the filter type on the filter:
+		if (DisplayedFilterResponse == LowPassFilter)
+		{
+			Filter.SetFilterType(Audio::EFilter::LowPass);
+		}
+		else if (DisplayedFilterResponse == HighPassFilter)
+		{
+			Filter.SetFilterType(Audio::EFilter::HighPass);
+		}
+		else if (DisplayedFilterResponse == BandPass)
+		{
+			Filter.SetFilterType(Audio::EFilter::BandPass);
+		}
+		else if (DisplayedFilterResponse == BandStop)
+		{
+			Filter.SetFilterType(Audio::EFilter::BandStop);
+		}
+
+		// Try and get all required filter params:
+		const TOptional<float> Frequency = GetPinValue<float>(MetaSoundNode, TEXT("Cutoff Frequency"));
+		const TOptional<float> Resonance = GetPinValue<float>(MetaSoundNode, TEXT("Resonance"));
+		const TOptional<float> BandStopControl = GetPinValue<float>(MetaSoundNode, TEXT("Band Stop Control"));
+
+		if (!Frequency.IsSet() || !Resonance.IsSet() || !BandStopControl.IsSet())
+		{
+			return false;
+		}
+
+		const float MaxCutoffFrequency = 0.5f * SampleRate;
+		const float CurrentFrequency = FMath::Clamp(*Frequency, 0.f, MaxCutoffFrequency);
+		const float CurrentResonance = FMath::Clamp(*Resonance, 0.f, 10.f);
+		const float CurrentBandStopControl = FMath::Clamp(*BandStopControl, 0.f, 1.f);
+
+		// Update the filter with the filter params:
+		Filter.SetQ(CurrentResonance);
+		Filter.SetFrequency(CurrentFrequency);
+		Filter.SetBandStopControl(CurrentBandStopControl);
+
+		Filter.Update();
+
+		return true;
+	}
+
+	void SMetaSoundStateVariableFilterGraphNode::ArrayCalculateFilterResponseInPlace(TArrayView<float> InOutComplexValues) const
+	{
+		Filter.ArrayCalculateResponseInPlace(InOutComplexValues);
+	}
+
+	void SMetaSoundStateVariableFilterGraphNode::ExtendSpectrumPlotContextMenu(FMenuBuilder& MenuBuilder)
+	{
+		using namespace SMetasoundFilterGraphNodesPrivate;
+
+		// Display filter response selection submenu if no outputs are connected, or more than one output is connected:
+		const int NumConnectedAudioOutputPins = Algo::CountIf(GetMetaSoundNode().Pins, IsConnectedAudioOutputPin);
+		if (NumConnectedAudioOutputPins == 0 || NumConnectedAudioOutputPins > 1)
+		{
+			MenuBuilder.AddSubMenu(
+				LOCTEXT("DisplayedFilterResponse", "Displayed Filter Response"),
+				FText(),
+				FNewMenuDelegate::CreateSP(this, &SMetaSoundStateVariableFilterGraphNode::BuildFilterOutputSubMenu));
+		}
+	}
+
+	void SMetaSoundStateVariableFilterGraphNode::BuildFilterOutputSubMenu(FMenuBuilder& SubMenu)
+	{
+		using namespace SMetasoundFilterGraphNodesPrivate;
+
+		const UMetasoundEditorGraphNode& MetaSoundNode = GetMetaSoundNode();
+
+		// Add menu entries for all connected audio output pins. If there are no connected audio output pins then add entries for all audio output pins:
+		const bool bHasConnectedAudioOutputPins = MetaSoundNode.Pins.ContainsByPredicate(IsConnectedAudioOutputPin);
+		for (UEdGraphPin* Pin : MetaSoundNode.Pins)
+		{
+			const bool bIsAudioOutputPin = Pin->Direction == EGPD_Output && Pin->PinType.PinCategory == FGraphBuilder::PinCategoryAudio;
+			if (bIsAudioOutputPin)
+			{
+				const bool bIsPinConnected = !Pin->LinkedTo.IsEmpty();
+				if (bIsPinConnected || !bHasConnectedAudioOutputPins)
+				{
+					SubMenu.AddMenuEntry(
+						FText::FromName(Pin->PinName),
+						FText::FromString(Pin->PinToolTip),
+						FSlateIcon(),
+						FUIAction(
+							FExecuteAction::CreateSPLambda(this, [this, PinName = Pin->PinName]() { DisplayedFilterResponse = PinName; }),
+							FCanExecuteAction(),
+							FIsActionChecked::CreateSPLambda(this, [this, PinName = Pin->PinName]() { return DisplayedFilterResponse == PinName; })
+						),
+						NAME_None,
+						EUserInterfaceActionType::ToggleButton);
+				}
+			}
+		}
+	}
+
+#pragma endregion
+
+}
+
+#undef LOCTEXT_NAMESPACE
diff --git a/Engine/Plugins/Runtime/Metasound/Source/MetasoundEditor/Private/SMetasoundFilterGraphNodes.h b/Engine/Plugins/Runtime/Metasound/Source/MetasoundEditor/Private/SMetasoundFilterGraphNodes.h
new file mode 100644
--- /dev/null
+++ b/Engine/Plugins/Runtime/Metasound/Source/MetasoundEditor/Private/SMetasoundFilterGraphNodes.h
@@ -0,0 +1,124 @@
+// Copyright Epic Games, Inc. All Rights Reserved.
+#pragma once
+
+#include "DSP/Filter.h"
+#include "DSP/InterpolatedOnePole.h"
+#include "SAudioSpectrumPlot.h"
+#include "SMetasoundGraphNode.h"
+
+namespace Metasound::Editor
+{
+	/**
+	* Abstract class implementing shared functionality for an SGraphNode for any Metasound basic filter. Allows for displaying a frequency response plot of the filter.
+	*/
+	class SMetaSoundFilterGraphNode : public SMetaSoundGraphNode
+	{
+	public:
+		virtual ~SMetaSoundFilterGraphNode();
+
+	protected:
+		// Update function for derived classes to update their state each frame. Implementations should return true if they have the required information to plot a frequency response, or false otherwise.
+		virtual bool UpdateFilterParams() = 0;
+
+		// Derived classes should apply the filter transfer function to each z-domain value in the given array (complex numbers given as interleaved floats).
+		virtual void ArrayCalculateFilterResponseInPlace(TArrayView<float> InOutComplexValues) const = 0;
+
+		// Derived classes can optionally add items to the spectrum plot menu by overriding this member function.
+		virtual void ExtendSpectrumPlotContextMenu(FMenuBuilder& MenuBuilder) {};
+
+		// Begin SWidget overrides.
+		virtual void Tick(const FGeometry& AllottedGeometry, const double InCurrentTime, const float InDeltaTime) override;
+		// End SWidget overrides.
+
+		// Begin SGraphNode overrides.
+		virtual void CreateBelowPinControls(TSharedPtr<SVerticalBox> MainBox) override;
+		// End SGraphNode overrides.
+
+		static const float SampleRate;
+		
+	private:
+		FAudioPowerSpectrumData GetAudioSpectrumData();
+
+		TSharedPtr<SAudioSpectrumPlot> FrequencyResponsePlot;
+		TSharedPtr<const FExtensionBase> ContextMenuExtension;
+		TArray<float> CenterFrequencies;
+		TArray<float> SquaredMagnitudes;
+		bool bHasFilterParams = false;
+	};
+
+
+	class SMetaSoundBiquadFilterGraphNode : public SMetaSoundFilterGraphNode
+	{
+	protected:
+		virtual bool UpdateFilterParams() override;
+		virtual void ArrayCalculateFilterResponseInPlace(TArrayView<float> InOutComplexValues) const override;
+
+	private:
+		Audio::FBiquadFilter Filter;
+	};
+
+
+	class SMetaSoundLadderFilterGraphNode : public SMetaSoundFilterGraphNode
+	{
+	public:
+		SMetaSoundLadderFilterGraphNode();
+
+	protected:
+		virtual bool UpdateFilterParams() override;
+		virtual void ArrayCalculateFilterResponseInPlace(TArrayView<float> InOutComplexValues) const override;
+
+	private:
+		Audio::FLadderFilter Filter;
+	};
+
+
+	class SMetaSoundOnePoleHighPassFilterGraphNode : public SMetaSoundFilterGraphNode
+	{
+	public:
+		SMetaSoundOnePoleHighPassFilterGraphNode();
+
+	protected:
+		virtual bool UpdateFilterParams() override;
+		virtual void ArrayCalculateFilterResponseInPlace(TArrayView<float> InOutComplexValues) const override;
+
+	private:
+		Audio::FInterpolatedHPF Filter;
+	};
+
+
+	class SMetaSoundOnePoleLowPassFilterGraphNode : public SMetaSoundFilterGraphNode
+	{
+	public:
+		SMetaSoundOnePoleLowPassFilterGraphNode();
+
+	protected:
+		virtual bool UpdateFilterParams() override;
+		virtual void ArrayCalculateFilterResponseInPlace(TArrayView<float> InOutComplexValues) const override;
+
+	private:
+		Audio::FInterpolatedLPF Filter;
+	};
+
+
+	class SMetaSoundStateVariableFilterGraphNode : public SMetaSoundFilterGraphNode
+	{
+	public:
+		SMetaSoundStateVariableFilterGraphNode();
+
+	protected:
+		virtual bool UpdateFilterParams() override;
+		virtual void ArrayCalculateFilterResponseInPlace(TArrayView<float> InOutComplexValues) const override;
+		virtual void ExtendSpectrumPlotContextMenu(FMenuBuilder& MenuBuilder) override;
+
+	private:
+		void BuildFilterOutputSubMenu(FMenuBuilder& SubMenu);
+
+		static const FName LowPassFilter;
+		static const FName HighPassFilter;
+		static const FName BandPass;
+		static const FName BandStop;
+
+		Audio::FStateVariableFilter Filter;
+		FName DisplayedFilterResponse;
+	};
+}
diff --git a/Engine/Source/Runtime/SignalProcessing/Private/Filter.cpp b/Engine/Source/Runtime/SignalProcessing/Private/Filter.cpp
--- a/Engine/Source/Runtime/SignalProcessing/Private/Filter.cpp
+++ b/Engine/Source/Runtime/SignalProcessing/Private/Filter.cpp
@@ -1,6 +1,8 @@
 // Copyright Epic Games, Inc. All Rights Reserved.
 
 #include "DSP/Filter.h"
+
+#include <complex>
 #include "DSP/Dsp.h"
 
 namespace Audio
@@ -252,6 +254,30 @@ namespace Audio
 		bEnabled = bInEnabled;
 	}
 
+	void FBiquadFilter::ArrayCalculateResponseInPlace(TArrayView<float> InOutComplexValues) const
+	{
+		constexpr int32 ChannelIndex = 0;
+		check(ChannelIndex < NumChannels);
+		const FBiquadCoeff& Coefficients = Biquad[ChannelIndex];
+
+		const int32 NumFloats = InOutComplexValues.Num();
+		check(NumFloats % 2 == 0);
+		for (int32 Index = 0; Index < NumFloats; Index += 2)
+		{
+			const float ZReal = InOutComplexValues[Index + 0];
+			const float ZImag = InOutComplexValues[Index + 1];
+			const std::complex<float> Z(ZReal, ZImag);
+
+			// H(z) == (a0*z^2 + a1*z + a2) / (z^2 + b1*z + b2)
+			const std::complex<float> Numerator = (Coefficients.A0 * Z + Coefficients.A1) * Z + Coefficients.A2;
+			const std::complex<float> Denominator = (Z + Coefficients.B1) * Z + Coefficients.B2;
+			const std::complex<float> H = Numerator / Denominator;
+
+			InOutComplexValues[Index + 0] = H.real();
+			InOutComplexValues[Index + 1] = H.imag();
+		}
+	}
+
 	void FBiquadFilter::CalculateBiquadCoefficients()
 	{
 		static const float NaturalLog2 = 0.69314718055994530942f;
@@ -687,6 +713,45 @@ namespace Audio
 		}
 	}
 
+	void FOnePoleFilter::ArrayCalculateResponseInPlace(TArrayView<float> InOutComplexValues) const
+	{
+		// Get the coefficients for the filter transfer function:
+		const float A1 = A0;
+		const float B1 = 2.0f * A0 - 1.0f;
+
+		const int32 NumFloats = InOutComplexValues.Num();
+		check(NumFloats % 2 == 0);
+		if (FilterType == EFilter::HighPass)
+		{
+			for (int32 Index = 0; Index < NumFloats; Index += 2)
+			{
+				const float ZReal = InOutComplexValues[Index + 0];
+				const float ZImag = InOutComplexValues[Index + 1];
+				const std::complex<float> Z(ZReal, ZImag);
+
+				const std::complex<float> LPF = (A0 * Z + A1) / (B1 + Z);
+				const std::complex<float> HPF = 1.0f - LPF;
+
+				InOutComplexValues[Index + 0] = HPF.real();
+				InOutComplexValues[Index + 1] = HPF.imag();
+			}
+		}
+		else
+		{
+			for (int32 Index = 0; Index < NumFloats; Index += 2)
+			{
+				const float ZReal = InOutComplexValues[Index + 0];
+				const float ZImag = InOutComplexValues[Index + 1];
+				const std::complex<float> Z(ZReal, ZImag);
+
+				const std::complex<float> LPF = (A0 * Z + A1) / (B1 + Z);
+
+				InOutComplexValues[Index + 0] = LPF.real();
+				InOutComplexValues[Index + 1] = LPF.imag();
+			}
+		}
+	}
+
 	FStateVariableFilter::FStateVariableFilter()
 		: InputScale(1.0f)
 		, A0(1.0f)
@@ -866,6 +931,62 @@ namespace Audio
 			}
 		}
 	}
+
+	void FStateVariableFilter::ArrayCalculateResponseInPlace(TArrayView<float> InOutComplexValues) const
+	{
+		// Get the params:
+		const float G = A0;
+		const float R = 0.5f * (Feedback - G); // Dampening
+
+		// Evaluate common subexpressions:
+		const float GSqr = G * G;
+		const float TwoRG = 2.0f * R * G;
+		const float TwoGSqr = 2.0f * GSqr;
+
+		const int32 NumFloats = InOutComplexValues.Num();
+		check(NumFloats % 2 == 0);
+		for (int32 Index = 0; Index < NumFloats; Index += 2)
+		{
+			const float ZReal = InOutComplexValues[Index + 0];
+			const float ZImag = InOutComplexValues[Index + 1];
+			const std::complex<float> Z(ZReal, ZImag);
+
+			// Evaluate the polynomials that form the rational fractions:
+			const std::complex<float> NumeratorLPF = (GSqr * Z + TwoGSqr) * Z + GSqr;
+			const std::complex<float> NumeratorHPF = (Z - 2.0f) * Z + 1.0f;
+			const std::complex<float> NumeratorBPF = G * Z * Z - G;
+			const std::complex<float> Denominator = ((GSqr + TwoRG + 1.0f) * Z + (TwoGSqr - 2.0f)) * Z + (GSqr - TwoRG + 1.0f);
+
+			// Evaluate the rational functions:
+			const std::complex<float> Reciprocal = 1.0f / Denominator;
+			const std::complex<float> LPF = NumeratorLPF * Reciprocal;
+			const std::complex<float> HPF = NumeratorHPF * Reciprocal;
+			const std::complex<float> BPF = NumeratorBPF * Reciprocal;
+			const std::complex<float> BSF = BandStopParam * HPF + (1.0f - BandStopParam) * LPF;
+
+			switch (FilterType)
+			{
+			default:
+			case EFilter::LowPass:
+				InOutComplexValues[Index + 0] = LPF.real();
+				InOutComplexValues[Index + 1] = LPF.imag();
+				break;
+			case EFilter::HighPass:
+				InOutComplexValues[Index + 0] = HPF.real();
+				InOutComplexValues[Index + 1] = HPF.imag();
+				break;
+			case EFilter::BandPass:
+				InOutComplexValues[Index + 0] = BPF.real();
+				InOutComplexValues[Index + 1] = BPF.imag();
+				break;
+			case EFilter::BandStop:
+				InOutComplexValues[Index + 0] = BSF.real();
+				InOutComplexValues[Index + 1] = BSF.imag();
+				break;
+			}
+		}
+	}
+
 	FLadderFilter::FLadderFilter()
 		: K(0.0f)
 		, Gamma(0.0f)
@@ -1054,4 +1175,33 @@ namespace Audio
 			}
 		}
 	}
+
+	void FLadderFilter::ArrayCalculateResponseInPlace(TArrayView<float> InOutComplexValues) const
+	{
+		check(FilterType == EFilter::LowPass);
+
+		// Get the one-pole LPF coefficients:
+		const float A0 = OnePoleFilters[0].GetCoefficient();
+		const float A1 = A0;
+		const float B1 = 2.0f * A0 - 1.0f;
+
+		const int32 NumFloats = InOutComplexValues.Num();
+		check(NumFloats % 2 == 0);
+		for (int32 Index = 0; Index < NumFloats; Index += 2)
+		{
+			const float ZReal = InOutComplexValues[Index + 0];
+			const float ZImag = InOutComplexValues[Index + 1];
+			const std::complex<float> Z(ZReal, ZImag);
+
+			// Calculate the response for the 4 one-pole filters in series:
+			const std::complex<float> OnePoleResponse = (A0 * Z + A1) / (B1 + Z);
+			const std::complex<float> Series = OnePoleResponse * OnePoleResponse * OnePoleResponse * OnePoleResponse;
+
+			// Apply pass-band gain-compensation and feedback:
+			const std::complex<float> H = (1.0f + PassBandGainCompensation * K) * Series / (K * Series + 1.0f);
+
+			InOutComplexValues[Index + 0] = H.real();
+			InOutComplexValues[Index + 1] = H.imag();
+		}
+	}
 }
diff --git a/Engine/Source/Runtime/SignalProcessing/Private/InterpolatedOnePole.cpp b/Engine/Source/Runtime/SignalProcessing/Private/InterpolatedOnePole.cpp
--- a/Engine/Source/Runtime/SignalProcessing/Private/InterpolatedOnePole.cpp
+++ b/Engine/Source/Runtime/SignalProcessing/Private/InterpolatedOnePole.cpp
@@ -3,7 +3,7 @@
 #include "DSP/InterpolatedOnePole.h"
 
 #include "AudioDefines.h"
-
+#include <complex>
 
 namespace Audio
 {
@@ -288,6 +288,26 @@ namespace Audio
 		Z1.AddZeroed(NumChannels);
 	}
 
+	void FInterpolatedLPF::ArrayCalculateResponseInPlace(TArrayView<float> InOutComplexValues) const
+	{
+		const float B1 = B1Curr;
+		const float A0 = (1.0f - B1Curr);
+
+		const int32 NumFloats = InOutComplexValues.Num();
+		check(NumFloats % 2 == 0);
+		for (int32 Index = 0; Index < NumFloats; Index += 2)
+		{
+			const float ZReal = InOutComplexValues[Index + 0];
+			const float ZImag = InOutComplexValues[Index + 1];
+			const std::complex<float> Z(ZReal, ZImag);
+
+			const std::complex<float> LPF = A0 / (1.0f - B1 * (1.0f / Z));
+
+			InOutComplexValues[Index + 0] = LPF.real();
+			InOutComplexValues[Index + 1] = LPF.imag();
+		}
+	}
+
 
 	// INTERPOLATED ONE-POLE HIGH-PASS IMPLEMENTATION
 	FInterpolatedHPF::FInterpolatedHPF()
@@ -476,5 +496,26 @@ namespace Audio
 		Z1.AddZeroed(NumChannels);
 	}
 
+	void FInterpolatedHPF::ArrayCalculateResponseInPlace(TArrayView<float> InOutComplexValues) const
+	{
+		const float A0 = A0Curr;
+		const float A1 = A0Curr;
+		const float B1 = 2.0f * A0 - 1.0f;
+
+		const int32 NumFloats = InOutComplexValues.Num();
+		check(NumFloats % 2 == 0);
+		for (int32 Index = 0; Index < NumFloats; Index += 2)
+		{
+			const float ZReal = InOutComplexValues[Index + 0];
+			const float ZImag = InOutComplexValues[Index + 1];
+			const std::complex<float> Z(ZReal, ZImag);
+
+			const std::complex<float> LPF = (A0 * Z + A1) / (B1 + Z);
+			const std::complex<float> HPF = 1.0f - LPF;
+
+			InOutComplexValues[Index + 0] = HPF.real();
+			InOutComplexValues[Index + 1] = HPF.imag();
+		}
+	}
 
 } // namespace Audio
\ No newline at end of file
diff --git a/Engine/Source/Runtime/SignalProcessing/Public/DSP/Filter.h b/Engine/Source/Runtime/SignalProcessing/Public/DSP/Filter.h
--- a/Engine/Source/Runtime/SignalProcessing/Public/DSP/Filter.h
+++ b/Engine/Source/Runtime/SignalProcessing/Public/DSP/Filter.h
@@ -73,6 +73,9 @@ namespace Audio
 		// Sets whether or no this filter is enabled (if disabled audio is passed through)
 		SIGNALPROCESSING_API void SetEnabled(const bool bInEnabled);
 
+		// Apply the filter transfer function to each z-domain value in the given array (complex numbers given as interleaved floats). Passing in z-domain values on the complex unit circle will give the frequency response.
+		SIGNALPROCESSING_API void ArrayCalculateResponseInPlace(TArrayView<float> InOutComplexValues) const;
+
 	protected:
 		struct FBiquadCoeff;
 
@@ -255,6 +258,9 @@ namespace Audio
 
 		float GetState(const int32 InChannel) const { return Z1[InChannel]; }
 
+		// Apply the filter transfer function to each z-domain value in the given array (complex numbers given as interleaved floats). Passing in z-domain values on the complex unit circle will give the frequency response.
+		SIGNALPROCESSING_API void ArrayCalculateResponseInPlace(TArrayView<float> InOutComplexValues) const;
+
 	protected:
 		float A0;
 		float* Z1;
@@ -282,6 +288,10 @@ namespace Audio
 		SIGNALPROCESSING_API void ProcessAudio(const float* const* InBuffers, const int32 InNumSamples,
 			float* const* LpfOutBuffers, float* const* HpfOutBuffers, float* const* BpfOutBuffers, float* const* BsfOutBuffers);
 
+		// Apply the filter transfer function to each z-domain value in the given array (complex numbers given as interleaved floats). Passing in z-domain values on the complex unit circle will give the frequency response.
+		// The transfer function utilized here represents the behavior of the filter if the non-linearity were removed.
+		SIGNALPROCESSING_API void ArrayCalculateResponseInPlace(TArrayView<float> InOutComplexValues) const;
+
 	protected:
 		float InputScale;
 		float A0;
@@ -320,6 +330,9 @@ namespace Audio
 		// Process a non-interleaved buffer of multichannel audio
 		SIGNALPROCESSING_API virtual void ProcessAudio(const float* const* InBuffers, const int32 InNumSamples, float* const* OutBuffers) override;
 
+		// Apply the filter transfer function to each z-domain value in the given array (complex numbers given as interleaved floats). Passing in z-domain values on the complex unit circle will give the frequency response.
+		// The transfer function utilized here represents the behavior of the filter if the non-linearity were removed.
+		SIGNALPROCESSING_API void ArrayCalculateResponseInPlace(TArrayView<float> InOutComplexValues) const;
 
 	protected:
 
diff --git a/Engine/Source/Runtime/SignalProcessing/Public/DSP/InterpolatedOnePole.h b/Engine/Source/Runtime/SignalProcessing/Public/DSP/InterpolatedOnePole.h
--- a/Engine/Source/Runtime/SignalProcessing/Public/DSP/InterpolatedOnePole.h
+++ b/Engine/Source/Runtime/SignalProcessing/Public/DSP/InterpolatedOnePole.h
@@ -58,6 +58,9 @@ namespace Audio
 		// This function is useful when there is a break between ProcessAudio calls.
 		SIGNALPROCESSING_API void ClearMemory();
 
+		// Apply the filter transfer function to each z-domain value in the given array (complex numbers given as interleaved floats). Passing in z-domain values on the complex unit circle will give the frequency response.
+		SIGNALPROCESSING_API void ArrayCalculateResponseInPlace(TArrayView<float> InOutComplexValues) const;
+
 	private:
 		float CutoffFrequency{ -1.0f };
 		float B1Curr{ 0.0f }; // coefficient
@@ -122,6 +125,9 @@ namespace Audio
 		// Clears memory without reevaluating coefficients.
 		// This function is useful when there is a break between ProcessAudio calls.
 		SIGNALPROCESSING_API void ClearMemory();
+
+		// Apply the filter transfer function to each z-domain value in the given array (complex numbers given as interleaved floats). Passing in z-domain values on the complex unit circle will give the frequency response.
+		SIGNALPROCESSING_API void ArrayCalculateResponseInPlace(TArrayView<float> InOutComplexValues) const;
 		
 	private:
 		float CutoffFrequency{ -1.0f };
