[Metasound] Just use late reflections for plate reverb node. Back out changes to plate reverb and early reflections, no longer necessary.

#rnx
#tests Metasound editor

#virtualized

diff --git a/Engine/Plugins/Runtime/Harmonix/Source/HarmonixMetasound/Public/HarmonixMetasound/Common.h b/Engine/Plugins/Runtime/Harmonix/Source/HarmonixMetasound/Public/HarmonixMetasound/Common.h
--- a/Engine/Plugins/Runtime/Harmonix/Source/HarmonixMetasound/Public/HarmonixMetasound/Common.h
+++ b/Engine/Plugins/Runtime/Harmonix/Source/HarmonixMetasound/Public/HarmonixMetasound/Common.h
@@ -8,10 +8,6 @@
 #include "UObject/NameTypes.h"
 #include "UObject/ObjectMacros.h"
 
-// TODO: move these into MetasoundParamHelper.h
-#define METASOUND_GET_PARAM_METADATA_ADVANCED(NAME) FDataVertexMetadata { NAME##Tooltip, NAME##DisplayName, true }
-#define METASOUND_GET_PARAM_NAME_AND_METADATA_ADVANCED(NAME) METASOUND_GET_PARAM_NAME(NAME), METASOUND_GET_PARAM_METADATA_ADVANCED(NAME)
-
 namespace HarmonixMetasound
 {
 	HARMONIXMETASOUND_API extern const FName HarmonixNodeNamespace;
diff --git a/Engine/Plugins/Runtime/Metasound/Source/MetasoundGraphCore/Public/MetasoundParamHelper.h b/Engine/Plugins/Runtime/Metasound/Source/MetasoundGraphCore/Public/MetasoundParamHelper.h
--- a/Engine/Plugins/Runtime/Metasound/Source/MetasoundGraphCore/Public/MetasoundParamHelper.h
+++ b/Engine/Plugins/Runtime/Metasound/Source/MetasoundGraphCore/Public/MetasoundParamHelper.h
@@ -67,6 +67,9 @@
 #define METASOUND_GET_PARAM_DISPLAYNAME(NAME) NAME##DisplayName
 #define METASOUND_GET_PARAM_NAME_AND_METADATA(NAME) METASOUND_GET_PARAM_NAME(NAME), METASOUND_GET_PARAM_METADATA(NAME)
 
+#define METASOUND_GET_PARAM_METADATA_ADVANCED(NAME) FDataVertexMetadata { NAME##Tooltip, NAME##DisplayName, true }
+#define METASOUND_GET_PARAM_NAME_AND_METADATA_ADVANCED(NAME) METASOUND_GET_PARAM_NAME(NAME), METASOUND_GET_PARAM_METADATA_ADVANCED(NAME)
+
 #define METASOUND_GET_PARAM_NAME_WITH_INDEX(NAME, INDEX) *FString::Format(*static_cast<FName>(NAME##Name).ToString(), {INDEX})
 #if WITH_EDITOR
 #define METASOUND_GET_PARAM_NAME_WITH_INDEX_AND_TT(NAME, INDEX)  METASOUND_GET_PARAM_NAME_WITH_INDEX(NAME, INDEX), FText::Format(NAME##Tooltip, INDEX)
diff --git a/Engine/Plugins/Runtime/Metasound/Source/MetasoundStandardNodes/Private/MetasoundPlateReverbNode.cpp b/Engine/Plugins/Runtime/Metasound/Source/MetasoundStandardNodes/Private/MetasoundPlateReverbNode.cpp
--- a/Engine/Plugins/Runtime/Metasound/Source/MetasoundStandardNodes/Private/MetasoundPlateReverbNode.cpp
+++ b/Engine/Plugins/Runtime/Metasound/Source/MetasoundStandardNodes/Private/MetasoundPlateReverbNode.cpp
@@ -8,7 +8,7 @@
 #include "MetasoundStandardNodesCategories.h"
 #include "MetasoundStandardNodesNames.h"
 #include "DSP/FloatArrayMath.h"
-#include "DSP/ReverbFast.h"
+#include "DSP/LateReflectionsFast.h"
 #include "Internationalization/Text.h"
 
 #define LOCTEXT_NAMESPACE "MetasoundStandardNodes_PlateReverb"
@@ -20,16 +20,23 @@ namespace Metasound
 		namespace Inputs
 		{
 			METASOUND_PARAM(Bypass, "Bypass", "Toggle to bypass the effect and send audio through unaltered.")
-			METASOUND_PARAM(AudioLeft, "In Left", "Left channel audio input.")
-			METASOUND_PARAM(AudioRight, "In Right", "Right channel audio input.")
+			METASOUND_PARAM(InAudioLeft, "In Left", "Left channel audio input.")
+			METASOUND_PARAM(InAudioRight, "In Right", "Right channel audio input.")
 			METASOUND_PARAM(DryLevel, "Dry Level", "The level of the dry signal (linear).")
 			METASOUND_PARAM(WetLevel, "Wet Level", "The level of the wet signal (linear).")
+			METASOUND_PARAM(LateReflectionsDelay, "Delay (ms)", "Pre-delay before late reflections")
+			METASOUND_PARAM(LateReflectionsGainDb, "Gain (dB)", "Initial attenuation of audio after it leaves the predelay")
+			METASOUND_PARAM(LateReflectionsBandwidth, "Bandwidth", "Frequency bandwidth of audio going into input diffusers. 0.999 is full bandwidth")
+			METASOUND_PARAM(LateReflectionsDiffusion, "Diffusion", "Amount of input diffusion (larger value results in more diffusion)")
+			METASOUND_PARAM(LateReflectionsDampening, "Dampening", "The amount of high-frequency dampening in plate feedback paths")
+			METASOUND_PARAM(LateReflectionsDecay, "Decay", "The amount of decay in the feedback path. Lower value is larger reverb time.")
+			METASOUND_PARAM(LateReflectionsDensity, "Density", "The amount of diffusion in decay path. Larger values is a more dense reverb.")
 		}
 
 		namespace Outputs
 		{
-			METASOUND_PARAM(AudioLeft, "Out Left", "Left channel audio output.")
-			METASOUND_PARAM(AudioRight, "Out Right", "Right channel audio output.")
+			METASOUND_PARAM(OutAudioLeft, "Out Left", "Left channel audio output.")
+			METASOUND_PARAM(OutAudioRight, "Out Right", "Right channel audio output.")
 		}
 	}
 
@@ -43,6 +50,13 @@ namespace Metasound
 			FAudioBufferReadRef AudioRight;
 			FFloatReadRef DryLevel;
 			FFloatReadRef WetLevel;
+			FFloatReadRef LateReflectionsDelay;
+			FFloatReadRef LateReflectionsGainDb;
+			FFloatReadRef LateReflectionsBandwidth;
+			FFloatReadRef LateReflectionsDiffusion;
+			FFloatReadRef LateReflectionsDampening;
+			FFloatReadRef LateReflectionsDecay;
+			FFloatReadRef LateReflectionsDensity;
 		};
 		
 		struct FOutputs
@@ -79,20 +93,29 @@ namespace Metasound
 		
 		static const FVertexInterface& GetVertexInterface()
 		{
+			static const Audio::FLateReflectionsFastSettings Defaults;
+			
 			static const FVertexInterface Interface
 			{
 				FInputVertexInterface
 				{
 					TInputDataVertex<bool>(METASOUND_GET_PARAM_NAME_AND_METADATA(PlateReverb::Inputs::Bypass), false),
-					TInputDataVertex<FAudioBuffer>(METASOUND_GET_PARAM_NAME_AND_METADATA(PlateReverb::Inputs::AudioLeft)),
-					TInputDataVertex<FAudioBuffer>(METASOUND_GET_PARAM_NAME_AND_METADATA(PlateReverb::Inputs::AudioRight)),
+					TInputDataVertex<FAudioBuffer>(METASOUND_GET_PARAM_NAME_AND_METADATA(PlateReverb::Inputs::InAudioLeft)),
+					TInputDataVertex<FAudioBuffer>(METASOUND_GET_PARAM_NAME_AND_METADATA(PlateReverb::Inputs::InAudioRight)),
 					TInputDataVertex<float>(METASOUND_GET_PARAM_NAME_AND_METADATA(PlateReverb::Inputs::DryLevel), 1.0f),
-					TInputDataVertex<float>(METASOUND_GET_PARAM_NAME_AND_METADATA(PlateReverb::Inputs::WetLevel), 1.0f)
+					TInputDataVertex<float>(METASOUND_GET_PARAM_NAME_AND_METADATA(PlateReverb::Inputs::WetLevel), 1.0f),
+					TInputDataVertex<float>(METASOUND_GET_PARAM_NAME_AND_METADATA_ADVANCED(PlateReverb::Inputs::LateReflectionsDelay), Defaults.LateDelayMsec),
+					TInputDataVertex<float>(METASOUND_GET_PARAM_NAME_AND_METADATA_ADVANCED(PlateReverb::Inputs::LateReflectionsGainDb), Defaults.LateGainDB),
+					TInputDataVertex<float>(METASOUND_GET_PARAM_NAME_AND_METADATA_ADVANCED(PlateReverb::Inputs::LateReflectionsBandwidth), Defaults.Bandwidth),
+					TInputDataVertex<float>(METASOUND_GET_PARAM_NAME_AND_METADATA_ADVANCED(PlateReverb::Inputs::LateReflectionsDiffusion), Defaults.Diffusion),
+					TInputDataVertex<float>(METASOUND_GET_PARAM_NAME_AND_METADATA_ADVANCED(PlateReverb::Inputs::LateReflectionsDampening), Defaults.Dampening),
+					TInputDataVertex<float>(METASOUND_GET_PARAM_NAME_AND_METADATA_ADVANCED(PlateReverb::Inputs::LateReflectionsDecay), Defaults.Decay),
+					TInputDataVertex<float>(METASOUND_GET_PARAM_NAME_AND_METADATA_ADVANCED(PlateReverb::Inputs::LateReflectionsDensity), Defaults.Density)
 				},
 				FOutputVertexInterface
 				{
-					TOutputDataVertex<FAudioBuffer>(METASOUND_GET_PARAM_NAME_AND_METADATA(PlateReverb::Outputs::AudioLeft)),
-					TOutputDataVertex<FAudioBuffer>(METASOUND_GET_PARAM_NAME_AND_METADATA(PlateReverb::Outputs::AudioRight))
+					TOutputDataVertex<FAudioBuffer>(METASOUND_GET_PARAM_NAME_AND_METADATA(PlateReverb::Outputs::OutAudioLeft)),
+					TOutputDataVertex<FAudioBuffer>(METASOUND_GET_PARAM_NAME_AND_METADATA(PlateReverb::Outputs::OutAudioRight))
 				}
 			};
 
@@ -104,19 +127,28 @@ namespace Metasound
 			FInputs Inputs
 			{
 				InParams.InputData.GetOrCreateDefaultDataReadReference<bool>(PlateReverb::Inputs::BypassName, InParams.OperatorSettings),
-				InParams.InputData.GetOrConstructDataReadReference<FAudioBuffer>(PlateReverb::Inputs::AudioLeftName, InParams.OperatorSettings),
-				InParams.InputData.GetOrConstructDataReadReference<FAudioBuffer>(PlateReverb::Inputs::AudioRightName, InParams.OperatorSettings),
+				InParams.InputData.GetOrConstructDataReadReference<FAudioBuffer>(PlateReverb::Inputs::InAudioLeftName, InParams.OperatorSettings),
+				InParams.InputData.GetOrConstructDataReadReference<FAudioBuffer>(PlateReverb::Inputs::InAudioRightName, InParams.OperatorSettings),
 				InParams.InputData.GetOrCreateDefaultDataReadReference<float>(PlateReverb::Inputs::DryLevelName, InParams.OperatorSettings),
-				InParams.InputData.GetOrCreateDefaultDataReadReference<float>(PlateReverb::Inputs::WetLevelName, InParams.OperatorSettings)
+				InParams.InputData.GetOrCreateDefaultDataReadReference<float>(PlateReverb::Inputs::WetLevelName, InParams.OperatorSettings),
+				InParams.InputData.GetOrCreateDefaultDataReadReference<float>(PlateReverb::Inputs::LateReflectionsDelayName, InParams.OperatorSettings),
+				InParams.InputData.GetOrCreateDefaultDataReadReference<float>(PlateReverb::Inputs::LateReflectionsGainDbName, InParams.OperatorSettings),
+				InParams.InputData.GetOrCreateDefaultDataReadReference<float>(PlateReverb::Inputs::LateReflectionsBandwidthName, InParams.OperatorSettings),
+				InParams.InputData.GetOrCreateDefaultDataReadReference<float>(PlateReverb::Inputs::LateReflectionsDiffusionName, InParams.OperatorSettings),
+				InParams.InputData.GetOrCreateDefaultDataReadReference<float>(PlateReverb::Inputs::LateReflectionsDampeningName, InParams.OperatorSettings),
+				InParams.InputData.GetOrCreateDefaultDataReadReference<float>(PlateReverb::Inputs::LateReflectionsDecayName, InParams.OperatorSettings),
+				InParams.InputData.GetOrCreateDefaultDataReadReference<float>(PlateReverb::Inputs::LateReflectionsDensityName, InParams.OperatorSettings)
 			};
 
 			return MakeUnique<FPlateReverbOperator>(InParams, MoveTemp(Inputs));
 		}
 
+		static constexpr int32 MaxReverbBufferSize = 512;
+
 		FPlateReverbOperator(const FBuildOperatorParams& BuildParams, FInputs&& Inputs)
 			: Inputs(MoveTemp(Inputs))
 			, Outputs({ FAudioBufferWriteRef::CreateNew(BuildParams.OperatorSettings), FAudioBufferWriteRef::CreateNew(BuildParams.OperatorSettings) })
-			, Reverb(BuildParams.OperatorSettings.GetSampleRate())
+			, Reverb(BuildParams.OperatorSettings.GetSampleRate(), MaxReverbBufferSize)
 		{
 			Reset(BuildParams);
 		}
@@ -124,58 +156,78 @@ namespace Metasound
 		virtual void BindInputs(FInputVertexInterfaceData& InOutVertexData) override
 		{
 			InOutVertexData.BindReadVertex(PlateReverb::Inputs::BypassName, Inputs.Bypass);
-			InOutVertexData.BindReadVertex(PlateReverb::Inputs::AudioLeftName, Inputs.AudioLeft);
-			InOutVertexData.BindReadVertex(PlateReverb::Inputs::AudioRightName, Inputs.AudioRight);
+			InOutVertexData.BindReadVertex(PlateReverb::Inputs::InAudioLeftName, Inputs.AudioLeft);
+			InOutVertexData.BindReadVertex(PlateReverb::Inputs::InAudioRightName, Inputs.AudioRight);
 			InOutVertexData.BindReadVertex(PlateReverb::Inputs::DryLevelName, Inputs.DryLevel);
 			InOutVertexData.BindReadVertex(PlateReverb::Inputs::WetLevelName, Inputs.WetLevel);
+			InOutVertexData.BindReadVertex(PlateReverb::Inputs::LateReflectionsDelayName, Inputs.LateReflectionsDelay);
+			InOutVertexData.BindReadVertex(PlateReverb::Inputs::LateReflectionsGainDbName, Inputs.LateReflectionsGainDb);
+			InOutVertexData.BindReadVertex(PlateReverb::Inputs::LateReflectionsBandwidthName, Inputs.LateReflectionsBandwidth);
+			InOutVertexData.BindReadVertex(PlateReverb::Inputs::LateReflectionsDiffusionName, Inputs.LateReflectionsDiffusion);
+			InOutVertexData.BindReadVertex(PlateReverb::Inputs::LateReflectionsDampeningName, Inputs.LateReflectionsDampening);
+			InOutVertexData.BindReadVertex(PlateReverb::Inputs::LateReflectionsDecayName, Inputs.LateReflectionsDecay);
+			InOutVertexData.BindReadVertex(PlateReverb::Inputs::LateReflectionsDensityName, Inputs.LateReflectionsDensity);
 		}
 		
 		virtual void BindOutputs(FOutputVertexInterfaceData& InOutVertexData) override
 		{
-			InOutVertexData.BindReadVertex(PlateReverb::Outputs::AudioLeftName, Outputs.AudioLeft);
-			InOutVertexData.BindReadVertex(PlateReverb::Outputs::AudioRightName, Outputs.AudioRight);
+			InOutVertexData.BindReadVertex(PlateReverb::Outputs::OutAudioLeftName, Outputs.AudioLeft);
+			InOutVertexData.BindReadVertex(PlateReverb::Outputs::OutAudioRightName, Outputs.AudioRight);
 		}
 
-		void Reset(const FResetParams&)
+		void Reset(const FResetParams& Params)
 		{
+			WorkBuffer.SetNumUninitialized(Params.OperatorSettings.GetNumFramesPerBlock());
+			Reverb.FlushAudio();
 		}
 		
 		void Execute()
 		{
-			const int32 NumFrames = Inputs.AudioLeft->Num();
-			check(Inputs.AudioRight->Num() == NumFrames);
-			
+			if (*Inputs.Bypass != WasBypassed)
+			{
+				// Flush the reverb if we just bypassed
+				if (*Inputs.Bypass)
+				{
+					Reverb.FlushAudio();
+				}
+
+				WasBypassed = *Inputs.Bypass;
+			}
+
+			// Pass through audio if bypassed
 			if (*Inputs.Bypass)
 			{
-				FMemory::Memcpy(Outputs.AudioLeft->GetData(), Inputs.AudioLeft->GetData(), NumFrames * sizeof(float));
-				FMemory::Memcpy(Outputs.AudioRight->GetData(), Inputs.AudioRight->GetData(), NumFrames * sizeof(float));
+				FMemory::Memcpy(Outputs.AudioLeft->GetData(), Inputs.AudioLeft->GetData(), Inputs.AudioLeft->Num() * sizeof(float));
+				FMemory::Memcpy(Outputs.AudioRight->GetData(), Inputs.AudioRight->GetData(), Inputs.AudioRight->Num() * sizeof(float));
 				return;
 			}
 
-			// Copy the inputs to the work buffers
-			WorkBufferLeft.SetNumUninitialized(NumFrames);
-			WorkBufferRight.SetNumUninitialized(NumFrames);
-			FMemory::Memcpy(WorkBufferLeft.GetData(), Inputs.AudioLeft->GetData(), NumFrames * sizeof(float));
-			FMemory::Memcpy(WorkBufferRight.GetData(), Inputs.AudioRight->GetData(), NumFrames * sizeof(float));
+			// Sum to mono. This happens in the late reflections code when you pass in interleaved, stereo audio.
+			// Doing it here and scaling below avoids the extra interleave memory and time.
+			Audio::ArraySum(*Inputs.AudioLeft, *Inputs.AudioRight, WorkBuffer);
 			
 			const float CurrentWetLevel = FMath::Clamp(*Inputs.WetLevel, 0.0f, 1.0f);
 
-			// Apply the wet gain to the input to preserve the reverb tail
+			// Apply the wet gain to the input to preserve the reverb tail.
+			// We scale by half because that's what happens when passing in stereo audio to the late reflections.
 			if (LastWetLevel >= 0.0f && !FMath::IsNearlyEqual(CurrentWetLevel, LastWetLevel))
 			{
-				Audio::ArrayFade(WorkBufferLeft, LastWetLevel, CurrentWetLevel);
-				Audio::ArrayFade(WorkBufferRight, LastWetLevel, CurrentWetLevel);
+				const float From = LastWetLevel * 0.5f;
+				const float To = CurrentWetLevel * 0.5f;
+				Audio::ArrayFade(WorkBuffer, From, To);
 			}
 			else
 			{
-				Audio::ArrayMultiplyByConstantInPlace(WorkBufferLeft, CurrentWetLevel);
-				Audio::ArrayMultiplyByConstantInPlace(WorkBufferRight, CurrentWetLevel);
+				Audio::ArrayMultiplyByConstantInPlace(WorkBuffer, CurrentWetLevel * 0.5f);
 			}
 			
 			LastWetLevel = CurrentWetLevel;
 
+			// Update the reverb settings
+			UpdateSettingsIfChanged();
+			
 			// Process
-			Reverb.ProcessAudioStereoNonInterleaved(WorkBufferLeft, WorkBufferRight, *Outputs.AudioLeft, *Outputs.AudioRight);
+			Reverb.ProcessAudio(WorkBuffer, 1, *Outputs.AudioLeft, *Outputs.AudioRight);
 
 			// Mix in the dry signal
 			const float DryLevel = FMath::Clamp(*Inputs.DryLevel, 0.0f, 1.0f);
@@ -184,12 +236,60 @@ namespace Metasound
 		}
 
 	private:
+		void UpdateSettingsIfChanged()
+		{
+			bool SettingsChanged = false;
+
+			if (*Inputs.LateReflectionsDelay != CurrentSettings.LateDelayMsec)
+			{
+				CurrentSettings.LateDelayMsec = *Inputs.LateReflectionsDelay;
+				SettingsChanged = true;
+			}
+			if (*Inputs.LateReflectionsGainDb != CurrentSettings.LateGainDB)
+			{
+				CurrentSettings.LateGainDB = *Inputs.LateReflectionsGainDb;
+				SettingsChanged = true;
+			}
+			if (*Inputs.LateReflectionsBandwidth != CurrentSettings.Bandwidth)
+			{
+				CurrentSettings.Bandwidth = *Inputs.LateReflectionsBandwidth;
+				SettingsChanged = true;
+			}
+			if (*Inputs.LateReflectionsDiffusion != CurrentSettings.Diffusion)
+			{
+				CurrentSettings.Diffusion = *Inputs.LateReflectionsDiffusion;
+				SettingsChanged = true;
+			}
+			if (*Inputs.LateReflectionsDampening != CurrentSettings.Dampening)
+			{
+				CurrentSettings.Dampening = *Inputs.LateReflectionsDampening;
+				SettingsChanged = true;
+			}
+			if (*Inputs.LateReflectionsDecay != CurrentSettings.Decay)
+			{
+				CurrentSettings.Decay = *Inputs.LateReflectionsDecay;
+				SettingsChanged = true;
+			}
+			if (*Inputs.LateReflectionsDensity != CurrentSettings.Density)
+			{
+				CurrentSettings.Density = *Inputs.LateReflectionsDensity;
+				SettingsChanged = true;
+			}
+
+			if (SettingsChanged)
+			{
+				Reverb.SetSettings(CurrentSettings);
+			}
+		}
+		
 		FInputs Inputs;
 		FOutputs Outputs;
 		
-		Audio::FPlateReverbFast Reverb;
-		Audio::FAlignedFloatBuffer WorkBufferLeft;
-		Audio::FAlignedFloatBuffer WorkBufferRight;
+		Audio::FLateReflectionsFast Reverb;
+		Audio::FLateReflectionsFastSettings CurrentSettings;
+		Audio::FAlignedFloatBuffer WorkBuffer;
+
+		bool WasBypassed;
 		float LastWetLevel{ -1 };
 	};
 
diff --git a/Engine/Source/Runtime/SignalProcessing/Private/EarlyReflectionsFast.cpp b/Engine/Source/Runtime/SignalProcessing/Private/EarlyReflectionsFast.cpp
--- a/Engine/Source/Runtime/SignalProcessing/Private/EarlyReflectionsFast.cpp
+++ b/Engine/Source/Runtime/SignalProcessing/Private/EarlyReflectionsFast.cpp
@@ -221,40 +221,6 @@ namespace Audio
 		ArrayMultiplyByConstantInPlace(OutRightSamples, Settings.Gain);
 	}
 
-	void FEarlyReflectionsFast::ProcessAudioStereoNonInterleaved(
-		const TArrayView<const float>& InLeft,
-		const TArrayView<const float>& InRight,
-		FAlignedFloatBuffer& OutLeft,
-		FAlignedFloatBuffer& OutRight)
-	{
-		const int32 InNumFrames = InLeft.Num();
-		check(InRight.Num() == InNumFrames);
-		check(OutLeft.Num() == InNumFrames);
-		check(OutRight.Num() == InNumFrames);
-
-		// Resize internal buffers
-		LeftWorkBufferA.SetNumUninitialized(InNumFrames);
-		LeftWorkBufferB.SetNumUninitialized(InNumFrames);
-		RightWorkBufferA.SetNumUninitialized(InNumFrames);
-		RightWorkBufferB.SetNumUninitialized(InNumFrames);
-
-		// predelay
-		LeftPreDelay.ProcessAudio(InLeft, LeftWorkBufferB);
-		RightPreDelay.ProcessAudio(InRight, RightWorkBufferB);
-
-		// lpf
-		LeftInputLPF.ProcessAudio(LeftWorkBufferB, LeftWorkBufferA);
-		RightInputLPF.ProcessAudio(RightWorkBufferB, RightWorkBufferA);
-
-		// feedback delay network
-		LeftFDN.ProcessAudio(LeftWorkBufferA, OutLeft);
-		RightFDN.ProcessAudio(RightWorkBufferA, OutRight);
-
-		// Apply Gain
-		ArrayMultiplyByConstantInPlace(OutLeft, Settings.Gain);
-		ArrayMultiplyByConstantInPlace(OutRight, Settings.Gain);
-	}
-
 	void FEarlyReflectionsFast::FlushAudio()
 	{
 		// predelay
diff --git a/Engine/Source/Runtime/SignalProcessing/Private/ReverbFast.cpp b/Engine/Source/Runtime/SignalProcessing/Private/ReverbFast.cpp
--- a/Engine/Source/Runtime/SignalProcessing/Private/ReverbFast.cpp
+++ b/Engine/Source/Runtime/SignalProcessing/Private/ReverbFast.cpp
@@ -148,68 +148,6 @@ namespace Audio {
 		InterleaveAndMixOutput(FrontLeftReverbSamples, FrontRightReverbSamples, OutSamples, OutNumChannels);
 	}
 
-	void FPlateReverbFast::ProcessAudioStereoNonInterleaved(
-		const TArrayView<const float>& InLeft,
-		const TArrayView<const float>& InRight,
-		FAlignedFloatBuffer& OutLeft,
-		FAlignedFloatBuffer& OutRight)
-	{
-		const int32 NumFrames = InLeft.Num();
-		check(NumFrames == InRight.Num());
-		check(NumFrames == OutLeft.Num());
-		check(NumFrames == OutRight.Num());
-
-		if (NumFrames == 0)
-		{
-			return;
-		}
-
-		// Reverb is disabled, so just zero the output buffer
-		if (!Settings.bEnableEarlyReflections && !Settings.bEnableLateReflections)
-		{
-			FMemory::Memzero(OutLeft.GetData(), NumFrames * sizeof(float));
-			FMemory::Memzero(OutRight.GetData(), NumFrames * sizeof(float));
-			
-			return;
-		}
-
-		// Both early and late reflections
-		if (Settings.bEnableEarlyReflections && Settings.bEnableLateReflections)
-		{
-			// Process the early reflections directly to the outputs
-			EarlyReflections.ProcessAudioStereoNonInterleaved(InLeft, InRight, OutLeft, OutRight);
-
-			// Process the late reflections and hang onto the result in the work buffers
-			// Sum to mono and scale by half because the late reflection algorithm does that internally. This saves some copies.
-			FrontLeftReverbSamples.SetNumUninitialized(NumFrames);
-			ArraySum(InLeft, InRight, FrontLeftReverbSamples);
-			ArrayMultiplyByConstantInPlace(FrontLeftReverbSamples, 0.5f);
-			LateReflections.ProcessAudio(FrontLeftReverbSamples, 1, FrontLeftLateReflectionsSamples, FrontRightLateReflectionsSamples);
-
-			// Mix in the early reflections
-			ArrayAddInPlace(FrontLeftLateReflectionsSamples, OutLeft);
-			ArrayAddInPlace(FrontRightLateReflectionsSamples, OutRight);
-		}
-		// Only early reflections
-		else if (Settings.bEnableEarlyReflections)
-		{
-			EarlyReflections.ProcessAudioStereoNonInterleaved(InLeft, InRight, OutLeft, OutRight);
-		}
-		// Only late reflections
-		else if (Settings.bEnableLateReflections)
-		{
-			// Sum to mono and scale by half because the late reflection algorithm does that internally. This saves some copies.
-			FrontLeftReverbSamples.SetNumUninitialized(NumFrames);
-			ArraySum(InLeft, InRight, FrontLeftReverbSamples);
-			ArrayMultiplyByConstantInPlace(FrontLeftReverbSamples, 0.5f);
-			LateReflections.ProcessAudio(FrontLeftReverbSamples, 1, OutLeft, OutRight);
-		}
-		else
-		{
-			checkNoEntry();
-		}
-	}
-
 	void FPlateReverbFast::ClampSettings(FPlateReverbFastSettings& InOutSettings)
 	{
 		// Clamp settings for this object and member objects.
diff --git a/Engine/Source/Runtime/SignalProcessing/Public/DSP/EarlyReflectionsFast.h b/Engine/Source/Runtime/SignalProcessing/Public/DSP/EarlyReflectionsFast.h
--- a/Engine/Source/Runtime/SignalProcessing/Public/DSP/EarlyReflectionsFast.h
+++ b/Engine/Source/Runtime/SignalProcessing/Public/DSP/EarlyReflectionsFast.h
@@ -66,13 +66,6 @@ namespace Audio
 		// Process the single audio frame
 		SIGNALPROCESSING_API void ProcessAudio(const FAlignedFloatBuffer& InSamples, const int32 InNumChannels, FAlignedFloatBuffer& OutLeftSamples, FAlignedFloatBuffer& OutRightSamples);
 
-		// Process a block of audio. Optimized for stereo, non-interleaved input.
-		SIGNALPROCESSING_API void ProcessAudioStereoNonInterleaved(
-			const TArrayView<const float>& InLeft,
-			const TArrayView<const float>& InRight,
-			FAlignedFloatBuffer& OutLeft,
-			FAlignedFloatBuffer& OutRight);
-		
 		// Silence internal audio.
 		SIGNALPROCESSING_API void FlushAudio();
 
diff --git a/Engine/Source/Runtime/SignalProcessing/Public/DSP/ReverbFast.h b/Engine/Source/Runtime/SignalProcessing/Public/DSP/ReverbFast.h
--- a/Engine/Source/Runtime/SignalProcessing/Public/DSP/ReverbFast.h
+++ b/Engine/Source/Runtime/SignalProcessing/Public/DSP/ReverbFast.h
@@ -72,14 +72,7 @@ namespace Audio
 			// InNumChannels can be 1 or 2 channels.
 			// OutSamples must be greater or equal to 2.
 			SIGNALPROCESSING_API void ProcessAudio(const FAlignedFloatBuffer& InSamples, const int32 InNumChannels, FAlignedFloatBuffer& OutSamples, const int32 OutNumChannels);
-		
-			// Process a block of audio. Optimized for stereo, non-interleaved input and output.
-			SIGNALPROCESSING_API void ProcessAudioStereoNonInterleaved(
-				const TArrayView<const float>& InLeft,
-				const TArrayView<const float>& InRight,
-				FAlignedFloatBuffer& OutLeft,
-				FAlignedFloatBuffer& OutRight);
-		
+
 			SIGNALPROCESSING_API void FlushAudio();
 
 			// Clamp individual settings to values supported by this class.
