- Checkpoint for Initial MetaSound Pages document model support
- Fix-ups for MS AssetManager leaking asset names if not registered by moving to MultiMap
- Fix-up for ResetBuilder call passing off initial delegates properly
- Make builders map mutation thread safe when loading assets async
#rnx
#tests -game, AudioQA, Unit Tests, Automated Tests

diff --git a/Engine/Plugins/Runtime/Metasound/Source/MetasoundEditor/Private/MetasoundEditor.cpp b/Engine/Plugins/Runtime/Metasound/Source/MetasoundEditor/Private/MetasoundEditor.cpp
--- a/Engine/Plugins/Runtime/Metasound/Source/MetasoundEditor/Private/MetasoundEditor.cpp
+++ b/Engine/Plugins/Runtime/Metasound/Source/MetasoundEditor/Private/MetasoundEditor.cpp
@@ -892,6 +892,11 @@ namespace Metasound
 				{
 					MetaSoundAsset->SetVersionedOnLoad();
 				}
+
+				// Hack to ensure validation is re-run on re-opening of the editor.
+				// This is needed to refresh errors potentially caused by unloading of
+				// references (ex. if a referenced asset is force deleted in the editor).
+				MetaSoundAsset->GetModifyContext().SetForceRefreshViews();
 			}
 
 			Builder.Reset(&Engine::FDocumentBuilderRegistry::GetChecked().FindOrBeginBuilding(*ObjectToEdit));
@@ -910,6 +915,7 @@ namespace Metasound
 			{
 				FGraphBuilder::RegisterGraphWithFrontend(*Metasound);
 			}
+
 			RefreshEditorContext();
 
 			BindGraphCommands();
@@ -1960,7 +1966,7 @@ namespace Metasound
 			UEdGraph* Graph = MetasoundAsset->GetGraph();
 			check(Graph);
 
-			return *CastChecked<UMetasoundEditorGraph>(MetasoundAsset->GetGraph());
+			return *CastChecked<UMetasoundEditorGraph>(Graph);
 		}
 
 		void FEditor::ExecuteNode()
@@ -3384,7 +3390,13 @@ namespace Metasound
 		void FEditor::RefreshEditorContext()
 		{
 			TRACE_CPUPROFILER_EVENT_SCOPE(Metasound::Editor::FEditor::RefreshEditorContext);
-			const bool bSynchronizedGraph = FGraphBuilder::SynchronizeGraph(*Metasound);
+
+			check(Builder.IsValid());
+			UMetasoundEditorGraph* Graph = nullptr;
+			FGraphBuilder::BindEditorGraph(Builder->GetConstBuilder(), &Graph);
+			check(Graph);
+
+			const bool bSynchronizedGraph = FGraphBuilder::SynchronizeGraph(Builder->GetConstBuilder(), *Graph);
 
 			FMetasoundAssetBase* MetasoundAsset = IMetasoundUObjectRegistry::Get().GetObjectAsAssetBase(Metasound);
 			check(MetasoundAsset);
@@ -3428,10 +3440,9 @@ namespace Metasound
 					// If no member was selected by an action (ex. undo/redo), select a modified member 
 					if (!SelectedMember)
 					{
-						UMetasoundEditorGraph& Graph = GetMetaSoundGraphChecked();
 						for (const FGuid& MemberGuid : MembersModified)
 						{
-							if (UObject* Member = Graph.FindMember(MemberGuid))
+							if (UObject* Member = Graph->FindMember(MemberGuid))
 							{
 								// Currently only one member can be selected at a time, so only first found is added
 								Selection.Add(Member);
@@ -3534,12 +3545,10 @@ namespace Metasound
 		{
 			using namespace Metasound::Editor;
 
-			if (Metasound)
+			if (Builder.IsValid())
 			{
-				const FMetasoundAssetBase* MetasoundAsset = IMetasoundUObjectRegistry::Get().GetObjectAsAssetBase(Metasound);
-				check(MetasoundAsset);
-				const FMetasoundFrontendDocument& Document = MetasoundAsset->GetConstDocumentChecked();
-				return Document.RootGraph.Graph.Style.bIsGraphEditable;
+				const FMetasoundFrontendGraph& Graph = Builder->GetConstBuilder().FindConstBuildGraphChecked();
+				return Graph.Style.bIsGraphEditable;
 			}
 
 			return false;
diff --git a/Engine/Plugins/Runtime/Metasound/Source/MetasoundEditor/Private/MetasoundEditorGraph.cpp b/Engine/Plugins/Runtime/Metasound/Source/MetasoundEditor/Private/MetasoundEditorGraph.cpp
--- a/Engine/Plugins/Runtime/Metasound/Source/MetasoundEditor/Private/MetasoundEditorGraph.cpp
+++ b/Engine/Plugins/Runtime/Metasound/Source/MetasoundEditor/Private/MetasoundEditorGraph.cpp
@@ -2007,8 +2007,7 @@ void UMetasoundEditorGraph::MigrateEditorDocumentData(FMetaSoundFrontendDocument
 	// 5. Remove input locations and ensure that all other nodes only have at most one
 	// location represented in the style/editor graph (0 is acceptable as some member
 	// node types (eg. variables) may not contain locations and that's ok).
-	const FMetasoundFrontendDocument& Document = OutBuilder.GetConstDocument();
-	const TArray<FMetasoundFrontendNode>& GraphNodes = Document.RootGraph.Graph.Nodes;
+	const TArray<FMetasoundFrontendNode>& GraphNodes = OutBuilder.FindConstBuildGraphChecked().Nodes;
 	TMap<FGuid, const UMetasoundEditorGraphNode*> EdNodeMap;
 	Algo::Transform(AllMetaSoundNodes, EdNodeMap, [](const UEdGraphNode* Node)
 	{
diff --git a/Engine/Plugins/Runtime/Metasound/Source/MetasoundEditor/Private/MetasoundEditorGraphBuilder.cpp b/Engine/Plugins/Runtime/Metasound/Source/MetasoundEditor/Private/MetasoundEditorGraphBuilder.cpp
--- a/Engine/Plugins/Runtime/Metasound/Source/MetasoundEditor/Private/MetasoundEditorGraphBuilder.cpp
+++ b/Engine/Plugins/Runtime/Metasound/Source/MetasoundEditor/Private/MetasoundEditorGraphBuilder.cpp
@@ -16,6 +16,7 @@
 #include "MetasoundAssetBase.h"
 #include "MetasoundAudioBuffer.h"
 #include "MetasoundDocumentBuilderRegistry.h"
+#include "MetasoundDocumentInterface.h"
 #include "MetasoundEditor.h"
 #include "MetasoundEditorGraph.h"
 #include "MetasoundEditorGraphCommentNode.h"
@@ -120,11 +121,13 @@ namespace Metasound
 				}
 			};
 
-			void SynchronizeGraphRecursively(UObject& InMetaSound)
+			void SynchronizeGraphRecursively(const FMetaSoundFrontendDocumentBuilder& InBuilder, UMetasoundEditorGraph& OutGraph)
 			{
+				using namespace Engine;
 				using namespace Frontend;
 
-				FMetasoundAssetBase* MetaSoundAsset = IMetasoundUObjectRegistry::Get().GetObjectAsAssetBase(&InMetaSound);
+				UObject& MetaSound = InBuilder.CastDocumentObjectChecked<UObject>();
+				FMetasoundAssetBase* MetaSoundAsset = IMetasoundUObjectRegistry::Get().GetObjectAsAssetBase(&MetaSound);
 				check(MetaSoundAsset);
 
 				// Synchronize referenced graphs first to ensure all editor data
@@ -134,7 +137,13 @@ namespace Metasound
 				for (FMetasoundAssetBase* Reference : References)
 				{
 					check(Reference);
-					SynchronizeGraphRecursively(*Reference->GetOwningAsset());
+					UObject* RefObject = Reference->GetOwningAsset();
+					check(RefObject);
+					FMetaSoundFrontendDocumentBuilder& RefBuilder = FDocumentBuilderRegistry::GetChecked().FindOrBeginBuilding(RefObject);
+
+					UMetasoundEditorGraph* Graph = nullptr;
+					FGraphBuilder::BindEditorGraph(RefBuilder, &Graph);
+					SynchronizeGraphRecursively(RefBuilder, *Graph);
 				}
 
 				if (!MetaSoundAsset->GetConstModifyContext().GetDocumentModified())
@@ -142,17 +151,11 @@ namespace Metasound
 					return;
 				}
 
-				// If no graph is set, MetaSound has been created outside of asset factory, so initialize it here.
-				if (!MetaSoundAsset->GetGraph())
-				{
-					UMetaSoundEditorSubsystem::GetChecked().InitEdGraph(InMetaSound);
-				}
-
-				FGraphBuilder::SynchronizeComments(InMetaSound);
-				FGraphBuilder::SynchronizeGraphMembers(InMetaSound);
-				FGraphBuilder::SynchronizeOutputNodes(InMetaSound);
-				FGraphBuilder::SynchronizeNodes(InMetaSound);
-				FGraphBuilder::SynchronizeConnections(InMetaSound);
+				FGraphBuilder::SynchronizeComments(InBuilder, OutGraph);
+				FGraphBuilder::SynchronizeGraphMembers(InBuilder, OutGraph);
+				FGraphBuilder::SynchronizeOutputNodes(InBuilder, OutGraph);
+				FGraphBuilder::SynchronizeNodes(InBuilder, OutGraph);
+				FGraphBuilder::SynchronizeConnections(InBuilder, OutGraph);
 			}
 		} // namespace GraphBuilderPrivate
 
@@ -1173,21 +1176,6 @@ namespace Metasound
 			}
 		}
 
-		bool FGraphBuilder::GraphContainsErrors(const UObject& InMetaSound)
-		{
-			const FMetasoundAssetBase* MetaSoundAsset = IMetasoundUObjectRegistry::Get().GetObjectAsAssetBase(&InMetaSound);
-			check(MetaSoundAsset);
-			const UMetasoundEditorGraph* EditorGraph = CastChecked<UMetasoundEditorGraph>(MetaSoundAsset->GetGraph());
-
-			// Get all editor nodes from editor graph (some nodes on graph may *NOT* be metasound ed nodes,
-			// such as comment boxes, etc, so just get nodes of class UMetasoundEditorGraph).
-			TArray<const UMetasoundEditorGraphNode*> EditorNodes;
-			EditorGraph->GetNodesOfClass(EditorNodes);
-
-			// Do not synchronize with errors present as the graph is expected to be malformed.
-			return Algo::AnyOf(EditorNodes, [](const UMetasoundEditorGraphNode* Node) { return Node->ErrorType == EMessageSeverity::Error; });
-		}
-
 		Frontend::FNodeHandle FGraphBuilder::AddInputTemplateNodeHandle(UObject& InMetaSound, Frontend::FNodeHandle& InputNodeHandle)
 		{
 			using namespace Frontend;
@@ -1684,6 +1672,28 @@ namespace Metasound
 			return ensure(Graph->RemoveNode(&InNode));
 		}
 
+		bool FGraphBuilder::BindEditorGraph(const FMetaSoundFrontendDocumentBuilder& InBuilder, UMetasoundEditorGraph** OutGraph)
+		{
+			using namespace Frontend;
+
+			bool bNewGraphBound = false;
+			FMetasoundAssetBase& MetaSound = InBuilder.GetMetasoundAsset();
+			UMetasoundEditorGraph* Graph = Cast<UMetasoundEditorGraph>(MetaSound.GetGraph());
+			if (!Graph)
+			{
+				Graph = NewObject<UMetasoundEditorGraph>(MetaSound.GetOwningAsset(), FName(), RF_Transactional);
+				Graph->Schema = UMetasoundEditorGraphSchema::StaticClass();
+				MetaSound.SetGraph(Graph);
+				bNewGraphBound = true;
+			}
+
+			if (OutGraph)
+			{
+				*OutGraph = Graph;
+			}
+			return bNewGraphBound;
+		}
+
 		void FGraphBuilder::RebuildNodePins(UMetasoundEditorGraphNode& InGraphNode)
 		{
 			using namespace Frontend;
@@ -1961,18 +1971,12 @@ namespace Metasound
 			return false;
 		}
 
-		bool FGraphBuilder::SynchronizeComments(UObject& InMetaSound)
+		bool FGraphBuilder::SynchronizeComments(const FMetaSoundFrontendDocumentBuilder& InBuilder, UMetasoundEditorGraph& OutGraph)
 		{
 			bool bModified = false;
 
-			FMetasoundAssetBase* MetaSoundAsset = IMetasoundUObjectRegistry::Get().GetObjectAsAssetBase(&InMetaSound);
-			check(MetaSoundAsset);
-
-			UMetasoundEditorGraph* EditorGraph = CastChecked<UMetasoundEditorGraph>(MetaSoundAsset->GetGraph());
-			check(EditorGraph);
-
 			TMap<FGuid, UMetasoundEditorGraphCommentNode*> CommentIDToEdNode;
-			Algo::TransformIf(EditorGraph->Nodes, CommentIDToEdNode,
+			Algo::TransformIf(OutGraph.Nodes, CommentIDToEdNode,
 				[](const UEdGraphNode* Node)
 				{
 					return Node->IsA<UMetasoundEditorGraphCommentNode>();
@@ -1983,8 +1987,8 @@ namespace Metasound
 					return TPair<FGuid, UMetasoundEditorGraphCommentNode*>(CommentNode->GetCommentID(), CommentNode);
 				});
 
-			const FMetasoundFrontendDocument& Document = MetaSoundAsset->GetConstDocumentChecked();
-			const TMap<FGuid, FMetaSoundFrontendGraphComment>& Comments = Document.RootGraph.Graph.Style.Comments;
+			const FMetasoundFrontendDocument& Document = InBuilder.GetConstDocumentChecked();
+			const TMap<FGuid, FMetaSoundFrontendGraphComment>& Comments = InBuilder.FindConstBuildGraphChecked().Style.Comments;
 			for (const TPair<FGuid, FMetaSoundFrontendGraphComment>& Pair : Comments)
 			{
 				UMetasoundEditorGraphCommentNode* CommentNode = nullptr;
@@ -2002,9 +2006,10 @@ namespace Metasound
 					const UGraphEditorSettings* GraphEditorSettings = GetDefault<UGraphEditorSettings>();
 					check(GraphEditorSettings);
 
-					UMetasoundEditorGraphCommentNode* NewNode = CastChecked<UMetasoundEditorGraphCommentNode>(DuplicateObject<UEdGraphNode>(CommentCDO, EditorGraph));
+					UMetasoundEditorGraphCommentNode* NewNode = CastChecked<UMetasoundEditorGraphCommentNode>(DuplicateObject<UEdGraphNode>(CommentCDO, &OutGraph));
 					NewNode->SetFlags(RF_Transactional);
-					EditorGraph->AddNode(NewNode, false, false);
+					OutGraph.AddNode(NewNode, false, false);
+
 					NewNode->CreateNewGuid();
 					NewNode->bCommentBubbleVisible_InDetailsPanel = GraphEditorSettings->bShowCommentBubbleWhenZoomedOut;
 
@@ -2020,13 +2025,13 @@ namespace Metasound
 			{
 				constexpr bool bMarkDirty = false;
 				constexpr bool bBreakAllLinks = true;
-				EditorGraph->RemoveNode(Comment.Value, bBreakAllLinks, bMarkDirty);
+				OutGraph.RemoveNode(Comment.Value, bBreakAllLinks, bMarkDirty);
 			}
 
 			return bModified;
 		}
 
-		bool FGraphBuilder::SynchronizeConnections(UObject& InMetaSound)
+		bool FGraphBuilder::SynchronizeConnections(const FMetaSoundFrontendDocumentBuilder& InBuilder, UMetasoundEditorGraph& OutGraph)
 		{
 			using namespace Frontend;
 
@@ -2034,10 +2039,7 @@ namespace Metasound
 
 			TArray<UMetasoundEditorGraphNode*> EditorNodes;
 			{
-				FMetasoundAssetBase* MetaSoundAsset = IMetasoundUObjectRegistry::Get().GetObjectAsAssetBase(&InMetaSound);
-				check(MetaSoundAsset);
-				UMetasoundEditorGraph* EditorGraph = CastChecked<UMetasoundEditorGraph>(MetaSoundAsset->GetGraph());
-				EditorGraph->GetNodesOfClass(EditorNodes);
+				OutGraph.GetNodesOfClass(EditorNodes);
 			}
 
 			TMap<FGuid, UMetasoundEditorGraphNode*> EditorNodesByFrontendID;
@@ -2156,13 +2158,14 @@ namespace Metasound
 			return bIsGraphDirty;
 		}
 
-		bool FGraphBuilder::SynchronizeGraph(UObject& InMetaSound)
+		bool FGraphBuilder::SynchronizeGraph(const FMetaSoundFrontendDocumentBuilder& InBuilder, UMetasoundEditorGraph& OutGraph)
 		{
 			TRACE_CPUPROFILER_EVENT_SCOPE(Metasound::Editor::FGraphBuilder::SynchronizeGraph);
 
 			using namespace Frontend;
 
-			FMetasoundAssetBase* MetaSoundAsset = IMetasoundUObjectRegistry::Get().GetObjectAsAssetBase(&InMetaSound);
+			UObject& MetaSound = InBuilder.CastDocumentObjectChecked<UObject>();
+			FMetasoundAssetBase* MetaSoundAsset = IMetasoundUObjectRegistry::Get().GetObjectAsAssetBase(&MetaSound);
 			check(MetaSoundAsset);
 
 			if (RecurseGetDocumentModified(*MetaSoundAsset))
@@ -2173,7 +2176,7 @@ namespace Metasound
 					TArray<UObject*> EditedAssets = GEditor->GetEditorSubsystem<UAssetEditorSubsystem>()->GetAllEditedAssets();
 					for (UObject* Asset : EditedAssets)
 					{
-						if (Asset != &InMetaSound)
+						if (Asset != &MetaSound)
 						{
 							if (FMetasoundAssetBase* EditedMetaSound = IMetasoundUObjectRegistry::Get().GetObjectAsAssetBase(Asset))
 							{
@@ -2189,7 +2192,7 @@ namespace Metasound
 				if (EditedReferencingMetaSounds.IsEmpty())
 				{
 					MetaSoundAsset->CacheRegistryMetadata();
-					GraphBuilderPrivate::SynchronizeGraphRecursively(InMetaSound);
+					GraphBuilderPrivate::SynchronizeGraphRecursively(InBuilder, OutGraph);
 					GraphBuilderPrivate::RecurseClearDocumentModified(*MetaSoundAsset);
 				}
 				else
@@ -2199,7 +2202,8 @@ namespace Metasound
 						check(EditedMetaSound);
 						UObject* OwningMetaSound = EditedMetaSound->GetOwningAsset();
 						check(OwningMetaSound);
-						SynchronizeGraph(*OwningMetaSound);
+						const FMetaSoundFrontendDocumentBuilder& EditedBuilder = IDocumentBuilderRegistry::GetChecked().FindOrBeginBuilding(OwningMetaSound);
+						SynchronizeGraph(EditedBuilder, *CastChecked<UMetasoundEditorGraph>(&EditedMetaSound->GetGraphChecked()));
 					}
 				}
 
@@ -2209,7 +2213,7 @@ namespace Metasound
 			return false;
 		}
 
-		bool FGraphBuilder::SynchronizeOutputNodes(UObject& InMetaSound)
+		bool FGraphBuilder::SynchronizeOutputNodes(const FMetaSoundFrontendDocumentBuilder& InBuilder, UMetasoundEditorGraph& OutGraph)
 		{
 			TRACE_CPUPROFILER_EVENT_SCOPE(Metasound::Editor::FGraphBuilder::SynchronizeOutputNodes);
 
@@ -2217,13 +2221,13 @@ namespace Metasound
 
 			bool bEditorGraphModified = false;
 
-			FMetasoundAssetBase* MetaSoundAsset = IMetasoundUObjectRegistry::Get().GetObjectAsAssetBase(&InMetaSound);
+			const UObject& MetaSoundObject = InBuilder.CastDocumentObjectChecked<const UObject>();
+			const FMetasoundAssetBase* MetaSoundAsset = IMetasoundUObjectRegistry::Get().GetObjectAsAssetBase(&MetaSoundObject);
 			check(MetaSoundAsset);
 			FConstGraphHandle GraphHandle = MetaSoundAsset->GetRootGraphHandle();
-			UMetasoundEditorGraph* EditorGraph = CastChecked<UMetasoundEditorGraph>(MetaSoundAsset->GetGraph());
 
 			TArray<UMetasoundEditorGraphOutputNode*> OutputNodes;
-			EditorGraph->GetNodesOfClassEx<UMetasoundEditorGraphOutputNode>(OutputNodes);
+			OutGraph.GetNodesOfClassEx<UMetasoundEditorGraphOutputNode>(OutputNodes);
 			for (UMetasoundEditorGraphOutputNode* Node : OutputNodes)
 			{
 				FConstNodeHandle NodeHandle = Node->GetConstNodeHandle();
@@ -2239,11 +2243,7 @@ namespace Metasound
 							if (Node->GetNodeHandle()->GetID() != Output->NodeID)
 							{
 								Node->SetNodeID(Output->NodeID);
-
-								// Requery handle as the id has been fixed up
-								NodeHandle = Node->GetConstNodeHandle();
-								FText OutputDisplayName = Node->GetDisplayName();
-								UE_LOG(LogMetasoundEditor, Verbose, TEXT("Editor Output Node '%s' interface versioned"), *OutputDisplayName.ToString());
+								UE_LOG(LogMetasoundEditor, Verbose, TEXT("Editor Output Node '%s' interface versioned"), *Node->GetDisplayName().ToString());
 
 								bEditorGraphModified = true;
 							}
@@ -2255,7 +2255,7 @@ namespace Metasound
 			return bEditorGraphModified;
 		}
 
-		void FGraphBuilder::SynchronizeNodes(UObject& InMetaSound)
+		void FGraphBuilder::SynchronizeNodes(const FMetaSoundFrontendDocumentBuilder& InBuilder, UMetasoundEditorGraph& OutGraph)
 		{
 			TRACE_CPUPROFILER_EVENT_SCOPE(Metasound::Editor::FGraphBuilder::SynchronizeNodes);
 
@@ -2263,13 +2263,13 @@ namespace Metasound
 
 			// Get all external nodes from Frontend graph.  Input and output references will only be added/synchronized
 			// if required when synchronizing connections (as they are not required to inhabit editor graph).
-			FMetasoundAssetBase* MetaSoundAsset = IMetasoundUObjectRegistry::Get().GetObjectAsAssetBase(&InMetaSound);
+			UObject& MetaSoundObject = InBuilder.CastDocumentObjectChecked<UObject>();
+			const FMetasoundAssetBase* MetaSoundAsset = IMetasoundUObjectRegistry::Get().GetObjectAsAssetBase(&MetaSoundObject);
 			check(MetaSoundAsset);
-			FConstGraphHandle GraphHandle = static_cast<const FMetasoundAssetBase*>(MetaSoundAsset)->GetRootGraphHandle();
+			FConstGraphHandle GraphHandle = MetaSoundAsset->GetRootGraphHandle();
 			TArray<FConstNodeHandle> FrontendNodes = GraphHandle->GetConstNodes();
-			UMetasoundEditorGraph* EditorGraph = CastChecked<UMetasoundEditorGraph>(MetaSoundAsset->GetGraph());
 			TArray<UMetasoundEditorGraphNode*> EditorNodes;
-			EditorGraph->GetNodesOfClass(EditorNodes);
+			OutGraph.GetNodesOfClass(EditorNodes);
 
 			// Find existing array of editor nodes associated with Frontend node
 			struct FAssociatedNodes
@@ -2330,7 +2330,7 @@ namespace Metasound
 			{
 				constexpr bool bMarkDirty = false;
 				constexpr bool bBreakAllLinks = true;
-				EditorGraph->RemoveNode(EditorNode, bBreakAllLinks, bMarkDirty);
+				OutGraph.RemoveNode(EditorNode, bBreakAllLinks, bMarkDirty);
 			}
 
 			// Add missing editor nodes marked as visible.
@@ -2345,7 +2345,7 @@ namespace Metasound
 					{
 						case EMetasoundFrontendClassType::External:
 						{
-							NewGraphNode = CastChecked<UMetasoundEditorGraphNode>(AddExternalNode(InMetaSound, Node->GetID(), ClassMetadata, false));
+							NewGraphNode = CastChecked<UMetasoundEditorGraphNode>(AddExternalNode(MetaSoundObject, Node->GetID(), ClassMetadata, false));
 						}
 						break;
 
@@ -2353,18 +2353,18 @@ namespace Metasound
 						{
 							if (Node->GetClassMetadata().GetClassName() == FInputNodeTemplate::ClassName)
 							{
-								NewGraphNode = CastChecked<UMetasoundEditorGraphNode>(AddInputNode(InMetaSound, Node, false));
+								NewGraphNode = CastChecked<UMetasoundEditorGraphNode>(AddInputNode(MetaSoundObject, Node, false));
 							}
 							else
 							{
-								NewGraphNode = CastChecked<UMetasoundEditorGraphNode>(AddTemplateNode(InMetaSound, Node->GetID(), Node->GetClassMetadata(), false));
+								NewGraphNode = CastChecked<UMetasoundEditorGraphNode>(AddTemplateNode(MetaSoundObject, Node->GetID(), Node->GetClassMetadata(), false));
 							}
 						}
 						break;
 
 						case EMetasoundFrontendClassType::Output:
 						{
-							NewGraphNode = CastChecked<UMetasoundEditorGraphNode>(AddOutputNode(InMetaSound, Node, false));
+							NewGraphNode = CastChecked<UMetasoundEditorGraphNode>(AddOutputNode(MetaSoundObject, Node, false));
 						}
 						break;
 
@@ -2373,7 +2373,7 @@ namespace Metasound
 						case EMetasoundFrontendClassType::VariableDeferredAccessor:
 						case EMetasoundFrontendClassType::Variable:
 						{
-							NewGraphNode = CastChecked<UMetasoundEditorGraphNode>(AddVariableNode(InMetaSound, Node, false));
+							NewGraphNode = CastChecked<UMetasoundEditorGraphNode>(AddVariableNode(MetaSoundObject, Node, false));
 						}
 						break;
 
@@ -2624,7 +2624,7 @@ namespace Metasound
 			return OldValue != InPin.DefaultValue;
 		}
 
-		bool FGraphBuilder::SynchronizeGraphMembers(UObject& InMetaSound)
+		bool FGraphBuilder::SynchronizeGraphMembers(const FMetaSoundFrontendDocumentBuilder& InBuilder, UMetasoundEditorGraph& OutGraph)
 		{
 			TRACE_CPUPROFILER_EVENT_SCOPE(Metasound::Editor::FGraphBuilder::SynchronizeGraphMembers);
 
@@ -2632,9 +2632,9 @@ namespace Metasound
 
 			bool bEditorGraphModified = false;
 
-			const FMetasoundAssetBase* MetaSoundAsset = IMetasoundUObjectRegistry::Get().GetObjectAsAssetBase(&InMetaSound);
+			const UObject& MetaSoundObject = InBuilder.CastDocumentObjectChecked<const UObject>();
+			const FMetasoundAssetBase* MetaSoundAsset = IMetasoundUObjectRegistry::Get().GetObjectAsAssetBase(&MetaSoundObject);
 			check(MetaSoundAsset);
-			UMetasoundEditorGraph* Graph = CastChecked<UMetasoundEditorGraph>(MetaSoundAsset->GetGraph());
 			FConstGraphHandle GraphHandle = MetaSoundAsset->GetRootGraphHandle();
 
 			TSet<UMetasoundEditorGraphInput*> Inputs;
@@ -2643,14 +2643,14 @@ namespace Metasound
 			// Collect all editor graph inputs with corresponding frontend inputs. 
 			GraphHandle->IterateConstNodes([&](FConstNodeHandle NodeHandle)
 			{
-				if (UMetasoundEditorGraphInput* Input = Graph->FindInput(NodeHandle->GetID()))
+				if (UMetasoundEditorGraphInput* Input = OutGraph.FindInput(NodeHandle->GetID()))
 				{
 					Inputs.Add(Input);
 					return;
 				}
 
 				// Add an editor input if none exist for a frontend input.
-				Inputs.Add(Graph->FindOrAddInput(NodeHandle));
+				Inputs.Add(OutGraph.FindOrAddInput(NodeHandle));
 				constexpr bool bIncludeNamespace = true;
 				FText NodeDisplayName = FGraphBuilder::GetDisplayName(*NodeHandle, bIncludeNamespace);
 				UE_LOG(LogMetasoundEditor, Verbose, TEXT("Synchronizing Inputs: Added missing input '%s'."), *NodeDisplayName.ToString());
@@ -2659,9 +2659,9 @@ namespace Metasound
 
 			// Remove stale inputs.
 			{
-				for (int32 Index = Graph->Inputs.Num() - 1; Index >= 0; --Index)
+				for (int32 Index = OutGraph.Inputs.Num() - 1; Index >= 0; --Index)
 				{
-					const UMetasoundEditorGraphInput* Input = Graph->Inputs[Index];
+					const UMetasoundEditorGraphInput* Input = OutGraph.Inputs[Index];
 					check(Input);
 					if (!Inputs.Contains(Input))
 					{
@@ -2673,7 +2673,7 @@ namespace Metasound
 
 						bEditorGraphModified = true;
 						UE_LOG(LogMetasoundEditor, Verbose, TEXT("Synchronizing Outputs: Removing stale output '%s'."), *Input->GetName());
-						Graph->Inputs.RemoveAtSwap(Index, EAllowShrinking::No);
+						OutGraph.Inputs.RemoveAtSwap(Index, EAllowShrinking::No);
 					}
 				}
 			}
@@ -2681,14 +2681,14 @@ namespace Metasound
 			// Collect all editor graph outputs with corresponding frontend outputs. 
 			GraphHandle->IterateConstNodes([&](FConstNodeHandle NodeHandle)
 			{
-				if (UMetasoundEditorGraphOutput* Output = Graph->FindOutput(NodeHandle->GetID()))
+				if (UMetasoundEditorGraphOutput* Output = OutGraph.FindOutput(NodeHandle->GetID()))
 				{
 					Outputs.Add(Output);
 					return;
 				}
 
 				// Add an editor output if none exist for a frontend output.
-				Outputs.Add(Graph->FindOrAddOutput(NodeHandle));
+				Outputs.Add(OutGraph.FindOrAddOutput(NodeHandle));
 				constexpr bool bIncludeNamespace = true;
 				FText NodeDisplayName = FGraphBuilder::GetDisplayName(*NodeHandle, bIncludeNamespace);
 				UE_LOG(LogMetasoundEditor, Verbose, TEXT("Synchronizing Outputs: Added missing output '%s'."), *NodeDisplayName.ToString());
@@ -2697,15 +2697,15 @@ namespace Metasound
 
 			// Remove stale outputs.
 			{
-				for (int32 Index = Graph->Outputs.Num() - 1; Index >= 0; --Index)
+				for (int32 Index = OutGraph.Outputs.Num() - 1; Index >= 0; --Index)
 				{
-					const UMetasoundEditorGraphOutput* Output = Graph->Outputs[Index];
+					const UMetasoundEditorGraphOutput* Output = OutGraph.Outputs[Index];
 					check(Output);
 					if (!Outputs.Contains(Output))
 					{
 						bEditorGraphModified = true;
 						UE_LOG(LogMetasoundEditor, Verbose, TEXT("Synchronizing Outputs: Removing stale output '%s'."), *Output->GetName());
-						Graph->Outputs.RemoveAtSwap(Index, EAllowShrinking::No);
+						OutGraph.Outputs.RemoveAtSwap(Index, EAllowShrinking::No);
 					}
 				}
 			}
@@ -2740,7 +2740,9 @@ namespace Metasound
 
 						if (DefaultLiteral.IsValid())
 						{
-							InVertex.GetLiteral()->SetFromLiteral(DefaultLiteral);
+							UMetasoundEditorGraphMemberDefaultLiteral* Literal = InVertex.GetLiteral();
+							Literal->SetFromLiteral(DefaultLiteral);
+							bEditorGraphModified = true;
 						}
 					}
 				}
@@ -2749,7 +2751,7 @@ namespace Metasound
 			// Synchronize data types & default values for input nodes.
 			GraphHandle->IterateConstNodes([&](FConstNodeHandle NodeHandle)
 			{
-				if (UMetasoundEditorGraphInput* Input = Graph->FindInput(NodeHandle->GetID()))
+				if (UMetasoundEditorGraphInput* Input = OutGraph.FindInput(NodeHandle->GetID()))
 				{
 					SynchronizeMemberDataType(*Input);
 
@@ -2778,7 +2780,7 @@ namespace Metasound
 			// Synchronize data types of output nodes.
 			GraphHandle->IterateConstNodes([&](FConstNodeHandle NodeHandle)
 			{
-				if (UMetasoundEditorGraphOutput* Output = Graph->FindOutput(NodeHandle->GetID()))
+				if (UMetasoundEditorGraphOutput* Output = OutGraph.FindOutput(NodeHandle->GetID()))
 				{
 					SynchronizeMemberDataType(*Output);
 
@@ -2792,27 +2794,28 @@ namespace Metasound
 
 			// Enforce all variables exists in the edgraph
 			TSet<FGuid> VariableIDs;
-			Algo::Transform(Graph->Variables, VariableIDs, [](const TObjectPtr<UMetasoundEditorGraphVariable>& Variable) { return Variable->GetVariableID(); });
+			Algo::Transform(OutGraph.Variables, VariableIDs, [](const TObjectPtr<UMetasoundEditorGraphVariable>& Variable) { return Variable->GetVariableID(); });
 			TArray<FConstVariableHandle> FrontendVariables = GraphHandle->GetVariables();
 			for (const FConstVariableHandle& VariableHandle : FrontendVariables)
 			{
 				const FGuid VariableID = VariableHandle->GetID();
 				if (VariableIDs.Remove(VariableID) == 0)
 				{
-					UMetasoundEditorGraphVariable* NewVariable = NewObject<UMetasoundEditorGraphVariable>(Graph, FName(), RF_Transactional);
+					UMetasoundEditorGraphVariable* NewVariable = NewObject<UMetasoundEditorGraphVariable>(&OutGraph, FName(), RF_Transactional);
 					if (ensure(NewVariable))
 					{
 						const FMetasoundFrontendLiteral DefaultLiteral = VariableHandle->GetLiteral();
 						NewVariable->InitMember(VariableHandle->GetDataType(), DefaultLiteral, VariableID);
-						Graph->Variables.Add(NewVariable);
+						OutGraph.Variables.Add(NewVariable);
+						bEditorGraphModified = true;
 					}
 				}
 			}
 
 			// Remove empty entries
-			bEditorGraphModified |= Graph->Inputs.RemoveAllSwap([](const TObjectPtr<UMetasoundEditorGraphInput>& Input) { return !Input; }) > 0;
-			bEditorGraphModified |= Graph->Outputs.RemoveAllSwap([](const TObjectPtr<UMetasoundEditorGraphOutput>& Output) { return !Output; }) > 0;
-			bEditorGraphModified |= Graph->Variables.RemoveAllSwap([&VariableIDs](const TObjectPtr<UMetasoundEditorGraphVariable>& Variable)
+			bEditorGraphModified |= OutGraph.Inputs.RemoveAllSwap([](const TObjectPtr<UMetasoundEditorGraphInput>& Input) { return !Input; }) > 0;
+			bEditorGraphModified |= OutGraph.Outputs.RemoveAllSwap([](const TObjectPtr<UMetasoundEditorGraphOutput>& Output) { return !Output; }) > 0;
+			bEditorGraphModified |= OutGraph.Variables.RemoveAllSwap([&VariableIDs](const TObjectPtr<UMetasoundEditorGraphVariable>& Variable)
 			{
 				return !Variable || VariableIDs.Contains(Variable->GetVariableID());
 			}) > 0;
diff --git a/Engine/Plugins/Runtime/Metasound/Source/MetasoundEditor/Private/MetasoundEditorGraphBuilder.h b/Engine/Plugins/Runtime/Metasound/Source/MetasoundEditor/Private/MetasoundEditorGraphBuilder.h
--- a/Engine/Plugins/Runtime/Metasound/Source/MetasoundEditor/Private/MetasoundEditorGraphBuilder.h
+++ b/Engine/Plugins/Runtime/Metasound/Source/MetasoundEditor/Private/MetasoundEditorGraphBuilder.h
@@ -193,9 +193,6 @@ namespace Metasound
 			// Generates a unique output name for the given MetaSound object
 			static FName GenerateUniqueNameByClassType(const UObject& InMetaSound, EMetasoundFrontendClassType InClassType, const FString& InBaseName);
 
-			// Whether or not associated editor graph is in an error state or not.
-			static bool GraphContainsErrors(const UObject& InMetaSound);
-
 			static TArray<FString> GetDataTypeNameCategories(const FName& InDataTypeName);
 
 			// Get the input handle from an input pin.  Ensures pin is an input pin.
@@ -242,6 +239,10 @@ namespace Metasound
 			// Deletes both the editor graph & frontend nodes from respective graphs
 			static bool DeleteNode(UEdGraphNode& InNode, bool bRemoveUnusedDependencies = true);
 
+			// Returns Editor Graph associated with the given builder's MetaSound object. If the editor graph was created, initialized, and
+			// bound to builder's MetaSound object, returns true (false if it already existed).  Sets (optional) pointer to the bound graph.
+			static bool BindEditorGraph(const FMetaSoundFrontendDocumentBuilder& InBuilder, UMetasoundEditorGraph** OutGraph = nullptr);
+
 			// Adds an Input UEdGraphPin to a UMetasoundEditorGraphNode
 			static UEdGraphPin* AddPinToNode(UMetasoundEditorGraphNode& InEditorNode, Frontend::FConstInputHandle InInputHandle);
 
@@ -254,30 +255,31 @@ namespace Metasound
 			// Adds and removes nodes, pins and connections so that the UEdGraph of the MetaSound matches the
 			// FMetasoundFrontendDocument model. Validates the graph (and those referenced recursively).
 			//
-			// @param InMetaSound - MetaSound to synchronize and optionally validate.
+			// @param InBuilder - Builder to synchronize ed graph with.
+			// @param OutGraph = Graph to mutate to conform to provided builder's selected build graph
 			// @return whether or not EditorGraph synchronization was performed.
-			static bool SynchronizeGraph(UObject& InMetaSound);
+			static bool SynchronizeGraph(const FMetaSoundFrontendDocumentBuilder& InBuilder, UMetasoundEditorGraph& OutGraph);
 
 			// Synchronizes editor nodes with frontend nodes, removing editor nodes that are not represented in the frontend, and adding editor nodes to represent missing frontend nodes.
-			static void SynchronizeNodes(UObject& InMetaSound);
+			static void SynchronizeNodes(const FMetaSoundFrontendDocumentBuilder& InBuilder, UMetasoundEditorGraph& OutGraph);
 
 			// Synchronizes and reports to log whether or not an output node's associated FrontendNode ID has changed and therefore been updated through node versioning.
 			//
 			// @return True if the UMetasoundEditorGraphNode was altered. False otherwise.
-			static bool SynchronizeOutputNodes(UObject& InMetaSound);
+			static bool SynchronizeOutputNodes(const FMetaSoundFrontendDocumentBuilder& InBuilder, UMetasoundEditorGraph& OutGraph);
 
 			// Adds and removes pins so that the UMetasoundEditorGraphNode matches the InNode.
 			//
 			// @return True if the UMetasoundEditorGraphNode was altered. False otherwise.
 			static bool SynchronizeNodePins(UMetasoundEditorGraphNode& InEditorNode, Frontend::FConstNodeHandle InNode, bool bRemoveUnusedPins = true, bool bLogChanges = true);
 
-			static bool SynchronizeComments(UObject& InMetaSound);
+			static bool SynchronizeComments(const FMetaSoundFrontendDocumentBuilder& InBuilder, UMetasoundEditorGraph& OutGraph);
 
 			// Adds and removes connections so that the UEdGraph of the MetaSound has the same
 			// connections as the FMetasoundFrontendDocument graph.
 			//
 			// @return True if the UEdGraph was altered. False otherwise.
-			static bool SynchronizeConnections(UObject& InMetaSound);
+			static bool SynchronizeConnections(const FMetaSoundFrontendDocumentBuilder& InBuilder, UMetasoundEditorGraph& OutGraph);
 
 			// Synchronizes literal for a given input with the EdGraph's pin value.
 			static bool SynchronizePinLiteral(UEdGraphPin& InPin);
@@ -288,7 +290,7 @@ namespace Metasound
 			// Synchronizes inputs, variables, and outputs for the given MetaSound.
 			//
 			// @return True if the UEdGraph was altered. False otherwise.
-			static bool SynchronizeGraphMembers(UObject& InMetaSound);
+			static bool SynchronizeGraphMembers(const FMetaSoundFrontendDocumentBuilder& InBuilder, UMetasoundEditorGraph& OutGraph);
 
 			// Returns true if the FInputHandle and UEdGraphPin match each other.
 			static bool IsMatchingInputHandleAndPin(const Frontend::FConstInputHandle& InInputHandle, const UEdGraphPin& InEditorPin);
diff --git a/Engine/Plugins/Runtime/Metasound/Source/MetasoundEditor/Private/MetasoundEditorGraphSchema.cpp b/Engine/Plugins/Runtime/Metasound/Source/MetasoundEditor/Private/MetasoundEditorGraphSchema.cpp
--- a/Engine/Plugins/Runtime/Metasound/Source/MetasoundEditor/Private/MetasoundEditorGraphSchema.cpp
+++ b/Engine/Plugins/Runtime/Metasound/Source/MetasoundEditor/Private/MetasoundEditorGraphSchema.cpp
@@ -19,6 +19,7 @@
 #include "MetasoundAssetManager.h"
 #include "MetasoundBuilderSubsystem.h"
 #include "MetasoundDataReference.h"
+#include "MetasoundDocumentBuilderRegistry.h"
 #include "MetasoundDocumentInterface.h"
 #include "MetasoundEditor.h"
 #include "MetasoundEditorCommands.h"
@@ -1937,6 +1938,7 @@ void UMetasoundEditorGraphSchema::GetAssetsPinHoverMessage(const TArray<FAssetDa
 void UMetasoundEditorGraphSchema::DroppedAssetsOnGraph(const TArray<FAssetData>& Assets, const FVector2D& GraphPosition, UEdGraph* Graph) const
 {
 	using namespace Metasound;
+	using namespace Metasound::Engine;
 	using namespace Metasound::Frontend;
 
 	FScopedTransaction Transaction(LOCTEXT("DropMetaSoundOnGraph", "Drop MetaSound On Graph"));
@@ -1948,6 +1950,8 @@ void UMetasoundEditorGraphSchema::DroppedAssetsOnGraph(const TArray<FAssetData>&
 
 	FMetasoundAssetBase* MetaSoundAsset = IMetasoundUObjectRegistry::Get().GetObjectAsAssetBase(&MetaSound);
 	check(MetaSoundAsset);
+	UMetaSoundBuilderBase& Builder = FDocumentBuilderRegistry::GetChecked().FindOrBeginBuilding(MetaSound);
+	EMetaSoundBuilderResult Result = EMetaSoundBuilderResult::Failed;
 
 	for (const FAssetData& DroppedAsset : Assets)
 	{
@@ -1971,36 +1975,17 @@ void UMetasoundEditorGraphSchema::DroppedAssetsOnGraph(const TArray<FAssetData>&
 				bModifiedObjects = true;
 			}
 
-			// This may not be necessary as dropping an asset on the graph may load it, thus triggering the registration from the MetaSoundAssetManager.
-			const FMetasoundFrontendDocument& DroppedDoc = DroppedMetaSoundAsset->GetConstDocumentChecked();
-			const FNodeRegistryKey RegistryKey = FNodeRegistryKey(DroppedDoc.RootGraph);
-			if (ensure(RegistryKey.IsValid()))
+			TScriptInterface<IMetaSoundDocumentInterface> DroppedDocInterface(DroppedObject);
+			FMetaSoundNodeHandle NodeHandle = Builder.AddNode(DroppedDocInterface, Result);
+			if (ensure(Result == EMetaSoundBuilderResult::Succeeded))
 			{
-				FMetaSoundAssetRegistrationOptions RegOptions;
-				RegOptions.bForceReregister = false;
-				DroppedMetaSoundAsset->RegisterGraphWithFrontend(RegOptions);
-			}
-
-			FMetasoundFrontendClass Class;
-			if (ensure(ISearchEngine::Get().FindClassWithHighestVersion(DroppedDoc.RootGraph.Metadata.GetClassName(), Class)))
-			{
-				UMetasoundEditorGraphNode* NewGraphNode = Metasound::Editor::FGraphBuilder::AddExternalNode(MetaSound, Class.Metadata);
-				NewGraphNode->UpdateFrontendNodeLocation(GraphPosition);
-				NewGraphNode->SyncLocationFromFrontendNode();
-
-				bTransactionSucceeded = true;
+				Builder.SetNodeLocation(NodeHandle.NodeID, GraphPosition, Result);
+				bTransactionSucceeded = ensure(Result == EMetaSoundBuilderResult::Succeeded);
 			}
 		}
 	}
 
-	if (bTransactionSucceeded)
-	{
-		// Reregister this asset after adding the dropped asset to update references
-		FMetaSoundAssetRegistrationOptions RegOptions;
-		RegOptions.bForceReregister = true;
-		MetaSoundAsset->RegisterGraphWithFrontend(RegOptions);
-	}
-	else
+	if (!bTransactionSucceeded)
 	{
 		Transaction.Cancel();
 	}
diff --git a/Engine/Plugins/Runtime/Metasound/Source/MetasoundEditor/Private/MetasoundEditorModule.cpp b/Engine/Plugins/Runtime/Metasound/Source/MetasoundEditor/Private/MetasoundEditorModule.cpp
--- a/Engine/Plugins/Runtime/Metasound/Source/MetasoundEditor/Private/MetasoundEditorModule.cpp
+++ b/Engine/Plugins/Runtime/Metasound/Source/MetasoundEditor/Private/MetasoundEditorModule.cpp
@@ -452,8 +452,9 @@ namespace Metasound
 							FGraphBuilder::UnregisterGraphWithFrontend(*AssetObject);
 						}
 
-						IMetaSoundAssetManager::GetChecked().RemoveAsset(InAssetData);
 					}
+
+					IMetaSoundAssetManager::GetChecked().RemoveAsset(InAssetData);
 				}
 			}
 
@@ -979,7 +980,7 @@ namespace Metasound
 					IMetaSoundAssetManager& AssetManager = IMetaSoundAssetManager::GetChecked();
 					for (FMetaSoundFrontendDocumentBuilder& MetaSoundBuilder : NewMetaSoundAssetBuilders)
 					{
-						const FMetasoundFrontendClassName OldReferencedClassName = MetaSoundBuilder.GetConstDocument().RootGraph.Metadata.GetClassName();
+						const FMetasoundFrontendClassName OldReferencedClassName = MetaSoundBuilder.GetConstDocumentChecked().RootGraph.Metadata.GetClassName();
 						const FMetasoundFrontendClassName NewReferencedClassName = MetaSoundBuilder.GenerateNewClassName();
 						OldToNewReferencedClassNames.FindOrAdd(OldReferencedClassName) = NewReferencedClassName;
 
diff --git a/Engine/Plugins/Runtime/Metasound/Source/MetasoundEditor/Private/MetasoundEditorSubsystem.cpp b/Engine/Plugins/Runtime/Metasound/Source/MetasoundEditor/Private/MetasoundEditorSubsystem.cpp
--- a/Engine/Plugins/Runtime/Metasound/Source/MetasoundEditor/Private/MetasoundEditorSubsystem.cpp
+++ b/Engine/Plugins/Runtime/Metasound/Source/MetasoundEditor/Private/MetasoundEditorSubsystem.cpp
@@ -4,6 +4,7 @@
 #include "IAssetTools.h"
 #include "Kismet/KismetSystemLibrary.h"
 #include "MetasoundDocumentBuilderRegistry.h"
+#include "MetasoundDocumentInterface.h"
 #include "MetasoundEditorGraph.h"
 #include "MetasoundEditorGraphBuilder.h"
 #include "MetasoundEditorGraphSchema.h"
@@ -228,27 +229,8 @@ void UMetaSoundEditorSubsystem::InitAsset(UObject& InNewMetaSound, UObject* InRe
 
 void UMetaSoundEditorSubsystem::InitEdGraph(UObject& InMetaSound)
 {
-	using namespace Metasound;
-	using namespace Metasound::Editor;
 	using namespace Metasound::Frontend;
-
-	TScriptInterface<IMetaSoundDocumentInterface> DocInterface = &InMetaSound;
-	const FMetasoundFrontendClassName& ClassName = DocInterface->GetConstDocument().RootGraph.Metadata.GetClassName();
-
-	FMetasoundAssetBase* MetaSoundAsset = IMetasoundUObjectRegistry::Get().GetObjectAsAssetBase(&InMetaSound);
-	checkf(MetaSoundAsset, TEXT("EdGraph can only be initialized on registered MetaSoundAsset type"));
-
-	UMetasoundEditorGraph* Graph = Cast<UMetasoundEditorGraph>(MetaSoundAsset->GetGraph());
-	if (!Graph)
-	{
-		Graph = NewObject<UMetasoundEditorGraph>(&InMetaSound, FName(), RF_Transactional);
-		Graph->Schema = UMetasoundEditorGraphSchema::StaticClass();
-		MetaSoundAsset->SetGraph(Graph);
-
-		// Has to be done inline to have valid graph initially when opening editor for the first
-		// time (as opposed to being applied on tick when the document's modify context has updates)
-		FGraphBuilder::SynchronizeGraph(InMetaSound);
-	}
+	Metasound::Editor::FGraphBuilder::BindEditorGraph(IDocumentBuilderRegistry::GetChecked().FindOrBeginBuilding(&InMetaSound));
 }
 
 void UMetaSoundEditorSubsystem::RegisterGraphWithFrontend(UObject& InMetaSound, bool bInForceViewSynchronization)
diff --git a/Engine/Plugins/Runtime/Metasound/Source/MetasoundEditor/Public/MetasoundEditorGraph.h b/Engine/Plugins/Runtime/Metasound/Source/MetasoundEditor/Public/MetasoundEditorGraph.h
--- a/Engine/Plugins/Runtime/Metasound/Source/MetasoundEditor/Public/MetasoundEditorGraph.h
+++ b/Engine/Plugins/Runtime/Metasound/Source/MetasoundEditor/Public/MetasoundEditorGraph.h
@@ -19,7 +19,7 @@
 
 
 // Forward Declarations
-struct FMetasoundFrontendDocument;
+struct FMetaSoundFrontendDocumentBuilder;
 struct FPropertyChangedChainEvent;
 struct FPropertyChangedEvent;
 
diff --git a/Engine/Plugins/Runtime/Metasound/Source/MetasoundEditor/Public/MetasoundEditorSubsystem.h b/Engine/Plugins/Runtime/Metasound/Source/MetasoundEditor/Public/MetasoundEditorSubsystem.h
--- a/Engine/Plugins/Runtime/Metasound/Source/MetasoundEditor/Public/MetasoundEditorSubsystem.h
+++ b/Engine/Plugins/Runtime/Metasound/Source/MetasoundEditor/Public/MetasoundEditorSubsystem.h
@@ -63,7 +63,7 @@ public:
 	// Initialize the UObject asset, with an optional MetaSound to be referenced if the asset is a preset
 	void InitAsset(UObject& InNewMetaSound, UObject* InReferencedMetaSound = nullptr);
 
-	// Initialize UMetasoundEditorGraph for a given MetaSound object
+	UE_DEPRECATED(5.5, "EdGraph is now transiently generated and privately managed for asset editor use only.")
 	void InitEdGraph(UObject& InMetaSound);
 
 	// Wraps RegisterGraphWithFrontend logic in Frontend with any additional logic required to refresh editor & respective editor object state.
diff --git a/Engine/Plugins/Runtime/Metasound/Source/MetasoundEngine/Private/MetasoundAssetSubsystem.cpp b/Engine/Plugins/Runtime/Metasound/Source/MetasoundEngine/Private/MetasoundAssetSubsystem.cpp
--- a/Engine/Plugins/Runtime/Metasound/Source/MetasoundEngine/Private/MetasoundAssetSubsystem.cpp
+++ b/Engine/Plugins/Runtime/Metasound/Source/MetasoundEngine/Private/MetasoundAssetSubsystem.cpp
@@ -90,13 +90,8 @@ namespace Metasound::Engine
 			check(IsInGameThread());
 			if (TArray<FTopLevelAssetPath>* MapAssetPaths = InMap.Find(AssetKey))
 			{
-				auto ComparePaths = [&AssetPath](const FTopLevelAssetPath& Path)
-				{
-					// Cook can strip package names on destruction, so only asset name is reliable
-					return IsRunningCookCommandlet()
-						? Path.GetAssetName() == AssetPath.GetAssetName()
-						: Path == AssetPath;
-				};
+				// Package names are stripped on destruction, so only asset name is reliable
+				auto ComparePaths = [&AssetPath](const FTopLevelAssetPath& Path) { return Path.GetAssetName() == AssetPath.GetAssetName(); };
 				if (MapAssetPaths->RemoveAllSwap(ComparePaths, EAllowShrinking::No) > 0)
 				{
 					if (MapAssetPaths->IsEmpty())
@@ -156,7 +151,7 @@ namespace Metasound::Engine
 
 		/* IMetaSoundAssetManager Implementation */
 #if WITH_EDITORONLY_DATA
-		virtual void AddAssetReferences(FMetasoundAssetBase& InAssetBase) override;
+		virtual bool AddAssetReferences(FMetasoundAssetBase& InAssetBase) override;
 #endif // WITH_EDITORONLY_DATA
 		virtual FAssetKey AddOrUpdateAsset(const FAssetData& InAssetData) override;
 		virtual FAssetKey AddOrUpdateAsset(const UObject& InObject) override;
@@ -230,7 +225,7 @@ namespace Metasound::Engine
 	}
 
 #if WITH_EDITORONLY_DATA
-	void FMetaSoundAssetManager::AddAssetReferences(FMetasoundAssetBase& InAssetBase)
+	bool FMetaSoundAssetManager::AddAssetReferences(FMetasoundAssetBase& InAssetBase)
 	{
 		using namespace Frontend;
 
@@ -260,7 +255,7 @@ namespace Metasound::Engine
 		// All keys are loaded
 		if (!bAddFromReferencedAssets)
 		{
-			return;
+			return false;
 		}
 
 		UE_LOG(LogMetaSound, Verbose, TEXT("Attempting preemptive reference load..."));
@@ -271,7 +266,7 @@ namespace Metasound::Engine
 			if (Asset)
 			{
 				const FMetasoundFrontendDocument& RefDocument = Asset->GetConstDocumentChecked();
-				const FNodeRegistryKey ClassKey = FNodeRegistryKey(RefDocument.RootGraph);
+				const FAssetKey ClassKey = FAssetKey(RefDocument.RootGraph);
 				if (!ContainsKey(ClassKey))
 				{
 					UE_LOG(LogMetaSound, Verbose,
@@ -291,6 +286,8 @@ namespace Metasound::Engine
 				UE_LOG(LogMetaSound, Warning, TEXT("Null referenced dependent asset in %s. Resaving asset in editor may fix the issue"), *InAssetBase.GetOwningAssetName());
 			}
 		}
+
+		return true;
 	}
 #endif // WITH_EDITORONLY_DATA
 
diff --git a/Engine/Plugins/Runtime/Metasound/Source/MetasoundEngine/Private/MetasoundBuilderSubsystem.cpp b/Engine/Plugins/Runtime/Metasound/Source/MetasoundEngine/Private/MetasoundBuilderSubsystem.cpp
--- a/Engine/Plugins/Runtime/Metasound/Source/MetasoundEngine/Private/MetasoundBuilderSubsystem.cpp
+++ b/Engine/Plugins/Runtime/Metasound/Source/MetasoundEngine/Private/MetasoundBuilderSubsystem.cpp
@@ -227,7 +227,7 @@ const Metasound::Engine::FOutputAudioFormatInfoPair* UMetaSoundSourceBuilder::Fi
 
 	auto Predicate = [this](const FOutputAudioFormatInfoPair& Pair)
 	{
-		const FMetasoundFrontendDocument& Document = Builder.GetConstDocument();
+		const FMetasoundFrontendDocument& Document = Builder.GetConstDocumentChecked();
 		return Document.Interfaces.Contains(Pair.Value.InterfaceVersion);
 	};
 
@@ -288,8 +288,7 @@ void UMetaSoundSourceBuilder::OnEdgeAdded(int32 EdgeIndex) const
 {
 	using namespace Metasound::DynamicGraph;
 
-	const FMetasoundFrontendDocument& Doc = Builder.GetConstDocument();
-	const FMetasoundFrontendEdge& NewEdge = Doc.RootGraph.Graph.Edges[EdgeIndex];
+	const FMetasoundFrontendEdge& NewEdge = Builder.FindConstBuildGraphChecked().Edges[EdgeIndex];
 	ExecuteAuditionableTransaction([this, &NewEdge](Metasound::DynamicGraph::FDynamicOperatorTransactor& Transactor)
 	{
 		const FMetaSoundFrontendDocumentBuilder& DocBuilder = GetConstBuilder();
@@ -326,7 +325,7 @@ void UMetaSoundSourceBuilder::OnInputAdded(int32 InputIndex)
 		using namespace Metasound;
 		using namespace Metasound::Frontend;
 
-		const FMetasoundFrontendDocument& Doc = Builder.GetConstDocument();
+		const FMetasoundFrontendDocument& Doc = Builder.GetConstDocumentChecked();
 		const FMetasoundFrontendGraphClass& GraphClass = Doc.RootGraph;
 		const FMetasoundFrontendClassInput& NewInput = GraphClass.Interface.Inputs[InputIndex];
 
@@ -363,10 +362,7 @@ void UMetaSoundSourceBuilder::OnNodeAdded(int32 NodeIndex) const
 		using namespace Metasound;
 		using namespace Metasound::Frontend;
 
-		const FMetasoundFrontendDocument& Doc = Builder.GetConstDocument();
-		const FMetasoundFrontendGraphClass& GraphClass = Doc.RootGraph;
-		const FMetasoundFrontendNode& AddedNode = GraphClass.Graph.Nodes[NodeIndex];
-
+		const FMetasoundFrontendNode& AddedNode = Builder.FindConstBuildGraphChecked().Nodes[NodeIndex];
 		const FMetasoundFrontendClass* NodeClass = Builder.FindDependency(AddedNode.ClassID);
 		checkf(NodeClass, TEXT("Node successfully added to graph but document is missing associated dependency"));
 
@@ -473,8 +469,7 @@ void UMetaSoundSourceBuilder::OnNodeInputLiteralSet(int32 NodeIndex, int32 Verte
 {
 	using namespace Metasound::DynamicGraph;
 
-	const FMetasoundFrontendGraphClass& GraphClass = Builder.GetConstDocument().RootGraph;
-	const FMetasoundFrontendNode& Node = GraphClass.Graph.Nodes[NodeIndex];
+	const FMetasoundFrontendNode& Node = Builder.FindConstBuildGraphChecked().Nodes[NodeIndex];
 	const FMetasoundFrontendVertex& Input = Node.Interface.Inputs[VertexIndex];
 
 	// Only send the literal down if not connected, as the graph core layer
@@ -502,7 +497,7 @@ void UMetaSoundSourceBuilder::OnOutputAdded(int32 OutputIndex) const
 	{
 		using namespace Metasound::Frontend;
 
-		const FMetasoundFrontendDocument& Doc = Builder.GetConstDocument();
+		const FMetasoundFrontendDocument& Doc = Builder.GetConstDocumentChecked();
 		const FMetasoundFrontendGraphClass& GraphClass = Doc.RootGraph;
 		const FMetasoundFrontendClassOutput& NewOutput = GraphClass.Interface.Outputs[OutputIndex];
 
@@ -515,8 +510,7 @@ void UMetaSoundSourceBuilder::OnRemoveSwappingEdge(int32 SwapIndex, int32 LastIn
 {
 	using namespace Metasound::DynamicGraph;
 
-	const FMetasoundFrontendDocument& Doc = Builder.GetConstDocument();
-	const FMetasoundFrontendEdge& EdgeBeingRemoved = Doc.RootGraph.Graph.Edges[SwapIndex];
+	const FMetasoundFrontendEdge& EdgeBeingRemoved = Builder.FindConstBuildGraphChecked().Edges[SwapIndex];
 	ExecuteAuditionableTransaction([this, EdgeBeingRemoved](FDynamicOperatorTransactor& Transactor)
 	{
 		using namespace Metasound;
@@ -566,7 +560,7 @@ void UMetaSoundSourceBuilder::OnRemovingInput(int32 InputIndex)
 		using namespace Metasound;
 		using namespace Metasound::Frontend;
 
-		const FMetasoundFrontendDocument& Doc = Builder.GetConstDocument();
+		const FMetasoundFrontendDocument& Doc = Builder.GetConstDocumentChecked();
 		const FMetasoundFrontendGraphClass& GraphClass = Doc.RootGraph;
 		const FMetasoundFrontendClassInput& InputBeingRemoved = GraphClass.Interface.Inputs[InputIndex];
 
@@ -604,9 +598,7 @@ void UMetaSoundSourceBuilder::OnRemoveSwappingNode(int32 SwapIndex, int32 LastIn
 	{
 		using namespace Metasound::Frontend;
 
-		const FMetasoundFrontendDocument& Doc = Builder.GetConstDocument();
-		const FMetasoundFrontendGraphClass& GraphClass = Doc.RootGraph;
-		const FMetasoundFrontendNode& NodeBeingRemoved = GraphClass.Graph.Nodes[SwapIndex];
+		const FMetasoundFrontendNode& NodeBeingRemoved = Builder.FindConstBuildGraphChecked().Nodes[SwapIndex];
 		const FGuid& NodeID = NodeBeingRemoved.GetID();
 		Transactor.RemoveNode(NodeID);
 		return true;
@@ -617,8 +609,7 @@ void UMetaSoundSourceBuilder::OnRemovingNodeInputLiteral(int32 NodeIndex, int32
 {
 	using namespace Metasound::DynamicGraph;
 
-	const FMetasoundFrontendGraphClass& GraphClass = Builder.GetConstDocument().RootGraph;
-	const FMetasoundFrontendNode& Node = GraphClass.Graph.Nodes[NodeIndex];
+	const FMetasoundFrontendNode& Node = Builder.FindConstBuildGraphChecked().Nodes[NodeIndex];
 	const FMetasoundFrontendVertex& Input = Node.Interface.Inputs[VertexIndex];
 
 	// Only send the literal down if not connected, as the graph core layer will disconnect.
@@ -630,7 +621,7 @@ void UMetaSoundSourceBuilder::OnRemovingNodeInputLiteral(int32 NodeIndex, int32
 			using namespace Metasound::Engine;
 			using namespace Metasound::Frontend;
 
-			const TArray<FMetasoundFrontendNode>& Nodes = Builder.GetConstDocument().RootGraph.Graph.Nodes;
+			const TArray<FMetasoundFrontendNode>& Nodes = Builder.FindConstBuildGraphChecked().Nodes;
 			const FMetasoundFrontendNode& Node = Nodes[NodeIndex];
 			const FMetasoundFrontendVertex& Input = Node.Interface.Inputs[VertexIndex];
 
@@ -655,7 +646,7 @@ void UMetaSoundSourceBuilder::OnRemovingOutput(int32 OutputIndex) const
 	{
 		using namespace Metasound::Frontend;
 
-		const FMetasoundFrontendDocument& Doc = Builder.GetConstDocument();
+		const FMetasoundFrontendDocument& Doc = Builder.GetConstDocumentChecked();
 		const FMetasoundFrontendGraphClass& GraphClass = Doc.RootGraph;
 		const FMetasoundFrontendClassOutput& OutputBeingRemoved = GraphClass.Interface.Outputs[OutputIndex];
 
@@ -691,7 +682,7 @@ void UMetaSoundSourceBuilder::SetFormat(EMetaSoundOutputAudioFormat OutputFormat
 
 	TArray<FMetasoundFrontendVersion> OutputFormatsToRemove;
 
-	const FMetasoundFrontendDocument& Document = GetConstBuilder().GetConstDocument();
+	const FMetasoundFrontendDocument& Document = GetConstBuilder().GetConstDocumentChecked();
 	for (const FOutputAudioFormatInfoPair& Pair : FormatMap)
 	{
 		const FMetasoundFrontendVersion& FormatVersion = Pair.Value.InterfaceVersion;
@@ -1055,10 +1046,7 @@ UMetaSoundSourceBuilder* UMetaSoundBuilderSubsystem::FindSourceBuilder(FName Bui
 void UMetaSoundBuilderSubsystem::InvalidateDocumentCache(const FMetasoundFrontendClassName& InClassName) const
 {
 	using namespace Metasound::Engine;
-	if (FMetaSoundFrontendDocumentBuilder* Builder = FDocumentBuilderRegistry::GetChecked().FindBuilder(InClassName))
-	{
-		Builder->Reload();
-	}
+	FDocumentBuilderRegistry::GetChecked().ReloadBuilder(InClassName);
 }
 
 bool UMetaSoundBuilderSubsystem::IsInterfaceRegistered(FName InInterfaceName) const
diff --git a/Engine/Plugins/Runtime/Metasound/Source/MetasoundEngine/Private/MetasoundDocumentBuilderRegistry.cpp b/Engine/Plugins/Runtime/Metasound/Source/MetasoundEngine/Private/MetasoundDocumentBuilderRegistry.cpp
--- a/Engine/Plugins/Runtime/Metasound/Source/MetasoundEngine/Private/MetasoundDocumentBuilderRegistry.cpp
+++ b/Engine/Plugins/Runtime/Metasound/Source/MetasoundEngine/Private/MetasoundDocumentBuilderRegistry.cpp
@@ -10,32 +10,25 @@ namespace Metasound::Engine
 {
 	FDocumentBuilderRegistry::~FDocumentBuilderRegistry()
 	{
-#if !NO_LOGGING
-		if (!Builders.IsEmpty())
+		FScopeLock Lock(&BuildersCriticalSection);
+		UE_CLOG(!Builders.IsEmpty(), LogMetaSound, Display, TEXT("BuilderRegistry is shutting down with the following %i active builder entries. Forcefully shutting down:"), Builders.Num());
+		int32 NumStale = 0;
+		for (const TPair<FMetasoundFrontendClassName, TWeakObjectPtr<UMetaSoundBuilderBase>>& Pair : Builders)
 		{
-			TArray<TPair<FMetasoundFrontendClassName, TWeakObjectPtr<UMetaSoundBuilderBase>>> PairsToFinish;
-			int32 NumStale = 0;
-			for (const TPair<FMetasoundFrontendClassName, TWeakObjectPtr<UMetaSoundBuilderBase>>& Pair : Builders)
-			{
-				if (Pair.Value.IsValid())
-				{
-					PairsToFinish.Add(Pair);
-				}
-				else
-				{
-					NumStale++;
-				}
-			}
-
-			UE_LOG(LogMetaSound, Warning, TEXT("BuilderRegistry is shutting down with existing records: %i stale entries never removed, the following %i are still active and will be forcefully shutdown:"), NumStale, PairsToFinish.Num());
-			for (const TPair<FMetasoundFrontendClassName, TWeakObjectPtr<UMetaSoundBuilderBase>>& Pair : PairsToFinish)
+			const bool bIsValid = Pair.Value.IsValid();
+			if (bIsValid)
 			{
+				UE_CLOG(bIsValid, LogMetaSound, Display, TEXT("- %s"), *Pair.Value->GetFullName());
 				constexpr bool bForceUnregister = true;
 				FinishBuilding(Pair.Key, bForceUnregister);
-				UE_LOG(LogMetaSound, Warning, TEXT("- %s"), *Pair.Value->GetFullName());
+			}
+			else
+			{
+				++NumStale;
 			}
 		}
-#endif // !NO_LOGGING
+		UE_CLOG(NumStale > 0, LogMetaSound, Display, TEXT("BuilderRegistry is shutting down with %i stale entries"), NumStale);
+		Builders.Reset();
 	}
 
 #if WITH_EDITORONLY_DATA
@@ -50,9 +43,9 @@ namespace Metasound::Engine
 
 	FMetaSoundFrontendDocumentBuilder* FDocumentBuilderRegistry::FindBuilder(TScriptInterface<IMetaSoundDocumentInterface> MetaSound) const
 	{
-		if (UObject* Object = MetaSound.GetObject())
+		if (UMetaSoundBuilderBase* Builder = FindBuilderObject(MetaSound))
 		{
-			return FindBuilder(MetaSound->GetConstDocument().RootGraph.Metadata.GetClassName());
+			return &Builder->GetBuilder();
 		}
 
 		return nullptr;
@@ -60,8 +53,7 @@ namespace Metasound::Engine
 
 	FMetaSoundFrontendDocumentBuilder* FDocumentBuilderRegistry::FindBuilder(const FMetasoundFrontendClassName& InClassName) const
 	{
-		using namespace Metasound::Engine;
-		if (UMetaSoundBuilderBase* Builder = Builders.FindRef(InClassName).Get())
+		if (UMetaSoundBuilderBase* Builder = FindBuilderObject(InClassName))
 		{
 			return &Builder->GetBuilder();
 		}
@@ -75,7 +67,30 @@ namespace Metasound::Engine
 		{
 			const FMetasoundFrontendDocument& Document = MetaSound->GetConstDocument();
 			const FMetasoundFrontendClassName& ClassName = Document.RootGraph.Metadata.GetClassName();
-			return Builders.FindRef(ClassName).Get();
+			TArray<TWeakObjectPtr<UMetaSoundBuilderBase>> Entries;
+
+			{
+				FScopeLock Lock(&BuildersCriticalSection);
+				Builders.MultiFind(ClassName, Entries);
+			}
+
+			if (!Entries.IsEmpty())
+			{
+				auto IsMetaSound = [MetaSoundObject = MetaSound.GetObject()](const TWeakObjectPtr<UMetaSoundBuilderBase>& BuilderPtr)
+				{
+					if (BuilderPtr.IsValid())
+					{
+						UObject& TestMetaSound = BuilderPtr->GetConstBuilder().CastDocumentObjectChecked<UObject>();
+						return &TestMetaSound == MetaSoundObject;
+					}
+					return false;
+				};
+				const TWeakObjectPtr<UMetaSoundBuilderBase>* Builder = Entries.FindByPredicate(IsMetaSound);
+				if (Builder)
+				{
+					return Builder->Get();
+				}
+			}
 		}
 
 		return nullptr;
@@ -83,7 +98,64 @@ namespace Metasound::Engine
 
 	UMetaSoundBuilderBase* FDocumentBuilderRegistry::FindBuilderObject(const FMetasoundFrontendClassName& ClassName) const
 	{
-		return Builders.FindRef(ClassName).Get();
+		TArray<TWeakObjectPtr<UMetaSoundBuilderBase>> Entries;
+		{
+			FScopeLock Lock(&BuildersCriticalSection);
+			Builders.MultiFind(ClassName, Entries);
+		}
+		if (!Entries.IsEmpty())
+		{
+#if !NO_LOGGING
+			if (Entries.Num() > 1)
+			{
+				UE_LOG(LogMetaSound, Error, TEXT("More than one asset registered with class name '%s': returning builder that may not be associated with desired object! This can happen if multiple assets were improperly copied. See following list of assets needing to be resolved:"),
+					*ClassName.ToString());
+				for (const TWeakObjectPtr<UMetaSoundBuilderBase>& BuilderPtr : Entries)
+				{
+					if (BuilderPtr.IsValid())
+					{
+						UE_LOG(LogMetaSound, Error, TEXT("- %s"), *BuilderPtr->GetConstBuilder().CastDocumentObjectChecked<UObject>().GetPathName());
+					}
+					else
+					{
+						UE_LOG(LogMetaSound, Error, TEXT("- STALE"));
+					}
+				}
+			}
+#endif // !NO_LOGGING
+			return Entries.Last().Get();
+		}
+
+		return nullptr;
+	}
+
+	TArray<UMetaSoundBuilderBase*> FDocumentBuilderRegistry::FindBuilderObjects(const FMetasoundFrontendClassName& ClassName) const
+	{
+		TArray<UMetaSoundBuilderBase*> FoundBuilders;
+		TArray<TWeakObjectPtr<UMetaSoundBuilderBase>> Entries;
+
+		{
+			FScopeLock Lock(&BuildersCriticalSection);
+			Builders.MultiFind(ClassName, Entries);
+		}
+
+		if (!Entries.IsEmpty())
+		{
+			Algo::TransformIf(Entries, FoundBuilders,
+				[](const TWeakObjectPtr<UMetaSoundBuilderBase>& BuilderPtr) { return BuilderPtr.IsValid(); },
+				[](const TWeakObjectPtr<UMetaSoundBuilderBase>& BuilderPtr) { return BuilderPtr.Get(); }
+			);
+		}
+
+		return FoundBuilders;
+	}
+
+	FMetaSoundFrontendDocumentBuilder* FDocumentBuilderRegistry::FindOutermostBuilder(const UObject& InSubObject) const
+	{
+		using namespace Metasound::Frontend;
+		TScriptInterface<IMetaSoundDocumentInterface> DocumentInterface = InSubObject.GetOutermostObject();
+		check(DocumentInterface.GetObject());
+		return FindBuilder(DocumentInterface);
 	}
 
 	bool FDocumentBuilderRegistry::FinishBuilding(const FMetasoundFrontendClassName& InClassName, bool bForceUnregister) const
@@ -91,38 +163,47 @@ namespace Metasound::Engine
 		using namespace Metasound;
 		using namespace Metasound::Engine;
 
-		if (TWeakObjectPtr<UMetaSoundBuilderBase>* BuilderPtr = Builders.Find(InClassName))
+		TArray<UMetaSoundBuilderBase*> FoundBuilders = FindBuilderObjects(InClassName);
+		for (UMetaSoundBuilderBase* Builder : FoundBuilders)
 		{
-			if (UMetaSoundBuilderBase* Builder = BuilderPtr->Get())
+			// If the builder has applied transactions to its document object that are not mirrored in the frontend registry,
+			// unregister version in registry. This will ensure that future requests for the builder's associated asset will
+			// register a fresh version from the object as the transaction history is intrinsically lost once this builder
+			// is destroyed. It is also possible that the DocBuilder's underlying object can be invalid if object was force
+			// deleted, so validity check is necessary.
+			FMetaSoundFrontendDocumentBuilder& DocBuilder = Builder->GetBuilder();
+			if (DocBuilder.IsValid())
 			{
-				// If the builder has applied transactions to its document object that are not mirrored in the frontend registry,
-				// unregister version in registry. This will ensure that future requests for the builder's associated asset will
-				// register a fresh version from the object as the transaction history is intrinsically lost once this builder
-				// is destroyed. It is also possible that the DocBuilder's underlying object can be invalid if object was force
-				// deleted, so validity check is necessary.
-				FMetaSoundFrontendDocumentBuilder& DocBuilder = Builder->GetBuilder();
-				if (DocBuilder.IsValid())
+				if (!IsRunningCookCommandlet())
 				{
-					if (!IsRunningCookCommandlet())
+					const int32 TransactionCount = DocBuilder.GetTransactionCount();
+					const int32 LastTransactionRegistered = Builder->GetLastTransactionRegistered();
+					if (bForceUnregister || LastTransactionRegistered != TransactionCount)
 					{
-						const int32 TransactionCount = DocBuilder.GetTransactionCount();
-						const int32 LastTransactionRegistered = Builder->GetLastTransactionRegistered();
-						if (bForceUnregister || LastTransactionRegistered != TransactionCount)
+						UObject& MetaSound = DocBuilder.CastDocumentObjectChecked<UObject>();
+						if (FMetasoundAssetBase* MetaSoundAsset = IMetasoundUObjectRegistry::Get().GetObjectAsAssetBase(&MetaSound))
 						{
-							UObject& MetaSound = DocBuilder.CastDocumentObjectChecked<UObject>();
-							if (FMetasoundAssetBase* MetaSoundAsset = IMetasoundUObjectRegistry::Get().GetObjectAsAssetBase(&MetaSound))
-							{
-								MetaSoundAsset->UnregisterGraphWithFrontend();
-							}
+							MetaSoundAsset->UnregisterGraphWithFrontend();
 						}
 					}
-					DocBuilder.FinishBuilding();
 				}
-			}	
+				DocBuilder.FinishBuilding();
+			}
+		}
 
-			// Still return true in this case as the builder likely has become inaccessible and may be in the "beginning" of destruction,
-			// so entries is still reporting that it was successfully removed.
-			ensureAlways(Builders.Remove(InClassName));
+		// Still return true in this case as the builder likely has become inaccessible and may be in the "beginning" of destruction,
+		// so entries is still reporting that it was successfully removed.
+		{
+			FScopeLock Lock(&BuildersCriticalSection);
+			return Builders.Remove(InClassName) > 0;
+		}
+	}
+
+	bool FDocumentBuilderRegistry::ReloadBuilder(const FMetasoundFrontendClassName& InClassName) const
+	{
+		if (UMetaSoundBuilderBase* Builder = FindBuilderObject(InClassName))
+		{
+			Builder->Reload();
 			return true;
 		}
 
diff --git a/Engine/Plugins/Runtime/Metasound/Source/MetasoundEngine/Private/MetasoundEngineAsset.h b/Engine/Plugins/Runtime/Metasound/Source/MetasoundEngine/Private/MetasoundEngineAsset.h
--- a/Engine/Plugins/Runtime/Metasound/Source/MetasoundEngine/Private/MetasoundEngineAsset.h
+++ b/Engine/Plugins/Runtime/Metasound/Source/MetasoundEngine/Private/MetasoundEngineAsset.h
@@ -4,7 +4,7 @@
 
 #include "Metasound.h"
 #include "MetasoundAssetManager.h"
-#include "MetasoundBuilderSubsystem.h"
+#include "MetasoundDocumentBuilderRegistry.h"
 #include "MetasoundFrontendDocumentIdGenerator.h"
 #include "MetasoundFrontendRegistryKey.h"
 #include "MetasoundUObjectRegistry.h"
@@ -28,7 +28,7 @@ namespace Metasound
 #if WITH_EDITOR
 		static void PostDuplicate(TScriptInterface<IMetaSoundDocumentInterface> MetaSound, EDuplicateMode::Type InDuplicateMode, FGuid& OutAssetClassID)
 		{
-			using namespace Metasound::Frontend;
+			using namespace Frontend;
 
 			if (InDuplicateMode == EDuplicateMode::Normal)
 			{
@@ -236,4 +236,3 @@ namespace Metasound
 #endif // WITH_EDITORONLY_DATA
 	};
 } // namespace Metasound
-
diff --git a/Engine/Plugins/Runtime/Metasound/Source/MetasoundEngine/Public/MetasoundBuilderBase.cpp b/Engine/Plugins/Runtime/Metasound/Source/MetasoundEngine/Public/MetasoundBuilderBase.cpp
--- a/Engine/Plugins/Runtime/Metasound/Source/MetasoundEngine/Public/MetasoundBuilderBase.cpp
+++ b/Engine/Plugins/Runtime/Metasound/Source/MetasoundEngine/Public/MetasoundBuilderBase.cpp
@@ -45,7 +45,7 @@ void UMetaSoundBuilderBase::BeginDestroy()
 	{
 		if (IDocumentBuilderRegistry* BuilderRegistry = IDocumentBuilderRegistry::Get())
 		{
-			const FMetasoundFrontendClassName& MetaSoundClassName = Builder.GetConstDocument().RootGraph.Metadata.GetClassName();
+			const FMetasoundFrontendClassName& MetaSoundClassName = Builder.GetConstDocumentChecked().RootGraph.Metadata.GetClassName();
 			BuilderRegistry->FinishBuilding(MetaSoundClassName);
 		}
 
@@ -76,7 +76,7 @@ FMetaSoundBuilderNodeOutputHandle UMetaSoundBuilderBase::AddGraphInputNode(FName
 		else
 		{
 			FDocumentIDGenerator& IDGenerator = FDocumentIDGenerator::Get();
-			const FMetasoundFrontendDocument& Doc = GetConstBuilder().GetConstDocument();
+			const FMetasoundFrontendDocument& Doc = GetConstBuilder().GetConstDocumentChecked();
 
 			FMetasoundFrontendClassInput Description;
 			Description.Name = Name;
@@ -122,7 +122,7 @@ FMetaSoundBuilderNodeInputHandle UMetaSoundBuilderBase::AddGraphOutputNode(FName
 		else
 		{
 			FDocumentIDGenerator& IDGenerator = FDocumentIDGenerator::Get();
-			const FMetasoundFrontendDocument& Doc = GetConstBuilder().GetConstDocument();
+			const FMetasoundFrontendDocument& Doc = GetConstBuilder().GetConstDocumentChecked();
 
 			FMetasoundFrontendClassOutput Description;
 			Description.Name = Name;
@@ -236,7 +236,7 @@ void UMetaSoundBuilderBase::BuildInternal(TScriptInterface<IMetaSoundDocumentInt
 	FMetaSoundFrontendDocumentBuilder NewDocBuilder(NewMetaSound);
 
 	constexpr bool bResetVersion = false;
-	NewDocBuilder.InitDocument(&GetConstBuilder().GetConstDocument(), InDocClassName, bResetVersion);
+	NewDocBuilder.InitDocument(&GetConstBuilder().GetConstDocumentChecked(), InDocClassName, bResetVersion);
 	NewMetaSound->ConformObjectToDocument();
 }
 
@@ -735,13 +735,13 @@ UObject* UMetaSoundBuilderBase::GetReferencedPresetAsset() const
 
 	// Find the single external node which is the referenced preset asset, 
 	// and find the asset with its registry key 
-	auto FindExternalNode = [this](const FMetasoundFrontendNode& Node) 
+	auto FindExternalNode = [this](const FMetasoundFrontendNode& Node)
 	{
 		const FMetasoundFrontendClass* Class = Builder.FindDependency(Node.ClassID);
 		check(Class);
 		return Class->Metadata.GetType() == EMetasoundFrontendClassType::External;
 	};
-	const FMetasoundFrontendNode* Node = Builder.GetConstDocument().RootGraph.Graph.Nodes.FindByPredicate(FindExternalNode);
+	const FMetasoundFrontendNode* Node = Builder.FindConstBuildGraphChecked().Nodes.FindByPredicate(FindExternalNode);
 	if (Node != nullptr)
 	{
 		const FMetasoundFrontendClass* NodeClass = Builder.FindDependency(Node->ClassID);
@@ -834,7 +834,7 @@ void UMetaSoundBuilderBase::OnDependencyAdded(int32 Index)
 	using namespace Metasound;
 	using namespace Metasound::Frontend;
 
-	const FMetasoundFrontendClass& NewDependency = Builder.GetConstDocument().Dependencies[Index];
+	const FMetasoundFrontendClass& NewDependency = Builder.GetConstDocumentChecked().Dependencies[Index];
 	if (NewDependency.Metadata.GetType() == EMetasoundFrontendClassType::External)
 	{
 		const FAssetKey AssetKey(NewDependency.Metadata);
@@ -850,7 +850,7 @@ void UMetaSoundBuilderBase::OnRemoveSwappingDependency(int32 Index, int32 LastIn
 	using namespace Metasound;
 	using namespace Metasound::Frontend;
 
-	const FMetasoundFrontendClass& NewDependency = Builder.GetConstDocument().Dependencies[Index];
+	const FMetasoundFrontendClass& NewDependency = Builder.GetConstDocumentChecked().Dependencies[Index];
 	if (NewDependency.Metadata.GetType() == EMetasoundFrontendClassType::External)
 	{
 		const FAssetKey AssetKey(NewDependency.Metadata);
@@ -1079,7 +1079,7 @@ FMetasoundFrontendVersion UMetaSoundBuilderBase::FindNodeClassVersion(const FMet
 
 FMetasoundFrontendClassName UMetaSoundBuilderBase::GetRootGraphClassName() const
 {
-	return Builder.GetConstDocument().RootGraph.Metadata.GetClassName();
+	return Builder.GetConstDocumentChecked().RootGraph.Metadata.GetClassName();
 }
 
 void UMetaSoundBuilderBase::GetNodeInputData(const FMetaSoundBuilderNodeInputHandle& InputHandle, FName& Name, FName& DataType, EMetaSoundBuilderResult& OutResult)
diff --git a/Engine/Plugins/Runtime/Metasound/Source/MetasoundEngine/Public/MetasoundBuilderBase.h b/Engine/Plugins/Runtime/Metasound/Source/MetasoundEngine/Public/MetasoundBuilderBase.h
--- a/Engine/Plugins/Runtime/Metasound/Source/MetasoundEngine/Public/MetasoundBuilderBase.h
+++ b/Engine/Plugins/Runtime/Metasound/Source/MetasoundEngine/Public/MetasoundBuilderBase.h
@@ -517,11 +517,10 @@ protected:
 	UPROPERTY()
 	FMetaSoundFrontendDocumentBuilder Builder;
 
-	// ClassName of associated document object.  Cached in case DocumentObject is deleted but continues to provide an identifier after destruction.
-	UPROPERTY()
+#if WITH_EDITORONLY_DATA
+	UPROPERTY(meta = (DeprecatedProperty, DeprecationMessage = "5.5 - No longer used. ClassName should be queried from associated FrontendBuilder's MetaSound"))
 	FMetasoundFrontendClassName ClassName;
 
-#if WITH_EDITORONLY_DATA
 	UPROPERTY(meta = (DeprecatedProperty, DeprecationMessage = "5.4 - All source builders now operate on an underlying document source document that is also used to audition."))
 	bool bIsAttached = false;
 #endif // WITH_EDITORONLY_DATA
@@ -530,7 +529,6 @@ private:
 	void BuildInternal(TScriptInterface<IMetaSoundDocumentInterface> NewMetaSound, const FMetasoundFrontendClassName* InDocClassName) const;
 	void OnDependencyAdded(int32 Index);
 	void OnRemoveSwappingDependency(int32 Index, int32 LastIndex);
-	void OnRenamingDependencyInternal(const int32 Index, const FMetasoundFrontendClassName& NewName);
 
 private:
 	int32 LastTransactionRegistered = 0;
diff --git a/Engine/Plugins/Runtime/Metasound/Source/MetasoundEngine/Public/MetasoundBuilderSubsystem.h b/Engine/Plugins/Runtime/Metasound/Source/MetasoundEngine/Public/MetasoundBuilderSubsystem.h
--- a/Engine/Plugins/Runtime/Metasound/Source/MetasoundEngine/Public/MetasoundBuilderSubsystem.h
+++ b/Engine/Plugins/Runtime/Metasound/Source/MetasoundEngine/Public/MetasoundBuilderSubsystem.h
@@ -154,7 +154,7 @@ private:
 	TMap<FName, TObjectPtr<UMetaSoundBuilderBase>> NamedBuilders;
 
 public:	
-	UE_DEPRECATED(5.5, "Direct call of 'Reload' can be applied to builder directly")
+	UE_DEPRECATED(5.5, "Call 'ReloadBuilder' in IDocumentBuilderRegistry instead ")
 	virtual void InvalidateDocumentCache(const FMetasoundFrontendClassName& InClassName) const;
 
 	static UMetaSoundBuilderSubsystem* Get();
diff --git a/Engine/Plugins/Runtime/Metasound/Source/MetasoundEngine/Public/MetasoundDocumentBuilderRegistry.h b/Engine/Plugins/Runtime/Metasound/Source/MetasoundEngine/Public/MetasoundDocumentBuilderRegistry.h
--- a/Engine/Plugins/Runtime/Metasound/Source/MetasoundEngine/Public/MetasoundDocumentBuilderRegistry.h
+++ b/Engine/Plugins/Runtime/Metasound/Source/MetasoundEngine/Public/MetasoundDocumentBuilderRegistry.h
@@ -13,7 +13,10 @@ namespace Metasound::Engine
 {
 	class METASOUNDENGINE_API FDocumentBuilderRegistry : public Frontend::IDocumentBuilderRegistry
 	{
-		mutable TMap<FMetasoundFrontendClassName, TWeakObjectPtr<UMetaSoundBuilderBase>> Builders;
+		mutable TMultiMap<FMetasoundFrontendClassName, TWeakObjectPtr<UMetaSoundBuilderBase>> Builders;
+
+		// Critical section primarily for allowing builder collection mutation during async loading of MetaSound assets.
+		mutable FCriticalSection BuildersCriticalSection;
 
 	public:
 		FDocumentBuilderRegistry() = default;
@@ -37,7 +40,7 @@ namespace Metasound::Engine
 			TObjectPtr<BuilderClass> NewBuilder = NewObject<BuilderClass>(TransientPackage, ObjectName, NewObjectFlags);
 			check(NewBuilder);
 			NewBuilder->Initialize();
-			const FMetasoundFrontendDocument& Document = NewBuilder->GetConstBuilder().GetConstDocument();
+			const FMetasoundFrontendDocument& Document = NewBuilder->GetConstBuilder().GetConstDocumentChecked();
 			const FMetasoundFrontendClassName& ClassName = Document.RootGraph.Metadata.GetClassName();
 			Builders.Add(ClassName, NewBuilder);
 			return *NewBuilder.Get();
@@ -48,27 +51,21 @@ namespace Metasound::Engine
 		BuilderClass& FindOrBeginBuilding(UObject& InMetaSoundObject) const
 		{
 			check(InMetaSoundObject.IsAsset());
-			checkf(IsInGameThread(), TEXT("Asset MetaSound Builder cannot be created in non-game thread as it may result in UObject creation"));
 
 			TScriptInterface<IMetaSoundDocumentInterface> DocInterface = &InMetaSoundObject;
 			check(DocInterface.GetObject());
 
-			const FMetasoundFrontendDocument& Document = DocInterface->GetConstDocument();
-			const FMetasoundFrontendClassName& FullClassName = Document.RootGraph.Metadata.GetClassName();
-
-			if (FullClassName.IsValid())
+			if (UMetaSoundBuilderBase* Builder = FindBuilderObject(&InMetaSoundObject))
 			{
-				TWeakObjectPtr<UMetaSoundBuilderBase> Builder = Builders.FindRef(FullClassName);
-				if (Builder.IsValid())
-				{
-					return *CastChecked<BuilderClass>(Builder.Get());
-				}
+				return *CastChecked<BuilderClass>(Builder);
 			}
 
 			TObjectPtr<UMetaSoundBuilderBase> NewBuilder = CastChecked<UMetaSoundBuilderBase>(NewObject<UObject>(&InMetaSoundObject, &DocInterface->GetBuilderUClass()));
 			FMetaSoundFrontendDocumentBuilder& BuilderRef = NewBuilder->GetBuilder();
 			BuilderRef = FMetaSoundFrontendDocumentBuilder(DocInterface);
 
+			const FMetasoundFrontendDocument& Document = DocInterface->GetConstDocument();
+			const FMetasoundFrontendClassName& FullClassName = Document.RootGraph.Metadata.GetClassName();
 			if (!FullClassName.IsValid())
 			{
 				BuilderRef.InitDocument();
@@ -88,9 +85,15 @@ namespace Metasound::Engine
 
 		virtual FMetaSoundFrontendDocumentBuilder* FindBuilder(TScriptInterface<IMetaSoundDocumentInterface> MetaSound) const override;
 		virtual FMetaSoundFrontendDocumentBuilder* FindBuilder(const FMetasoundFrontendClassName& InClassName) const override;
+
+		virtual FMetaSoundFrontendDocumentBuilder* FindOutermostBuilder(const UObject& InSubObject) const override;
+
 		virtual bool FinishBuilding(const FMetasoundFrontendClassName& InClassName, bool bForceUnregister = false) const override;
 
 		UMetaSoundBuilderBase* FindBuilderObject(TScriptInterface<const IMetaSoundDocumentInterface> MetaSound) const;
 		UMetaSoundBuilderBase* FindBuilderObject(const FMetasoundFrontendClassName& ClassName) const;
+		TArray<UMetaSoundBuilderBase*> FindBuilderObjects(const FMetasoundFrontendClassName& ClassName) const;
+
+		bool ReloadBuilder(const FMetasoundFrontendClassName& InClassName) const override;
 	};
 } // namespace Metasound::Engine
\ No newline at end of file
diff --git a/Engine/Plugins/Runtime/Metasound/Source/MetasoundFrontend/Private/MetasoundAssetBase.cpp b/Engine/Plugins/Runtime/Metasound/Source/MetasoundFrontend/Private/MetasoundAssetBase.cpp
--- a/Engine/Plugins/Runtime/Metasound/Source/MetasoundFrontend/Private/MetasoundAssetBase.cpp
+++ b/Engine/Plugins/Runtime/Metasound/Source/MetasoundFrontend/Private/MetasoundAssetBase.cpp
@@ -592,7 +592,7 @@ void FMetasoundAssetBase::CacheRegistryMetadata()
 	}
 
 	// 2. Copy metadata for inputs/outputs managed by interfaces, removing them from maps generated
-	auto CacheInterfaceMetadata = [](const FMetasoundFrontendVertexMetadata & InRegistryMetadata, FMetasoundFrontendVertexMetadata& OutMetadata)
+	auto CacheInterfaceMetadata = [](const FMetasoundFrontendVertexMetadata& InRegistryMetadata, FMetasoundFrontendVertexMetadata& OutMetadata)
 	{
 		const int32 CachedSortOrderIndex = OutMetadata.SortOrderIndex;
 		OutMetadata = InRegistryMetadata;
@@ -718,7 +718,6 @@ bool FMetasoundAssetBase::IsReferencedAsset(const FMetasoundAssetBase& InAsset)
 		ensureAlways(IMetaSoundAssetManager::GetChecked().TryLoadReferencedAssets(ChildAsset, ChildRefs));
 		Algo::Transform(ChildRefs, Children, [](FMetasoundAssetBase* Child) { return Child; });
 		return Children;
-
 	});
 
 	return bIsReferenced;
@@ -898,9 +897,11 @@ void FMetasoundAssetBase::RebuildReferencedAssetClasses()
 	using namespace Metasound::Frontend;
 
 	IMetaSoundAssetManager& AssetManager = IMetaSoundAssetManager::GetChecked();
-	AssetManager.AddAssetReferences(*this);
-	TSet<IMetaSoundAssetManager::FAssetInfo> ReferencedAssetClasses = AssetManager.GetReferencedAssetClasses(*this);
-	SetReferencedAssetClasses(MoveTemp(ReferencedAssetClasses));
+	if (AssetManager.AddAssetReferences(*this))
+	{
+		TSet<IMetaSoundAssetManager::FAssetInfo> ReferencedAssetClasses = AssetManager.GetReferencedAssetClasses(*this);
+		SetReferencedAssetClasses(MoveTemp(ReferencedAssetClasses));
+	}
 }
 #endif // WITH_EDITOR
 
diff --git a/Engine/Plugins/Runtime/Metasound/Source/MetasoundFrontend/Private/MetasoundFrontendDocument.cpp b/Engine/Plugins/Runtime/Metasound/Source/MetasoundFrontend/Private/MetasoundFrontendDocument.cpp
--- a/Engine/Plugins/Runtime/Metasound/Source/MetasoundFrontend/Private/MetasoundFrontendDocument.cpp
+++ b/Engine/Plugins/Runtime/Metasound/Source/MetasoundFrontend/Private/MetasoundFrontendDocument.cpp
@@ -870,6 +870,127 @@ FMetasoundFrontendGraphClass::FMetasoundFrontendGraphClass()
 	Metadata.SetType(EMetasoundFrontendClassType::Graph);
 }
 
+void FMetasoundFrontendGraphClass::AddNewGraphPage(const FGuid& InPageID)
+{
+	checkf(ContainsGraphPage(InPageID), TEXT("Cannot add new graph page with existing PageID"));
+	checkf(InPageID.IsValid(), TEXT("Cannot add graph with invalid PageID"))
+
+	FMetasoundFrontendGraph& NewGraph = PagedGraphs.AddDefaulted_GetRef();
+	NewGraph.PageID = InPageID;
+}
+
+bool FMetasoundFrontendGraphClass::ContainsGraphPage(const FGuid& InPageID) const
+{
+	if (InPageID == Graph.PageID)
+	{
+		return true;
+	}
+
+	auto MatchesPageID = [&InPageID](const FMetasoundFrontendGraph& Iter) { return Iter.PageID == InPageID; };
+	return PagedGraphs.ContainsByPredicate(MatchesPageID);
+}
+
+void FMetasoundFrontendGraphClass::DuplicateLastGraphPage(const FGuid& InPageID)
+{
+	checkf(ContainsGraphPage(InPageID), TEXT("Cannot add new graph page with existing PageID"));
+	checkf(InPageID.IsValid(), TEXT("Cannot add graph with invalid PageID"))
+
+	FMetasoundFrontendGraph* ToDuplicate = PagedGraphs.IsEmpty() ? &Graph : &PagedGraphs.Last();
+	FMetasoundFrontendGraph& NewGraph = PagedGraphs.Add_GetRef(*ToDuplicate);
+	NewGraph.PageID = InPageID;
+}
+
+void FMetasoundFrontendGraphClass::RemoveAllGraphPages()
+{
+	PagedGraphs.Empty();
+	Graph.PageID = FGuid();
+}
+
+bool FMetasoundFrontendGraphClass::RemoveGraphPage(const FGuid& InPageID)
+{
+	const bool bRemoved = PagedGraphs.RemoveAllSwap([&InPageID](const FMetasoundFrontendGraph& Iter)
+	{
+		return Iter.PageID == InPageID;
+	}, EAllowShrinking::Yes) > 0;
+
+	return bRemoved;
+}
+
+FMetasoundFrontendGraph* FMetasoundFrontendGraphClass::FindGraph(const FGuid& InPageID)
+{
+	if (InPageID == Graph.PageID)
+	{
+		return &Graph;
+	}
+
+	auto MatchesPageID = [this, &InPageID](const FMetasoundFrontendGraph& Iter) { return Iter.PageID == InPageID; };
+	FMetasoundFrontendGraph* PageGraph = PagedGraphs.FindByPredicate(MatchesPageID);
+	return PageGraph;
+}
+
+FMetasoundFrontendGraph& FMetasoundFrontendGraphClass::FindGraphChecked(const FGuid& InPageID)
+{
+	FMetasoundFrontendGraph* FoundGraph = FindGraph(InPageID);
+	check(FoundGraph);
+	return *FoundGraph;
+}
+
+const FMetasoundFrontendGraph* FMetasoundFrontendGraphClass::FindConstGraph(const FGuid& InPageID) const
+{
+	if (InPageID == Graph.PageID)
+	{
+		return &Graph;
+	}
+
+	auto MatchesPageID = [this, &InPageID](const FMetasoundFrontendGraph& Iter) { return Iter.PageID == InPageID; };
+	const FMetasoundFrontendGraph* PageGraph = PagedGraphs.FindByPredicate(MatchesPageID);
+	return PageGraph;
+}
+
+const FMetasoundFrontendGraph& FMetasoundFrontendGraphClass::FindConstGraphChecked(const FGuid& InPageID) const
+{
+	const FMetasoundFrontendGraph* FoundGraph = FindConstGraph(InPageID);
+	check(FoundGraph);
+	return *FoundGraph;
+}
+
+FMetasoundFrontendGraph& FMetasoundFrontendGraphClass::GetDefaultGraph()
+{
+	return FindGraphChecked(FGuid());
+}
+
+const FMetasoundFrontendGraph& FMetasoundFrontendGraphClass::GetConstDefaultGraph() const
+{
+	return FindConstGraphChecked(FGuid());
+}
+
+void FMetasoundFrontendGraphClass::IterateGraphPages(TFunctionRef<void(FMetasoundFrontendGraph&)> IterFunc)
+{
+	IterFunc(Graph);
+
+	for (FMetasoundFrontendGraph& Iter : PagedGraphs)
+	{
+		IterFunc(Iter);
+	}
+}
+
+void FMetasoundFrontendGraphClass::IterateGraphPages(TFunctionRef<void(const FMetasoundFrontendGraph&)> IterFunc) const
+{
+	IterFunc(Graph);
+
+	for (const FMetasoundFrontendGraph& Iter : PagedGraphs)
+	{
+		IterFunc(Iter);
+	}
+}
+
+void FMetasoundFrontendGraphClass::ResetGraphs()
+{
+	PagedGraphs.Empty();
+	Graph.Nodes.Reset();
+	Graph.Edges.Reset();
+}
+
 FMetasoundFrontendVersionNumber FMetasoundFrontendDocument::GetMaxVersion()
 {
 	return Metasound::Frontend::GetMaxDocumentVersion();
@@ -968,12 +1089,12 @@ namespace Metasound::Frontend
 	{
 		ForEachLiteral(static_cast<const FMetasoundFrontendClass&>(InGraphClass), OnLiteral);
 
-		for (const FMetasoundFrontendNode& Node : InGraphClass.Graph.Nodes)
+		for (const FMetasoundFrontendNode& Node : InGraphClass.GetConstDefaultGraph().Nodes)
 		{
 			ForEachLiteral(Node, OnLiteral);
 		}
 
-		for (const FMetasoundFrontendVariable& Variable : InGraphClass.Graph.Variables)
+		for (const FMetasoundFrontendVariable& Variable : InGraphClass.GetConstDefaultGraph().Variables)
 		{
 			OnLiteral(Variable.TypeName, Variable.Literal);
 		}
diff --git a/Engine/Plugins/Runtime/Metasound/Source/MetasoundFrontend/Private/MetasoundFrontendDocumentAccessPtr.cpp b/Engine/Plugins/Runtime/Metasound/Source/MetasoundFrontend/Private/MetasoundFrontendDocumentAccessPtr.cpp
--- a/Engine/Plugins/Runtime/Metasound/Source/MetasoundFrontend/Private/MetasoundFrontendDocumentAccessPtr.cpp
+++ b/Engine/Plugins/Runtime/Metasound/Source/MetasoundFrontend/Private/MetasoundFrontendDocumentAccessPtr.cpp
@@ -1,12 +1,44 @@
 // Copyright Epic Games, Inc. All Rights Reserved.
 
 #include "MetasoundFrontendDocumentAccessPtr.h"
+
+#include "MetasoundDocumentInterface.h"
 #include "MetasoundFrontendRegistries.h"
+#include "MetasoundFrontendDocumentBuilder.h"
 
 namespace Metasound
 {
 	namespace Frontend
 	{
+		namespace BuilderPrivate
+		{
+			FMetasoundFrontendGraph& FindBuildGraphChecked(FMetasoundFrontendGraphClass& InGraphClass)
+			{
+				if (IDocumentBuilderRegistry* BuilderRegistry = IDocumentBuilderRegistry::Get())
+				{
+					if (FMetaSoundFrontendDocumentBuilder* Builder = BuilderRegistry->FindBuilder(InGraphClass.Metadata.GetClassName()))
+					{
+						return Builder->FindBuildGraphChecked();
+					}
+				}
+
+				return InGraphClass.GetDefaultGraph();
+			}
+
+			const FMetasoundFrontendGraph& FindConstBuildGraphChecked(const FMetasoundFrontendGraphClass& InGraphClass)
+			{
+				if (IDocumentBuilderRegistry* BuilderRegistry = IDocumentBuilderRegistry::Get())
+				{
+					if (FMetaSoundFrontendDocumentBuilder* Builder = BuilderRegistry->FindBuilder(InGraphClass.Metadata.GetClassName()))
+					{
+						return Builder->FindConstBuildGraphChecked();
+					}
+				}
+
+				return InGraphClass.GetConstDefaultGraph();
+			}
+		}
+
 		namespace MetasoundFrontendDocumentAccessPtrPrivate
 		{
 			/** TFindInArray is a callable object for finding child objects with a parent 
@@ -215,22 +247,22 @@ namespace Metasound
 			/** A callable object for accessing arrays of nodes from a graph class. */
 			struct FGetGraphClassNodes
 			{
-			 	TArray<FMetasoundFrontendNode>& operator()(FMetasoundFrontendGraphClass& InGraphClass) const { return InGraphClass.Graph.Nodes; }
-			 	const TArray<FMetasoundFrontendNode>& operator()(const FMetasoundFrontendGraphClass& InGraphClass) const { return InGraphClass.Graph.Nodes; }
+			 	TArray<FMetasoundFrontendNode>& operator()(FMetasoundFrontendGraphClass& InGraphClass) const { return BuilderPrivate::FindBuildGraphChecked(InGraphClass).Nodes; }
+			 	const TArray<FMetasoundFrontendNode>& operator()(const FMetasoundFrontendGraphClass& InGraphClass) const { return BuilderPrivate::FindConstBuildGraphChecked(InGraphClass).Nodes; }
 			};
 
 			/** A callable object for accessing arrays of variables from a graph class. */
 			/*
 			struct FGetGraphClassVariables
 			{
-			 	TMap<FGuid, FMetasoundFrontendVariable>& operator()(FMetasoundFrontendGraphClass& InGraphClass) const { return InGraphClass.Graph.Variables; }
-			 	const TMap<FGuid, FMetasoundFrontendVariable>& operator()(const FMetasoundFrontendGraphClass& InGraphClass) const { return InGraphClass.Graph.Variables; }
+			 	TMap<FGuid, FMetasoundFrontendVariable>& operator()(FMetasoundFrontendGraphClass& InGraphClass) const { return BuilderPrivate::FindBuildGraphChecked(InGraphClass).Variables; }
+			 	const TMap<FGuid, FMetasoundFrontendVariable>& operator()(const FMetasoundFrontendGraphClass& InGraphClass) const { return BuilderPrivate::FindConstBuildGraphChecked(InGraphClass).Variables; }
 			};
 			*/
 			struct FGetGraphClassVariables
 			{
-			 	TArray<FMetasoundFrontendVariable>& operator()(FMetasoundFrontendGraphClass& InGraphClass) const { return InGraphClass.Graph.Variables; }
-			 	const TArray<FMetasoundFrontendVariable>& operator()(const FMetasoundFrontendGraphClass& InGraphClass) const { return InGraphClass.Graph.Variables; }
+			 	TArray<FMetasoundFrontendVariable>& operator()(FMetasoundFrontendGraphClass& InGraphClass) const { return BuilderPrivate::FindBuildGraphChecked(InGraphClass).Variables; }
+			 	const TArray<FMetasoundFrontendVariable>& operator()(const FMetasoundFrontendGraphClass& InGraphClass) const { return BuilderPrivate::FindConstBuildGraphChecked(InGraphClass).Variables; }
 			};
 
 			/** A callable object for accessing arrays of subgraphs from a document. */
@@ -666,7 +698,7 @@ namespace Metasound
 		{
 			auto FindGraph = [](FMetasoundFrontendGraphClass& InGraphClass) -> FMetasoundFrontendGraph* 
 			{ 
-				return &InGraphClass.Graph; 
+				return &BuilderPrivate::FindBuildGraphChecked(InGraphClass); 
 			};
 
 			return GetMemberAccessPtr<FGraphAccessPtr>(FindGraph);
@@ -738,7 +770,7 @@ namespace Metasound
 		{
 			auto FindGraph = [](const FMetasoundFrontendGraphClass& InGraphClass) -> const FMetasoundFrontendGraph* 
 			{ 
-				return &InGraphClass.Graph; 
+				return &BuilderPrivate::FindConstBuildGraphChecked(InGraphClass); 
 			};
 
 			return GetMemberAccessPtr<FConstGraphAccessPtr>(FindGraph);
@@ -810,7 +842,7 @@ namespace Metasound
 		{
 			auto FindGraph = [](const FMetasoundFrontendGraphClass& InGraphClass) -> const FMetasoundFrontendGraph* 
 			{ 
-				return &InGraphClass.Graph; 
+				return &BuilderPrivate::FindConstBuildGraphChecked(InGraphClass); 
 			};
 
 			return GetMemberAccessPtr<FConstGraphAccessPtr>(FindGraph);
diff --git a/Engine/Plugins/Runtime/Metasound/Source/MetasoundFrontend/Private/MetasoundFrontendDocumentBuilder.cpp b/Engine/Plugins/Runtime/Metasound/Source/MetasoundFrontend/Private/MetasoundFrontendDocumentBuilder.cpp
--- a/Engine/Plugins/Runtime/Metasound/Source/MetasoundFrontend/Private/MetasoundFrontendDocumentBuilder.cpp
+++ b/Engine/Plugins/Runtime/Metasound/Source/MetasoundFrontend/Private/MetasoundFrontendDocumentBuilder.cpp
@@ -503,7 +503,7 @@ namespace Metasound::Frontend
 #if WITH_EDITORONLY_DATA
 				if (bAddedVertices && Options.bSetDefaultNodeLocations && !IsRunningCookCommandlet())
 				{
-					TArray<FMetasoundFrontendNode>& Nodes = Document.RootGraph.Graph.Nodes;
+					TArray<FMetasoundFrontendNode>& Nodes = OutBuilder.FindBuildGraphChecked().Nodes;
 					// Sort/Place Inputs
 					{
 						TSet<FName> NamesToSort;
@@ -701,7 +701,7 @@ const FMetasoundFrontendClass* FMetaSoundFrontendDocumentBuilder::AddDependency(
 {
 	using namespace Metasound::Frontend;
 
-	FMetasoundFrontendDocument& Document = GetDocument();
+	FMetasoundFrontendDocument& Document = GetDocumentChecked();
 	const FMetasoundFrontendClass* Dependency = nullptr;
 
 	FMetasoundFrontendClass NewDependency = InClass;
@@ -736,8 +736,8 @@ void FMetaSoundFrontendDocumentBuilder::AddEdge(FMetasoundFrontendEdge&& InNewEd
 	}
 #endif // DO_CHECK
 
-	FMetasoundFrontendDocument& Document = GetDocument();
-	FMetasoundFrontendGraph& Graph = Document.RootGraph.Graph;
+	FMetasoundFrontendDocument& Document = GetDocumentChecked();
+	FMetasoundFrontendGraph& Graph = FindBuildGraphChecked();
 	Graph.Edges.Add(MoveTemp(InNewEdge));
 	const int32 NewIndex = Graph.Edges.Num() - 1;
 	DocumentDelegates->EdgeDelegates.OnEdgeAdded.Broadcast(NewIndex);
@@ -747,8 +747,9 @@ bool FMetaSoundFrontendDocumentBuilder::AddNamedEdges(const TSet<Metasound::Fron
 {
 	using namespace Metasound::Frontend;
 
-	FMetasoundFrontendDocument& Document = GetDocument();
+	FMetasoundFrontendDocument& Document = GetDocumentChecked();
 	const IDocumentGraphNodeCache& NodeCache = DocumentCache->GetNodeCache();
+	FMetasoundFrontendGraph& Graph = FindBuildGraphChecked();
 
 	if (OutNewEdges)
 	{
@@ -767,7 +768,6 @@ bool FMetaSoundFrontendDocumentBuilder::AddNamedEdges(const TSet<Metasound::Fron
 	TArray<FNewEdgeData> EdgesToAdd;
 	for (const FNamedEdge& Edge : EdgesToMake)
 	{
-		const FMetasoundFrontendGraph& Graph = Document.RootGraph.Graph;
 		const FMetasoundFrontendVertex* OutputVertex = NodeCache.FindOutputVertex(Edge.OutputNodeID, Edge.OutputName);
 		const FMetasoundFrontendVertex* InputVertex = NodeCache.FindInputVertex(Edge.InputNodeID, Edge.InputName);
 
@@ -787,7 +787,7 @@ bool FMetaSoundFrontendDocumentBuilder::AddNamedEdges(const TSet<Metasound::Fron
 		}
 	}
 
-	const TArray<FMetasoundFrontendEdge>& Edges = GetDocument().RootGraph.Graph.Edges;
+	const TArray<FMetasoundFrontendEdge>& Edges = FindConstBuildGraphChecked().Edges;
 	const int32 LastIndex = Edges.Num() - 1;
 	for (FNewEdgeData& EdgeToAdd : EdgesToAdd)
 	{
@@ -836,7 +836,7 @@ bool FMetaSoundFrontendDocumentBuilder::AddNamedEdges(const TSet<Metasound::Fron
 			FMetasoundFrontendEdge EdgeToRemove;
 			if (const int32* EdgeIndex = DocumentCache->GetEdgeCache().FindEdgeIndexToNodeInput(EdgeToAdd.NewEdge.ToNodeID, EdgeToAdd.NewEdge.ToVertexID))
 			{
-				EdgeToRemove = Document.RootGraph.Graph.Edges[*EdgeIndex];
+				EdgeToRemove = Graph.Edges[*EdgeIndex];
 			}
 
 			const FMetasoundFrontendVertex* Input = FindNodeInput(EdgeToAdd.NewEdge.ToNodeID, EdgeToAdd.NewEdge.ToVertexID);
@@ -898,7 +898,7 @@ bool FMetaSoundFrontendDocumentBuilder::AddEdgesFromMatchingInterfaceNodeOutputs
 
 	const IDocumentGraphNodeCache& NodeCache = DocumentCache->GetNodeCache();
 	const IDocumentGraphInterfaceCache& InterfaceCache = DocumentCache->GetInterfaceCache();
-	const TSet<FMetasoundFrontendVersion> CommonInterfaces = NodeInterfaces.Intersect(GetDocument().Interfaces);
+	const TSet<FMetasoundFrontendVersion> CommonInterfaces = NodeInterfaces.Intersect(GetDocumentChecked().Interfaces);
 
 	TSet<FNamedEdge> EdgesToMake;
 	for (const FMetasoundFrontendVersion& Version : CommonInterfaces)
@@ -908,7 +908,7 @@ bool FMetaSoundFrontendDocumentBuilder::AddEdgesFromMatchingInterfaceNodeOutputs
 		{
 			Algo::Transform(RegistryEntry->GetInterface().Outputs, EdgesToMake, [this, &NodeCache, &InterfaceCache, InNodeID](const FMetasoundFrontendClassOutput& Output)
 			{
-				const FMetasoundFrontendGraph& Graph = GetDocument().RootGraph.Graph;
+				const FMetasoundFrontendGraph& Graph = FindConstBuildGraphChecked();
 				const FMetasoundFrontendVertex* NodeVertex = NodeCache.FindOutputVertex(InNodeID, Output.Name);
 				check(NodeVertex);
 				const FMetasoundFrontendClassOutput* OutputClass = InterfaceCache.FindOutput(Output.Name);
@@ -942,7 +942,7 @@ bool FMetaSoundFrontendDocumentBuilder::AddEdgesFromMatchingInterfaceNodeInputsT
 
 	const IDocumentGraphNodeCache& NodeCache = DocumentCache->GetNodeCache();
 	const IDocumentGraphInterfaceCache& InterfaceCache = DocumentCache->GetInterfaceCache();
-	const TSet<FMetasoundFrontendVersion> CommonInterfaces = NodeInterfaces.Intersect(GetDocument().Interfaces);
+	const TSet<FMetasoundFrontendVersion> CommonInterfaces = NodeInterfaces.Intersect(GetDocumentChecked().Interfaces);
 
 	TSet<FNamedEdge> EdgesToMake;
 	for (const FMetasoundFrontendVersion& Version : CommonInterfaces)
@@ -952,7 +952,7 @@ bool FMetaSoundFrontendDocumentBuilder::AddEdgesFromMatchingInterfaceNodeInputsT
 		{
 			Algo::Transform(RegistryEntry->GetInterface().Inputs, EdgesToMake, [this, &NodeCache, &InterfaceCache, InNodeID](const FMetasoundFrontendClassInput& Input)
 			{
-				const FMetasoundFrontendGraph& Graph = GetDocument().RootGraph.Graph;
+				const FMetasoundFrontendGraph& Graph = FindConstBuildGraphChecked();
 				const FMetasoundFrontendVertex* NodeVertex = NodeCache.FindInputVertex(InNodeID, Input.Name);
 				check(NodeVertex);
 				const FMetasoundFrontendClassInput* InputClass = InterfaceCache.FindInput(Input.Name);
@@ -973,7 +973,7 @@ const FMetasoundFrontendNode* FMetaSoundFrontendDocumentBuilder::AddGraphInput(c
 {
 	using namespace Metasound::Frontend;
 
-	FMetasoundFrontendDocument& Document = GetDocument();
+	FMetasoundFrontendDocument& Document = GetDocumentChecked();
 	FMetasoundFrontendGraphClass& RootGraph = Document.RootGraph;
 
 	checkf(InClassInput.NodeID.IsValid(), TEXT("Unassigned NodeID when adding graph input"));
@@ -1014,12 +1014,25 @@ const FMetasoundFrontendNode* FMetaSoundFrontendDocumentBuilder::AddGraphInput(c
 			// auditioned MetaSound.
 			DocumentBuilderPrivate::SetDefaultLiteralOnInputNode(InOutNode, InClassInput);
 		};
-		if (FMetasoundFrontendNode* NewNode = AddNodeInternal(Class.Metadata, FinalizeNode, InClassInput.NodeID))
+
+		// Must add input node to all paged graphs to maintain API parity for all page implementations
+		FMetasoundFrontendNode* NewNode = nullptr;
+		Document.RootGraph.IterateGraphPages([this, &NewNode, &Class, &FinalizeNode, &InClassInput](const FMetasoundFrontendGraph& Graph)
 		{
+			constexpr int32* NewNodeIndex = nullptr;
+			FMetasoundFrontendNode* NewPageNode = AddNodeInternal(Class.Metadata, FinalizeNode, InClassInput.NodeID, NewNodeIndex, &Graph.PageID);
+			if (Graph.PageID == BuildPageID)
+			{
+				NewNode = NewPageNode;
+			}
+
 			// Remove the default literal on the node added during the "FinalizeNode" call. This matches how 
 			// nodes are serialized in editor. The default literals are only stored on the FMetasoundFrontendClassInputs.
-			NewNode->InputLiterals.Reset();
+			NewPageNode->InputLiterals.Reset();
+		});
 
+		if (NewNode)
+		{
 			const int32 NewIndex = RootGraph.Interface.Inputs.Num();
 			FMetasoundFrontendClassInput& NewInput = RootGraph.Interface.Inputs.Add_GetRef(InClassInput);
 			if (!NewInput.VertexID.IsValid())
@@ -1046,7 +1059,7 @@ const FMetasoundFrontendNode* FMetaSoundFrontendDocumentBuilder::AddGraphOutput(
 	checkf(InClassOutput.NodeID.IsValid(), TEXT("Unassigned NodeID when adding graph output"));
 	checkf(InClassOutput.VertexID.IsValid(), TEXT("Unassigned VertexID when adding graph output"));
 
-	FMetasoundFrontendDocument& Document = GetDocument();
+	FMetasoundFrontendDocument& Document = GetDocumentChecked();
 	FMetasoundFrontendGraphClass& RootGraph = Document.RootGraph;
 
 	if (InClassOutput.TypeName.IsNone())
@@ -1105,7 +1118,7 @@ bool FMetaSoundFrontendDocumentBuilder::AddInterface(FName InterfaceName)
 	FMetasoundFrontendInterface Interface;
 	if (ISearchEngine::Get().FindInterfaceWithHighestVersion(InterfaceName, Interface))
 	{
-		if (GetDocument().Interfaces.Contains(Interface.Version))
+		if (GetDocumentChecked().Interfaces.Contains(Interface.Version))
 		{
 			UE_LOG(LogMetaSound, VeryVerbose, TEXT("MetaSound interface '%s' already found on document. MetaSoundBuilder skipping add request."), *InterfaceName.ToString());
 			return true;
@@ -1224,7 +1237,7 @@ const FMetasoundFrontendNode* FMetaSoundFrontendDocumentBuilder::AddNodeByTempla
 	return NewNode;
 }
 
-FMetasoundFrontendNode* FMetaSoundFrontendDocumentBuilder::AddNodeInternal(const FMetasoundFrontendClassMetadata& InClassMetadata, FFinalizeNodeFunctionRef FinalizeNode, FGuid InNodeID, int32* NewNodeIndex)
+FMetasoundFrontendNode* FMetaSoundFrontendDocumentBuilder::AddNodeInternal(const FMetasoundFrontendClassMetadata& InClassMetadata, FFinalizeNodeFunctionRef FinalizeNode, FGuid InNodeID, int32* NewNodeIndex, const FGuid* InGraphPageID)
 {
 	METASOUND_TRACE_CPUPROFILER_EVENT_SCOPE(FMetaSoundFrontendDocumentBuilder::AddNodeInternal);
 
@@ -1233,25 +1246,28 @@ FMetasoundFrontendNode* FMetaSoundFrontendDocumentBuilder::AddNodeInternal(const
 	const FNodeRegistryKey ClassKey = FNodeRegistryKey(InClassMetadata);
 	if (const FMetasoundFrontendClass* Dependency = DocumentCache->FindDependency(ClassKey))
 	{
-		FMetasoundFrontendDocument& Document = GetDocument();
-		FMetasoundFrontendGraph& Graph = Document.RootGraph.Graph;
+		FMetasoundFrontendDocument& Document = GetDocumentChecked();
+		FMetasoundFrontendGraph& Graph = Document.RootGraph.FindGraphChecked(InGraphPageID ? *InGraphPageID : BuildPageID);
 		TArray<FMetasoundFrontendNode>& Nodes = Graph.Nodes;
 		FMetasoundFrontendNode& Node = Nodes.Emplace_GetRef(*Dependency);
 		Node.UpdateID(InNodeID);
 		FinalizeNode(Node, ClassKey);
 
-		const int32 NewIndex = Nodes.Num() - 1;
-		const IDocumentGraphNodeCache& NodeCache = DocumentCache->GetNodeCache();
-		DocumentDelegates->NodeDelegates.OnNodeAdded.Broadcast(NewIndex);
-
-		if (NewNodeIndex)
+		if (!InGraphPageID || *InGraphPageID == BuildPageID)
 		{
-			*NewNodeIndex = NewIndex;
-		}
+			const int32 NewIndex = Nodes.Num() - 1;
+			const IDocumentGraphNodeCache& NodeCache = DocumentCache->GetNodeCache();
+			DocumentDelegates->NodeDelegates.OnNodeAdded.Broadcast(NewIndex);
+
+			if (NewNodeIndex)
+			{
+				*NewNodeIndex = NewIndex;
+			}
 
 #if WITH_EDITORONLY_DATA
-		Document.Metadata.ModifyContext.AddNodeIDModified(InNodeID);
+			Document.Metadata.ModifyContext.AddNodeIDModified(InNodeID);
 #endif // WITH_EDITORONLY_DATA
+		}
 		return &Node;
 	}
 
@@ -1262,7 +1278,7 @@ bool FMetaSoundFrontendDocumentBuilder::CanAddEdge(const FMetasoundFrontendEdge&
 {
 	using namespace Metasound::Frontend;
 
-	const FMetasoundFrontendDocument& Document = GetConstDocument();
+	const FMetasoundFrontendDocument& Document = GetConstDocumentChecked();
 	const IDocumentGraphEdgeCache& EdgeCache = DocumentCache->GetEdgeCache();
 
 	if (!EdgeCache.IsNodeInputConnected(InEdge.ToNodeID, InEdge.ToVertexID))
@@ -1273,39 +1289,27 @@ bool FMetaSoundFrontendDocumentBuilder::CanAddEdge(const FMetasoundFrontendEdge&
 	return false;
 }
 
-bool FMetaSoundFrontendDocumentBuilder::RemoveUnusedDependencies()
+void FMetaSoundFrontendDocumentBuilder::ClearDocument(TSharedRef<Metasound::Frontend::FDocumentModifyDelegates> ModifyDelegates)
 {
-	bool bDidEdit = false;
+	FMetasoundFrontendDocument& Doc = GetDocumentChecked();
 
-	TArray<FMetasoundFrontendClass>& Dependencies = GetDocument().Dependencies;
-	for (int32 i = Dependencies.Num() - 1; i >= 0; --i)
-	{
-		const FGuid& ClassID = Dependencies[i].ID;
-		if (!DocumentCache->GetNodeCache().ContainsNodesOfClassID(ClassID))
-		{
-			const bool bRemoved = RemoveDependency(ClassID);
-			bDidEdit |= ensureAlwaysMsgf(bRemoved, TEXT("Failed to remove dependency that was found on document and was not referenced by nodes"));
-		}
-	}
+	Doc.Interfaces.Reset();
+	Doc.Dependencies.Reset();
 
-	return bDidEdit;
-}
+#if WITH_EDITORONLY_DATA
+	Doc.Metadata.MemberMetadata.Reset();
+#endif // WITH_EDITORONLY_DATA
 
-void FMetaSoundFrontendDocumentBuilder::ClearGraph(TSharedRef<Metasound::Frontend::FDocumentModifyDelegates> ModifyDelegates)
-{
-	FMetasoundFrontendDocument& Doc = GetDocument();
 	FMetasoundFrontendGraphClass& GraphClass = Doc.RootGraph;
-	GraphClass.Graph.Nodes.Reset();
-	GraphClass.Graph.Edges.Reset();
 	GraphClass.Interface.Inputs.Reset();
 	GraphClass.Interface.Outputs.Reset();
 	GraphClass.PresetOptions.InputsInheritingDefault.Reset();
-	Doc.Interfaces.Reset();
-	Doc.Dependencies.Reset();
+	GraphClass.PresetOptions.bIsPreset = false;
 
-#if WITH_EDITORONLY_DATA
-	Doc.Metadata.MemberMetadata.Reset();
-#endif // WITH_EDITORONLY_DATA
+	GraphClass.RemoveAllGraphPages();
+	FMetasoundFrontendGraph& Graph = GraphClass.GetDefaultGraph();
+	Graph.Nodes.Reset();
+	Graph.Edges.Reset();
 
 	Reload(ModifyDelegates);
 }
@@ -1313,7 +1317,7 @@ void FMetaSoundFrontendDocumentBuilder::ClearGraph(TSharedRef<Metasound::Fronten
 #if WITH_EDITORONLY_DATA
 bool FMetaSoundFrontendDocumentBuilder::ClearMemberMetadata(const FGuid& InMemberID)
 {
-	return GetDocument().Metadata.MemberMetadata.Remove(InMemberID) > 0;
+	return GetDocumentChecked().Metadata.MemberMetadata.Remove(InMemberID) > 0;
 }
 #endif // WITH_EDITORONLY_DATA
 
@@ -1321,11 +1325,11 @@ bool FMetaSoundFrontendDocumentBuilder::ConformGraphInputNodeToClass(const FMeta
 {
 	using namespace Metasound::Frontend;
 
-	FMetasoundFrontendDocument& Document = GetDocument();
+	FMetasoundFrontendDocument& Document = GetDocumentChecked();
 	const IDocumentGraphNodeCache& NodeCache = DocumentCache->GetNodeCache();
 	if (const int32* NodeIndexPtr = NodeCache.FindNodeIndex(GraphInput.NodeID))
 	{
-		TArray<FMetasoundFrontendNode>& Nodes = Document.RootGraph.Graph.Nodes;
+		TArray<FMetasoundFrontendNode>& Nodes = FindBuildGraphChecked().Nodes;
 		FMetasoundFrontendNode& Node = Nodes[*NodeIndexPtr];
 		FMetasoundFrontendClass Class;
 		const bool bClassFound = DocumentBuilderPrivate::FindInputRegistryClass(GraphInput.TypeName, GraphInput.AccessType, Class);
@@ -1375,11 +1379,11 @@ bool FMetaSoundFrontendDocumentBuilder::ConformGraphOutputNodeToClass(const FMet
 {
 	using namespace Metasound::Frontend;
 
-	FMetasoundFrontendDocument& Document = GetDocument();
+	FMetasoundFrontendDocument& Document = GetDocumentChecked();
 	const IDocumentGraphNodeCache& NodeCache = DocumentCache->GetNodeCache();
 	if (const int32* NodeIndexPtr = NodeCache.FindNodeIndex(GraphOutput.NodeID))
 	{
-		TArray<FMetasoundFrontendNode>& Nodes = Document.RootGraph.Graph.Nodes;
+		TArray<FMetasoundFrontendNode>& Nodes = FindBuildGraphChecked().Nodes;
 		FMetasoundFrontendNode& Node = Nodes[*NodeIndexPtr];
 		FMetasoundFrontendClass Class;
 		const bool bClassFound = DocumentBuilderPrivate::FindOutputRegistryClass(GraphOutput.TypeName, GraphOutput.AccessType, Class);
@@ -1441,13 +1445,13 @@ bool FMetaSoundFrontendDocumentBuilder::ConvertFromPreset()
 
 	if (IsPreset())
 	{
-		FMetasoundFrontendDocument& Document = GetDocument();
+		FMetasoundFrontendDocument& Document = GetDocumentChecked();
 		FMetasoundFrontendGraphClass& RootGraphClass = Document.RootGraph;
 		FMetasoundFrontendGraphClassPresetOptions& PresetOptions = RootGraphClass.PresetOptions;
 		PresetOptions.bIsPreset = false;
 
 #if WITH_EDITOR
-		FMetasoundFrontendGraphStyle& Style = Document.RootGraph.Graph.Style;
+		FMetasoundFrontendGraphStyle& Style = FindBuildGraphChecked().Style;
 		Style.bIsGraphEditable = true;
 #endif // WITH_EDITOR
 
@@ -1462,28 +1466,27 @@ bool FMetaSoundFrontendDocumentBuilder::ConvertToPreset(const FMetasoundFrontend
 	using namespace Metasound;
 	using namespace Metasound::Frontend;
 
-	ClearGraph(ModifyDelegates);
+	ClearDocument(ModifyDelegates);
 
-	FMetasoundFrontendGraphClass& PresetAssetRootGraph = GetDocument().RootGraph;
-	FMetasoundFrontendGraph& PresetAssetGraph = PresetAssetRootGraph.Graph;
-	// Mark preset as auto-update and non-editable
+	FMetasoundFrontendGraphClass& PresetAssetRootGraph = GetDocumentChecked().RootGraph;
+	PresetAssetRootGraph.IterateGraphPages([](FMetasoundFrontendGraph& PresetAssetGraph)
+	{
 #if WITH_EDITORONLY_DATA
-	PresetAssetGraph.Style.bIsGraphEditable = false;
+		PresetAssetGraph.Style.bIsGraphEditable = false;
 #endif // WITH_EDITORONLY_DATA
+	});
 
 	// Mark all inputs as inherited by default
 	{
-		TArray<FName> InputsInheritingDefault;
+		PresetAssetRootGraph.PresetOptions.InputsInheritingDefault.Reset();
 		auto GetInputName = [](const FMetasoundFrontendClassInput& Input) { return Input.Name; };
-		Algo::Transform(PresetAssetRootGraph.Interface.Inputs, InputsInheritingDefault, GetInputName);
-
+		Algo::Transform(PresetAssetRootGraph.Interface.Inputs, PresetAssetRootGraph.PresetOptions.InputsInheritingDefault, GetInputName);
 		PresetAssetRootGraph.PresetOptions.bIsPreset = true;
-		PresetAssetRootGraph.PresetOptions.InputsInheritingDefault = TSet<FName>(InputsInheritingDefault);
 	}
 
 	// Apply root graph transform 
 	FRebuildPresetRootGraph RebuildPresetRootGraph(InReferencedDocument);
-	if (RebuildPresetRootGraph.Transform(GetDocument()))
+	if (RebuildPresetRootGraph.Transform(GetDocumentChecked()))
 	{
 		DocumentInterface->ConformObjectToDocument();
 
@@ -1497,13 +1500,22 @@ bool FMetaSoundFrontendDocumentBuilder::ConvertToPreset(const FMetasoundFrontend
 		return true;
 	}
 
-	// Hack until prior processes are all moved over to builder API
 	return false;
 }
 
+FMetasoundFrontendGraph& FMetaSoundFrontendDocumentBuilder::FindBuildGraphChecked()
+{
+	return GetDocumentChecked().RootGraph.FindGraphChecked(BuildPageID);
+}
+
+const FMetasoundFrontendGraph& FMetaSoundFrontendDocumentBuilder::FindConstBuildGraphChecked() const
+{
+	return GetConstDocumentChecked().RootGraph.FindConstGraphChecked(BuildPageID);
+}
+
 bool FMetaSoundFrontendDocumentBuilder::FindDeclaredInterfaces(TArray<const Metasound::Frontend::IInterfaceRegistryEntry*>& OutInterfaces) const
 {
-	return FindDeclaredInterfaces(GetConstDocument(), OutInterfaces);
+	return FindDeclaredInterfaces(GetConstDocumentChecked(), OutInterfaces);
 }
 
 bool FMetaSoundFrontendDocumentBuilder::FindDeclaredInterfaces(const FMetasoundFrontendDocument& InDocument, TArray<const Metasound::Frontend::IInterfaceRegistryEntry*>& OutInterfaces)
@@ -1555,13 +1567,13 @@ TArray<const FMetasoundFrontendEdge*> FMetaSoundFrontendDocumentBuilder::FindEdg
 #if WITH_EDITOR
 const FMetaSoundFrontendGraphComment* FMetaSoundFrontendDocumentBuilder::FindGraphComment(const FGuid& InCommentID) const
 {
-	const TMap<FGuid, FMetaSoundFrontendGraphComment>& Comments = GetConstDocument().RootGraph.Graph.Style.Comments;
+	const TMap<FGuid, FMetaSoundFrontendGraphComment>& Comments = FindConstBuildGraphChecked().Style.Comments;
 	return Comments.Find(InCommentID);
 }
 
 FMetaSoundFrontendGraphComment* FMetaSoundFrontendDocumentBuilder::FindGraphComment(const FGuid& InCommentID)
 {
-	TMap<FGuid, FMetaSoundFrontendGraphComment>& Comments = GetDocument().RootGraph.Graph.Style.Comments;
+	TMap<FGuid, FMetaSoundFrontendGraphComment>& Comments = FindBuildGraphChecked().Style.Comments;
 	return Comments.Find(InCommentID);
 }
 #endif // WITH_EDITOR
@@ -1575,7 +1587,7 @@ bool FMetaSoundFrontendDocumentBuilder::FindInterfaceInputNodes(FName InterfaceN
 	FMetasoundFrontendInterface Interface;
 	if (ISearchEngine::Get().FindInterfaceWithHighestVersion(InterfaceName, Interface))
 	{
-		if (GetConstDocument().Interfaces.Contains(Interface.Version))
+		if (GetConstDocumentChecked().Interfaces.Contains(Interface.Version))
 		{
 			const IDocumentGraphNodeCache& NodeCache = DocumentCache->GetNodeCache();
 			const IDocumentGraphInterfaceCache& InterfaceCache = DocumentCache->GetInterfaceCache();
@@ -1616,7 +1628,7 @@ bool FMetaSoundFrontendDocumentBuilder::FindInterfaceOutputNodes(FName Interface
 	FMetasoundFrontendInterface Interface;
 	if (ISearchEngine::Get().FindInterfaceWithHighestVersion(InterfaceName, Interface))
 	{
-		if (GetConstDocument().Interfaces.Contains(Interface.Version))
+		if (GetConstDocumentChecked().Interfaces.Contains(Interface.Version))
 		{
 			const IDocumentGraphNodeCache& NodeCache = DocumentCache->GetNodeCache();
 			const IDocumentGraphInterfaceCache& InterfaceCache = DocumentCache->GetInterfaceCache();
@@ -1687,7 +1699,7 @@ const FMetasoundFrontendNode* FMetaSoundFrontendDocumentBuilder::FindGraphOutput
 #if WITH_EDITOR
 UMetaSoundFrontendMemberMetadata* FMetaSoundFrontendDocumentBuilder::FindMemberMetadata(const FGuid& InMemberID)
 {
-	FMetasoundFrontendDocument& Document = GetDocument();
+	FMetasoundFrontendDocument& Document = GetDocumentChecked();
 	TMap<FGuid, TObjectPtr<UMetaSoundFrontendMemberMetadata>>& LiteralMetadata = Document.Metadata.MemberMetadata;
 	TObjectPtr<UMetaSoundFrontendMemberMetadata> ToReturn = LiteralMetadata.FindRef(InMemberID);
 	return ToReturn;
@@ -1706,7 +1718,7 @@ bool FMetaSoundFrontendDocumentBuilder::FindNodeClassInterfaces(const FGuid& InN
 	using namespace Metasound;
 	using namespace Metasound::Frontend;
 
-	const FMetasoundFrontendDocument& Document = GetConstDocument();
+	const FMetasoundFrontendDocument& Document = GetConstDocumentChecked();
 	const IDocumentGraphNodeCache& NodeCache = DocumentCache->GetNodeCache();
 	if (const FMetasoundFrontendNode* Node = NodeCache.FindNode(InNodeID))
 	{
@@ -1746,7 +1758,7 @@ TArray<const FMetasoundFrontendVertex*> FMetaSoundFrontendDocumentBuilder::FindN
 	const IDocumentGraphEdgeCache& EdgeCache = DocumentCache->GetEdgeCache();
 	const IDocumentGraphNodeCache& NodeCache = DocumentCache->GetNodeCache();
 
-	const FMetasoundFrontendDocument& Document = GetConstDocument();
+	const FMetasoundFrontendDocument& Document = GetConstDocumentChecked();
 
 	if (ConnectedInputNodes)
 	{
@@ -1754,10 +1766,11 @@ TArray<const FMetasoundFrontendVertex*> FMetaSoundFrontendDocumentBuilder::FindN
 	}
 
 	TArray<const FMetasoundFrontendVertex*> Inputs;
+	const FMetasoundFrontendGraph& Graph = FindConstBuildGraphChecked();
 	const TArrayView<const int32> Indices = EdgeCache.FindEdgeIndicesFromNodeOutput(InOutputNodeID, InOutputVertexID);
-	Algo::Transform(Indices, Inputs, [&Document, &NodeCache, &ConnectedInputNodes](const int32& Index)
+	Algo::Transform(Indices, Inputs, [&Graph, &Document, &NodeCache, &ConnectedInputNodes](const int32& Index)
 	{
-		const FMetasoundFrontendEdge& Edge = Document.RootGraph.Graph.Edges[Index];
+		const FMetasoundFrontendEdge& Edge = Graph.Edges[Index];
 		if (ConnectedInputNodes)
 		{
 			ConnectedInputNodes->Add(NodeCache.FindNode(Edge.ToNodeID));
@@ -1777,7 +1790,7 @@ FMetasoundFrontendNode* FMetaSoundFrontendDocumentBuilder::FindNodeInternal(cons
 	const IDocumentGraphNodeCache& NodeCache = DocumentCache->GetNodeCache();
 	if (const int32* NodeIndex = NodeCache.FindNodeIndex(InNodeID))
 	{
-		return &GetDocument().RootGraph.Graph.Nodes[*NodeIndex];
+		return &FindBuildGraphChecked().Nodes[*NodeIndex];
 	}
 
 	return nullptr;
@@ -1809,8 +1822,8 @@ const FMetasoundFrontendVertex* FMetaSoundFrontendDocumentBuilder::FindNodeOutpu
 	const IDocumentGraphEdgeCache& EdgeCache = DocumentCache->GetEdgeCache();
 	if (const int32* Index = EdgeCache.FindEdgeIndexToNodeInput(InInputNodeID, InInputVertexID))
 	{
-		const FMetasoundFrontendDocument& Document = GetConstDocument();
-		const FMetasoundFrontendEdge& Edge = Document.RootGraph.Graph.Edges[*Index];
+		const FMetasoundFrontendDocument& Document = GetConstDocumentChecked();
+		const FMetasoundFrontendEdge& Edge = FindConstBuildGraphChecked().Edges[*Index];
 		const IDocumentGraphNodeCache& NodeCache = DocumentCache->GetNodeCache();
 		if (ConnectedOutputNode)
 		{
@@ -1830,7 +1843,7 @@ const FMetasoundFrontendVertex* FMetaSoundFrontendDocumentBuilder::FindNodeOutpu
 FMetaSoundFrontendGraphComment& FMetaSoundFrontendDocumentBuilder::FindOrAddGraphComment(const FGuid& InCommentID)
 {
 	check(InCommentID.IsValid());
-	TMap<FGuid, FMetaSoundFrontendGraphComment>& Comments = GetDocument().RootGraph.Graph.Style.Comments;
+	TMap<FGuid, FMetaSoundFrontendGraphComment>& Comments = FindBuildGraphChecked().Style.Comments;
 	return Comments.FindOrAdd(InCommentID);
 }
 #endif // WITH_EDITOR
@@ -1838,7 +1851,7 @@ FMetaSoundFrontendGraphComment& FMetaSoundFrontendDocumentBuilder::FindOrAddGrap
 FMetasoundFrontendClassName FMetaSoundFrontendDocumentBuilder::GenerateNewClassName()
 {
 	using namespace Metasound::Frontend;
-	FMetasoundFrontendClassMetadata& Metadata = GetDocument().RootGraph.Metadata;
+	FMetasoundFrontendClassMetadata& Metadata = GetDocumentChecked().RootGraph.Metadata;
 	const FMetasoundFrontendClassName OldClassName = Metadata.GetClassName();
 	const FMetasoundFrontendClassName NewClassName(FName(), FName(*FGuid::NewGuid().ToString()), FName());
 	Metadata.SetClassName(NewClassName);
@@ -1852,36 +1865,36 @@ const FTopLevelAssetPath FMetaSoundFrontendDocumentBuilder::GetBuilderClassPath(
 	return Interface->GetBaseMetaSoundUClass().GetClassPathName();
 }
 
-const FMetasoundFrontendDocument& FMetaSoundFrontendDocumentBuilder::GetConstDocument() const
+const FMetasoundFrontendDocument& FMetaSoundFrontendDocumentBuilder::GetConstDocumentChecked() const
 {
-	return DocumentInterface->GetConstDocument();
+	return GetConstDocumentInterfaceChecked().GetConstDocument();
+}
+
+const IMetaSoundDocumentInterface& FMetaSoundFrontendDocumentBuilder::GetConstDocumentInterfaceChecked() const
+{
+	const IMetaSoundDocumentInterface* Interface = DocumentInterface.GetInterface();
+	checkf(Interface, TEXT("Failed to return document; interface must always be valid while builder is operating on MetaSound UObject!"));
+	return *Interface;
 }
 
 const FString FMetaSoundFrontendDocumentBuilder::GetDebugName() const
 {
 	using namespace Metasound::Frontend;
 
-	const FMetasoundFrontendClassMetadata& Metadata = GetConstDocument().RootGraph.Metadata;
-	const FMetasoundFrontendVersionNumber& Version = Metadata.GetVersion();
-	const FAssetKey AssetKey(Metadata.GetClassName(), Version);
-	if (const FMetasoundAssetBase* Asset = IMetaSoundAssetManager::GetChecked().FindAsset(AssetKey))
-	{
-		return Asset->GetOwningAssetName();
-	}
-
-	return Metadata.GetClassName().ToString();
+	UObject& MetaSoundObject = CastDocumentObjectChecked<UObject>();
+	return MetaSoundObject.GetPathName();
 }
 
-FMetasoundFrontendDocument& FMetaSoundFrontendDocumentBuilder::GetDocument()
+const FMetasoundFrontendDocument& FMetaSoundFrontendDocumentBuilder::GetDocument() const
 {
-	IMetaSoundDocumentInterface* Interface = DocumentInterface.GetInterface();
+	const IMetaSoundDocumentInterface* Interface = DocumentInterface.GetInterface();
 	checkf(Interface, TEXT("Failed to return document; interface must always be valid while builder is operating on MetaSound UObject!"));
-	return Interface->GetDocument();
+	return Interface->GetConstDocument();
 }
 
-const FMetasoundFrontendDocument& FMetaSoundFrontendDocumentBuilder::GetDocument() const
+FMetasoundFrontendDocument& FMetaSoundFrontendDocumentBuilder::GetDocumentChecked() const
 {
-	return DocumentInterface->GetConstDocument();
+	return GetDocumentInterfaceChecked().GetDocument();
 }
 
 const Metasound::Frontend::FDocumentModifyDelegates& FMetaSoundFrontendDocumentBuilder::GetDocumentDelegates() const
@@ -1890,13 +1903,20 @@ const Metasound::Frontend::FDocumentModifyDelegates& FMetaSoundFrontendDocumentB
 }
 
 const IMetaSoundDocumentInterface& FMetaSoundFrontendDocumentBuilder::GetDocumentInterface() const
+{
+	const IMetaSoundDocumentInterface* Interface = DocumentInterface.GetInterface();
+	checkf(Interface, TEXT("Failed to return document; interface must always be valid while builder is operating on MetaSound UObject!"));
+	return *Interface;
+}
+
+IMetaSoundDocumentInterface& FMetaSoundFrontendDocumentBuilder::GetDocumentInterfaceChecked() const
 {
 	IMetaSoundDocumentInterface* Interface = DocumentInterface.GetInterface();
-	check(Interface);
+	checkf(Interface, TEXT("Failed to return document; interface must always be valid while builder is operating on MetaSound UObject!"));
 	return *Interface;
 }
 
-FMetasoundAssetBase& FMetaSoundFrontendDocumentBuilder::GetMetasoundAsset()
+FMetasoundAssetBase& FMetaSoundFrontendDocumentBuilder::GetMetasoundAsset() const
 {
 	using namespace Metasound::Frontend;
 
@@ -1907,15 +1927,9 @@ FMetasoundAssetBase& FMetaSoundFrontendDocumentBuilder::GetMetasoundAsset()
 	return *Asset;
 }
 
-const FMetasoundAssetBase& FMetaSoundFrontendDocumentBuilder::GetMetasoundAsset() const
+const FGuid& FMetaSoundFrontendDocumentBuilder::GetBuildPageID() const
 {
-	using namespace Metasound::Frontend;
-
-	UObject* Object = DocumentInterface.GetObject();
-	check(Object);
-	const FMetasoundAssetBase* Asset = IMetaSoundAssetManager::GetChecked().GetAsAsset(*Object);
-	check(Asset);
-	return *Asset;
+	return BuildPageID;
 }
 
 EMetasoundFrontendVertexAccessType FMetaSoundFrontendDocumentBuilder::GetNodeInputAccessType(const FGuid& InNodeID, const FGuid& InVertexID) const
@@ -1925,7 +1939,7 @@ EMetasoundFrontendVertexAccessType FMetaSoundFrontendDocumentBuilder::GetNodeInp
 	const IDocumentGraphNodeCache& NodeCache = DocumentCache->GetNodeCache();
 	if (const int32* NodeIndex = NodeCache.FindNodeIndex(InNodeID))
 	{
-		const FMetasoundFrontendGraph& Graph = GetConstDocument().RootGraph.Graph;
+		const FMetasoundFrontendGraph& Graph = FindConstBuildGraphChecked();
 		const FMetasoundFrontendNode& Node = Graph.Nodes[*NodeIndex];
 		auto IsVertexID = [&InVertexID](const FMetasoundFrontendVertex& Vertex) { return Vertex.VertexID == InVertexID; };
 		if (const FMetasoundFrontendClass* Class = DocumentCache->FindDependency(Node.ClassID))
@@ -1979,8 +1993,7 @@ const FMetasoundFrontendLiteral* FMetaSoundFrontendDocumentBuilder::GetNodeInput
 	const IDocumentGraphNodeCache& NodeCache = DocumentCache->GetNodeCache();
 	if (const int32* NodeIndex = NodeCache.FindNodeIndex(InNodeID))
 	{
-		const FMetasoundFrontendGraph& Graph = GetConstDocument().RootGraph.Graph;
-		const FMetasoundFrontendNode& Node = Graph.Nodes[*NodeIndex];
+		const FMetasoundFrontendNode& Node = FindConstBuildGraphChecked().Nodes[*NodeIndex];
 		auto IsVertexID = [&InVertexID](const FMetasoundFrontendVertex& Vertex) { return Vertex.VertexID == InVertexID; };
 		if (const FMetasoundFrontendVertex* Vertex = Node.Interface.Inputs.FindByPredicate(IsVertexID))
 		{
@@ -2023,8 +2036,7 @@ const FMetasoundFrontendLiteral* FMetaSoundFrontendDocumentBuilder::GetNodeInput
 	const IDocumentGraphNodeCache& NodeCache = DocumentCache->GetNodeCache();
 	if (const int32* NodeIndex = NodeCache.FindNodeIndex(InNodeID))
 	{
-		const FMetasoundFrontendGraph& Graph = GetConstDocument().RootGraph.Graph;
-		const FMetasoundFrontendNode& Node = Graph.Nodes[*NodeIndex];
+		const FMetasoundFrontendNode& Node = FindConstBuildGraphChecked().Nodes[*NodeIndex];
 
 		auto IsVertex = [&InVertexID](const FMetasoundFrontendVertex& Vertex) { return Vertex.VertexID == InVertexID; };
 		const int32 VertexIndex = Node.Interface.Inputs.IndexOfByPredicate(IsVertex);
@@ -2052,8 +2064,7 @@ EMetasoundFrontendVertexAccessType FMetaSoundFrontendDocumentBuilder::GetNodeOut
 	const IDocumentGraphNodeCache& NodeCache = DocumentCache->GetNodeCache();
 	if (const int32* NodeIndex = NodeCache.FindNodeIndex(InNodeID))
 	{
-		const FMetasoundFrontendGraph& Graph = GetConstDocument().RootGraph.Graph;
-		const FMetasoundFrontendNode& Node = Graph.Nodes[*NodeIndex];
+		const FMetasoundFrontendNode& Node = FindConstBuildGraphChecked().Nodes[*NodeIndex];
 		if (const FMetasoundFrontendClass* Class = DocumentCache->FindDependency(Node.ClassID))
 		{
 			const EMetasoundFrontendClassType ClassType = Class->Metadata.GetType();
@@ -2106,7 +2117,7 @@ void FMetaSoundFrontendDocumentBuilder::InitDocument(const FMetasoundFrontendDoc
 
 	METASOUND_TRACE_CPUPROFILER_EVENT_SCOPE(FMetaSoundFrontendDocumentBuilder::InitDocument);
 
-	FMetasoundFrontendDocument& Document = GetDocument();
+	FMetasoundFrontendDocument& Document = GetDocumentChecked();
 
 	// 1. Set default class Metadata.
 	if (InDocumentTemplate)
@@ -2174,7 +2185,7 @@ void FMetaSoundFrontendDocumentBuilder::InitGraphClassMetadata(FMetasoundFronten
 
 void FMetaSoundFrontendDocumentBuilder::InitGraphClassMetadata(bool bResetVersion, const FMetasoundFrontendClassName* NewClassName)
 {
-	InitGraphClassMetadata(GetDocument().RootGraph.Metadata, bResetVersion, NewClassName);
+	InitGraphClassMetadata(GetDocumentChecked().RootGraph.Metadata, bResetVersion, NewClassName);
 }
 
 void FMetaSoundFrontendDocumentBuilder::InitNodeLocations()
@@ -2183,18 +2194,16 @@ void FMetaSoundFrontendDocumentBuilder::InitNodeLocations()
 	using namespace Metasound;
 	using namespace Metasound::Frontend;
 
-	FMetasoundFrontendGraph& Graph = GetDocument().RootGraph.Graph;
-
 	FVector2D InputNodeLocation = FVector2D::ZeroVector;
 	FVector2D ExternalNodeLocation = InputNodeLocation + DisplayStyle::NodeLayout::DefaultOffsetX;
 	FVector2D OutputNodeLocation = ExternalNodeLocation + DisplayStyle::NodeLayout::DefaultOffsetX;
 
-	TArray<FMetasoundFrontendNode>& Nodes = Graph.Nodes;
+	TArray<FMetasoundFrontendNode>& Nodes = FindBuildGraphChecked().Nodes;
 	for (FMetasoundFrontendNode& Node : Nodes)
 	{
 		if (const int32* ClassIndex = DocumentCache->FindDependencyIndex(Node.ClassID))
 		{
-			FMetasoundFrontendClass& Class = GetDocument().Dependencies[*ClassIndex];
+			FMetasoundFrontendClass& Class = GetDocumentChecked().Dependencies[*ClassIndex];
 
 			const EMetasoundFrontendClassType NodeType = Class.Metadata.GetType();
 			FVector2D NewLocation;
@@ -2264,12 +2273,12 @@ bool FMetaSoundFrontendDocumentBuilder::IsInterfaceDeclared(FName InInterfaceNam
 
 bool FMetaSoundFrontendDocumentBuilder::IsInterfaceDeclared(const FMetasoundFrontendVersion& InInterfaceVersion) const
 {
-	return GetConstDocument().Interfaces.Contains(InInterfaceVersion);
+	return GetConstDocumentChecked().Interfaces.Contains(InInterfaceVersion);
 }
 
 bool FMetaSoundFrontendDocumentBuilder::IsPreset() const
 {
-	return GetConstDocument().RootGraph.PresetOptions.bIsPreset;
+	return GetConstDocumentChecked().RootGraph.PresetOptions.bIsPreset;
 }
 
 Metasound::Frontend::EInvalidEdgeReason FMetaSoundFrontendDocumentBuilder::IsValidEdge(const FMetasoundFrontendEdge& InEdge) const
@@ -2311,7 +2320,7 @@ void FMetaSoundFrontendDocumentBuilder::IterateNodesConnectedWithVertex(const FM
 {
 	using namespace Metasound::Frontend;
 
-	FMetasoundFrontendDocument& Document = GetDocument();
+	FMetasoundFrontendGraph& Graph = FindBuildGraphChecked();
 	TArray<FMetasoundFrontendEdge> EdgesToConnectedNodes; // Have to cache to avoid pointers becoming garbage in subsequent removal loop
 	const IDocumentGraphEdgeCache& EdgeCache = DocumentCache->GetEdgeCache();
 	const IDocumentGraphNodeCache& NodeCache = DocumentCache->GetNodeCache();
@@ -2322,7 +2331,7 @@ void FMetaSoundFrontendDocumentBuilder::IterateNodesConnectedWithVertex(const FM
 		const FGuid& ConnectedNodeID = Edge.ToNodeID == Vertex.NodeID ? Edge.FromNodeID : Edge.ToNodeID;
 		if (const int32* ConnectedNodeIndex = NodeCache.FindNodeIndex(ConnectedNodeID))
 		{
-			FMetasoundFrontendNode& Node = Document.RootGraph.Graph.Nodes[*ConnectedNodeIndex];
+			FMetasoundFrontendNode& Node = Graph.Nodes[*ConnectedNodeIndex];
 			NodeIndexIterFunc(Edge, Node);
 		}
 	}
@@ -2332,7 +2341,7 @@ bool FMetaSoundFrontendDocumentBuilder::ModifyInterfaces(Metasound::Frontend::FM
 {
 	using namespace Metasound::Frontend;
 
-	FMetasoundFrontendDocument& Doc = GetDocument();
+	FMetasoundFrontendDocument& Doc = GetDocumentChecked();
 	DocumentBuilderPrivate::FModifyInterfacesImpl Context(Doc, MoveTemp(InOptions));
 	return Context.Execute(*this, *DocumentDelegates);
 }
@@ -2351,8 +2360,7 @@ bool FMetaSoundFrontendDocumentBuilder::TransformTemplateNodes()
 	};
 	using FTemplateTransformParamsMap = TSortedMap<FGuid, FTemplateTransformParams>;
 
-	FMetasoundFrontendDocument& Document = GetDocument();
-	FMetasoundFrontendGraph& Graph = Document.RootGraph.Graph;
+	FMetasoundFrontendDocument& Document = GetDocumentChecked();
 	TArray<FMetasoundFrontendClass>& Dependencies = Document.Dependencies;
 
 	FTemplateTransformParamsMap TemplateParams;
@@ -2376,6 +2384,7 @@ bool FMetaSoundFrontendDocumentBuilder::TransformTemplateNodes()
 	// which allows for addition/removal of nodes to/from original array container
 	// without template transform having to worry about mutation while iterating
 	TArray<FGuid> TemplateNodeIDs;
+	FMetasoundFrontendGraph& Graph = FindBuildGraphChecked();
 	for (const FMetasoundFrontendNode& Node : Graph.Nodes)
 	{
 		if (FTemplateTransformParams* Params = TemplateParams.Find(Node.ClassID))
@@ -2427,7 +2436,7 @@ void FMetaSoundFrontendDocumentBuilder::BeginBuilding(TSharedPtr<Metasound::Fron
 	if (DocumentInterface)
 	{
 		DocumentInterface->OnBeginActiveBuilder();
-		DocumentCache = FDocumentCache::Create(GetDocument(), DocumentDelegates.ToSharedRef(), bPrimeCache);
+		DocumentCache = FDocumentCache::Create(GetDocumentChecked(), DocumentDelegates.ToSharedRef(), BuildPageID, bPrimeCache);
 	}
 }
 
@@ -2449,7 +2458,7 @@ bool FMetaSoundFrontendDocumentBuilder::RemoveDependency(const FGuid& InClassID)
 {
 	using namespace Metasound::Frontend;
 
-	FMetasoundFrontendDocument& Document = GetDocument();
+	FMetasoundFrontendDocument& Document = GetDocumentChecked();
 	const IDocumentGraphNodeCache& NodeCache = DocumentCache->GetNodeCache();
 
 	if (const int32* IndexPtr = DocumentCache->FindDependencyIndex(InClassID))
@@ -2478,7 +2487,7 @@ bool FMetaSoundFrontendDocumentBuilder::RemoveDependency(EMetasoundFrontendClass
 {
 	using namespace Metasound::Frontend;
 
-	FMetasoundFrontendDocument& Document = GetDocument();
+	FMetasoundFrontendDocument& Document = GetDocumentChecked();
 	const IDocumentGraphNodeCache& NodeCache = DocumentCache->GetNodeCache();
 
 	const FNodeRegistryKey ClassKey(ClassType, InClassName, InClassVersionNumber);
@@ -2508,8 +2517,8 @@ bool FMetaSoundFrontendDocumentBuilder::RemoveEdge(const FMetasoundFrontendEdge&
 {
 	using namespace Metasound::Frontend;
 
-	FMetasoundFrontendDocument& Document = GetDocument();
-	FMetasoundFrontendGraph& Graph = Document.RootGraph.Graph;
+	FMetasoundFrontendDocument& Document = GetDocumentChecked();
+	FMetasoundFrontendGraph& Graph = FindBuildGraphChecked();
 	TArray<FMetasoundFrontendEdge>& Edges = Graph.Edges;
 	const IDocumentGraphEdgeCache& EdgeCache = DocumentCache->GetEdgeCache();
 	if (const int32* IndexPtr = EdgeCache.FindEdgeIndexToNodeInput(EdgeToRemove.ToNodeID, EdgeToRemove.ToVertexID))
@@ -2532,7 +2541,7 @@ bool FMetaSoundFrontendDocumentBuilder::RemoveNamedEdges(const TSet<Metasound::F
 {
 	using namespace Metasound::Frontend;
 
-	FMetasoundFrontendDocument& Document = GetDocument();
+	FMetasoundFrontendDocument& Document = GetDocumentChecked();
 	const IDocumentGraphNodeCache& NodeCache = DocumentCache->GetNodeCache();
 
 	if (OutRemovedEdges)
@@ -2545,7 +2554,7 @@ bool FMetaSoundFrontendDocumentBuilder::RemoveNamedEdges(const TSet<Metasound::F
 	TArray<FMetasoundFrontendEdge> EdgesToRemove;
 	for (const FNamedEdge& NamedEdge : InNamedEdgesToRemove)
 	{
-		const FMetasoundFrontendGraph& Graph = Document.RootGraph.Graph;
+		const FMetasoundFrontendGraph& Graph = FindBuildGraphChecked();
 		const FMetasoundFrontendVertex* OutputVertex = NodeCache.FindOutputVertex(NamedEdge.OutputNodeID, NamedEdge.OutputName);
 		const FMetasoundFrontendVertex* InputVertex = NodeCache.FindInputVertex(NamedEdge.InputNodeID, NamedEdge.InputName);
 
@@ -2596,7 +2605,7 @@ void FMetaSoundFrontendDocumentBuilder::Reload(TSharedPtr<Metasound::Frontend::F
 
 	if (DocumentInterface)
 	{
-		DocumentCache = FDocumentCache::Create(GetDocument(), DocumentDelegates.ToSharedRef(), bPrimeCache);
+		DocumentCache = FDocumentCache::Create(GetDocumentChecked(), DocumentDelegates.ToSharedRef(), BuildPageID, bPrimeCache);
 		DocumentInterface->OnBeginActiveBuilder();
 	}
 }
@@ -2608,7 +2617,7 @@ bool FMetaSoundFrontendDocumentBuilder::RemoveNodeInputDefault(const FGuid& InNo
 	const IDocumentGraphNodeCache& NodeCache = DocumentCache->GetNodeCache();
 	if (const int32* NodeIndex = NodeCache.FindNodeIndex(InNodeID))
 	{
-		FMetasoundFrontendGraph& Graph = GetDocument().RootGraph.Graph;
+		FMetasoundFrontendGraph& Graph = FindBuildGraphChecked();
 		FMetasoundFrontendNode& Node = Graph.Nodes[*NodeIndex];
 
 		auto IsVertex = [&InVertexID](const FMetasoundFrontendVertex& Vertex) { return Vertex.VertexID == InVertexID; };
@@ -2696,8 +2705,8 @@ bool FMetaSoundFrontendDocumentBuilder::RemoveEdgesFromNodeOutput(const FGuid& I
 	const TArrayView<const int32> Indices = EdgeCache.FindEdgeIndicesFromNodeOutput(InNodeID, InVertexID);
 	if (!Indices.IsEmpty())
 	{
-		FMetasoundFrontendDocument& Document = GetDocument();
-		FMetasoundFrontendGraph& Graph = Document.RootGraph.Graph;
+		FMetasoundFrontendDocument& Document = GetDocumentChecked();
+		FMetasoundFrontendGraph& Graph = FindBuildGraphChecked();
 
 		// Copy off indices and sort descending as the edge array will be modified when notifying the cache in the loop below
 		TArray<int32> IndicesCopy(Indices.GetData(), Indices.Num());
@@ -2722,8 +2731,8 @@ bool FMetaSoundFrontendDocumentBuilder::RemoveEdgeToNodeInput(const FGuid& InNod
 	const IDocumentGraphEdgeCache& EdgeCache = DocumentCache->GetEdgeCache();
 	if (const int32* IndexPtr = EdgeCache.FindEdgeIndexToNodeInput(InNodeID, InVertexID))
 	{
-		FMetasoundFrontendDocument& Document = GetDocument();
-		FMetasoundFrontendGraph& Graph = Document.RootGraph.Graph;
+		FMetasoundFrontendDocument& Document = GetDocumentChecked();
+		FMetasoundFrontendGraph& Graph = FindBuildGraphChecked();
 		const int32 Index = *IndexPtr; // Copy off indices as the pointer may be modified when notifying the cache below
 		const int32 LastIndex = Graph.Edges.Num() - 1;
 		DocumentDelegates->EdgeDelegates.OnRemoveSwappingEdge.Broadcast(Index, LastIndex);
@@ -2742,8 +2751,8 @@ bool FMetaSoundFrontendDocumentBuilder::RemoveEdgeToNodeInput(const FGuid& InNod
 #if WITH_EDITOR
 bool FMetaSoundFrontendDocumentBuilder::RemoveGraphComment(const FGuid& InCommentID)
 {
-	FMetasoundFrontendDocument& Document = GetDocument();
-	TMap<FGuid, FMetaSoundFrontendGraphComment>& Comments = Document.RootGraph.Graph.Style.Comments;
+	FMetasoundFrontendDocument& Document = GetDocumentChecked();
+	TMap<FGuid, FMetaSoundFrontendGraphComment>& Comments = FindBuildGraphChecked().Style.Comments;
 
 	if (Comments.Remove(InCommentID) > 0)
 	{
@@ -2764,7 +2773,7 @@ TArray<const FMetasoundFrontendNode*> FMetaSoundFrontendDocumentBuilder::GetGrap
 
 	TArray<const FMetasoundFrontendNode*> TemplateNodes;
 
-	FMetasoundFrontendDocument& Document = GetDocument();
+	FMetasoundFrontendDocument& Document = GetDocumentChecked();
 	if (const int32* Index = DocumentCache->GetInterfaceCache().FindInputIndex(InInputName))
 	{
 		const FMetasoundFrontendClassInput& InputClass = Document.RootGraph.Interface.Inputs[*Index];
@@ -2780,7 +2789,7 @@ TArray<const FMetasoundFrontendNode*> FMetaSoundFrontendDocumentBuilder::GetGrap
 				check(Edge);
 				if (const int32* ConnectedNodeIndex = NodeCache.FindNodeIndex(Edge->ToNodeID))
 				{
-					const FMetasoundFrontendNode& ConnectedNode = Document.RootGraph.Graph.Nodes[*ConnectedNodeIndex];
+					const FMetasoundFrontendNode& ConnectedNode = FindConstBuildGraphChecked().Nodes[*ConnectedNodeIndex];
 					if (const FMetasoundFrontendClass* ConnectedNodeClass = FindDependency(ConnectedNode.ClassID))
 					{
 						if (ConnectedNodeClass->Metadata.GetClassName() == FInputNodeTemplate::ClassName)
@@ -2800,10 +2809,10 @@ bool FMetaSoundFrontendDocumentBuilder::RemoveGraphInput(FName InInputName)
 {
 	using namespace Metasound::Frontend;
 
-	FMetasoundFrontendDocument& Document = GetDocument();
+	FMetasoundFrontendDocument& Document = GetDocumentChecked();
 	if (const int32* IndexPtr = DocumentCache->GetInterfaceCache().FindInputIndex(InInputName))
 	{
-		TArray<FMetasoundFrontendClassInput>& Inputs = GetDocument().RootGraph.Interface.Inputs;
+		TArray<FMetasoundFrontendClassInput>& Inputs = Document.RootGraph.Interface.Inputs;
 		const FGuid NodeID = Inputs[*IndexPtr].NodeID;
 		FGuid ClassID;
 		TArray<FGuid> NodeIDsToRemove { NodeID };
@@ -2875,7 +2884,7 @@ bool FMetaSoundFrontendDocumentBuilder::RemoveGraphOutput(FName InOutputName)
 		const FGuid NodeID = Node->GetID();
 		if (RemoveNode(NodeID))
 		{
-			FMetasoundFrontendDocument& Document = GetDocument();
+			FMetasoundFrontendDocument& Document = GetDocumentChecked();
 
 #if WITH_EDITORONLY_DATA
 			Document.Metadata.ModifyContext.AddNodeIDModified(NodeID);
@@ -2923,7 +2932,7 @@ bool FMetaSoundFrontendDocumentBuilder::RemoveInterface(FName InterfaceName)
 	FMetasoundFrontendInterface Interface;
 	if (ISearchEngine::Get().FindInterfaceWithHighestVersion(InterfaceName, Interface))
 	{
-		if (!GetDocument().Interfaces.Contains(Interface.Version))
+		if (!GetDocumentChecked().Interfaces.Contains(Interface.Version))
 		{
 			UE_LOG(LogMetaSound, VeryVerbose, TEXT("MetaSound interface '%s' not found on document. MetaSoundBuilder skipping remove request."), *InterfaceName.ToString());
 			return true;
@@ -2957,13 +2966,14 @@ bool FMetaSoundFrontendDocumentBuilder::RemoveNode(const FGuid& InNodeID)
 	using namespace Metasound::Frontend;
 
 	const IDocumentGraphNodeCache& NodeCache = DocumentCache->GetNodeCache();
+	const IDocumentGraphEdgeCache& EdgeCache = DocumentCache->GetEdgeCache();
 
 	if (const int32* IndexPtr = NodeCache.FindNodeIndex(InNodeID))
 	{
 		const int32 Index = *IndexPtr; // Copy off indices as the pointer may be modified when notifying the cache below
 
-		FMetasoundFrontendDocument& Document = GetDocument();
-		FMetasoundFrontendGraph& Graph = Document.RootGraph.Graph;
+		FMetasoundFrontendDocument& Document = GetDocumentChecked();
+		FMetasoundFrontendGraph& Graph = FindBuildGraphChecked();
 		TArray<FMetasoundFrontendNode>& Nodes = Graph.Nodes;
 		const FMetasoundFrontendNode& Node = Nodes[Index];
 		const FGuid& NodeID = Node.GetID();
@@ -3008,7 +3018,7 @@ int32 FMetaSoundFrontendDocumentBuilder::RemoveNodeLocation(const FGuid& InNodeI
 	const IDocumentGraphNodeCache& NodeCache = DocumentCache->GetNodeCache();
 	if (const int32* NodeIndex = NodeCache.FindNodeIndex(InNodeID))
 	{
-		FMetasoundFrontendNode& Node = GetDocument().RootGraph.Graph.Nodes[*NodeIndex];
+		FMetasoundFrontendNode& Node = FindBuildGraphChecked().Nodes[*NodeIndex];
 		FMetasoundFrontendNodeStyle& Style = Node.Style;
 		if (InLocationGuid)
 		{
@@ -3026,6 +3036,24 @@ int32 FMetaSoundFrontendDocumentBuilder::RemoveNodeLocation(const FGuid& InNodeI
 }
 #endif // WITH_EDITOR
 
+bool FMetaSoundFrontendDocumentBuilder::RemoveUnusedDependencies()
+{
+	bool bDidEdit = false;
+
+	TArray<FMetasoundFrontendClass>& Dependencies = GetDocumentChecked().Dependencies;
+	for (int32 i = Dependencies.Num() - 1; i >= 0; --i)
+	{
+		const FGuid& ClassID = Dependencies[i].ID;
+		if (!DocumentCache->GetNodeCache().ContainsNodesOfClassID(ClassID))
+		{
+			const bool bRemoved = RemoveDependency(ClassID);
+			bDidEdit |= ensureAlwaysMsgf(bRemoved, TEXT("Failed to remove dependency that was found on document and was not referenced by nodes"));
+		}
+	}
+
+	return bDidEdit;
+}
+
 bool FMetaSoundFrontendDocumentBuilder::RenameRootGraphClass(const FMetasoundFrontendClassName& InName)
 {
 	return false;
@@ -3041,11 +3069,31 @@ void FMetaSoundFrontendDocumentBuilder::ReloadCache()
 #if WITH_EDITOR
 void FMetaSoundFrontendDocumentBuilder::SetAuthor(const FString& InAuthor)
 {
-	FMetasoundFrontendClassMetadata& ClassMetadata = GetDocument().RootGraph.Metadata;
+	FMetasoundFrontendClassMetadata& ClassMetadata = GetDocumentChecked().RootGraph.Metadata;
 	ClassMetadata.SetAuthor(InAuthor);
 }
 #endif // WITH_EDITOR
 
+bool FMetaSoundFrontendDocumentBuilder::SetBuildPageID(const FGuid& InBuildPageID)
+{
+	using namespace Metasound::Frontend;
+
+	FMetasoundFrontendDocument& Document = GetDocumentChecked();
+	if (const FMetasoundFrontendGraph* BuildGraph = Document.RootGraph.FindConstGraph(InBuildPageID))
+	{
+		if (BuildPageID != BuildGraph->PageID)
+		{
+			BuildPageID = BuildGraph->PageID;
+
+			constexpr bool bPrimeCache = false;
+			DocumentCache = FDocumentCache::Create(Document, DocumentDelegates.ToSharedRef(), BuildPageID, bPrimeCache);
+		}
+		return true;
+	}
+
+	return false;
+}
+
 bool FMetaSoundFrontendDocumentBuilder::SetGraphInputAccessType(FName InputName, EMetasoundFrontendVertexAccessType AccessType)
 {
 	using namespace Metasound::Frontend;
@@ -3057,7 +3105,7 @@ bool FMetaSoundFrontendDocumentBuilder::SetGraphInputAccessType(FName InputName,
 
 	if (const int32* Index = DocumentCache->GetInterfaceCache().FindInputIndex(InputName))
 	{
-		FMetasoundFrontendDocument& Document = GetDocument();
+		FMetasoundFrontendDocument& Document = GetDocumentChecked();
 		FMetasoundFrontendClassInput& GraphInput = Document.RootGraph.Interface.Inputs[*Index];
 
 		if (GraphInput.AccessType != AccessType)
@@ -3068,7 +3116,7 @@ bool FMetaSoundFrontendDocumentBuilder::SetGraphInputAccessType(FName InputName,
 				const IDocumentGraphNodeCache& NodeCache = DocumentCache->GetNodeCache();
 				if (const int32* NodeIndex = NodeCache.FindNodeIndex(GraphInput.NodeID))
 				{
-					FMetasoundFrontendNode& Node = Document.RootGraph.Graph.Nodes[*NodeIndex];
+					FMetasoundFrontendNode& Node = FindBuildGraphChecked().Nodes[*NodeIndex];
 					const FMetasoundFrontendVertex& NodeOutput = Node.Interface.Outputs.Last();
 					IterateNodesConnectedWithVertex({ GraphInput.NodeID, NodeOutput.VertexID }, [this, &AccessType, &Document](const FMetasoundFrontendEdge& Edge, FMetasoundFrontendNode& ConnectedNode)
 					{
@@ -3126,7 +3174,7 @@ bool FMetaSoundFrontendDocumentBuilder::SetGraphInputDataType(FName InputName, F
 	{
 		if (const int32* Index = DocumentCache->GetInterfaceCache().FindInputIndex(InputName))
 		{
-			FMetasoundFrontendDocument& Document = GetDocument();
+			FMetasoundFrontendDocument& Document = GetDocumentChecked();
 			FMetasoundFrontendClassInput& GraphInput = Document.RootGraph.Interface.Inputs[*Index];
 			if (GraphInput.TypeName != DataType)
 			{
@@ -3138,7 +3186,7 @@ bool FMetaSoundFrontendDocumentBuilder::SetGraphInputDataType(FName InputName, F
 				const IDocumentGraphNodeCache& NodeCache = DocumentCache->GetNodeCache();
 				if (const int32* NodeIndex = NodeCache.FindNodeIndex(GraphInput.NodeID))
 				{
-					FMetasoundFrontendNode& Node = Document.RootGraph.Graph.Nodes[*NodeIndex];
+					FMetasoundFrontendNode& Node = FindBuildGraphChecked().Nodes[*NodeIndex];
 					FMetasoundFrontendVertex& NodeOutput = Node.Interface.Outputs.Last();
 					IterateNodesConnectedWithVertex({ GraphInput.NodeID, NodeOutput.VertexID }, [this, &DataType](const FMetasoundFrontendEdge& Edge, FMetasoundFrontendNode& ConnectedNode)
 					{
@@ -3187,7 +3235,7 @@ bool FMetaSoundFrontendDocumentBuilder::SetGraphInputDefault(FName InputName, co
 	using namespace Metasound::Frontend;
 
 	auto NameMatchesInput = [&InputName](const FMetasoundFrontendClassInput& Input) { return Input.Name == InputName; };
-	TArray<FMetasoundFrontendClassInput>& Inputs = GetDocument().RootGraph.Interface.Inputs;
+	TArray<FMetasoundFrontendClassInput>& Inputs = GetDocumentChecked().RootGraph.Interface.Inputs;
 
 	const int32 Index = Inputs.IndexOfByPredicate(NameMatchesInput);
 	if (Index != INDEX_NONE)
@@ -3215,7 +3263,7 @@ bool FMetaSoundFrontendDocumentBuilder::SetGraphInputDefault(FName InputName, co
 
 bool FMetaSoundFrontendDocumentBuilder::SetGraphInputInheritsDefault(FName InName, bool bInputInheritsDefault)
 {
-	FMetasoundFrontendGraphClassPresetOptions& PresetOptions = GetDocument().RootGraph.PresetOptions;
+	FMetasoundFrontendGraphClassPresetOptions& PresetOptions = GetDocumentChecked().RootGraph.PresetOptions;
 	if (bInputInheritsDefault)
 	{
 		if (PresetOptions.bIsPreset)
@@ -3245,7 +3293,7 @@ bool FMetaSoundFrontendDocumentBuilder::SetGraphOutputAccessType(FName OutputNam
 
 	if (const int32* Index = DocumentCache->GetInterfaceCache().FindOutputIndex(OutputName))
 	{
-		FMetasoundFrontendDocument& Document = GetDocument();
+		FMetasoundFrontendDocument& Document = GetDocumentChecked();
 		FMetasoundFrontendClassOutput& GraphOutput = Document.RootGraph.Interface.Outputs[*Index];
 		if (GraphOutput.AccessType != AccessType)
 		{
@@ -3255,7 +3303,7 @@ bool FMetaSoundFrontendDocumentBuilder::SetGraphOutputAccessType(FName OutputNam
 				const IDocumentGraphNodeCache& NodeCache = DocumentCache->GetNodeCache();
 				if (const int32* NodeIndex = NodeCache.FindNodeIndex(GraphOutput.NodeID))
 				{
-					FMetasoundFrontendNode& Node = Document.RootGraph.Graph.Nodes[*NodeIndex];
+					FMetasoundFrontendNode& Node = FindBuildGraphChecked().Nodes[*NodeIndex];
 					const FMetasoundFrontendVertex& NodeInput = Node.Interface.Inputs.Last();
 					IterateNodesConnectedWithVertex({ GraphOutput.NodeID, NodeInput.VertexID }, [this, &AccessType](const FMetasoundFrontendEdge& Edge, FMetasoundFrontendNode& ConnectedNode)
 					{
@@ -3295,7 +3343,7 @@ bool FMetaSoundFrontendDocumentBuilder::SetGraphOutputDataType(FName OutputName,
 	{
 		if (const int32* Index = DocumentCache->GetInterfaceCache().FindOutputIndex(OutputName))
 		{
-			FMetasoundFrontendDocument& Document = GetDocument();
+			FMetasoundFrontendDocument& Document = GetDocumentChecked();
 			FMetasoundFrontendClassOutput& GraphOutput = Document.RootGraph.Interface.Outputs[*Index];
 			if (GraphOutput.TypeName != DataType)
 			{
@@ -3304,7 +3352,7 @@ bool FMetaSoundFrontendDocumentBuilder::SetGraphOutputDataType(FName OutputName,
 				const IDocumentGraphNodeCache& NodeCache = DocumentCache->GetNodeCache();
 				if (const int32* NodeIndex = NodeCache.FindNodeIndex(GraphOutput.NodeID))
 				{
-					FMetasoundFrontendNode& Node = Document.RootGraph.Graph.Nodes[*NodeIndex];
+					FMetasoundFrontendNode& Node = FindBuildGraphChecked().Nodes[*NodeIndex];
 
 					FMetasoundFrontendLiteral DefaultLiteral;
 					DefaultLiteral.SetFromLiteral(IDataTypeRegistry::Get().CreateDefaultLiteral(DataType));
@@ -3344,7 +3392,7 @@ void FMetaSoundFrontendDocumentBuilder::SetMemberMetadata(UMetaSoundFrontendMemb
 {
 	check(NewMetadata.MemberID.IsValid());
 
-	TMap<FGuid, TObjectPtr<UMetaSoundFrontendMemberMetadata>>& LiteralMetadata = GetDocument().Metadata.MemberMetadata;
+	TMap<FGuid, TObjectPtr<UMetaSoundFrontendMemberMetadata>>& LiteralMetadata = GetDocumentChecked().Metadata.MemberMetadata;
 	LiteralMetadata.Remove(NewMetadata.MemberID);
 	LiteralMetadata.Add(NewMetadata.MemberID, NewMetadata);
 }
@@ -3379,7 +3427,7 @@ bool FMetaSoundFrontendDocumentBuilder::SetNodeInputDefault(const FGuid& InNodeI
 	const IDocumentGraphNodeCache& NodeCache = DocumentCache->GetNodeCache();
 	if (const int32* NodeIndex = NodeCache.FindNodeIndex(InNodeID))
 	{
-		FMetasoundFrontendGraph& Graph = GetDocument().RootGraph.Graph;
+		FMetasoundFrontendGraph& Graph = FindBuildGraphChecked();
 		FMetasoundFrontendNode& Node = Graph.Nodes[*NodeIndex];
 
 		auto IsVertex = [&InVertexID](const FMetasoundFrontendVertex& Vertex) { return Vertex.VertexID == InVertexID; };
@@ -3449,7 +3497,7 @@ bool FMetaSoundFrontendDocumentBuilder::SetNodeLocation(const FGuid& InNodeID, c
 
 void FMetaSoundFrontendDocumentBuilder::SetVersionNumber(const FMetasoundFrontendVersionNumber& InDocumentVersionNumber)
 {
-	GetDocument().Metadata.Version.Number = InDocumentVersionNumber;
+	GetDocumentChecked().Metadata.Version.Number = InDocumentVersionNumber;
 }
 
 bool FMetaSoundFrontendDocumentBuilder::SpliceVariableNodeFromStack(const FGuid& InNodeID)
@@ -3457,11 +3505,8 @@ bool FMetaSoundFrontendDocumentBuilder::SpliceVariableNodeFromStack(const FGuid&
 	using namespace Metasound;
 	using namespace Metasound::Frontend;
 
-	FMetasoundFrontendGraphClass& GraphClass = GetDocument().RootGraph;
-	FMetasoundFrontendGraph& Graph = GraphClass.Graph;
-
+	FMetasoundFrontendGraph& Graph = FindBuildGraphChecked();
 	const IDocumentGraphEdgeCache& EdgeCache = DocumentCache->GetEdgeCache();
-
 	FMetasoundFrontendVertexHandle FromVariableVertexHandle;
 	{
 		const FMetasoundFrontendVertex* InputVertex = FindNodeInput(InNodeID, VariableNames::InputVariableName);
@@ -3562,7 +3607,7 @@ bool FMetaSoundFrontendDocumentBuilder::SwapGraphInput(const FMetasoundFrontendC
 	{
 		const int32* NodeIndex = NodeCache.FindNodeIndex(NewInputNode->GetID());
 		checkf(NodeIndex, TEXT("Cache was not updated to reflect newly added input node"));
-		FMetasoundFrontendNode& NewNode = GetDocument().RootGraph.Graph.Nodes[*NodeIndex];
+		FMetasoundFrontendNode& NewNode = FindBuildGraphChecked().Nodes[*NodeIndex];
 		NewNode.Style.Display.Locations = Locations;
 	}
 #endif // WITH_EDITOR
@@ -3639,7 +3684,7 @@ bool FMetaSoundFrontendDocumentBuilder::SwapGraphOutput(const FMetasoundFrontend
 	{
 		const int32* NodeIndex = NodeCache.FindNodeIndex(NewOutputNode->GetID());
 		checkf(NodeIndex, TEXT("Cache was not updated to reflect newly added input node"));
-		FMetasoundFrontendNode& NewNode = GetDocument().RootGraph.Graph.Nodes[*NodeIndex];
+		FMetasoundFrontendNode& NewNode = FindBuildGraphChecked().Nodes[*NodeIndex];
 		NewNode.Style.Display.Locations = Locations;
 	}
 #endif // WITH_EDITOR
@@ -3660,8 +3705,7 @@ bool FMetaSoundFrontendDocumentBuilder::UnlinkVariableNode(const FGuid& InNodeID
 {
 	auto IsNodeID = [&InNodeID](const FGuid& TestID) { return TestID == InNodeID; };
 
-	FMetasoundFrontendGraphClass& GraphClass = GetDocument().RootGraph;
-	FMetasoundFrontendGraph& Graph = GraphClass.Graph;
+	FMetasoundFrontendGraph& Graph = FindBuildGraphChecked();
 	for (FMetasoundFrontendVariable& Variable : Graph.Variables)
 	{
 		if (Variable.MutatorNodeID == InNodeID)
@@ -3698,7 +3742,7 @@ bool FMetaSoundFrontendDocumentBuilder::UnlinkVariableNode(const FGuid& InNodeID
 
 bool FMetaSoundFrontendDocumentBuilder::UpdateDependencyClassNames(const TMap<FMetasoundFrontendClassName, FMetasoundFrontendClassName>& OldToNewReferencedClassNames)
 {
-	for (FMetasoundFrontendClass& Dependency : GetDocument().Dependencies)
+	for (FMetasoundFrontendClass& Dependency : GetDocumentChecked().Dependencies)
 	{
 		const FMetasoundFrontendClassName OldName = Dependency.Metadata.GetClassName();
 		if (const FMetasoundFrontendClassName* NewName = OldToNewReferencedClassNames.Find(OldName))
diff --git a/Engine/Plugins/Runtime/Metasound/Source/MetasoundFrontend/Private/MetasoundFrontendDocumentCache.cpp b/Engine/Plugins/Runtime/Metasound/Source/MetasoundFrontend/Private/MetasoundFrontendDocumentCache.cpp
--- a/Engine/Plugins/Runtime/Metasound/Source/MetasoundFrontend/Private/MetasoundFrontendDocumentCache.cpp
+++ b/Engine/Plugins/Runtime/Metasound/Source/MetasoundFrontend/Private/MetasoundFrontendDocumentCache.cpp
@@ -97,19 +97,23 @@ namespace Metasound::Frontend
 	{
 	}
 
-	void FDocumentCache::Init(bool bPrimeCache)
+	void FDocumentCache::Init(const FGuid& InPageID, bool bPrimeCache)
 	{
 		METASOUND_TRACE_CPUPROFILER_EVENT_SCOPE(Metasound::Frontend::FDocumentCache::Init);
 		check(Document);
 
+		PageID = InPageID;
 		if (bPrimeCache)
 		{
 			METASOUND_TRACE_CPUPROFILER_EVENT_SCOPE(Metasound::Frontend::FDocumentCache::Init_Prime);
 
 			TSharedRef<const FDocumentCache> ThisShared = StaticCastSharedRef<const FDocumentCache>(AsShared());
 			DependencyCache = MakeShared<FDocumentDependencyCache>(GetDocument());
-			EdgeCache = FDocumentGraphEdgeCache::Create(ThisShared, ModifyDelegates->EdgeDelegates);
-			NodeCache = FDocumentGraphNodeCache::Create(ThisShared, ModifyDelegates->NodeDelegates);
+
+			// Currently just hardcoded to use cache for default graph (i.e. PageID of zero guid).  Needs to be
+			// transitioned to support multiple caches for each page, where cache is generated on demand per page.
+			EdgeCache = FDocumentGraphEdgeCache::Create(ThisShared, PageID, ModifyDelegates->EdgeDelegates);
+			NodeCache = FDocumentGraphNodeCache::Create(ThisShared, PageID, ModifyDelegates->NodeDelegates);
 			InterfaceCache = FDocumentGraphInterfaceCache::Create(ThisShared, ModifyDelegates->InterfaceDelegates);
 		}
 
@@ -124,13 +128,16 @@ namespace Metasound::Frontend
 		return *Document;
 	}
 
-	TSharedRef<FDocumentCache> FDocumentCache::Create(const FMetasoundFrontendDocument& InDocument, TSharedRef<FDocumentModifyDelegates> Delegates, bool bPrimeCache = false)
+	TSharedRef<FDocumentCache> FDocumentCache::Create(const FMetasoundFrontendDocument& InDocument, TSharedRef<FDocumentModifyDelegates> Delegates, const FGuid& InPageID, bool bPrimeCache)
 	{
 		METASOUND_TRACE_CPUPROFILER_EVENT_SCOPE(Metasound::Frontend::FDocumentCache::Create);
 
 		// Factory style constructor as restrictions on construction of shared pointers disallows passing of this document cache's pointer to sub-caches
 		TSharedRef<FDocumentCache> Cache = MakeShared<FDocumentCache>(InDocument, Delegates);
-		Cache->Init(bPrimeCache);
+
+		// Right now the assumption is the whole doc cache is invalidated when swapping pages.  Need to migrate to supporting associated caches per page
+		// (i.e. a node and edge cache per page)
+		Cache->Init(InPageID, bPrimeCache);
 		return Cache;
 	}
 
@@ -204,7 +211,7 @@ namespace Metasound::Frontend
 		if (!EdgeCache.IsValid())
 		{
 			TSharedRef<const FDocumentCache> ThisShared = StaticCastSharedRef<const FDocumentCache>(AsShared());
-			EdgeCache = FDocumentGraphEdgeCache::Create(ThisShared, ModifyDelegates->EdgeDelegates);
+			EdgeCache = FDocumentGraphEdgeCache::Create(ThisShared, PageID, ModifyDelegates->EdgeDelegates);
 		}
 		return *EdgeCache;
 	}
@@ -214,7 +221,7 @@ namespace Metasound::Frontend
 		if (!NodeCache.IsValid())
 		{
 			TSharedRef<const FDocumentCache> ThisShared = StaticCastSharedRef<const FDocumentCache>(AsShared());
-			NodeCache = FDocumentGraphNodeCache::Create(ThisShared, ModifyDelegates->NodeDelegates);
+			NodeCache = FDocumentGraphNodeCache::Create(ThisShared, PageID, ModifyDelegates->NodeDelegates);
 		}
 		return *NodeCache;
 	}
@@ -423,10 +430,11 @@ namespace Metasound::Frontend
 	{
 	}
 
-	FDocumentGraphNodeCache::FDocumentGraphNodeCache(TSharedRef<const FDocumentCache> ParentCache)
+	FDocumentGraphNodeCache::FDocumentGraphNodeCache(TSharedRef<const FDocumentCache> ParentCache, const FGuid& InPageID)
 		: Parent(ParentCache)
+		, PageID(InPageID)
 	{
-		const FMetasoundFrontendGraph& Graph = Parent->GetDocument().RootGraph.Graph;
+		const FMetasoundFrontendGraph& Graph = Parent->GetDocument().RootGraph.FindConstGraphChecked(InPageID);
 		const TArray<FMetasoundFrontendNode>& Nodes = Graph.Nodes;
 		for (int32 Index = 0; Index < Nodes.Num(); ++Index)
 		{
@@ -445,12 +453,12 @@ namespace Metasound::Frontend
 		return IDToIndex.Contains(InNodeID);
 	}
 
-	TSharedRef<FDocumentGraphNodeCache> FDocumentGraphNodeCache::Create(TSharedRef<const FDocumentCache> ParentCache, FNodeModifyDelegates& OutDelegates)
+	TSharedRef<FDocumentGraphNodeCache> FDocumentGraphNodeCache::Create(TSharedRef<const FDocumentCache> ParentCache, const FGuid& InPageID, FNodeModifyDelegates& OutDelegates)
 	{
 		METASOUND_TRACE_CPUPROFILER_EVENT_SCOPE(Metasound::Frontend::FDocumentGraphNodeCache::Create);
 
 		// Factory style constructor as restrictions on construction of shared pointers disallows passing of this document cache's pointer to sub-caches
-		TSharedRef<FDocumentGraphNodeCache> Cache = MakeShared<FDocumentGraphNodeCache>(ParentCache);
+		TSharedRef<FDocumentGraphNodeCache> Cache = MakeShared<FDocumentGraphNodeCache>(ParentCache, InPageID);
 		Cache->Init(OutDelegates);
 		return Cache;
 	}
@@ -466,7 +474,7 @@ namespace Metasound::Frontend
 		if (const TArray<int32>* NodeIndices = ClassIDToNodeIndices.Find(InClassID))
 		{
 			const FMetasoundFrontendDocument& Document = Parent->GetDocument();
-			const FMetasoundFrontendGraph& Graph = Document.RootGraph.Graph;
+			const FMetasoundFrontendGraph& Graph = Document.RootGraph.FindConstGraphChecked(PageID);
 			Algo::Transform(*NodeIndices, Nodes, [&Graph](const int32& Index) { return &Graph.Nodes[Index]; });
 		}
 		return Nodes;
@@ -477,7 +485,7 @@ namespace Metasound::Frontend
 		if (const int32* NodeIndex = IDToIndex.Find(InNodeID))
 		{
 			const FMetasoundFrontendDocument& Document = Parent->GetDocument();
-			const FMetasoundFrontendGraph& Graph = Document.RootGraph.Graph;
+			const FMetasoundFrontendGraph& Graph = Document.RootGraph.FindConstGraphChecked(PageID);
 			return &Graph.Nodes[*NodeIndex];
 		}
 
@@ -602,7 +610,7 @@ namespace Metasound::Frontend
 				const FMetasoundFrontendDocument& Doc = Parent->GetDocument();
 				for (int32 Index : EdgeIndices)
 				{
-					const FMetasoundFrontendEdge& EdgeConnectedToOutput = Doc.RootGraph.Graph.Edges[Index];
+					const FMetasoundFrontendEdge& EdgeConnectedToOutput = Doc.RootGraph.FindConstGraphChecked(PageID).Edges[Index];
 					TArray<const FMetasoundFrontendVertex*> ReroutedVertices;
 					if (ConnectedNodes)
 					{
@@ -657,7 +665,7 @@ namespace Metasound::Frontend
 				const FMetasoundFrontendDocument& Doc = Parent->GetDocument();
 				for (int32 Index : EdgeIndices)
 				{
-					const FMetasoundFrontendEdge& EdgeConnectedToOutput = Doc.RootGraph.Graph.Edges[Index];
+					const FMetasoundFrontendEdge& EdgeConnectedToOutput = Doc.RootGraph.FindConstGraphChecked(PageID).Edges[Index];
 
 					TArray<const FMetasoundFrontendVertex*> ReroutedVertices;
 					if (ConnectedNodes)
@@ -710,7 +718,7 @@ namespace Metasound::Frontend
 				const FMetasoundFrontendDocument& Doc = Parent->GetDocument();
 				if (const int32* ConnectedEdgeIndex = EdgeCache.FindEdgeIndexToNodeInput(InNodeID, RerouteInput.VertexID))
 				{
-					const FMetasoundFrontendEdge& EdgeConnectedToInput = Doc.RootGraph.Graph.Edges[*ConnectedEdgeIndex];
+					const FMetasoundFrontendEdge& EdgeConnectedToInput = Doc.RootGraph.FindConstGraphChecked(PageID).Edges[*ConnectedEdgeIndex];
 					return FindReroutedOutputVertex(EdgeConnectedToInput.FromNodeID, EdgeConnectedToInput.FromVertexID, ConnectedNode, bOutIsRerouted);
 				}
 
@@ -749,7 +757,7 @@ namespace Metasound::Frontend
 				const FMetasoundFrontendDocument& Doc = Parent->GetDocument();
 				if (const int32* ConnectedEdgeIndex = EdgeCache.FindEdgeIndexToNodeInput(InNodeID, RerouteInput.VertexID))
 				{
-					const FMetasoundFrontendEdge& EdgeConnectedToInput = Doc.RootGraph.Graph.Edges[*ConnectedEdgeIndex];
+					const FMetasoundFrontendEdge& EdgeConnectedToInput = Doc.RootGraph.FindConstGraphChecked(PageID).Edges[*ConnectedEdgeIndex];
 					return FindReroutedOutputVertex(EdgeConnectedToInput.FromNodeID, EdgeConnectedToInput.FromVertexID, ConnectedNode, bOutIsRerouted);
 				}
 
@@ -781,7 +789,7 @@ namespace Metasound::Frontend
 	void FDocumentGraphNodeCache::OnNodeAdded(int32 InNewIndex)
 	{
 		const FMetasoundFrontendDocument& Document = Parent->GetDocument();
-		const FMetasoundFrontendGraph& Graph = Document.RootGraph.Graph;
+		const FMetasoundFrontendGraph& Graph = Document.RootGraph.FindConstGraphChecked(PageID);
 		const FMetasoundFrontendNode& Node = Graph.Nodes[InNewIndex];
 		IDToIndex.Add(Node.GetID(), InNewIndex);
 		ClassIDToNodeIndices.FindOrAdd(Node.ClassID).Add(InNewIndex);
@@ -799,7 +807,7 @@ namespace Metasound::Frontend
 		using namespace Metasound::Frontend;
 
 		const FMetasoundFrontendDocument& Document = Parent->GetDocument();
-		const TArray<FMetasoundFrontendNode>& Nodes = Document.RootGraph.Graph.Nodes;
+		const TArray<FMetasoundFrontendNode>& Nodes = Document.RootGraph.FindConstGraphChecked(PageID).Nodes;
 		const FMetasoundFrontendNode& SwapNode = Nodes[SwapIndex];
 		const FMetasoundFrontendNode& LastNode = Nodes[LastIndex];
 
@@ -819,10 +827,11 @@ namespace Metasound::Frontend
 	{
 	}
 
-	FDocumentGraphEdgeCache::FDocumentGraphEdgeCache(TSharedRef<const FDocumentCache> ParentCache)
+	FDocumentGraphEdgeCache::FDocumentGraphEdgeCache(TSharedRef<const FDocumentCache> ParentCache, const FGuid& InPageID)
 		: Parent(ParentCache)
+		, PageID(InPageID)
 	{
-		const TArray<FMetasoundFrontendEdge>& Edges = Parent->GetDocument().RootGraph.Graph.Edges;
+		const TArray<FMetasoundFrontendEdge>& Edges = Parent->GetDocument().RootGraph.FindConstGraphChecked(InPageID).Edges;
 		for (int32 Index = 0; Index < Edges.Num(); ++Index)
 		{
 			const FMetasoundFrontendEdge& Edge = Edges[Index];
@@ -837,19 +846,19 @@ namespace Metasound::Frontend
 		if (const int32* Index = InputToEdgeIndex.Find(InputPair))
 		{
 			const FMetasoundFrontendDocument& Document = Parent->GetDocument();
-			const FMetasoundFrontendEdge& Edge = Document.RootGraph.Graph.Edges[*Index];
+			const FMetasoundFrontendEdge& Edge = Document.RootGraph.FindConstGraphChecked(PageID).Edges[*Index];
 			return Edge.FromNodeID == InEdge.FromNodeID && Edge.FromVertexID == InEdge.FromVertexID;
 		}
 
 		return false;
 	}
 
-	TSharedRef<FDocumentGraphEdgeCache> FDocumentGraphEdgeCache::Create(TSharedRef<const FDocumentCache> ParentCache, FEdgeModifyDelegates& OutDelegates)
+	TSharedRef<FDocumentGraphEdgeCache> FDocumentGraphEdgeCache::Create(TSharedRef<const FDocumentCache> ParentCache, const FGuid& InPageID, FEdgeModifyDelegates& OutDelegates)
 	{
 		METASOUND_TRACE_CPUPROFILER_EVENT_SCOPE(Metasound::Frontend::FDocumentGraphEdgeCache::Create);
 
 		// Factory style constructor as restrictions on construction of shared pointers disallows passing of this document cache's pointer to sub-caches
-		TSharedRef<FDocumentGraphEdgeCache> Cache = MakeShared<FDocumentGraphEdgeCache>(ParentCache);
+		TSharedRef<FDocumentGraphEdgeCache> Cache = MakeShared<FDocumentGraphEdgeCache>(ParentCache, InPageID);
 		Cache->Init(OutDelegates);
 		return Cache;
 	}
@@ -859,7 +868,7 @@ namespace Metasound::Frontend
 		TArray<const FMetasoundFrontendEdge*> Edges;
 
 		const FMetasoundFrontendDocument& Document = Parent->GetDocument();
-		const FMetasoundFrontendGraph& Graph = Document.RootGraph.Graph;
+		const FMetasoundFrontendGraph& Graph = Document.RootGraph.FindConstGraphChecked(PageID);
 
 		const FMetasoundFrontendVertexHandle Handle { InNodeID, InVertexID };
 		if (const int32* Index = InputToEdgeIndex.Find(Handle))
@@ -914,7 +923,7 @@ namespace Metasound::Frontend
 	void FDocumentGraphEdgeCache::OnEdgeAdded(int32 InNewIndex)
 	{
 		const FMetasoundFrontendDocument& Document = Parent->GetDocument();
-		const FMetasoundFrontendGraph& Graph = Document.RootGraph.Graph;
+		const FMetasoundFrontendGraph& Graph = Document.RootGraph.FindConstGraphChecked(PageID);
 		const FMetasoundFrontendEdge& NewEdge = Graph.Edges[InNewIndex];
 		InputToEdgeIndex.Add(NewEdge.GetToVertexHandle(), InNewIndex);
 		OutputToEdgeIndices.FindOrAdd(NewEdge.GetFromVertexHandle()).Add(InNewIndex);
@@ -925,7 +934,7 @@ namespace Metasound::Frontend
 	void FDocumentGraphEdgeCache::OnRemoveSwappingEdge(int32 SwapIndex, int32 LastIndex)
 	{
 		const FMetasoundFrontendDocument& Document = Parent->GetDocument();
-		const TArray<FMetasoundFrontendEdge>& Edges = Document.RootGraph.Graph.Edges;
+		const TArray<FMetasoundFrontendEdge>& Edges = Document.RootGraph.FindConstGraphChecked(PageID).Edges;
 
 		const FMetasoundFrontendEdge& SwapEdge = Edges[SwapIndex];
 		const FMetasoundFrontendEdge& LastEdge = Edges[LastIndex];
diff --git a/Engine/Plugins/Runtime/Metasound/Source/MetasoundFrontend/Private/MetasoundFrontendDocumentCache.h b/Engine/Plugins/Runtime/Metasound/Source/MetasoundFrontend/Private/MetasoundFrontendDocumentCache.h
--- a/Engine/Plugins/Runtime/Metasound/Source/MetasoundFrontend/Private/MetasoundFrontendDocumentCache.h
+++ b/Engine/Plugins/Runtime/Metasound/Source/MetasoundFrontend/Private/MetasoundFrontendDocumentCache.h
@@ -23,10 +23,10 @@ namespace Metasound::Frontend
 
 		// Only exists to make TSharedFromThis happy, but should never be called
 		FDocumentGraphEdgeCache();
-		FDocumentGraphEdgeCache(TSharedRef<const FDocumentCache> ParentCache);
+		FDocumentGraphEdgeCache(TSharedRef<const FDocumentCache> ParentCache, const FGuid& InPageID);
 
 	public:
-		static TSharedRef<FDocumentGraphEdgeCache> Create(TSharedRef<const FDocumentCache> ParentCache, FEdgeModifyDelegates& OutDelegates);
+		static TSharedRef<FDocumentGraphEdgeCache> Create(TSharedRef<const FDocumentCache> ParentCache, const FGuid& InPageID, FEdgeModifyDelegates& OutDelegates);
 		virtual ~FDocumentGraphEdgeCache() = default;
 
 		// IDocumentGraphEdgeCache implementation
@@ -53,7 +53,7 @@ namespace Metasound::Frontend
 		TMap<FMetasoundFrontendVertexHandle, int32> InputToEdgeIndex;
 
 		TSharedRef<const FDocumentCache> Parent;
-
+		FGuid PageID;
 		int32 TransactionCount = 0;
 	};
 
@@ -65,10 +65,10 @@ namespace Metasound::Frontend
 
 		// Only exists to make TSharedFromThis happy, but should never be called
 		FDocumentGraphNodeCache();
-		FDocumentGraphNodeCache(TSharedRef<const FDocumentCache> ParentCache);
+		FDocumentGraphNodeCache(TSharedRef<const FDocumentCache> ParentCache, const FGuid& InPageID);
 
 	public:
-		static TSharedRef<FDocumentGraphNodeCache> Create(TSharedRef<const FDocumentCache> ParentCache, FNodeModifyDelegates& OutDelegates);
+		static TSharedRef<FDocumentGraphNodeCache> Create(TSharedRef<const FDocumentCache> ParentCache, const FGuid& InPageID, FNodeModifyDelegates& OutDelegates);
 		virtual ~FDocumentGraphNodeCache() = default;
 
 		// IDocumentGraphNodeCache implementation
@@ -108,7 +108,7 @@ namespace Metasound::Frontend
 		TSortedMap<FGuid, TArray<int32>> ClassIDToNodeIndices;
 
 		TSharedRef<const FDocumentCache> Parent;
-
+		FGuid PageID;
 		int32 TransactionCount = 0;
 	};
 
@@ -162,7 +162,7 @@ namespace Metasound::Frontend
 		FDocumentCache(const FMetasoundFrontendDocument& InDocument, TSharedRef<FDocumentModifyDelegates> Delegates);
 
 	public:
-		static TSharedRef<FDocumentCache> Create(const FMetasoundFrontendDocument& InDocument, TSharedRef<FDocumentModifyDelegates> Delegates, bool bPrimeCache);
+		static TSharedRef<FDocumentCache> Create(const FMetasoundFrontendDocument& InDocument, TSharedRef<FDocumentModifyDelegates> Delegates, const FGuid& InPageID, bool bPrimeCache);
 		virtual ~FDocumentCache() = default;
 
 		virtual bool ContainsDependency(const FNodeRegistryKey& InClassKey) const override;
@@ -181,7 +181,7 @@ namespace Metasound::Frontend
 		int32 GetTransactionCount() const;
 
 	private:
-		void Init(bool bPrimeCache);
+		void Init(const FGuid& InPageID, bool bPrimeCache);
 
 		void OnDependencyAdded(int32 InNewIndex);
 		void OnRemoveSwappingDependency(int32 SwapIndex, int32 LastIndex);
@@ -216,5 +216,6 @@ namespace Metasound::Frontend
 
 		// Number of transactions processed since builder was instantiated
 		int32 TransactionCount = 0;
+		FGuid PageID;
 	};
 } // namespace Metasound::Frontend
diff --git a/Engine/Plugins/Runtime/Metasound/Source/MetasoundFrontend/Private/MetasoundFrontendDocumentController.cpp b/Engine/Plugins/Runtime/Metasound/Source/MetasoundFrontend/Private/MetasoundFrontendDocumentController.cpp
--- a/Engine/Plugins/Runtime/Metasound/Source/MetasoundFrontend/Private/MetasoundFrontendDocumentController.cpp
+++ b/Engine/Plugins/Runtime/Metasound/Source/MetasoundFrontend/Private/MetasoundFrontendDocumentController.cpp
@@ -89,7 +89,7 @@ namespace Metasound
 				bool bSuccess = true;
 				FMetasoundFrontendGraphClass SubgraphCopy(InGraphToCopy);
 
-				for (FMetasoundFrontendNode& Node : SubgraphCopy.Graph.Nodes)
+				for (FMetasoundFrontendNode& Node : SubgraphCopy.GetDefaultGraph().Nodes)
 				{
 					const FGuid OriginalClassID = Node.ClassID;
 
@@ -443,11 +443,11 @@ namespace Metasound
 
 					auto AddGraphNodeClassIDsToSet = [&](const FMetasoundFrontendGraphClass& GraphClass)
 					{
-						Algo::ForEach(GraphClass.Graph.Nodes, AddNodeClassIDToSet);
+						Algo::ForEach(FindConstBuildGraphChecked(GraphClass).Nodes, AddNodeClassIDToSet);
 					};
 
 					// Referenced dependencies in root class
-					Algo::ForEach(Document->RootGraph.Graph.Nodes, AddNodeClassIDToSet);
+					Algo::ForEach(FindConstBuildGraphChecked(Document->RootGraph).Nodes, AddNodeClassIDToSet);
 
 					// Referenced dependencies in subgraphs
 					Algo::ForEach(Document->Subgraphs, AddGraphNodeClassIDsToSet);
diff --git a/Engine/Plugins/Runtime/Metasound/Source/MetasoundFrontend/Private/MetasoundFrontendDocumentController.h b/Engine/Plugins/Runtime/Metasound/Source/MetasoundFrontend/Private/MetasoundFrontendDocumentController.h
--- a/Engine/Plugins/Runtime/Metasound/Source/MetasoundFrontend/Private/MetasoundFrontendDocumentController.h
+++ b/Engine/Plugins/Runtime/Metasound/Source/MetasoundFrontend/Private/MetasoundFrontendDocumentController.h
@@ -30,10 +30,7 @@ namespace Metasound
 				if (IDocumentBuilderRegistry* Registry = IDocumentBuilderRegistry::Get())
 				{
 					const FMetasoundFrontendClassName& ClassName = InDocument.Get()->RootGraph.Metadata.GetClassName();
-
-					PRAGMA_DISABLE_DEPRECATION_WARNINGS
-					Registry->InvalidateDocumentCache(ClassName);
-					PRAGMA_ENABLE_DEPRECATION_WARNINGS
+					Registry->ReloadBuilder(ClassName);
 				}
 				return MakeShared<FDocumentController>(InDocument);
 			}
diff --git a/Engine/Plugins/Runtime/Metasound/Source/MetasoundFrontend/Private/MetasoundFrontendDocumentVersioning.cpp b/Engine/Plugins/Runtime/Metasound/Source/MetasoundFrontend/Private/MetasoundFrontendDocumentVersioning.cpp
--- a/Engine/Plugins/Runtime/Metasound/Source/MetasoundFrontend/Private/MetasoundFrontendDocumentVersioning.cpp
+++ b/Engine/Plugins/Runtime/Metasound/Source/MetasoundFrontend/Private/MetasoundFrontendDocumentVersioning.cpp
@@ -55,7 +55,7 @@ namespace Metasound::Frontend
 
 			bool Transform(FMetaSoundFrontendDocumentBuilder& OutDocumentBuilder) const
 			{
-				const FMetasoundFrontendDocumentMetadata& Metadata = OutDocumentBuilder.GetConstDocument().Metadata;
+				const FMetasoundFrontendDocumentMetadata& Metadata = OutDocumentBuilder.GetConstDocumentChecked().Metadata;
 
 				const FMetasoundFrontendVersionNumber TargetVersion = GetTargetVersion();
 				if (Metadata.Version.Number < TargetVersion)
diff --git a/Engine/Plugins/Runtime/Metasound/Source/MetasoundFrontend/Private/MetasoundFrontendGraph.cpp b/Engine/Plugins/Runtime/Metasound/Source/MetasoundFrontend/Private/MetasoundFrontendGraph.cpp
--- a/Engine/Plugins/Runtime/Metasound/Source/MetasoundFrontend/Private/MetasoundFrontendGraph.cpp
+++ b/Engine/Plugins/Runtime/Metasound/Source/MetasoundFrontend/Private/MetasoundFrontendGraph.cpp
@@ -352,13 +352,13 @@ namespace Metasound
 	bool FFrontendGraphBuilder::AddNodesToGraph(FBuildGraphContext& InGraphContext)
 	{
 		TSet<FNodeIDVertexID> GraphEdgeDestinations;
-		const TArray<FMetasoundFrontendEdge>& GraphEdges = InGraphContext.GraphClass.Graph.Edges;
+		const TArray<FMetasoundFrontendEdge>& GraphEdges = InGraphContext.GraphClass.GetConstDefaultGraph().Edges;
 		Algo::Transform(GraphEdges, GraphEdgeDestinations, [](const FMetasoundFrontendEdge& Edge) 
 		{
 			return FNodeIDVertexID{Edge.ToNodeID, Edge.ToVertexID};
 		});
 
-		for (const FMetasoundFrontendNode& Node : InGraphContext.GraphClass.Graph.Nodes)
+		for (const FMetasoundFrontendNode& Node : InGraphContext.GraphClass.GetConstDefaultGraph().Nodes)
 		{
 			const FMetasoundFrontendClass* NodeClass = InGraphContext.BuildContext.FrontendClasses.FindRef(Node.ClassID);
 
@@ -427,7 +427,7 @@ namespace Metasound
 
 					case EMetasoundFrontendClassType::Variable:
 					{
-						TSharedPtr<const INode> VariableNode(CreateVariableNode(InGraphContext.BuildContext, Node, InGraphContext.GraphClass.Graph).Release());
+						TSharedPtr<const INode> VariableNode(CreateVariableNode(InGraphContext.BuildContext, Node, InGraphContext.GraphClass.GetConstDefaultGraph()).Release());
 						InGraphContext.Graph->AddNode(Node.GetID(), VariableNode);
 					}
 					break;
@@ -466,7 +466,7 @@ namespace Metasound
 		TMap<FNodeIDVertexID, FCoreNodeAndFrontendVertex> NodeDestinationsByID;
 
 		// Add nodes to NodeID/VertexID map
-		for (const FMetasoundFrontendNode& Node : InGraphContext.GraphClass.Graph.Nodes)
+		for (const FMetasoundFrontendNode& Node : InGraphContext.GraphClass.GetConstDefaultGraph().Nodes)
 		{
 			const INode* CoreNode = InGraphContext.Graph->FindNode(Node.GetID());
 			if (nullptr == CoreNode)
@@ -486,7 +486,7 @@ namespace Metasound
 			}
 		};
 
-		for (const FMetasoundFrontendEdge& Edge : InGraphContext.GraphClass.Graph.Edges)
+		for (const FMetasoundFrontendEdge& Edge : InGraphContext.GraphClass.GetConstDefaultGraph().Edges)
 		{
 			const FNodeIDVertexID DestinationKey(Edge.ToNodeID, Edge.ToVertexID);
 			const FCoreNodeAndFrontendVertex* DestinationNodeAndVertex = NodeDestinationsByID.Find(DestinationKey);
@@ -654,7 +654,9 @@ namespace Metasound
 			return false;
 		}
 
+PRAGMA_DISABLE_DEPRECATION_WARNINGS
 		return IsFlat(InDocument.RootGraph, InDocument.Dependencies);
+PRAGMA_ENABLE_DEPRECATION_WARNINGS
 	}
 
 	bool FFrontendGraphBuilder::IsFlat(const FMetasoundFrontendGraphClass& InRoot, const TArray<FMetasoundFrontendClass>& InDependencies)
@@ -684,7 +686,7 @@ namespace Metasound
 			return AvailableDependencies.Contains(InNode.ClassID);
 		};
 
-		const bool bIsEveryDependencyMet = Algo::AllOf(InRoot.Graph.Nodes, IsDependencyMet);
+		const bool bIsEveryDependencyMet = Algo::AllOf(InRoot.GetConstDefaultGraph().Nodes, IsDependencyMet);
 
 		return bIsEveryDependencyMet;
 	}
@@ -706,7 +708,7 @@ namespace Metasound
 				// Cache subgraph dependencies.
 				for (const FMetasoundFrontendGraphClass* GraphClass : InGraphs)
 				{
-					for (const FMetasoundFrontendNode& Node : GraphClass->Graph.Nodes)
+					for (const FMetasoundFrontendNode& Node : GraphClass->GetConstDefaultGraph().Nodes)
 					{
 						if (ClassIDAndGraph.Contains(Node.ClassID))
 						{
diff --git a/Engine/Plugins/Runtime/Metasound/Source/MetasoundFrontend/Private/MetasoundFrontendGraphController.cpp b/Engine/Plugins/Runtime/Metasound/Source/MetasoundFrontend/Private/MetasoundFrontendGraphController.cpp
--- a/Engine/Plugins/Runtime/Metasound/Source/MetasoundFrontend/Private/MetasoundFrontendGraphController.cpp
+++ b/Engine/Plugins/Runtime/Metasound/Source/MetasoundFrontend/Private/MetasoundFrontendGraphController.cpp
@@ -4,8 +4,10 @@
 
 #include "Algo/NoneOf.h"
 #include "Internationalization/Text.h"
+#include "MetasoundDocumentInterface.h"
 #include "MetasoundFrontendDocument.h"
 #include "MetasoundFrontendDocumentAccessPtr.h"
+#include "MetasoundFrontendDocumentBuilder.h"
 #include "MetasoundFrontendDocumentIdGenerator.h"
 #include "MetasoundFrontendGraph.h"
 #include "MetasoundFrontendNodeController.h"
@@ -25,6 +27,34 @@ namespace Metasound
 {
 	namespace Frontend
 	{
+		const FMetasoundFrontendGraph& FindConstBuildGraphChecked(const FMetasoundFrontendGraphClass& InGraphClass)
+		{
+			// Registry can be null in test builds for legacy controller implementation, so doesn't use GetChecked
+			if (IDocumentBuilderRegistry* BuilderRegistry = IDocumentBuilderRegistry::Get())
+			{
+				if (const FMetaSoundFrontendDocumentBuilder* Builder = BuilderRegistry->FindBuilder(InGraphClass.Metadata.GetClassName()))
+				{
+					return Builder->FindConstBuildGraphChecked();
+				}
+			}
+
+			return InGraphClass.GetConstDefaultGraph();
+		}
+
+		FMetasoundFrontendGraph& FindBuildGraphChecked(FMetasoundFrontendGraphClass& InGraphClass)
+		{
+			// Registry can be null in test builds for legacy controller implementation, so doesn't use GetChecked
+			if (IDocumentBuilderRegistry* BuilderRegistry = IDocumentBuilderRegistry::Get())
+			{
+				if (FMetaSoundFrontendDocumentBuilder* Builder = BuilderRegistry->FindBuilder(InGraphClass.Metadata.GetClassName()))
+				{
+					return Builder->FindBuildGraphChecked();
+				}
+			}
+
+			return InGraphClass.GetDefaultGraph();
+		}
+
 		// 
 		// FGraphController
 		//
@@ -184,7 +214,7 @@ namespace Metasound
 		{
 			if (const FMetasoundFrontendGraphClass* GraphClass = GraphClassPtr.Get())
 			{
-				return GraphClass->Graph.Style;
+				return GraphClass->GetConstDefaultGraph().Style;
 			}
 
 			return Invalid::GetInvalidGraphStyle();
@@ -214,7 +244,7 @@ namespace Metasound
 		{
 			if (FMetasoundFrontendGraphClass* GraphClass = GraphClassPtr.Get())
 			{
-				GraphClass->Graph.Style = InStyle;
+				FindBuildGraphChecked(*GraphClass).Style = InStyle;
 			}
 		}
 
@@ -302,7 +332,7 @@ namespace Metasound
 						if (InitNode->IsValid())
 						{
 							Variable.VariableNodeID = InitNode->GetID();
-							GraphClass->Graph.Variables.Add(MoveTemp(Variable));
+							FindBuildGraphChecked(*GraphClass).Variables.Add(MoveTemp(Variable));
 						}
 					}
 
@@ -384,7 +414,7 @@ namespace Metasound
 					{
 						return InVariable.ID == InVariableID;
 					};
-					GraphClass->Graph.Variables.RemoveAllSwap(IsVariableWithID);
+					FindBuildGraphChecked(*GraphClass).Variables.RemoveAllSwap(IsVariableWithID);
 					return true;
 				}
 			}
@@ -398,7 +428,7 @@ namespace Metasound
 
 			if (FMetasoundFrontendGraphClass* GraphClass = GraphClassPtr.Get())
 			{
-				for (const FMetasoundFrontendVariable& Variable : GraphClass->Graph.Variables)
+				for (const FMetasoundFrontendVariable& Variable : GraphClass->GetConstDefaultGraph().Variables)
 				{
 					VariableHandles.Add(FindVariable(Variable.ID));
 				}
@@ -413,7 +443,7 @@ namespace Metasound
 
 			if (const FMetasoundFrontendGraphClass* GraphClass = GraphClassPtr.Get())
 			{
-				for (const FMetasoundFrontendVariable& Variable : GraphClass->Graph.Variables)
+				for (const FMetasoundFrontendVariable& Variable : GraphClass->GetConstDefaultGraph().Variables)
 				{
 					VariableHandles.Add(FindVariable(Variable.ID));
 				}
@@ -593,7 +623,7 @@ namespace Metasound
 				{
 					return InVariable.ID == InVariableID;
 				};
-				return GraphClass->Graph.Variables.FindByPredicate(IsVariableWithID);
+				return FindBuildGraphChecked(*GraphClass).Variables.FindByPredicate(IsVariableWithID);
 			}
 			return nullptr;
 		}
@@ -606,7 +636,7 @@ namespace Metasound
 				{
 					return InVariable.ID == InVariableID;
 				};
-				return GraphClass->Graph.Variables.FindByPredicate(IsVariableWithID);
+				return GraphClass->GetConstDefaultGraph().Variables.FindByPredicate(IsVariableWithID);
 			}
 			return nullptr;
 		}
@@ -622,7 +652,7 @@ namespace Metasound
 						return (InNodeID == InVariable.VariableNodeID) || (InNodeID == InVariable.MutatorNodeID) || InVariable.AccessorNodeIDs.Contains(InNodeID) || InVariable.DeferredAccessorNodeIDs.Contains(InNodeID);
 
 					};
-					if (const FMetasoundFrontendVariable* Variable = GraphClass->Graph.Variables.FindByPredicate(ContainsNodeWithID))
+					if (const FMetasoundFrontendVariable* Variable = GraphClass->GetConstDefaultGraph().Variables.FindByPredicate(ContainsNodeWithID))
 					{
 						return Variable->ID;
 					}
@@ -686,7 +716,7 @@ namespace Metasound
 			{
 				const FGuid NodeID = InNode.GetID();
 				
-				for (FMetasoundFrontendVariable& Variable : GraphClass->Graph.Variables)
+				for (FMetasoundFrontendVariable& Variable : FindBuildGraphChecked(*GraphClass).Variables)
 				{
 					if (NodeID == Variable.VariableNodeID)
 					{
@@ -743,8 +773,8 @@ namespace Metasound
 		{
 			if (FMetasoundFrontendGraphClass* GraphClass = GraphClassPtr.Get())
 			{
-				GraphClass->Graph.Nodes.Reset();
-				GraphClass->Graph.Edges.Reset();
+				FindBuildGraphChecked(*GraphClass).Nodes.Reset();
+				FindBuildGraphChecked(*GraphClass).Edges.Reset();
 				GraphClass->Interface.Inputs.Reset();
 				GraphClass->Interface.Outputs.Reset();
 				GraphClass->PresetOptions.InputsInheritingDefault.Reset();
@@ -764,7 +794,7 @@ namespace Metasound
 		{
 			if (FMetasoundFrontendGraphClass* GraphClass = GraphClassPtr.Get())
 			{
-				for (const FMetasoundFrontendNode& Node : GraphClass->Graph.Nodes)
+				for (const FMetasoundFrontendNode& Node : GraphClass->GetConstDefaultGraph().Nodes)
 				{
 					FConstClassAccessPtr NodeClassPtr = OwningDocument->FindClassWithID(Node.ClassID);
 					if (const FMetasoundFrontendClass* NodeClass = NodeClassPtr.Get())
@@ -788,7 +818,7 @@ namespace Metasound
 		{
 			if (FMetasoundFrontendGraphClass* GraphClass = GraphClassPtr.Get())
 			{
-				for (const FMetasoundFrontendNode& Node : GraphClass->Graph.Nodes)
+				for (const FMetasoundFrontendNode& Node : GraphClass->GetConstDefaultGraph().Nodes)
 				{
 					FConstClassAccessPtr NodeClassPtr = OwningDocument->FindClassWithID(Node.ClassID);
 					if (const FMetasoundFrontendClass* NodeClass = NodeClassPtr.Get())
@@ -980,7 +1010,7 @@ namespace Metasound
 							const FVertexName NewName = InClassInput.Name;
 
 							// Setup input node
-							FMetasoundFrontendNode& Node = GraphClass->Graph.Nodes.Emplace_GetRef(*InputClass);
+							FMetasoundFrontendNode& Node = FindBuildGraphChecked(*GraphClass).Nodes.Emplace_GetRef(*InputClass);
 							Node.Name = NewName;
 
 							FGuid NodeID = InClassInput.NodeID;
@@ -1101,7 +1131,7 @@ namespace Metasound
 						{
 							const FVertexName NewName = InClassOutput.Name;
 
-							FMetasoundFrontendNode& Node = GraphClass->Graph.Nodes.Add_GetRef(*OutputClass);
+							FMetasoundFrontendNode& Node = FindBuildGraphChecked(*GraphClass).Nodes.Add_GetRef(*OutputClass);
 							Node.Name = NewName;
 
 							FGuid NodeID = InClassOutput.NodeID;
@@ -1421,7 +1451,7 @@ namespace Metasound
 				{
 					if (const FMetasoundFrontendClass* NodeClass = Class.Get())
 					{
-						FMetasoundFrontendNode& Node = GraphClass->Graph.Nodes.Emplace_GetRef(*NodeClass);
+						FMetasoundFrontendNode& Node = FindBuildGraphChecked(*GraphClass).Nodes.Emplace_GetRef(*NodeClass);
 						Node.UpdateID(InNodeGuid);
 						Node.Interface = InTemplate.GenerateNodeInterface(MoveTemp(Params));
 						FNodeAccessPtr NodePtr = GraphClassPtr.GetNodeWithNodeID(Node.GetID());
@@ -1489,7 +1519,7 @@ namespace Metasound
 			{
 				FGuid NodeID = InNode.GetID();
 				auto IsNodeWithSameID = [&](const FMetasoundFrontendNode& InFrontendNode) { return InFrontendNode.GetID() == NodeID; };
-				if (const FMetasoundFrontendNode* FrontendNode = GraphClass->Graph.Nodes.FindByPredicate(IsNodeWithSameID))
+				if (const FMetasoundFrontendNode* FrontendNode = GraphClass->GetConstDefaultGraph().Nodes.FindByPredicate(IsNodeWithSameID))
 				{
 					switch(NodeClassType)
 					{
@@ -1627,7 +1657,7 @@ namespace Metasound
 			{
 				if (const FMetasoundFrontendClass* NodeClass = InExistingDependency.Get())
 				{
-					FMetasoundFrontendNode& Node = GraphClass->Graph.Nodes.Emplace_GetRef(*NodeClass);
+					FMetasoundFrontendNode& Node = FindBuildGraphChecked(*GraphClass).Nodes.Emplace_GetRef(*NodeClass);
 
 					// Cache the asset name on the node if it node is reference to asset-defined graph.
 					// AssetManager may not exist if this is called from a build that does not load the
@@ -1667,12 +1697,11 @@ namespace Metasound
 				auto IsEdgeForThisNode = [&](const FMetasoundFrontendEdge& ConDesc) { return (ConDesc.FromNodeID == InDesc.GetID()) || (ConDesc.ToNodeID == InDesc.GetID()); };
 
 				// Remove any reference connections
-				int32 NumRemoved = GraphClass->Graph.Edges.RemoveAll(IsEdgeForThisNode);
+				int32 NumRemoved = FindBuildGraphChecked(*GraphClass).Edges.RemoveAll(IsEdgeForThisNode);
 
 				auto IsNodeWithID = [&](const FMetasoundFrontendNode& Desc) { return InDesc.GetID() == Desc.GetID(); };
 
-				NumRemoved += GraphClass->Graph.Nodes.RemoveAll(IsNodeWithID);
-
+				NumRemoved += FindBuildGraphChecked(*GraphClass).Nodes.RemoveAll(IsNodeWithID);
 				OwningDocument->RemoveUnreferencedDependencies();
 
 #if WITH_EDITOR
@@ -1756,7 +1785,7 @@ namespace Metasound
 		{
 			if (FMetasoundFrontendGraphClass* GraphClass = GraphClassPtr.Get())
 			{
-				for (FMetasoundFrontendNode& Node : GraphClass->Graph.Nodes)
+				for (FMetasoundFrontendNode& Node : FindBuildGraphChecked(*GraphClass).Nodes)
 				{
 					const FMetasoundFrontendClass* NodeClass = OwningDocument->FindClassWithID(Node.ClassID).Get();
 					if (nullptr != NodeClass)
@@ -1782,7 +1811,7 @@ namespace Metasound
 
 			if (FMetasoundFrontendGraphClass* GraphClass = GraphClassPtr.Get())
 			{
-				for (FMetasoundFrontendNode& Node : GraphClass->Graph.Nodes)
+				for (FMetasoundFrontendNode& Node : FindBuildGraphChecked(*GraphClass).Nodes)
 				{
 					FNodeAccessPtr NodePtr = GraphClassPtr.GetNodeWithNodeID(Node.GetID());
 					FConstClassAccessPtr NodeClassPtr = OwningDocument->FindClassWithID(Node.ClassID);
@@ -1810,7 +1839,7 @@ namespace Metasound
 
 			if (const FMetasoundFrontendGraphClass* GraphClass = GraphClassPtr.Get())
 			{
-				for (const FMetasoundFrontendNode& Node : GraphClass->Graph.Nodes)
+				for (const FMetasoundFrontendNode& Node : GraphClass->GetConstDefaultGraph().Nodes)
 				{
 					FConstNodeAccessPtr NodePtr = GraphClassPtr.GetNodeWithNodeID(Node.GetID());
 					FConstClassAccessPtr NodeClassPtr = OwningDocument->FindClassWithID(Node.ClassID);
@@ -1838,7 +1867,7 @@ namespace Metasound
 
 			if (FMetasoundFrontendGraphClass* GraphClass = GraphClassPtr.Get())
 			{
-				for (FMetasoundFrontendNode& Node : GraphClass->Graph.Nodes)
+				for (FMetasoundFrontendNode& Node : FindBuildGraphChecked(*GraphClass).Nodes)
 				{
 					FConstClassAccessPtr NodeClassPtr = OwningDocument->FindClassWithID(Node.ClassID);
 					if (const FMetasoundFrontendClass* NodeClass = NodeClassPtr.Get())
@@ -1865,7 +1894,7 @@ namespace Metasound
 
 			if (const FMetasoundFrontendGraphClass* GraphClass = GraphClassPtr.Get())
 			{
-				for (const FMetasoundFrontendNode& Node : GraphClass->Graph.Nodes)
+				for (const FMetasoundFrontendNode& Node : GraphClass->GetConstDefaultGraph().Nodes)
 				{
 					FConstClassAccessPtr NodeClassPtr = OwningDocument->FindClassWithID(Node.ClassID);
 					if (const FMetasoundFrontendClass* NodeClass = NodeClassPtr.Get())
diff --git a/Engine/Plugins/Runtime/Metasound/Source/MetasoundFrontend/Private/MetasoundFrontendGraphController.h b/Engine/Plugins/Runtime/Metasound/Source/MetasoundFrontend/Private/MetasoundFrontendGraphController.h
--- a/Engine/Plugins/Runtime/Metasound/Source/MetasoundFrontend/Private/MetasoundFrontendGraphController.h
+++ b/Engine/Plugins/Runtime/Metasound/Source/MetasoundFrontend/Private/MetasoundFrontendGraphController.h
@@ -18,6 +18,11 @@ namespace Metasound
 
 	namespace Frontend
 	{
+		// Convenience functions used for finding targeted graph to mutate as the transition away from controllers
+		// is on-going, however documents can now contain multiple graph topologies (i.e. paged graphs) to select from.
+		const FMetasoundFrontendGraph& FindConstBuildGraphChecked(const FMetasoundFrontendGraphClass& InGraphClass);
+		FMetasoundFrontendGraph& FindBuildGraphChecked(FMetasoundFrontendGraphClass& InGraphClass);
+
 		/** FGraphController represents a Metasound graph class. */
 		class FGraphController : public IGraphController
 		{
diff --git a/Engine/Plugins/Runtime/Metasound/Source/MetasoundFrontend/Private/NodeTemplates/MetasoundFrontendNodeTemplateInput.cpp b/Engine/Plugins/Runtime/Metasound/Source/MetasoundFrontend/Private/NodeTemplates/MetasoundFrontendNodeTemplateInput.cpp
--- a/Engine/Plugins/Runtime/Metasound/Source/MetasoundFrontend/Private/NodeTemplates/MetasoundFrontendNodeTemplateInput.cpp
+++ b/Engine/Plugins/Runtime/Metasound/Source/MetasoundFrontend/Private/NodeTemplates/MetasoundFrontendNodeTemplateInput.cpp
@@ -205,8 +205,7 @@ namespace Metasound::Frontend
 	{
 		bool bInjectedNodes = false;
 
-		const FMetaSoundFrontendDocumentBuilder& ConstBuilder = static_cast<const FMetaSoundFrontendDocumentBuilder&>(InOutBuilder);
-		const TArray<FMetasoundFrontendClassInput>& Inputs = ConstBuilder.GetConstDocument().RootGraph.Interface.Inputs;
+		const TArray<FMetasoundFrontendClassInput>& Inputs = InOutBuilder.GetConstDocumentChecked().RootGraph.Interface.Inputs;
 		for (const FMetasoundFrontendClassInput& Input : Inputs)
 		{
 			// Potentially not used input, which is perfectly valid so early out
diff --git a/Engine/Plugins/Runtime/Metasound/Source/MetasoundFrontend/Public/MetasoundAssetManager.h b/Engine/Plugins/Runtime/Metasound/Source/MetasoundFrontend/Public/MetasoundAssetManager.h
--- a/Engine/Plugins/Runtime/Metasound/Source/MetasoundFrontend/Public/MetasoundAssetManager.h
+++ b/Engine/Plugins/Runtime/Metasound/Source/MetasoundFrontend/Public/MetasoundAssetManager.h
@@ -149,7 +149,8 @@ namespace Metasound::Frontend
 #if WITH_EDITORONLY_DATA
 		// Adds missing assets using the provided asset's local reference class cache. Used
 		// to prime system from asset attempting to register prior to asset scan being complete.
-		virtual void AddAssetReferences(FMetasoundAssetBase& InAssetBase) = 0;
+		// Returns true if references were added, false if it they are already found.
+		virtual bool AddAssetReferences(FMetasoundAssetBase& InAssetBase) = 0;
 #endif
 
 		// Add or Update a MetaSound Asset's entry data
diff --git a/Engine/Plugins/Runtime/Metasound/Source/MetasoundFrontend/Public/MetasoundDocumentInterface.h b/Engine/Plugins/Runtime/Metasound/Source/MetasoundFrontend/Public/MetasoundDocumentInterface.h
--- a/Engine/Plugins/Runtime/Metasound/Source/MetasoundFrontend/Public/MetasoundDocumentInterface.h
+++ b/Engine/Plugins/Runtime/Metasound/Source/MetasoundFrontend/Public/MetasoundDocumentInterface.h
@@ -80,6 +80,7 @@ namespace Metasound::Frontend
 
 		virtual FMetaSoundFrontendDocumentBuilder* FindBuilder(TScriptInterface<IMetaSoundDocumentInterface> MetaSound) const = 0;
 		virtual FMetaSoundFrontendDocumentBuilder* FindBuilder(const FMetasoundFrontendClassName& ClassName) const = 0;
+		virtual FMetaSoundFrontendDocumentBuilder* FindOutermostBuilder(const UObject& InSubObject) const = 0;
 
 #if WITH_EDITORONLY_DATA
 		// Find the existing builder for the given MetaSound, or optionally begin building by attaching a new builder.  Only available
@@ -92,9 +93,13 @@ namespace Metasound::Frontend
 
 		virtual bool FinishBuilding(const FMetasoundFrontendClassName& InClassName, bool bForceUnregister = false) const = 0;
 
-		UE_DEPRECATED(5.5, "Document cache can now be invalidated by retrieving an asset builder and calling 'Reload'")
+		UE_DEPRECATED(5.5, "Document cache can now be invalidated by retrieving an asset builder and calling 'ReloadBuilder'")
 		virtual void InvalidateDocumentCache(const FMetasoundFrontendClassName& InClassName) const { }
 
+		// Reloads the given builder, maintaining all modify delegate subscriptions. Returns true if builder was found and reloaded,
+		// false if not found.
+		virtual bool ReloadBuilder(const FMetasoundFrontendClassName& InClassName) const = 0;
+
 		static IDocumentBuilderRegistry* Get();
 		static IDocumentBuilderRegistry& GetChecked();
 		static void Deinitialize();
diff --git a/Engine/Plugins/Runtime/Metasound/Source/MetasoundFrontend/Public/MetasoundFrontendDocument.h b/Engine/Plugins/Runtime/Metasound/Source/MetasoundFrontend/Public/MetasoundFrontendDocument.h
--- a/Engine/Plugins/Runtime/Metasound/Source/MetasoundFrontend/Public/MetasoundFrontendDocument.h
+++ b/Engine/Plugins/Runtime/Metasound/Source/MetasoundFrontend/Public/MetasoundFrontendDocument.h
@@ -745,6 +745,9 @@ struct METASOUNDFRONTEND_API FMetasoundFrontendGraph
 	UPROPERTY()
 	FMetasoundFrontendGraphStyle Style;
 #endif // WITH_EDITORONLY_DATA
+
+	UPROPERTY()
+	FGuid PageID;
 };
 
 // Metadata associated with a vertex.
@@ -1641,15 +1644,38 @@ struct METASOUNDFRONTEND_API FMetasoundFrontendGraphClass : public FMetasoundFro
 {
 	GENERATED_BODY()
 
+public:
 	FMetasoundFrontendGraphClass();
-
 	virtual ~FMetasoundFrontendGraphClass() = default;
 
-	UPROPERTY()
+	UPROPERTY(meta = (DeprecationMessage = "5.5 - Direct Graph access will be revoked and page manipulation applied via public API below in future builds."))
 	FMetasoundFrontendGraph Graph;
 
+private:
+	UPROPERTY()
+	TArray<FMetasoundFrontendGraph> PagedGraphs;
+
+public:
 	UPROPERTY()
 	FMetasoundFrontendGraphClassPresetOptions PresetOptions;
+
+public:
+	void AddNewGraphPage(const FGuid& InPageID);
+	void DuplicateLastGraphPage(const FGuid& InPageID);
+	void RemoveAllGraphPages();
+	bool RemoveGraphPage(const FGuid& InPageID);
+
+	bool ContainsGraphPage(const FGuid& InPageID) const;
+	void IterateGraphPages(TFunctionRef<void(FMetasoundFrontendGraph&)> IterFunc);
+	void IterateGraphPages(TFunctionRef<void(const FMetasoundFrontendGraph&)> IterFunc) const;
+
+	FMetasoundFrontendGraph* FindGraph(const FGuid& InPageID);
+	FMetasoundFrontendGraph& FindGraphChecked(const FGuid& InPageID);
+	const FMetasoundFrontendGraph* FindConstGraph(const FGuid& InPageID) const;
+	const FMetasoundFrontendGraph& FindConstGraphChecked(const FGuid& InPageID) const;
+	FMetasoundFrontendGraph& GetDefaultGraph();
+	const FMetasoundFrontendGraph& GetConstDefaultGraph() const;
+	void ResetGraphs();
 };
 
 UCLASS()
diff --git a/Engine/Plugins/Runtime/Metasound/Source/MetasoundFrontend/Public/MetasoundFrontendDocumentBuilder.h b/Engine/Plugins/Runtime/Metasound/Source/MetasoundFrontend/Public/MetasoundFrontendDocumentBuilder.h
--- a/Engine/Plugins/Runtime/Metasound/Source/MetasoundFrontend/Public/MetasoundFrontendDocumentBuilder.h
+++ b/Engine/Plugins/Runtime/Metasound/Source/MetasoundFrontend/Public/MetasoundFrontendDocumentBuilder.h
@@ -172,9 +172,13 @@ public:
 	// is not already connected and the edge is valid (see function 'IsValidEdge').
 	bool CanAddEdge(const FMetasoundFrontendEdge& InEdge) const;
 
-	// Clears graph completely of all nodes, edges, dependencies, & member metadata.
-	// Reloads the builder state, so external delegates must be relinked if desired.
-	void ClearGraph(TSharedRef<Metasound::Frontend::FDocumentModifyDelegates> ModifyDelegates = { });
+	// Clears document completely of all graph page data (nodes, edges, & member metadata), dependencies,
+	// interfaces, member metadata, preset state, etc. Leaves ClassMetadata intact. Reloads the builder state,
+	// so external delegates must be relinked if desired.
+	void ClearDocument(TSharedRef<Metasound::Frontend::FDocumentModifyDelegates> ModifyDelegates = { });
+
+	UE_DEPRECATED(5.5, "Use ClearDocument instead")
+	void ClearGraph() {  }
 
 #if WITH_EDITORONLY_DATA
 	bool ClearMemberMetadata(const FGuid& InMemberID);
@@ -209,6 +213,9 @@ public:
 	bool FindInterfaceInputNodes(FName InterfaceName, TArray<const FMetasoundFrontendNode*>& OutInputs) const;
 	bool FindInterfaceOutputNodes(FName InterfaceName, TArray<const FMetasoundFrontendNode*>& OutOutputs) const;
 
+	FMetasoundFrontendGraph& FindBuildGraphChecked();
+	const FMetasoundFrontendGraph& FindConstBuildGraphChecked() const;
+
 	const FMetasoundFrontendNode* FindNode(const FGuid& InNodeID) const;
 
 	const FMetasoundFrontendVertex* FindNodeInput(const FGuid& InNodeID, const FGuid& InVertexID) const;
@@ -221,12 +228,19 @@ public:
 	TArray<const FMetasoundFrontendVertex*> FindNodeOutputs(const FGuid& InNodeID, FName TypeName = FName()) const;
 	const FMetasoundFrontendVertex* FindNodeOutputConnectedToNodeInput(const FGuid& InInputNodeID, const FGuid& InInputVertexID, const FMetasoundFrontendNode** ConnectedOutputNode = nullptr) const;
 
-	const FMetasoundFrontendDocument& GetConstDocument() const;
+	const FMetasoundFrontendDocument& GetConstDocumentChecked() const;
+	const IMetaSoundDocumentInterface& GetConstDocumentInterfaceChecked() const;
 	const FString GetDebugName() const;
 
-	UE_DEPRECATED(5.5, "Use GetConstDocument() instead")
+	// The graph ID used when graph is executed (i.e. utilized for node class registration)
+	const FGuid& GetExecutionGraphID() const;
+
+	UE_DEPRECATED(5.5, "Use GetConstDocumentChecked() instead")
 	const FMetasoundFrontendDocument& GetDocument() const;
 
+	// The graph ID used when requests are made to mutate specific paged graph topology (ex. adding or removing nodes or edges)
+	const FGuid& GetBuildPageID() const;
+
 	template<typename TObjectType>
 	TObjectType& CastDocumentObjectChecked() const
 	{
@@ -242,9 +256,10 @@ public:
 	FMetasoundFrontendClassName GenerateNewClassName();
 
 	const Metasound::Frontend::FDocumentModifyDelegates& GetDocumentDelegates() const;
+
+	UE_DEPRECATED(5.5, "Use GetConstDocumentInterfaceChecked instead")
 	const IMetaSoundDocumentInterface& GetDocumentInterface() const;
-	FMetasoundAssetBase& GetMetasoundAsset();
-	const FMetasoundAssetBase& GetMetasoundAsset() const;
+	FMetasoundAssetBase& GetMetasoundAsset() const;
 
 	int32 GetTransactionCount() const;
 
@@ -326,6 +341,13 @@ public:
 
 #endif // WITH_EDITOR
 
+	// Sets the builder's targeted paged graph ID to the given ID if it exists.
+	// Returns true if the builder is already targeting the given ID or if it successfully
+	// found a page implementation with the given ID and was able to switch to it, false if not.
+	// Swapping the targeted build graph ID clears the local cache, so swapping frequently can
+	// induce cash thrashing.
+	bool SetBuildPageID(const FGuid& InBuildPageID);
+
 	bool SetGraphInputAccessType(FName InputName, EMetasoundFrontendVertexAccessType AccessType);
 	bool SetGraphInputDataType(FName InputName, FName DataType);
 	bool SetGraphInputDefault(FName InputName, const FMetasoundFrontendLiteral& InDefaultLiteral);
@@ -369,7 +391,7 @@ public:
 private:
 	using FFinalizeNodeFunctionRef = TFunctionRef<void(FMetasoundFrontendNode&, const Metasound::Frontend::FNodeRegistryKey&)>;
 
-	FMetasoundFrontendNode* AddNodeInternal(const FMetasoundFrontendClassMetadata& InClassMetadata, Metasound::Frontend::FFinalizeNodeFunctionRef FinalizeNode, FGuid InNodeID = FGuid::NewGuid(), int32* NewNodeIndex = nullptr);
+	FMetasoundFrontendNode* AddNodeInternal(const FMetasoundFrontendClassMetadata& InClassMetadata, Metasound::Frontend::FFinalizeNodeFunctionRef FinalizeNode, FGuid InNodeID = FGuid::NewGuid(), int32* NewNodeIndex = nullptr, const FGuid* InGraphPageID = nullptr);
 	void BeginBuilding(TSharedPtr<Metasound::Frontend::FDocumentModifyDelegates> Delegates = {}, bool bPrimeCache = false);
 
 	// Conforms GraphOutput node's ClassID, Access & Data Type with the GraphOutput.
@@ -390,7 +412,8 @@ private:
 	void IterateNodesConnectedWithVertex(const FMetasoundFrontendVertexHandle& Vertex, TFunctionRef<void(const FMetasoundFrontendEdge&, FMetasoundFrontendNode&)> NodeIndexIterFunc);
 
 	const FTopLevelAssetPath GetBuilderClassPath() const;
-	FMetasoundFrontendDocument& GetDocument();
+	FMetasoundFrontendDocument& GetDocumentChecked() const;
+	IMetaSoundDocumentInterface& GetDocumentInterfaceChecked() const;
 
 	bool SetGraphInputInheritsDefault(FName InName, bool bInputInheritsDefault);
 
@@ -400,6 +423,8 @@ private:
 	UPROPERTY(Transient)
 	TScriptInterface<IMetaSoundDocumentInterface> DocumentInterface;
 
+	FGuid BuildPageID;
+
 	TSharedPtr<Metasound::Frontend::IDocumentCache> DocumentCache;
 	TSharedPtr<Metasound::Frontend::FDocumentModifyDelegates> DocumentDelegates;
 };
diff --git a/Engine/Plugins/Runtime/Metasound/Source/MetasoundFrontend/Public/MetasoundFrontendGraph.h b/Engine/Plugins/Runtime/Metasound/Source/MetasoundFrontend/Public/MetasoundFrontendGraph.h
--- a/Engine/Plugins/Runtime/Metasound/Source/MetasoundFrontend/Public/MetasoundFrontendGraph.h
+++ b/Engine/Plugins/Runtime/Metasound/Source/MetasoundFrontend/Public/MetasoundFrontendGraph.h
@@ -90,8 +90,10 @@ namespace Metasound
 		 *
 		 * @return True if all dependencies are C++ classes. False otherwise.
 		 */
+		UE_DEPRECATED(5.5, "No longer supported (not useful nor maintainable with paged graph support")
 		static bool IsFlat(const FMetasoundFrontendDocument& InDocument);
 
+		UE_DEPRECATED(5.5, "No longer supported (not useful nor maintainable with paged graph support")
 		static bool IsFlat(const FMetasoundFrontendGraphClass& InRoot, const TArray<FMetasoundFrontendClass>& InDependencies);
 
 		/* Create a FFrontendGraph from a FMetasoundFrontendDocument.*/
