Bus Reader-Writer Node Consistent Latency Fix

#tests EngineTest Audio Bus functional tests all pass

diff --git a/Engine/Plugins/Runtime/Metasound/Source/MetasoundEngine/Private/MetasoundAudioBusReaderNode.cpp b/Engine/Plugins/Runtime/Metasound/Source/MetasoundEngine/Private/MetasoundAudioBusReaderNode.cpp
--- a/Engine/Plugins/Runtime/Metasound/Source/MetasoundEngine/Private/MetasoundAudioBusReaderNode.cpp
+++ b/Engine/Plugins/Runtime/Metasound/Source/MetasoundEngine/Private/MetasoundAudioBusReaderNode.cpp
@@ -3,7 +3,6 @@
 #include "AudioDevice.h"
 #include "Sound/AudioBus.h"
 #include "AudioBusSubsystem.h"
-#include "DSP/ConvertDeinterleave.h"
 #include "Internationalization/Text.h"
 #include "MetasoundAudioBuffer.h"
 #include "MetasoundAudioBus.h"
@@ -32,6 +31,29 @@ namespace Metasound
 		METASOUND_PARAM(OutParamAudio, "Out {0}", "Audio bus output for channel {0}.");
 	}
 
+	int32 AudioBusReaderNodeInitialNumBlocks(int32 BlockSizeFrames, int32 AudioMixerOutputFrames)
+	{
+		// One extra block is required to cover the first metasound iteration.
+		int32 ExtraBlocks = 1;
+
+		// Find the number of whole blocks that fit in the mixer output.
+		int32 WholeBlocks = FMath::DivideAndRoundDown(AudioMixerOutputFrames, BlockSizeFrames);
+
+		// Determine if any frames remain.
+		int32 FramesRemainder = AudioMixerOutputFrames % BlockSizeFrames;
+		if (FramesRemainder > 0)
+		{
+			// Find the number of extra frames required to consistently cover the remainder.
+			int32 ExtraFrames = FMath::DivideAndRoundUp(BlockSizeFrames, FramesRemainder) * FramesRemainder;
+
+			// Find the number of blocks required to cover the extra frames.
+			ExtraBlocks += FMath::DivideAndRoundUp(ExtraFrames, BlockSizeFrames);
+		}
+
+		// Determine the total number of blocks required.
+		return WholeBlocks + ExtraBlocks;
+	}
+
 	template <uint32 NumChannels>
 	class TAudioBusReaderOperator : public TExecutableOperator<TAudioBusReaderOperator<NumChannels>>
 	{
@@ -150,9 +172,6 @@ namespace Metasound
 
 		void Execute()
 		{
-			const int32 BlockSizeFrames = AudioOutputs[0]->Num();
-			const int32 NumSamplesToPop = BlockSizeFrames * AudioBusChannels;
-
 			const FAudioBusProxyPtr& BusProxy = AudioBusAsset->GetAudioBusProxy();
 			if (!BusProxy.IsValid() || BusProxy->NumChannels <= 0)
 			{
@@ -162,76 +181,43 @@ namespace Metasound
 
 			if (BusProxy->AudioBusId != AudioBusId)
 			{
-				AudioBusPatchOutput.Reset();
+				InterleavedBuffer.Reset();
 			}
 			
-			if (!AudioBusPatchOutput.IsValid())
+			if (InterleavedBuffer.IsEmpty())
 			{
 				// if environment vars & a valid audio bus have been set since starting, try to create the patch now
 				if (SampleRate > 0.f && BusProxy.IsValid())
 				{
-					CreatePatchOutput(BlockSizeFrames);
-					if (!AudioBusPatchOutput.IsValid())
-					{
-						return;
-					}
-				}
-				else
-				{
-					return;
+					CreatePatchOutput();
 				}
 			}
 
-			bool bPerformPop = false;
-
-			if (bFirstBlock)
+			if (InterleavedBuffer.IsEmpty())
 			{
-				// Ensure there are enough samples in the patch output to support the maximum metasound executions the mixer requires to fill its output frames before popping begins.
-				if (uint32(AudioBusPatchOutput->GetNumSamplesAvailable()) >= FMath::DivideAndRoundUp(AudioMixerOutputFrames, BlockSizeFrames) * BlockSizeFrames * AudioBusChannels)
-				{
-					bPerformPop = true;
-
-					InterleavedBuffer.Reset();
-					InterleavedBuffer.AddUninitialized(NumSamplesToPop);
-				}
+				return;
 			}
-			else
+
+			// Pop off the interleaved data from the audio bus
+			int32 NumSamplesToPop = BlockSizeFrames * AudioBusChannels;
+			int32 SamplesPopped = AudioBusPatchOutput->PopAudio(InterleavedBuffer.GetData(), NumSamplesToPop, false);
+			if (SamplesPopped < NumSamplesToPop)
 			{
-				// Give input patch samples some time to arrive.
-				if (AudioBusPatchOutput->WaitUntilNumSamplesAvailable(NumSamplesToPop, uint32(AudioBusReaderNodePatchWaitTimeout)))
-				{
-					bPerformPop = true;
-				}
+				UE_LOG(LogMetaSound, Warning, TEXT("Underrun detected in audio bus reader node."));
 			}
 
-			if (bPerformPop)
+			const uint32 MinChannels = FMath::Min(NumChannels, AudioBusChannels);
+			for (uint32 ChannelIndex = 0; ChannelIndex < MinChannels; ++ChannelIndex)
 			{
-				bFirstBlock = false;
-
-				// Pop off the interleaved data from the audio bus
-				int32 SamplesPopped = AudioBusPatchOutput->PopAudio(InterleavedBuffer.GetData(), NumSamplesToPop, false);
-				if (SamplesPopped < NumSamplesToPop)
-				{
-					UE_LOG(LogMetaSound, Warning, TEXT("Underrun detected in audio bus reader node."));
-				}
-
-				const uint32 MinChannels = FMath::Min(NumChannels, AudioBusChannels);
-				for (uint32 ChannelIndex = 0; ChannelIndex < MinChannels; ++ChannelIndex)
+				float* AudioOutputBufferPtr = AudioOutputs[ChannelIndex]->GetData();
+				for (int32 FrameIndex = 0; FrameIndex < BlockSizeFrames; ++FrameIndex)
 				{
-					float* AudioOutputBufferPtr = AudioOutputs[ChannelIndex]->GetData();
-					for (int32 FrameIndex = 0; FrameIndex < BlockSizeFrames; ++FrameIndex)
-					{
-						AudioOutputBufferPtr[FrameIndex] = InterleavedBuffer[FrameIndex * AudioBusChannels + ChannelIndex];
-					}
+					AudioOutputBufferPtr[FrameIndex] = InterleavedBuffer[FrameIndex * AudioBusChannels + ChannelIndex];
 				}
 			}
-			else
-			{
-				UE_CLOG(!bFirstBlock, LogMetaSound, Warning, TEXT("Underrun detected in audio bus reader node."));
-			}
 		}
 
-		void CreatePatchOutput(const int32 BlockSizeFrames)
+		void CreatePatchOutput()
 		{
 			const FAudioBusProxyPtr& AudioBusProxy = AudioBusAsset->GetAudioBusProxy();
 			if (AudioBusProxy.IsValid())
@@ -241,28 +227,34 @@ namespace Metasound
 					UE_LOG(LogMetaSound, Warning, TEXT("AudioBusProxy is invalid (NumChannels = %i)."), AudioBusProxy->NumChannels);
 					return;
 				}
+
+				UAudioBusSubsystem* AudioBusSubsystem = nullptr;
 				if (FAudioDeviceManager* ADM = FAudioDeviceManager::Get())
 				{
 					if (FAudioDevice* AudioDevice = ADM->GetAudioDeviceRaw(AudioDeviceId))
 					{
-						// Start the audio bus in case it's not already started
-						AudioBusChannels = static_cast<uint32>(FMath::Min(AudioBusProxy->NumChannels, static_cast<int32>(EAudioBusChannels::MaxChannelCount)));
-						AudioBusId = AudioBusProxy->AudioBusId;
-						
-						const Audio::FAudioBusKey AudioBusKey = Audio::FAudioBusKey(AudioBusId);
-						UAudioBusSubsystem* AudioBusSubsystem = AudioDevice->GetSubsystem<UAudioBusSubsystem>();
-						
+						AudioBusSubsystem = AudioDevice->GetSubsystem<UAudioBusSubsystem>();
 						check(AudioBusSubsystem);
-						
-						AudioBusSubsystem->StartAudioBus(AudioBusKey, AudioBusChannels, false);
-
-						// Create a bus patch output with enough room for the number of samples we expect and some buffering
-						AudioBusPatchOutput = AudioBusSubsystem->AddPatchOutputForAudioBus(AudioBusKey, BlockSizeFrames, AudioBusChannels);
+						if (!AudioBusSubsystem)
+						{
+							return;
+						}
 					}
 				}
-			}
 
-			bFirstBlock = true;
+				AudioBusChannels = uint32(FMath::Min(AudioBusProxy->NumChannels, int32(EAudioBusChannels::MaxChannelCount)));
+				AudioBusId = AudioBusProxy->AudioBusId;
+
+				Audio::FAudioBusKey AudioBusKey(AudioBusId);
+				AudioBusSubsystem->StartAudioBus(AudioBusKey, AudioBusChannels, false);
+
+				AudioBusPatchOutput = AudioBusSubsystem->AddPatchOutputForSoundAndAudioBus(InstanceID, AudioBusKey, BlockSizeFrames, int32(AudioBusChannels));
+				PatchInput = AudioBusPatchOutput;
+				PatchInput.PushAudio(nullptr, NumBlocksToNumSamples(InitialNumBlocks()));
+
+				InterleavedBuffer.Reset();
+				InterleavedBuffer.AddUninitialized(NumBlocksToNumSamples(1));
+			}
 		}
 		
 		void Reset(const IOperator::FResetParams& InParams)
@@ -270,19 +262,31 @@ namespace Metasound
 			using namespace Frontend;
 			using namespace AudioBusReaderNode;
 
+			InterleavedBuffer.Reset();
+			AudioMixerOutputFrames = INDEX_NONE;
+			AudioDeviceId = INDEX_NONE;
+			InstanceID = 0;
+			AudioBusId = 0;
+			SampleRate = 0.0f;
+			AudioBusPatchOutput.Reset();
+			PatchInput.Reset();
+			AudioBusChannels = INDEX_NONE;
+			BlockSizeFrames = 0;
+
 			bool bHasEnvironmentVars = InParams.Environment.Contains<Audio::FDeviceId>(SourceInterface::Environment::DeviceID);
 			bHasEnvironmentVars &= InParams.Environment.Contains<int32>(SourceInterface::Environment::AudioMixerNumOutputFrames);
+			bHasEnvironmentVars &= InParams.Environment.Contains<uint64>(SourceInterface::Environment::TransmitterID);
 
 			if (bHasEnvironmentVars)
 			{
-				SampleRate = InParams.OperatorSettings.GetSampleRate();
 				AudioDeviceId = InParams.Environment.GetValue<Audio::FDeviceId>(SourceInterface::Environment::DeviceID);
 				AudioMixerOutputFrames = InParams.Environment.GetValue<int32>(SourceInterface::Environment::AudioMixerNumOutputFrames);
+				InstanceID = InParams.Environment.GetValue<uint64>(SourceInterface::Environment::TransmitterID);
 			}
 			else
 			{
-				UE_LOG(LogMetaSound, Warning, TEXT("Audio bus reader node requires audio device ID '%s' and audio mixer num output frames '%s' environment variables")
-					, *SourceInterface::Environment::DeviceID.ToString(), *SourceInterface::Environment::AudioMixerNumOutputFrames.ToString());
+				UE_LOG(LogMetaSound, Warning, TEXT("Audio bus reader node requires audio device ID '%s', audio mixer num output frames '%s' and transmitter id '%s' environment variables")
+					, *SourceInterface::Environment::DeviceID.ToString(), *SourceInterface::Environment::AudioMixerNumOutputFrames.ToString(), *SourceInterface::Environment::TransmitterID.ToString());
 			}
 			
 			for(const FAudioBufferWriteRef& Buffer : AudioOutputs)
@@ -290,25 +294,36 @@ namespace Metasound
 				Buffer->Zero();
 			}
 			
-			// Wait until Execute to connect the AudioBusPatchOutput.
-			// Otherwise the mixer will start filling the patch before the MetaSound starts.
-			AudioBusPatchOutput.Reset();
+			SampleRate = InParams.OperatorSettings.GetSampleRate();
+			BlockSizeFrames = InParams.OperatorSettings.GetNumFramesPerBlock();
+
+			CreatePatchOutput();
 		}
 
 	private:
+		int32 InitialNumBlocks() const
+		{
+			return AudioBusReaderNodeInitialNumBlocks(BlockSizeFrames, AudioMixerOutputFrames);
+		}
+
+		int32 NumBlocksToNumSamples(int32 NumBlocks) const
+		{
+			return NumBlocks * BlockSizeFrames * AudioBusChannels;
+		}
+
 		FAudioBusAssetReadRef AudioBusAsset;
 		TArray<FAudioBufferWriteRef> AudioOutputs;
 
 		TArray<float> InterleavedBuffer;
 		int32 AudioMixerOutputFrames = INDEX_NONE;
 		Audio::FDeviceId AudioDeviceId = INDEX_NONE;
+		uint64 InstanceID = 0;
 		uint32 AudioBusId = 0;
 		float SampleRate = 0.0f;
 		Audio::FPatchOutputStrongPtr AudioBusPatchOutput;
-		TUniquePtr<Audio::IConvertDeinterleave> ConvertDeinterleave;
-		Audio::FMultichannelBuffer DeinterleavedBuffer;
+		Audio::FPatchInput PatchInput;
 		uint32 AudioBusChannels = INDEX_NONE;
-		bool bFirstBlock = true;
+		int32 BlockSizeFrames = 0;
 	};
 
 	template<uint32 NumChannels>
diff --git a/Engine/Plugins/Runtime/Metasound/Source/MetasoundEngine/Private/MetasoundAudioBusWriterNode.cpp b/Engine/Plugins/Runtime/Metasound/Source/MetasoundEngine/Private/MetasoundAudioBusWriterNode.cpp
--- a/Engine/Plugins/Runtime/Metasound/Source/MetasoundEngine/Private/MetasoundAudioBusWriterNode.cpp
+++ b/Engine/Plugins/Runtime/Metasound/Source/MetasoundEngine/Private/MetasoundAudioBusWriterNode.cpp
@@ -24,6 +24,13 @@ namespace Metasound
 		METASOUND_PARAM(InParamAudio, "In {0}", "Audio input for channel {0}.");
 	}
 
+	int32 AudioBusWriterNodeInitialNumBlocks(int32 BlockSizeFrames, int32 AudioMixerOutputFrames)
+	{
+		// One less block is required because the metasound will write the final block.
+		int32 MaxSizeFrames = FMath::Max(AudioMixerOutputFrames, BlockSizeFrames), MinSizeFrames = FMath::Min(AudioMixerOutputFrames, BlockSizeFrames);
+		return FMath::DivideAndRoundUp(MaxSizeFrames, MinSizeFrames) - 1;
+	}
+
 	template<uint32 NumChannels>
 	class TAudioBusWriterOperator : public TExecutableOperator<TAudioBusWriterOperator<NumChannels>>
 	{
@@ -84,6 +91,7 @@ namespace Metasound
 
 			bool bHasEnvironmentVars = InParams.Environment.Contains<Audio::FDeviceId>(SourceInterface::Environment::DeviceID);
 			bHasEnvironmentVars &= InParams.Environment.Contains<int32>(SourceInterface::Environment::AudioMixerNumOutputFrames);
+			bHasEnvironmentVars &= InParams.Environment.Contains<uint64>(SourceInterface::Environment::TransmitterID);
 
 			if (bHasEnvironmentVars)
 			{
@@ -109,8 +117,8 @@ namespace Metasound
 			}
 			else
 			{
-				UE_LOG(LogMetaSound, Warning, TEXT("Audio bus writer node requires audio device ID '%s' and audio mixer num output frames '%s' environment variables")
-					, *SourceInterface::Environment::DeviceID.ToString(), *SourceInterface::Environment::AudioMixerNumOutputFrames.ToString());
+				UE_LOG(LogMetaSound, Warning, TEXT("Audio bus writer node requires audio device ID '%s', audio mixer num output frames '%s' and transmitter id '%s' environment variables")
+					, *SourceInterface::Environment::DeviceID.ToString(), *SourceInterface::Environment::AudioMixerNumOutputFrames.ToString(), *SourceInterface::Environment::TransmitterID.ToString());
 				return nullptr;
 			}
 		}
@@ -118,7 +126,6 @@ namespace Metasound
 		TAudioBusWriterOperator(const FBuildOperatorParams& InParams, FAudioBusAssetReadRef InAudioBusAsset, TArray<FAudioBufferReadRef> InAudioInputs, FString InGraphName)
 			: AudioBusAsset(MoveTemp(InAudioBusAsset))
 			, AudioInputs(MoveTemp(InAudioInputs))
-			, GraphName(MoveTemp(InGraphName))
 		{
 			Reset(InParams);
 		}
@@ -128,29 +135,41 @@ namespace Metasound
 			const FAudioBusProxyPtr& AudioBusProxy = AudioBusAsset->GetAudioBusProxy();
 			if (AudioBusProxy.IsValid())
 			{
+				if (AudioBusProxy->NumChannels <= 0)
+				{
+					UE_LOG(LogMetaSound, Warning, TEXT("AudioBusProxy is invalid (NumChannels = %i)."), AudioBusProxy->NumChannels);
+					return;
+				}
+
+				UAudioBusSubsystem* AudioBusSubsystem = nullptr;
 				if (FAudioDeviceManager* ADM = FAudioDeviceManager::Get())
 				{
 					if (FAudioDevice* AudioDevice = ADM->GetAudioDeviceRaw(AudioDeviceId))
 					{
-						UAudioBusSubsystem* AudioBusSubsystem = AudioDevice->GetSubsystem<UAudioBusSubsystem>();
+						AudioBusSubsystem = AudioDevice->GetSubsystem<UAudioBusSubsystem>();
 						check(AudioBusSubsystem);
-						
-						AudioBusId = AudioBusProxy->AudioBusId;
-						const Audio::FAudioBusKey AudioBusKey = Audio::FAudioBusKey(AudioBusId);
-						
-						// Start the audio bus in case it's not already started					
-						AudioBusChannels = AudioBusProxy->NumChannels;
-						AudioBusSubsystem->StartAudioBus(AudioBusKey, AudioBusChannels, false);
+						if (!AudioBusSubsystem)
+						{
+							return;
+						}
+					}
+				}
 
-						InterleavedBuffer.Reset();
-						InterleavedBuffer.Reserve(NumBlocksToNumSamples(InitialNumBlocks()));
+				AudioBusChannels = uint32(FMath::Min(AudioBusProxy->NumChannels, int32(EAudioBusChannels::MaxChannelCount)));
+				AudioBusId = AudioBusProxy->AudioBusId;
 
-						// Create a bus patch input with enough room for the number of samples we expect and some buffering
-						AudioBusPatchInput = AudioBusSubsystem->AddPatchInputForAudioBus(AudioBusKey, BlockSizeFrames, AudioBusChannels);
+				Audio::FAudioBusKey AudioBusKey(AudioBusId);
+				AudioBusSubsystem->StartAudioBus(AudioBusKey, AudioBusChannels, false);
 
-						ConnectionState = EConnectionState::Disconnected;
-					}
+				AudioBusPatchInput = AudioBusSubsystem->AddPatchInputForSoundAndAudioBus(InstanceID, AudioBusKey, BlockSizeFrames, int32(AudioBusChannels));
+				int32 NumBlocksToPush = InitialNumBlocks();
+				if (NumBlocksToPush > 0)
+				{
+					AudioBusPatchInput.PushAudio(nullptr, NumBlocksToNumSamples(NumBlocksToPush));
 				}
+
+				InterleavedBuffer.Reset();
+				InterleavedBuffer.AddUninitialized(NumBlocksToNumSamples(1));
 			}
 		}
 		
@@ -159,24 +178,32 @@ namespace Metasound
 			using namespace Frontend;
 			using namespace AudioBusWriterNode;
 
+			InterleavedBuffer.Reset();
+			AudioMixerOutputFrames = INDEX_NONE;
+			AudioDeviceId = INDEX_NONE;
+			SampleRate = InParams.OperatorSettings.GetSampleRate();
+			AudioBusPatchInput.Reset();
+			AudioBusChannels = INDEX_NONE;
+			AudioBusId = 0;
+			InstanceID = 0;
+			BlockSizeFrames = InParams.OperatorSettings.GetNumFramesPerBlock();
+			bWasUnderrunReported = false;
+
 			bool bHasEnvironmentVars = InParams.Environment.Contains<Audio::FDeviceId>(SourceInterface::Environment::DeviceID);
 			bHasEnvironmentVars &= InParams.Environment.Contains<int32>(SourceInterface::Environment::AudioMixerNumOutputFrames);
+			bHasEnvironmentVars &= InParams.Environment.Contains<uint64>(SourceInterface::Environment::TransmitterID);
 
 			if (bHasEnvironmentVars)
 			{
-				SampleRate = InParams.OperatorSettings.GetSampleRate();
 				AudioDeviceId = InParams.Environment.GetValue<Audio::FDeviceId>(SourceInterface::Environment::DeviceID);
 				AudioMixerOutputFrames = InParams.Environment.GetValue<int32>(SourceInterface::Environment::AudioMixerNumOutputFrames);
+				InstanceID = InParams.Environment.GetValue<uint64>(SourceInterface::Environment::TransmitterID);
 			}
 			else
 			{
-				UE_LOG(LogMetaSound, Warning, TEXT("Audio bus writer node requires audio device ID '%s' and audio mixer num output frames '%s' environment variables")
-					, *SourceInterface::Environment::DeviceID.ToString(), *SourceInterface::Environment::AudioMixerNumOutputFrames.ToString());
+				UE_LOG(LogMetaSound, Warning, TEXT("Audio bus writer node requires audio device ID '%s', audio mixer num output frames '%s' and transmitter id '%s' environment variables")
+					, *SourceInterface::Environment::DeviceID.ToString(), *SourceInterface::Environment::AudioMixerNumOutputFrames.ToString(), *SourceInterface::Environment::TransmitterID.ToString());
 			}
-			
-			BlockSizeFrames = InParams.OperatorSettings.GetNumFramesPerBlock();
-
-			bWasUnderrunReported = false;
 
 			CreatePatchInput();
 		}
@@ -219,23 +246,19 @@ namespace Metasound
 			const FAudioBusProxyPtr& BusProxy = AudioBusAsset->GetAudioBusProxy();
 			if (BusProxy.IsValid() && BusProxy->AudioBusId != AudioBusId)
 			{
-				AudioBusPatchInput.Reset();
+				InterleavedBuffer.Reset();
 			}
 			
-			if (!AudioBusPatchInput.IsValid())
+			if (InterleavedBuffer.IsEmpty())
 			{
 				// if environment vars & a valid audio bus have been set since starting, try to create the patch now
 				if (SampleRate > 0.f && BusProxy.IsValid())
 				{
 					CreatePatchInput();
 				}
-				else
-				{
-					return;
-				}
 			}
 
-			if (!AudioBusPatchInput.IsOutputStillActive())
+			if (InterleavedBuffer.IsEmpty())
 			{
 				return;
 			}
@@ -246,9 +269,7 @@ namespace Metasound
 			{
 				AudioInputBufferPtrs[ChannelIndex] = AudioInputs[ChannelIndex]->GetData();
 			}
-			int32 InterleavedBufferOffset = InterleavedBuffer.Num();
-			InterleavedBuffer.SetNum(InterleavedBufferOffset + BlockSizeFrames * NumChannels);
-			float* InterleavedBufferPtr = InterleavedBuffer.GetData() + InterleavedBufferOffset;
+			float* InterleavedBufferPtr = InterleavedBuffer.GetData();
 
 			if (AudioBusChannels == 1)
 			{
@@ -272,45 +293,6 @@ namespace Metasound
 				}
 			}
 
-			switch (ConnectionState)
-			{
-			case EConnectionState::Disconnected:
-				{
-					// Wait until InterleavedBuffer contains enough samples to satisfy any mix eventualities!
-					// A single mix requires enough MetaSound executions to satisfy its buffer size.
-					// Mixing can occur concurrently with the first MetaSound execution intended for the next mix,
-					// and can steal that execution's patch output.
-					if (InterleavedBuffer.Num() < NumBlocksToNumSamples(InitialNumBlocks()))
-					{
-						return;
-					}
-
-					ConnectionState = EConnectionState::ConnectionPending;
-					break;
-				}
-
-			case EConnectionState::ConnectionPending:
-				{
-					int32 InitialNumSamples = NumBlocksToNumSamples(InitialNumBlocks());
-
-					// Determine if the pending connection has been established, by detecting if samples have been consumed.
-					if (AudioBusPatchInput.GetNumSamplesAvailable() == InitialNumSamples)
-					{
-						// If the connection hasn't been established by the time as many executions have occurred again,
-						// something has probably gone wrong.
-						if (InterleavedBuffer.Num() == InitialNumSamples)
-						{
-							UE_LOG(LogMetaSound, Warning, TEXT("Graph %s: Writer node executed before mixer patch connection established, with buffer size %d."), *GraphName, InterleavedBuffer.Num());
-							return;
-						}
-						return;
-					}
-
-					ConnectionState = EConnectionState::Connected;
-					break;
-				}
-			}
-
 			// Pushes the interleaved data to the audio bus
 			const int32 SamplesPushed = AudioBusPatchInput.PushAudio(InterleavedBuffer.GetData(), InterleavedBuffer.Num());
 			if (SamplesPushed < InterleavedBuffer.Num() && !bWasUnderrunReported)
@@ -318,22 +300,12 @@ namespace Metasound
 				UE_LOG(LogMetaSound, Warning, TEXT("Underrun detected in audio bus writer node."));
 				bWasUnderrunReported = true;
 			}
-			InterleavedBuffer.Reset();
 		}
 
 	private:
 		int32 InitialNumBlocks() const
 		{
-			if (AudioMixerOutputFrames == BlockSizeFrames)
-			{
-				return 1;
-			}
-
-			// We need enough blocks for as many executions as the mixer can consume at once, plus one more,
-			// because the last execution can be concurrent with the mixer, and could contribute to either the
-			// current mix or the next.  That could leave us with one block too few if we didn't have an extra.
-			int32 MaxSizeFrames = FMath::Max(AudioMixerOutputFrames, BlockSizeFrames), MinSizeFrames = FMath::Min(AudioMixerOutputFrames, BlockSizeFrames);
-			return 1 + FMath::DivideAndRoundUp(MaxSizeFrames, MinSizeFrames);
+			return AudioBusWriterNodeInitialNumBlocks(BlockSizeFrames, AudioMixerOutputFrames);
 		}
 
 		int32 NumBlocksToNumSamples(int32 NumBlocks) const
@@ -343,24 +315,17 @@ namespace Metasound
 
 		FAudioBusAssetReadRef AudioBusAsset;
 		TArray<FAudioBufferReadRef> AudioInputs;
-		FString GraphName;
 
 		TArray<float> InterleavedBuffer;
 		int32 AudioMixerOutputFrames = INDEX_NONE;
 		Audio::FDeviceId AudioDeviceId = INDEX_NONE;
 		float SampleRate = 0.0f;
 		Audio::FPatchInput AudioBusPatchInput;
+		uint64 InstanceID = 0;
 		uint32 AudioBusChannels = INDEX_NONE;
 		uint32 AudioBusId = 0;
-		bool bWasUnderrunReported = false;
 		int32 BlockSizeFrames = 0;
-		enum class EConnectionState : uint8
-		{
-			Disconnected,
-			ConnectionPending,
-			Connected
-		}
-		ConnectionState = EConnectionState::Disconnected;
+		bool bWasUnderrunReported = false;
 	};
 
 	template<uint32 NumChannels>
diff --git a/Engine/Plugins/Runtime/Metasound/Source/MetasoundEngine/Private/MetasoundSource.cpp b/Engine/Plugins/Runtime/Metasound/Source/MetasoundEngine/Private/MetasoundSource.cpp
--- a/Engine/Plugins/Runtime/Metasound/Source/MetasoundEngine/Private/MetasoundSource.cpp
+++ b/Engine/Plugins/Runtime/Metasound/Source/MetasoundEngine/Private/MetasoundSource.cpp
@@ -5,6 +5,8 @@
 #include "Algo/Find.h"
 #include "Algo/Transform.h"
 #include "AssetRegistry/AssetRegistryModule.h"
+#include "AudioBusSubsystem.h"
+#include "AudioDevice.h"
 #include "AudioDeviceManager.h"
 #include "Containers/Ticker.h"
 #include "IAudioParameterInterfaceRegistry.h"
@@ -1043,6 +1045,19 @@ ISoundGeneratorPtr UMetaSoundSource::CreateSoundGenerator(const FSoundGeneratorI
 	if (Generator.IsValid())
 	{
 		TrackGenerator(InParams.AudioComponentId, Generator);
+
+		Generator->AddGraphSetCallback(Metasound::FOnSetGraph::FDelegate::CreateLambda([AudioDeviceID = InParams.AudioDeviceID, InstanceID = InParams.InstanceID]()
+			{
+				if (FAudioDeviceManager* ADM = FAudioDeviceManager::Get())
+				{
+					if (FAudioDevice* AudioDevice = ADM->GetAudioDeviceRaw(AudioDeviceID))
+					{
+						UAudioBusSubsystem* AudioBusSubsystem = AudioDevice->GetSubsystem<UAudioBusSubsystem>();
+						check(AudioBusSubsystem);
+						AudioBusSubsystem->ReadyToConnect(InstanceID);
+					}
+				}
+			}));
 	}
 
 	return ISoundGeneratorPtr(Generator);
diff --git a/Engine/Plugins/Runtime/Metasound/Source/MetasoundEngine/Public/MetasoundAudioBus.h b/Engine/Plugins/Runtime/Metasound/Source/MetasoundEngine/Public/MetasoundAudioBus.h
--- a/Engine/Plugins/Runtime/Metasound/Source/MetasoundEngine/Public/MetasoundAudioBus.h
+++ b/Engine/Plugins/Runtime/Metasound/Source/MetasoundEngine/Public/MetasoundAudioBus.h
@@ -42,5 +42,8 @@ namespace Metasound
 	};
 
 	DECLARE_METASOUND_DATA_REFERENCE_TYPES(FAudioBusAsset, METASOUNDENGINE_API, FAudioBusAssetTypeInfo, FAudioBusAssetReadRef, FAudioBusAssetWriteRef)
+
+	METASOUNDENGINE_API int32 AudioBusReaderNodeInitialNumBlocks(int32 BlockSizeFrames, int32 AudioMixerOutputFrames);
+	METASOUNDENGINE_API int32 AudioBusWriterNodeInitialNumBlocks(int32 BlockSizeFrames, int32 AudioMixerOutputFrames);
 }
  
diff --git a/Engine/Source/Runtime/AudioMixer/Private/AudioBusSubsystem.cpp b/Engine/Source/Runtime/AudioMixer/Private/AudioBusSubsystem.cpp
--- a/Engine/Source/Runtime/AudioMixer/Private/AudioBusSubsystem.cpp
+++ b/Engine/Source/Runtime/AudioMixer/Private/AudioBusSubsystem.cpp
@@ -145,6 +145,98 @@ Audio::FPatchOutputStrongPtr UAudioBusSubsystem::AddPatchOutputForAudioBus(Audio
 	return PatchOutput;
 }
 
+Audio::FPatchInput UAudioBusSubsystem::AddPatchInputForSoundAndAudioBus(uint64 SoundInstanceID, Audio::FAudioBusKey AudioBusKey, int32 InFrames, int32 NumChannels, float InGain)
+{
+	Audio::FMixerDevice* MixerDevice = GetMutableMixerDevice();
+	if (!MixerDevice)
+	{
+		return {};
+	}
+
+	if (Audio::FPatchOutputStrongPtr PatchOutput = MixerDevice->MakePatch(InFrames, NumChannels, InGain))
+	{
+		Audio::FPatchInput PatchInput = MoveTemp(PatchOutput);
+		AddPendingConnection(SoundInstanceID, FPendingConnection{ FPendingConnection::FPatchVariant(TInPlaceType<Audio::FPatchInput>(), PatchInput), MoveTemp(AudioBusKey), InFrames, NumChannels });
+		return PatchInput;
+	}
+
+	return {};
+}
+
+Audio::FPatchOutputStrongPtr UAudioBusSubsystem::AddPatchOutputForSoundAndAudioBus(uint64 SoundInstanceID, Audio::FAudioBusKey AudioBusKey, int32 InFrames, int32 NumChannels, float InGain)
+{
+	Audio::FMixerDevice* MixerDevice = GetMutableMixerDevice();
+	if (!MixerDevice)
+	{
+		return {};
+	}
+
+	if (Audio::FPatchOutputStrongPtr PatchOutput = MixerDevice->MakePatch(InFrames, NumChannels, InGain))
+	{
+		AddPendingConnection(SoundInstanceID, FPendingConnection{ FPendingConnection::FPatchVariant(TInPlaceType<Audio::FPatchOutputStrongPtr>(), PatchOutput), MoveTemp(AudioBusKey), InFrames, NumChannels });
+		return PatchOutput;
+	}
+
+	return {};
+}
+
+void UAudioBusSubsystem::AddPendingConnection(uint64 SoundInstanceID, FPendingConnection&& PendingConnection)
+{
+	FScopeLock ScopeLock(&Mutex);
+	FSoundInstanceConnections& SoundInstanceConnections = SoundInstanceConnectionMap.FindOrAdd(SoundInstanceID);
+	SoundInstanceConnections.PendingConnections.Add(MoveTemp(PendingConnection));
+}
+
+void UAudioBusSubsystem::ReadyToConnect(uint64 SoundInstanceID)
+{
+	FScopeLock ScopeLock(&Mutex);
+	FSoundInstanceConnections& SoundInstanceConnections = SoundInstanceConnectionMap.FindOrAdd(SoundInstanceID);
+	SoundInstanceConnections.bSoundInstanceReady = true;
+}
+
+void UAudioBusSubsystem::ConnectPatches(uint64 SoundInstanceID)
+{
+	TArray<FPendingConnection> PendingConnections = ExtractPendingConnectionsIfReady(SoundInstanceID);
+	if (!PendingConnections.IsEmpty())
+	{
+		Audio::FMixerSourceManager* SourceManager = GetMutableSourceManager();
+		check(SourceManager);
+		for (FPendingConnection& PendingConnection : PendingConnections)
+		{
+			switch (PendingConnection.PatchVariant.GetIndex())
+			{
+			case FPendingConnection::FPatchVariant::IndexOfType<Audio::FPatchInput>():
+				SourceManager->AddPendingAudioBusConnection(MoveTemp(PendingConnection.AudioBusKey), PendingConnection.NumChannels, PendingConnection.bIsAutomatic, MoveTemp(PendingConnection.PatchVariant.Get<Audio::FPatchInput>()));
+				break;
+			case FPendingConnection::FPatchVariant::IndexOfType<Audio::FPatchOutputStrongPtr>():
+				SourceManager->AddPendingAudioBusConnection(MoveTemp(PendingConnection.AudioBusKey), PendingConnection.NumChannels, PendingConnection.bIsAutomatic, MoveTemp(PendingConnection.PatchVariant.Get<Audio::FPatchOutputStrongPtr>()));
+				break;
+			}
+		}
+	}
+}
+
+void UAudioBusSubsystem::RemoveSound(uint64 SoundInstanceID)
+{
+	FScopeLock ScopeLock(&Mutex);
+	SoundInstanceConnectionMap.Remove(SoundInstanceID);
+}
+
+TArray<UAudioBusSubsystem::FPendingConnection> UAudioBusSubsystem::ExtractPendingConnectionsIfReady(uint64 SoundInstanceID)
+{
+	FScopeLock ScopeLock(&Mutex);
+	if (FSoundInstanceConnections* SoundInstanceConnections = SoundInstanceConnectionMap.Find(SoundInstanceID))
+	{
+		if (SoundInstanceConnections->bSoundInstanceReady)
+		{
+			TArray<FPendingConnection> PendingConnections = MoveTemp(SoundInstanceConnections->PendingConnections);
+			SoundInstanceConnections->PendingConnections.Empty();
+			return PendingConnections;
+		}
+	}
+	return {};
+}
+
 void UAudioBusSubsystem::InitDefaultAudioBuses()
 {
 	if (!ensure(IsInGameThread()))
diff --git a/Engine/Source/Runtime/AudioMixer/Private/AudioMixerSourceBuffer.cpp b/Engine/Source/Runtime/AudioMixer/Private/AudioMixerSourceBuffer.cpp
--- a/Engine/Source/Runtime/AudioMixer/Private/AudioMixerSourceBuffer.cpp
+++ b/Engine/Source/Runtime/AudioMixer/Private/AudioMixerSourceBuffer.cpp
@@ -78,6 +78,7 @@ namespace Audio
 		, BufferType(InArgs.Buffer->GetType())
 		, NumPrecacheFrames(InArgs.SoundWave->NumPrecacheFrames)
 		, AuioDeviceID(InArgs.AudioDeviceID)
+		, InstanceID(InArgs.InstanceID)
 		, WaveName(InArgs.SoundWave->GetFName())
 #if ENABLE_AUDIO_DEBUG
 		, SampleRate(InArgs.SampleRate)
@@ -490,6 +491,16 @@ namespace Audio
 			SubmitRealTimeSourceData(bIsFinishedOrLooped);
 		}
 
+		if (FAudioDeviceManager* ADM = FAudioDeviceManager::Get())
+		{
+			if (FAudioDevice* AudioDevice = ADM->GetAudioDeviceRaw(AuioDeviceID))
+			{
+				UAudioBusSubsystem* AudioBusSubsystem = AudioDevice->GetSubsystem<UAudioBusSubsystem>();
+				check(AudioBusSubsystem);
+				AudioBusSubsystem->ConnectPatches(InstanceID);
+			}
+		}
+
 		if (!AsyncRealtimeAudioTask)
 		{
 			// Update the buffer index
@@ -688,6 +699,16 @@ namespace Audio
 			{
 				SoundWave->OnEndGenerate(SoundGenerator);
 			}
+
+			if (FAudioDeviceManager* ADM = FAudioDeviceManager::Get())
+			{
+				if (FAudioDevice* AudioDevice = ADM->GetAudioDeviceRaw(AuioDeviceID))
+				{
+					UAudioBusSubsystem* AudioBusSubsystem = AudioDevice->GetSubsystem<UAudioBusSubsystem>();
+					check(AudioBusSubsystem);
+					AudioBusSubsystem->RemoveSound(InstanceID);
+				}
+			}
 		}
 		else
 		{
diff --git a/Engine/Source/Runtime/AudioMixer/Private/AudioMixerSourceBuffer.h b/Engine/Source/Runtime/AudioMixer/Private/AudioMixerSourceBuffer.h
--- a/Engine/Source/Runtime/AudioMixer/Private/AudioMixerSourceBuffer.h
+++ b/Engine/Source/Runtime/AudioMixer/Private/AudioMixerSourceBuffer.h
@@ -50,7 +50,7 @@ namespace Audio
 	{
 		FDeviceId AudioDeviceID = 0;
 		uint64 AudioComponentID = 0;
-		uint32 InstanceID = 0;
+		uint64 InstanceID = 0;
 		int32 SampleRate = 0;
 		int32 AudioMixerNumOutputFrames = 0;
 		FMixerBuffer* Buffer = nullptr;
@@ -156,6 +156,7 @@ namespace Audio
 		Audio::EBufferType::Type BufferType;
 		int32 NumPrecacheFrames;
 		Audio::FDeviceId AuioDeviceID;
+		uint64 InstanceID;
 		TArray<uint8> CachedRealtimeFirstBuffer;
 		FName WaveName;
 		uint64 AsyncTaskStartTimeInCycles=0;
diff --git a/Engine/Source/Runtime/AudioMixer/Public/AudioBusSubsystem.h b/Engine/Source/Runtime/AudioMixer/Public/AudioBusSubsystem.h
--- a/Engine/Source/Runtime/AudioMixer/Public/AudioBusSubsystem.h
+++ b/Engine/Source/Runtime/AudioMixer/Public/AudioBusSubsystem.h
@@ -2,6 +2,8 @@
 
 #pragma once
 
+#include "DSP/MultithreadedPatching.h"
+#include "Misc/TVariant.h"
 #include "Subsystems/AudioEngineSubsystem.h"
 #include "Templates/TypeHash.h"
 #include "UObject/StrongObjectPtr.h"
@@ -15,9 +17,6 @@ namespace Audio
 	// Forward declarations 
 	class FMixerAudioBus;
 	class FMixerSourceManager;
-	class FPatchInput;
-	struct FPatchOutput;
-	typedef TSharedPtr<FPatchOutput, ESPMode::ThreadSafe> FPatchOutputStrongPtr;
 
 	struct FAudioBusKey
 	{
@@ -87,6 +86,12 @@ public:
 	AUDIOMIXER_API Audio::FPatchInput AddPatchInputForAudioBus(Audio::FAudioBusKey InAudioBusKey, int32 InFrames, int32 InChannels, float InGain = 1.f);
 	AUDIOMIXER_API Audio::FPatchOutputStrongPtr AddPatchOutputForAudioBus(Audio::FAudioBusKey InAudioBusKey, int32 InFrames, int32 InChannels, float InGain = 1.f);
 
+	AUDIOMIXER_API Audio::FPatchInput AddPatchInputForSoundAndAudioBus(uint64 SoundInstanceID, Audio::FAudioBusKey AudioBusKey, int32 InFrames, int32 NumChannels, float InGain = 1.f);
+	AUDIOMIXER_API Audio::FPatchOutputStrongPtr AddPatchOutputForSoundAndAudioBus(uint64 SoundInstanceID, Audio::FAudioBusKey AudioBusKey, int32 InFrames, int32 NumChannels, float InGain = 1.f);
+	AUDIOMIXER_API void ReadyToConnect(uint64 SoundInstanceID);
+	AUDIOMIXER_API void ConnectPatches(uint64 SoundInstanceID);
+	AUDIOMIXER_API void RemoveSound(uint64 SoundInstanceID);
+
 	AUDIOMIXER_API void InitDefaultAudioBuses();
 	AUDIOMIXER_API void ShutdownDefaultAudioBuses();
 
@@ -101,4 +106,27 @@ private:
 	TArray<TStrongObjectPtr<UAudioBus>> DefaultAudioBuses; 
 	// The active audio bus list accessible on the game thread
 	TMap<Audio::FAudioBusKey, FActiveBusData> ActiveAudioBuses_GameThread;
+
+	struct FPendingConnection
+	{
+		using FPatchVariant = TVariant<Audio::FPatchInput, Audio::FPatchOutputStrongPtr>;
+		FPatchVariant PatchVariant;
+		Audio::FAudioBusKey AudioBusKey;
+		int32 BlockSizeFrames = 0;
+		int32 NumChannels = 0;
+		bool bIsAutomatic = false;
+	};
+
+	void AddPendingConnection(uint64 SoundInstanceID, FPendingConnection&& PendingConnection);
+
+	struct FSoundInstanceConnections
+	{
+		TArray<FPendingConnection> PendingConnections;
+		bool bSoundInstanceReady = false;
+	};
+
+	TArray<FPendingConnection> ExtractPendingConnectionsIfReady(uint64 SoundInstanceID);
+
+	TMap<uint64, FSoundInstanceConnections> SoundInstanceConnectionMap;
+	FCriticalSection Mutex;
 };
