- Make Version call to generate MetaSound Builder async GC-safe
- Misc clean-up
#rnx

diff --git a/Engine/Plugins/Runtime/Metasound/Source/MetasoundEditor/Private/MetasoundEditor.cpp b/Engine/Plugins/Runtime/Metasound/Source/MetasoundEditor/Private/MetasoundEditor.cpp
--- a/Engine/Plugins/Runtime/Metasound/Source/MetasoundEditor/Private/MetasoundEditor.cpp
+++ b/Engine/Plugins/Runtime/Metasound/Source/MetasoundEditor/Private/MetasoundEditor.cpp
@@ -884,23 +884,22 @@ namespace Metasound
 			Metasound->SetFlags(RF_Transactional);
 
 			// Typically sounds are versioned on load of the asset. There are certain instances where an asset is not versioned on reload.
-			// This is here as a hack around that fact to force versioning the document on load prior to the editor synchronizing and building
-			// the editor graph if an asset is reloaded while the asset editor was open.
+			// This forces versioning the document on load prior to the editor synchronizing and building the editor graph if an asset is
+			// reloaded while the asset editor was open.
+			Builder.Reset(&Engine::FDocumentBuilderRegistry::GetChecked().FindOrBeginBuilding(*ObjectToEdit));
 			if (FMetasoundAssetBase* MetaSoundAsset = IMetasoundUObjectRegistry::Get().GetObjectAsAssetBase(Metasound))
 			{
-				if (MetaSoundAsset->VersionAsset())
+				if (MetaSoundAsset->VersionAsset(Builder->GetBuilder()))
 				{
 					MetaSoundAsset->SetVersionedOnLoad();
 				}
 
-				// Hack to ensure validation is re-run on re-opening of the editor.
+				// Ensures validation is re-run on re-opening of the editor.
 				// This is needed to refresh errors potentially caused by unloading of
 				// references (ex. if a referenced asset is force deleted in the editor).
 				MetaSoundAsset->GetModifyContext().SetForceRefreshViews();
 			}
 
-			Builder.Reset(&Engine::FDocumentBuilderRegistry::GetChecked().FindOrBeginBuilding(*ObjectToEdit));
-
 			GEditor->RegisterForUndo(this);
 
 			FGraphEditorCommands::Register();
diff --git a/Engine/Plugins/Runtime/Metasound/Source/MetasoundEngine/Private/Metasound.cpp b/Engine/Plugins/Runtime/Metasound/Source/MetasoundEngine/Private/Metasound.cpp
--- a/Engine/Plugins/Runtime/Metasound/Source/MetasoundEngine/Private/Metasound.cpp
+++ b/Engine/Plugins/Runtime/Metasound/Source/MetasoundEngine/Private/Metasound.cpp
@@ -106,13 +106,13 @@ const FMetasoundFrontendDocument& UMetaSoundPatch::GetConstDocument() const
 void UMetaSoundPatch::PostDuplicate(EDuplicateMode::Type InDuplicateMode)
 {
 	Super::PostDuplicate(InDuplicateMode);
-	Metasound::FMetaSoundEngineAssetHelper::PostDuplicate(this, InDuplicateMode, AssetClassID);
+	Metasound::Engine::FAssetHelper::PostDuplicate(this, InDuplicateMode, AssetClassID);
 }
 
 void UMetaSoundPatch::PostEditUndo()
 {
 	Super::PostEditUndo();
-	Metasound::FMetaSoundEngineAssetHelper::PostEditUndo(*this);
+	Metasound::Engine::FAssetHelper::PostEditUndo(*this);
 }
 #endif // WITHEDITOR
 
@@ -125,13 +125,13 @@ void UMetaSoundPatch::BeginDestroy()
 void UMetaSoundPatch::PreSave(FObjectPreSaveContext InSaveContext)
 {
 	Super::PreSave(InSaveContext);
-	Metasound::FMetaSoundEngineAssetHelper::PreSaveAsset(*this, InSaveContext);
+	Metasound::Engine::FAssetHelper::PreSaveAsset(*this, InSaveContext);
 }
 
 void UMetaSoundPatch::Serialize(FArchive& InArchive)
 {
 	Super::Serialize(InArchive);
-	Metasound::FMetaSoundEngineAssetHelper::SerializeToArchive(*this, InArchive);
+	Metasound::Engine::FAssetHelper::SerializeToArchive(*this, InArchive);
 }
 
 #if WITH_EDITORONLY_DATA
@@ -165,32 +165,32 @@ FText UMetaSoundPatch::GetDisplayName() const
 
 void UMetaSoundPatch::SetRegistryAssetClassInfo(const Metasound::Frontend::FNodeClassInfo& InNodeInfo)
 {
-	Metasound::FMetaSoundEngineAssetHelper::SetMetaSoundRegistryAssetClassInfo(*this, InNodeInfo);
+	Metasound::Engine::FAssetHelper::SetMetaSoundRegistryAssetClassInfo(*this, InNodeInfo);
 }
 #endif // WITH_EDITORONLY_DATA
 
 
 FTopLevelAssetPath UMetaSoundPatch::GetAssetPathChecked() const
 {
-	return Metasound::FMetaSoundEngineAssetHelper::GetAssetPathChecked(*this);
+	return Metasound::Engine::FAssetHelper::GetAssetPathChecked(*this);
 }
 
 void UMetaSoundPatch::PostLoad() 
 {
 	Super::PostLoad();
-	Metasound::FMetaSoundEngineAssetHelper::PostLoad(*this);
+	Metasound::Engine::FAssetHelper::PostLoad(*this);
 }
 
 #if WITH_EDITOR
 void UMetaSoundPatch::SetReferencedAssetClasses(TSet<Metasound::Frontend::IMetaSoundAssetManager::FAssetInfo>&& InAssetClasses)
 {
-	Metasound::FMetaSoundEngineAssetHelper::SetReferencedAssetClasses(*this, MoveTemp(InAssetClasses));
+	Metasound::Engine::FAssetHelper::SetReferencedAssetClasses(*this, MoveTemp(InAssetClasses));
 }
 #endif
 
 TArray<FMetasoundAssetBase*> UMetaSoundPatch::GetReferencedAssets()
 {
-	return Metasound::FMetaSoundEngineAssetHelper::GetReferencedAssets(*this);
+	return Metasound::Engine::FAssetHelper::GetReferencedAssets(*this);
 }
 
 const TSet<FSoftObjectPath>& UMetaSoundPatch::GetAsyncReferencedAssetClassPaths() const 
@@ -200,7 +200,7 @@ const TSet<FSoftObjectPath>& UMetaSoundPatch::GetAsyncReferencedAssetClassPaths(
 
 void UMetaSoundPatch::OnAsyncReferencedAssetsLoaded(const TArray<FMetasoundAssetBase*>& InAsyncReferences)
 {
-	Metasound::FMetaSoundEngineAssetHelper::OnAsyncReferencedAssetsLoaded(*this, InAsyncReferences);
+	Metasound::Engine::FAssetHelper::OnAsyncReferencedAssetsLoaded(*this, InAsyncReferences);
 }
 
 bool UMetaSoundPatch::IsActivelyBuilding() const
diff --git a/Engine/Plugins/Runtime/Metasound/Source/MetasoundEngine/Private/MetasoundDocumentBuilderRegistry.cpp b/Engine/Plugins/Runtime/Metasound/Source/MetasoundEngine/Private/MetasoundDocumentBuilderRegistry.cpp
--- a/Engine/Plugins/Runtime/Metasound/Source/MetasoundEngine/Private/MetasoundDocumentBuilderRegistry.cpp
+++ b/Engine/Plugins/Runtime/Metasound/Source/MetasoundEngine/Private/MetasoundDocumentBuilderRegistry.cpp
@@ -158,7 +158,7 @@ namespace Metasound::Engine
 		return FindBuilder(DocumentInterface);
 	}
 
-	bool FDocumentBuilderRegistry::FinishBuilding(const FMetasoundFrontendClassName& InClassName, bool bForceUnregister) const
+	bool FDocumentBuilderRegistry::FinishBuilding(const FMetasoundFrontendClassName& InClassName, bool bForceUnregisterNodeClass) const
 	{
 		using namespace Metasound;
 		using namespace Metasound::Engine;
@@ -178,7 +178,7 @@ namespace Metasound::Engine
 				{
 					const int32 TransactionCount = DocBuilder.GetTransactionCount();
 					const int32 LastTransactionRegistered = Builder->GetLastTransactionRegistered();
-					if (bForceUnregister || LastTransactionRegistered != TransactionCount)
+					if (bForceUnregisterNodeClass || LastTransactionRegistered != TransactionCount)
 					{
 						UObject& MetaSound = DocBuilder.CastDocumentObjectChecked<UObject>();
 						if (FMetasoundAssetBase* MetaSoundAsset = IMetasoundUObjectRegistry::Get().GetObjectAsAssetBase(&MetaSound))
diff --git a/Engine/Plugins/Runtime/Metasound/Source/MetasoundEngine/Private/MetasoundEngineAsset.h b/Engine/Plugins/Runtime/Metasound/Source/MetasoundEngine/Private/MetasoundEngineAsset.h
--- a/Engine/Plugins/Runtime/Metasound/Source/MetasoundEngine/Private/MetasoundEngineAsset.h
+++ b/Engine/Plugins/Runtime/Metasound/Source/MetasoundEngine/Private/MetasoundEngineAsset.h
@@ -12,18 +12,21 @@
 #include "Serialization/Archive.h"
 
 #if WITH_EDITORONLY_DATA
-#include "MetasoundFrontendRegistries.h"
 #include "Algo/Transform.h"
+#include "MetasoundFrontendRegistryContainer.h"
+#include "UObject/StrongObjectPtrTemplates.h"
+#include "UObject/ObjectMacros.h"
+#include "UObject/GarbageCollection.h"
 #endif // WITH_EDITORONLY_DATA
 
 
-namespace Metasound
+namespace Metasound::Engine
 {
 	/** MetaSound Engine Asset helper provides routines for UObject based MetaSound assets. 
 	 * Any UObject deriving from FMetaSoundAssetBase should use these helper functions
 	 * in their UObject overrides. 
 	 */
-	struct FMetaSoundEngineAssetHelper
+	struct FAssetHelper
 	{
 #if WITH_EDITOR
 		static void PostDuplicate(TScriptInterface<IMetaSoundDocumentInterface> MetaSound, EDuplicateMode::Type InDuplicateMode, FGuid& OutAssetClassID)
@@ -149,15 +152,27 @@ namespace Metasound
 		template <typename TMetaSoundObject>
 		static void SerializeToArchive(TMetaSoundObject& InMetaSound, FArchive& InArchive)
 		{
+#if WITH_EDITORONLY_DATA
+			bool bVersionedAsset = false;
+
 			if (InArchive.IsLoading())
 			{
-				if (InMetaSound.VersionAsset())
+				TStrongObjectPtr<UMetaSoundBuilderBase> Builder;
 				{
-#if WITH_EDITORONLY_DATA
-					InMetaSound.SetVersionedOnLoad();
-#endif // WITH_EDITORONLY_DATA
+					FGCScopeGuard ScopeGuard;
+					Builder.Reset(&FDocumentBuilderRegistry::GetChecked().FindOrBeginBuilding(InMetaSound));
 				}
+
+				check(Builder.IsValid());
+				bVersionedAsset = InMetaSound.VersionAsset(Builder->GetBuilder());
+				Builder->ClearInternalFlags(EInternalObjectFlags::Async);
 			}
+
+			if (bVersionedAsset)
+			{
+				InMetaSound.SetVersionedOnLoad();
+			}
+#endif // WITH_EDITORONLY_DATA
 		}
 
 		template<typename TMetaSoundObject>
@@ -235,4 +250,4 @@ namespace Metasound
 		}
 #endif // WITH_EDITORONLY_DATA
 	};
-} // namespace Metasound
+} // namespace Metasound::Engine
diff --git a/Engine/Plugins/Runtime/Metasound/Source/MetasoundEngine/Private/MetasoundSource.cpp b/Engine/Plugins/Runtime/Metasound/Source/MetasoundEngine/Private/MetasoundSource.cpp
--- a/Engine/Plugins/Runtime/Metasound/Source/MetasoundEngine/Private/MetasoundSource.cpp
+++ b/Engine/Plugins/Runtime/Metasound/Source/MetasoundEngine/Private/MetasoundSource.cpp
@@ -360,13 +360,13 @@ const FMetasoundFrontendDocument& UMetaSoundSource::GetConstDocument() const
 void UMetaSoundSource::PostEditUndo()
 {
 	Super::PostEditUndo();
-	Metasound::FMetaSoundEngineAssetHelper::PostEditUndo(*this);
+	Metasound::Engine::FAssetHelper::PostEditUndo(*this);
 }
 
 void UMetaSoundSource::PostDuplicate(EDuplicateMode::Type InDuplicateMode)
 {
 	Super::PostDuplicate(InDuplicateMode);
-	Metasound::FMetaSoundEngineAssetHelper::PostDuplicate(this, InDuplicateMode, AssetClassID);
+	Metasound::Engine::FAssetHelper::PostDuplicate(this, InDuplicateMode, AssetClassID);
 }
 
 void UMetaSoundSource::PostEditChangeProperty(FPropertyChangedEvent& InEvent)
@@ -490,7 +490,7 @@ bool UMetaSoundSource::ConformObjectToDocument()
 
 FTopLevelAssetPath UMetaSoundSource::GetAssetPathChecked() const
 {
-	return Metasound::FMetaSoundEngineAssetHelper::GetAssetPathChecked(*this);
+	return Metasound::Engine::FAssetHelper::GetAssetPathChecked(*this);
 }
 
 void UMetaSoundSource::BeginDestroy()
@@ -502,16 +502,16 @@ void UMetaSoundSource::BeginDestroy()
 void UMetaSoundSource::PreSave(FObjectPreSaveContext InSaveContext)
 {
 	Super::PreSave(InSaveContext);
-	Metasound::FMetaSoundEngineAssetHelper::PreSaveAsset(*this, InSaveContext);
+	Metasound::Engine::FAssetHelper::PreSaveAsset(*this, InSaveContext);
 }
 
 void UMetaSoundSource::Serialize(FArchive& InArchive)
 {
+	using namespace Metasound::SourcePrivate;
+
 	Super::Serialize(InArchive);
-	Metasound::FMetaSoundEngineAssetHelper::SerializeToArchive(*this, InArchive);
-	
-	using namespace Metasound::SourcePrivate;	
-	
+	Metasound::Engine::FAssetHelper::SerializeToArchive(*this, InArchive);
+
 	// Load/Save cooked data.
 	if (InArchive.IsCooking() || (FPlatformProperties::RequiresCookedData() && InArchive.IsLoading()) || IsCookedForEditor(InArchive, this))
 	{
@@ -607,13 +607,13 @@ void UMetaSoundSource::SerializeCookedQualitySettings(const FName PlatformName,
 #if WITH_EDITOR
 void UMetaSoundSource::SetReferencedAssetClasses(TSet<Metasound::Frontend::IMetaSoundAssetManager::FAssetInfo>&& InAssetClasses)
 {
-	Metasound::FMetaSoundEngineAssetHelper::SetReferencedAssetClasses(*this, MoveTemp(InAssetClasses));
+	Metasound::Engine::FAssetHelper::SetReferencedAssetClasses(*this, MoveTemp(InAssetClasses));
 }
 #endif // WITH_EDITOR
 
 TArray<FMetasoundAssetBase*> UMetaSoundSource::GetReferencedAssets()
 {
-	return Metasound::FMetaSoundEngineAssetHelper::GetReferencedAssets(*this);
+	return Metasound::Engine::FAssetHelper::GetReferencedAssets(*this);
 }
 
 const TSet<FSoftObjectPath>& UMetaSoundSource::GetAsyncReferencedAssetClassPaths() const 
@@ -623,7 +623,7 @@ const TSet<FSoftObjectPath>& UMetaSoundSource::GetAsyncReferencedAssetClassPaths
 
 void UMetaSoundSource::OnAsyncReferencedAssetsLoaded(const TArray<FMetasoundAssetBase*>& InAsyncReferences)
 {
-	Metasound::FMetaSoundEngineAssetHelper::OnAsyncReferencedAssetsLoaded(*this, InAsyncReferences);
+	Metasound::Engine::FAssetHelper::OnAsyncReferencedAssetsLoaded(*this, InAsyncReferences);
 }
 
 #if WITH_EDITORONLY_DATA
@@ -658,14 +658,14 @@ FText UMetaSoundSource::GetDisplayName() const
 
 void UMetaSoundSource::SetRegistryAssetClassInfo(const Metasound::Frontend::FNodeClassInfo& InNodeInfo)
 {
-	Metasound::FMetaSoundEngineAssetHelper::SetMetaSoundRegistryAssetClassInfo(*this, InNodeInfo);
+	Metasound::Engine::FAssetHelper::SetMetaSoundRegistryAssetClassInfo(*this, InNodeInfo);
 }
 #endif // WITH_EDITORONLY_DATA
 
 void UMetaSoundSource::PostLoad()
 {
 	Super::PostLoad();
-	Metasound::FMetaSoundEngineAssetHelper::PostLoad(*this);
+	Metasound::Engine::FAssetHelper::PostLoad(*this);
 
 	Duration = GetDuration();
 	bLooping = IsLooping();
diff --git a/Engine/Plugins/Runtime/Metasound/Source/MetasoundEngine/Public/Metasound.h b/Engine/Plugins/Runtime/Metasound/Source/MetasoundEngine/Public/Metasound.h
--- a/Engine/Plugins/Runtime/Metasound/Source/MetasoundEngine/Public/Metasound.h
+++ b/Engine/Plugins/Runtime/Metasound/Source/MetasoundEngine/Public/Metasound.h
@@ -19,6 +19,13 @@
 #include "Metasound.generated.h"
 
 
+// Forward Declarations
+namespace Metasound::Engine
+{
+	struct FAssetHelper;
+} // namespace Metasound::Engine
+
+
 UCLASS(Abstract)
 class METASOUNDENGINE_API UMetasoundEditorGraphBase : public UEdGraph
 {
@@ -53,13 +60,6 @@ public:
 };
 
 
-namespace Metasound
-{
-	// Forward declare
-	struct FMetaSoundEngineAssetHelper;
-}
-
-
 /**
  * This asset type is used for Metasound assets that can only be used as nodes in other Metasound graphs.
  * Because of this, they contain no required inputs or outputs.
@@ -69,7 +69,7 @@ class METASOUNDENGINE_API UMetaSoundPatch : public UObject, public FMetasoundAss
 {
 	GENERATED_BODY()
 
-	friend struct Metasound::FMetaSoundEngineAssetHelper;
+	friend struct Metasound::Engine::FAssetHelper;
 	friend class UMetaSoundPatchBuilder;
 
 protected:
diff --git a/Engine/Plugins/Runtime/Metasound/Source/MetasoundEngine/Public/MetasoundDocumentBuilderRegistry.h b/Engine/Plugins/Runtime/Metasound/Source/MetasoundEngine/Public/MetasoundDocumentBuilderRegistry.h
--- a/Engine/Plugins/Runtime/Metasound/Source/MetasoundEngine/Public/MetasoundDocumentBuilderRegistry.h
+++ b/Engine/Plugins/Runtime/Metasound/Source/MetasoundEngine/Public/MetasoundDocumentBuilderRegistry.h
@@ -2,11 +2,13 @@
 #pragma once
 
 #include "Algo/Find.h"
+#include "Async/Async.h"
 #include "MetasoundBuilderBase.h"
 #include "MetasoundDocumentInterface.h"
 #include "MetasoundFrontendDocument.h"
 #include "MetasoundFrontendRegistries.h"
 #include "MetasoundLog.h"
+#include "Misc/ScopeLock.h"
 
 
 namespace Metasound::Engine
@@ -60,20 +62,24 @@ namespace Metasound::Engine
 				return *CastChecked<BuilderClass>(Builder);
 			}
 
-			TObjectPtr<UMetaSoundBuilderBase> NewBuilder = CastChecked<UMetaSoundBuilderBase>(NewObject<UObject>(&InMetaSoundObject, &DocInterface->GetBuilderUClass()));
-			FMetaSoundFrontendDocumentBuilder& BuilderRef = NewBuilder->GetBuilder();
-			BuilderRef = FMetaSoundFrontendDocumentBuilder(DocInterface);
-
-			const FMetasoundFrontendDocument& Document = DocInterface->GetConstDocument();
-			const FMetasoundFrontendClassName& FullClassName = Document.RootGraph.Metadata.GetClassName();
-			if (!FullClassName.IsValid())
+			TObjectPtr<UMetaSoundBuilderBase> NewBuilder;
 			{
-				BuilderRef.InitDocument();
+				FScopeLock Lock(&BuildersCriticalSection);
+				NewBuilder = CastChecked<UMetaSoundBuilderBase>(NewObject<UObject>(&InMetaSoundObject, &DocInterface->GetBuilderUClass()));
+				FMetaSoundFrontendDocumentBuilder& BuilderRef = NewBuilder->GetBuilder();
+				BuilderRef = FMetaSoundFrontendDocumentBuilder(DocInterface);
+
+				const FMetasoundFrontendDocument& Document = DocInterface->GetConstDocument();
+				const FMetasoundFrontendClassName& FullClassName = Document.RootGraph.Metadata.GetClassName();
+				if (!FullClassName.IsValid())
+				{
+					BuilderRef.InitDocument();
+				}
+
+				checkf(FullClassName.IsValid(), TEXT("Document initialization must result in a valid class name being generated"));
+				Builders.Add(FullClassName, NewBuilder);
 			}
 
-			checkf(FullClassName.IsValid(), TEXT("Document initialization must result in a valid class name being generated"));
-			TObjectPtr<UMetaSoundBuilderBase> NewBuilderBase = CastChecked<UMetaSoundBuilderBase>(NewBuilder);
-			Builders.Add(FullClassName, NewBuilderBase);
 			return *CastChecked<BuilderClass>(NewBuilder);
 		}
 #endif // WITH_EDITORONLY_DATA
@@ -88,7 +94,7 @@ namespace Metasound::Engine
 
 		virtual FMetaSoundFrontendDocumentBuilder* FindOutermostBuilder(const UObject& InSubObject) const override;
 
-		virtual bool FinishBuilding(const FMetasoundFrontendClassName& InClassName, bool bForceUnregister = false) const override;
+		virtual bool FinishBuilding(const FMetasoundFrontendClassName& InClassName, bool bForceUnregisterNodeClass = false) const override;
 
 		UMetaSoundBuilderBase* FindBuilderObject(TScriptInterface<const IMetaSoundDocumentInterface> MetaSound) const;
 		UMetaSoundBuilderBase* FindBuilderObject(const FMetasoundFrontendClassName& ClassName) const;
diff --git a/Engine/Plugins/Runtime/Metasound/Source/MetasoundEngine/Public/MetasoundSource.h b/Engine/Plugins/Runtime/Metasound/Source/MetasoundEngine/Public/MetasoundSource.h
--- a/Engine/Plugins/Runtime/Metasound/Source/MetasoundEngine/Public/MetasoundSource.h
+++ b/Engine/Plugins/Runtime/Metasound/Source/MetasoundEngine/Public/MetasoundSource.h
@@ -35,11 +35,19 @@ namespace Audio
 
 namespace Metasound
 {
-	struct FMetaSoundEngineAssetHelper;
 	struct FMetasoundGeneratorInitParams;
-
 	class FMetasoundGenerator;
 
+	namespace DynamicGraph
+	{
+		class FDynamicOperatorTransactor;
+	} // namespace DynamicGraph
+
+	namespace Engine
+	{
+		struct FAssetHelper;
+	} // namespace Engine;
+
 	namespace Frontend
 	{
 		class IDataTypeRegistry;
@@ -51,10 +59,6 @@ namespace Metasound
 		using FCookedQualitySettings = FMetaSoundQualitySettings;
 	} // namespace SourcePrivate
 
-	namespace DynamicGraph
-	{
-		class FDynamicOperatorTransactor;
-	} // namespace DynamicGraph
 } // namespace Metasound
 
 
@@ -69,7 +73,7 @@ class METASOUNDENGINE_API UMetaSoundSource : public USoundWaveProcedural, public
 {
 	GENERATED_BODY()
 
-	friend struct Metasound::FMetaSoundEngineAssetHelper;
+	friend struct Metasound::Engine::FAssetHelper;
 	friend class UMetaSoundSourceBuilder;
 
 	// FRuntimeInput represents an input to a MetaSound which can be manipulated.
diff --git a/Engine/Plugins/Runtime/Metasound/Source/MetasoundFrontend/Private/MetasoundAssetBase.cpp b/Engine/Plugins/Runtime/Metasound/Source/MetasoundFrontend/Private/MetasoundAssetBase.cpp
--- a/Engine/Plugins/Runtime/Metasound/Source/MetasoundFrontend/Private/MetasoundAssetBase.cpp
+++ b/Engine/Plugins/Runtime/Metasound/Source/MetasoundFrontend/Private/MetasoundAssetBase.cpp
@@ -500,7 +500,8 @@ void FMetasoundAssetBase::SetDocument(FMetasoundFrontendDocument InDocument, boo
 	}
 }
 
-bool FMetasoundAssetBase::VersionAsset()
+#if WITH_EDITORONLY_DATA
+bool FMetasoundAssetBase::VersionAsset(FMetaSoundFrontendDocumentBuilder& Builder)
 {
 	using namespace Metasound;
 	using namespace Metasound::Frontend;
@@ -509,43 +510,29 @@ bool FMetasoundAssetBase::VersionAsset()
 	constexpr bool bIsDeterministic = true;
 	FDocumentIDGenerator::FScopeDeterminism DeterminismScope(bIsDeterministic);
 
-	UObject* OwningAsset = GetOwningAsset();
-	check(OwningAsset);
-
 	bool bDidEdit = false;
 
-#if WITH_EDITORONLY_DATA
 	// Version Document Model
-	if (GetConstDocumentChecked().RequiresInterfaceVersioning())
-	{
-		GetDocumentChecked().VersionInterfaces();
-		bDidEdit = true;
-	}
-#endif // WITH_EDITORONLY_DATA
-
-	bDidEdit |= Metasound::Frontend::VersionDocument(*this);
+	bDidEdit = Builder.VersionInterfaces();
+	bDidEdit |= Metasound::Frontend::VersionDocument(*this, Builder);
 
-	// Version Interfaces. Has to be re-run until no pass reports an update in case
-	// versions fork (ex. an interface splits into two newly named interfaces).
-	// TODO: Need to add function to determine if interfaces require updating before
-	// running update to avoid invalidating builder cache when not necessary. This
-	// will require either re-writing old transforms with the builder API or just
-	// leaving handle/controller logic and moving to deprecated implementation that
-	// is bypassed on versioned documents in most cases.
+	// TODO: Move this logic to builder API above, which will require rewriting update transforms to
+	// take in builder instead of DocumentHandle.
 	{
+		const FMetasoundFrontendDocument& Document = Builder.GetConstDocumentChecked();
 		bool bInterfaceUpdated = false;
 		bool bPassUpdated = true;
-		TScriptInterface<IMetaSoundDocumentInterface> Interface(OwningAsset);
-		const IMetaSoundDocumentInterface* ConstInterface = Interface.GetInterface();
+
+		// Has to be re-run until no pass reports an update in case versions
+		// fork (ex. an interface splits into two newly named interfaces).
 		while (bPassUpdated)
 		{
 			bPassUpdated = false;
 
-			const TArray<FMetasoundFrontendVersion> Versions = ConstInterface->GetConstDocument().Interfaces.Array();
+			const TArray<FMetasoundFrontendVersion> Versions = Document.Interfaces.Array();
 			for (const FMetasoundFrontendVersion& Version : Versions)
 			{
-				FUpdateRootGraphInterface UpdateTransform(Version, GetOwningAssetName());
-				bPassUpdated = TryUpdateInterfaceFromVersion(Version);
+				bPassUpdated |= TryUpdateInterfaceFromVersion(Version);
 			}
 
 			bInterfaceUpdated |= bPassUpdated;
@@ -553,6 +540,7 @@ bool FMetasoundAssetBase::VersionAsset()
 
 		if (bInterfaceUpdated)
 		{
+			TScriptInterface<IMetaSoundDocumentInterface> Interface(GetOwningAsset());
 			Interface->ConformObjectToDocument();
 		}
 		bDidEdit |= bInterfaceUpdated;
@@ -560,6 +548,7 @@ bool FMetasoundAssetBase::VersionAsset()
 
 	return bDidEdit;
 }
+#endif // WITH_EDITORONLY_DATA
 
 #if WITH_EDITOR
 void FMetasoundAssetBase::CacheRegistryMetadata()
@@ -853,10 +842,7 @@ const FMetasoundFrontendDocument& FMetasoundAssetBase::GetConstDocumentChecked()
 {
 	const UObject* Owner = GetOwningAsset();
 	check(Owner);
-
-	// Annoying const cast requirement because ScriptInterface isn't ctor isn't const correct :(
-	TScriptInterface<IMetaSoundDocumentInterface> DocInterface = const_cast<UObject*>(Owner);
-
+	TScriptInterface<const IMetaSoundDocumentInterface> DocInterface = Owner;
 	return DocInterface->GetConstDocument();
 }
 
@@ -984,9 +970,7 @@ bool FMetasoundAssetBase::TryUpdateInterfaceFromVersion(const FMetasoundFrontend
 	using namespace Metasound::Frontend;
 	using namespace AssetBasePrivate;
 
-	FConstDocumentHandle ConstDoc = IDocumentController::CreateDocumentHandle(GetDocumentConstAccessPtr());
-	FMetasoundFrontendInterface TargetInterface = GetInterfaceToVersion(Version, ConstDoc);
-
+	FMetasoundFrontendInterface TargetInterface = GetInterfaceToVersion(Version);
 	if (TargetInterface.Version.IsValid())
 	{
 		TArray<const IInterfaceRegistryEntry*> UpgradePath;
@@ -999,7 +983,7 @@ bool FMetasoundAssetBase::TryUpdateInterfaceFromVersion(const FMetasoundFrontend
 	return false;
 }
 
-FMetasoundFrontendInterface FMetasoundAssetBase::GetInterfaceToVersion(const FMetasoundFrontendVersion& InterfaceVersion, Metasound::Frontend::FConstDocumentHandle InDocument) const
+FMetasoundFrontendInterface FMetasoundAssetBase::GetInterfaceToVersion(const FMetasoundFrontendVersion& InterfaceVersion) const
 {
 	using namespace Metasound::Frontend;
 
diff --git a/Engine/Plugins/Runtime/Metasound/Source/MetasoundFrontend/Private/MetasoundFrontendDocumentBuilder.cpp b/Engine/Plugins/Runtime/Metasound/Source/MetasoundFrontend/Private/MetasoundFrontendDocumentBuilder.cpp
--- a/Engine/Plugins/Runtime/Metasound/Source/MetasoundFrontend/Private/MetasoundFrontendDocumentBuilder.cpp
+++ b/Engine/Plugins/Runtime/Metasound/Source/MetasoundFrontend/Private/MetasoundFrontendDocumentBuilder.cpp
@@ -3756,3 +3756,17 @@ bool FMetaSoundFrontendDocumentBuilder::UpdateDependencyClassNames(const TMap<FM
 	
 	return true;
 }
+
+#if WITH_EDITORONLY_DATA
+bool FMetaSoundFrontendDocumentBuilder::VersionInterfaces()
+{
+	FMetasoundFrontendDocument& Document = GetDocumentChecked();
+	if (Document.RequiresInterfaceVersioning())
+	{
+		Document.VersionInterfaces();
+		return true;
+	}
+
+	return false;
+}
+#endif // WITH_EDITORONLY_DATA
diff --git a/Engine/Plugins/Runtime/Metasound/Source/MetasoundFrontend/Private/MetasoundFrontendDocumentVersioning.cpp b/Engine/Plugins/Runtime/Metasound/Source/MetasoundFrontend/Private/MetasoundFrontendDocumentVersioning.cpp
--- a/Engine/Plugins/Runtime/Metasound/Source/MetasoundFrontend/Private/MetasoundFrontendDocumentVersioning.cpp
+++ b/Engine/Plugins/Runtime/Metasound/Source/MetasoundFrontend/Private/MetasoundFrontendDocumentVersioning.cpp
@@ -2,6 +2,7 @@
 #include "MetasoundFrontendDocumentVersioning.h"
 
 #include "Algo/Transform.h"
+#include "CoreGlobals.h"
 #include "Interfaces/MetasoundFrontendInterface.h"
 #include "Interfaces/MetasoundFrontendInterfaceRegistry.h"
 #include "MetasoundAccessPtr.h"
@@ -596,7 +597,7 @@ namespace Metasound::Frontend
 		}
 	} // namespace VersioningPrivate
 
-	bool VersionDocument(FMetasoundAssetBase& InAssetBase)
+	bool VersionDocument(FMetasoundAssetBase& InAssetBase, FMetaSoundFrontendDocumentBuilder& Builder)
 	{
 		bool bWasUpdated = false;
 
@@ -630,14 +631,7 @@ namespace Metasound::Frontend
 				// No longer supported, new versions should go in VersioningPrivate::VersionBuilderDocument
 			}
 
-#if WITH_EDITORONLY_DATA
-			FMetaSoundFrontendDocumentBuilder& Builder = IDocumentBuilderRegistry::GetChecked().FindOrBeginBuilding(DocumentInterface);
-			bWasUpdated |= VersioningPrivate::VersionBuilderDocument(InAssetBase, Builder);
-#else
-			FMetaSoundFrontendDocumentBuilder Builder(DocumentInterface);
 			bWasUpdated |= VersioningPrivate::VersionBuilderDocument(InAssetBase, Builder);
-#endif // WITH_EDITORONLY_DATA
-
 			if (bWasUpdated)
 			{
 				const FMetasoundFrontendVersionNumber& NewVersionNumber = Document.Metadata.Version.Number;
diff --git a/Engine/Plugins/Runtime/Metasound/Source/MetasoundFrontend/Private/MetasoundFrontendDocumentVersioning.h b/Engine/Plugins/Runtime/Metasound/Source/MetasoundFrontend/Private/MetasoundFrontendDocumentVersioning.h
--- a/Engine/Plugins/Runtime/Metasound/Source/MetasoundFrontend/Private/MetasoundFrontendDocumentVersioning.h
+++ b/Engine/Plugins/Runtime/Metasound/Source/MetasoundFrontend/Private/MetasoundFrontendDocumentVersioning.h
@@ -10,6 +10,7 @@
 #include "UObject/ScriptInterface.h"
 
 class IMetaSoundDocumentInterface;
+struct FMetaSoundFrontendDocumentBuilder;
 
 
 namespace Metasound::Frontend
@@ -21,5 +22,5 @@ namespace Metasound::Frontend
 
 	// Versions Frontend Document. Passed as AssetBase for backward compat to
 	// version asset documents predating the IMetaSoundDocumentInterface
-	bool VersionDocument(FMetasoundAssetBase& InAssetBase);
+	bool VersionDocument(FMetasoundAssetBase& InAssetBase, FMetaSoundFrontendDocumentBuilder& Builder);
 } // namespace Metasound::Frontend
diff --git a/Engine/Plugins/Runtime/Metasound/Source/MetasoundFrontend/Public/MetasoundAssetBase.h b/Engine/Plugins/Runtime/Metasound/Source/MetasoundFrontend/Public/MetasoundAssetBase.h
--- a/Engine/Plugins/Runtime/Metasound/Source/MetasoundFrontend/Public/MetasoundAssetBase.h
+++ b/Engine/Plugins/Runtime/Metasound/Source/MetasoundFrontend/Public/MetasoundAssetBase.h
@@ -154,7 +154,9 @@ public:
 	UE_DEPRECATED(5.4, "Use GetGraphRegistryKey instead.")
 	const Metasound::Frontend::FNodeRegistryKey& GetRegistryKey() const;
 
-	bool VersionAsset();
+#if WITH_EDITORONLY_DATA
+	bool VersionAsset(FMetaSoundFrontendDocumentBuilder& Builder);
+#endif // WITH_EDITORONLY_DATA
 
 #if WITH_EDITOR
 	/*
@@ -273,7 +275,7 @@ private:
 
 	// Returns new interface to be versioned to from the given version. If no interface versioning is
 	// required, returns invalid interface (interface with no name and invalid version number).
-	FMetasoundFrontendInterface GetInterfaceToVersion(const FMetasoundFrontendVersion& InterfaceVersion, Metasound::Frontend::FConstDocumentHandle InDocument) const;
+	FMetasoundFrontendInterface GetInterfaceToVersion(const FMetasoundFrontendVersion& InterfaceVersion) const;
 
 	Metasound::Frontend::FGraphRegistryKey GraphRegistryKey;
 };
diff --git a/Engine/Plugins/Runtime/Metasound/Source/MetasoundFrontend/Public/MetasoundDocumentInterface.h b/Engine/Plugins/Runtime/Metasound/Source/MetasoundFrontend/Public/MetasoundDocumentInterface.h
--- a/Engine/Plugins/Runtime/Metasound/Source/MetasoundFrontend/Public/MetasoundDocumentInterface.h
+++ b/Engine/Plugins/Runtime/Metasound/Source/MetasoundFrontend/Public/MetasoundDocumentInterface.h
@@ -91,7 +91,9 @@ namespace Metasound::Frontend
 		virtual FMetaSoundFrontendDocumentBuilder& FindOrBeginBuilding(TScriptInterface<IMetaSoundDocumentInterface> MetaSound) = 0;
 #endif // WITH_EDITORONLY_DATA
 
-		virtual bool FinishBuilding(const FMetasoundFrontendClassName& InClassName, bool bForceUnregister = false) const = 0;
+		// Removes builder from registry, clearing any cached builder state. (Optionally) forces unregistration from the Frontend Node Class Registry
+		// (If the builder has outstanding transactions, unregistration from the Node Class Registry will occur regardless).
+		virtual bool FinishBuilding(const FMetasoundFrontendClassName& InClassName, bool bForceUnregisterNodeClass = false) const = 0;
 
 		UE_DEPRECATED(5.5, "Document cache can now be invalidated by retrieving an asset builder and calling 'ReloadBuilder'")
 		virtual void InvalidateDocumentCache(const FMetasoundFrontendClassName& InClassName) const { }
diff --git a/Engine/Plugins/Runtime/Metasound/Source/MetasoundFrontend/Public/MetasoundFrontendDocumentBuilder.h b/Engine/Plugins/Runtime/Metasound/Source/MetasoundFrontend/Public/MetasoundFrontendDocumentBuilder.h
--- a/Engine/Plugins/Runtime/Metasound/Source/MetasoundFrontend/Public/MetasoundFrontendDocumentBuilder.h
+++ b/Engine/Plugins/Runtime/Metasound/Source/MetasoundFrontend/Public/MetasoundFrontendDocumentBuilder.h
@@ -232,9 +232,6 @@ public:
 	const IMetaSoundDocumentInterface& GetConstDocumentInterfaceChecked() const;
 	const FString GetDebugName() const;
 
-	// The graph ID used when graph is executed (i.e. utilized for node class registration)
-	const FGuid& GetExecutionGraphID() const;
-
 	UE_DEPRECATED(5.5, "Use GetConstDocumentChecked() instead")
 	const FMetasoundFrontendDocument& GetDocument() const;
 
@@ -386,6 +383,9 @@ public:
 	// Transforms template nodes within the given builder's document, which can include swapping associated edges and/or
 	// replacing nodes with other, registry-defined concrete node class instances. Returns true if any template nodes were processed.
 	bool TransformTemplateNodes();
+
+	// Versions legacy document members that contained interface information
+	bool VersionInterfaces();
 #endif // WITH_EDITORONLY_DATA
 
 private:
