MetaSound Graph Render Cost

The main elements of this change are:
- Incorporation of "Relative Render Cost" to the AudioDevice in order to do more nuanced voice limiting
- Added a method for tracking a RenderCost in MetaSounds
- Live information on RenderCost is displayed in the MetaSound editor when a sound is previewed

There was some small cleanup as well
- Moved the GraphHierarhcy MetaSound environment variable
- Renamed the MetaSound Environment Variable  TransmissionID to InstanceID
- Put CPU cost and Render Cost into their own widget in the MS editor.

diff --git a/Engine/Plugins/Runtime/Metasound/Source/MetasoundEditor/Private/MetasoundEditor.cpp b/Engine/Plugins/Runtime/Metasound/Source/MetasoundEditor/Private/MetasoundEditor.cpp
--- a/Engine/Plugins/Runtime/Metasound/Source/MetasoundEditor/Private/MetasoundEditor.cpp
+++ b/Engine/Plugins/Runtime/Metasound/Source/MetasoundEditor/Private/MetasoundEditor.cpp
@@ -71,6 +71,7 @@
 #include "ScopedTransaction.h"
 #include "SMetasoundActionMenu.h"
 #include "SMetasoundPalette.h"
+#include "SMetasoundRenderStats.h"
 #include "SNodePanel.h"
 #include "Stats/Stats.h"
 #include "Styling/AppStyle.h"
@@ -701,7 +702,7 @@ namespace Metasound
 
 			FAssetEditorToolkit::RegisterTabSpawners(InTabManager);
 
-			InTabManager->RegisterTabSpawner(TabFactory::Names::GraphCanvas, FOnSpawnTab::CreateLambda([InPlayTimeWidget = PlayTimeWidget, InCPUCoreUtilizationWidget = CPUCoreUtilizationWidget, InMetasoundGraphEditor = MetasoundGraphEditor](const FSpawnTabArgs& Args)
+			InTabManager->RegisterTabSpawner(TabFactory::Names::GraphCanvas, FOnSpawnTab::CreateLambda([InPlayTimeWidget = PlayTimeWidget, InRenderStatsWidget = RenderStatsWidget, InMetasoundGraphEditor = MetasoundGraphEditor](const FSpawnTabArgs& Args)
 			{
 				return TabFactory::CreateGraphCanvasTab(SNew(SOverlay)
 					+ SOverlay::Slot()
@@ -715,7 +716,7 @@ namespace Metasound
 					+ SOverlay::Slot()
 					.VAlign(VAlign_Bottom)
 					[
-						InCPUCoreUtilizationWidget.ToSharedRef()
+						InRenderStatsWidget.ToSharedRef()
 					]
 					.Padding(5.0f, 5.0f)
 				, Args);
@@ -1855,7 +1856,14 @@ namespace Metasound
 				MetasoundGraphEditor->RegisterActiveTimer(0.0f,
 					FWidgetActiveTimerDelegate::CreateLambda([this](double InCurrentTime, float InDeltaTime)
 					{
-						if (IsPlaying())
+						const bool bIsPlaying = IsPlaying();
+						if (RenderStatsWidget.IsValid())
+						{
+							check(IsInGameThread());
+							RenderStatsWidget->Update(bIsPlaying, Cast<const UMetaSoundSource>(Metasound));
+						}
+
+						if (bIsPlaying)
 						{
 							if (PlayTimeWidget.IsValid())
 							{
@@ -1866,32 +1874,6 @@ namespace Metasound
 								PlayTimeString.ReplaceInline(TEXT("+"), TEXT(""));
 								PlayTimeWidget->SetText(FText::FromString(PlayTimeString));
 							}
-							if (CPUCoreUtilizationWidget.IsValid())
-							{
-								// TODO: Need to protect against shipping builds of editor
-								check(IsInGameThread());
-								double CPUCoreUtilization = 0;
-								if (UMetaSoundSource* Source = Cast<UMetaSoundSource>(Metasound))
-								{
-									if (const UAudioComponent* PreviewComponent = GEditor->GetPreviewAudioComponent())
-									{
-										TSharedPtr<FMetasoundGenerator> Generator = Source->GetGeneratorForAudioComponent(PreviewComponent->GetAudioComponentID()).Pin();
-										if (Generator.IsValid())
-										{
-											CPUCoreUtilization = Generator->GetCPUCoreUtilization();
-										}
-									}
-								}
-								if (CPUCoreUtilization > 0)
-								{
-									FString CPUCoreUtilizationString = FString::Printf(TEXT("%5.2f %% CPU Core"), 100. * CPUCoreUtilization);
-									CPUCoreUtilizationWidget->SetText(FText::FromString(CPUCoreUtilizationString));
-								}
-								else
-								{
-									CPUCoreUtilizationWidget->SetText(FText::GetEmpty());
-								}
-							}
 
 							return EActiveTimerReturnType::Continue;
 						}
@@ -1900,10 +1882,6 @@ namespace Metasound
 							SetPreviewID(INDEX_NONE);
 							PlayTime = 0.0;
 							PlayTimeWidget->SetText(FText::GetEmpty());
-							if (CPUCoreUtilizationWidget.IsValid())
-							{
-								CPUCoreUtilizationWidget->SetText(FText::GetEmpty());
-							}
 							GraphConnectionManager = MakeUnique<FGraphConnectionManager>();
 
 							return EActiveTimerReturnType::Stop;
@@ -2249,10 +2227,7 @@ namespace Metasound
 				.TextStyle(FAppStyle::Get(), "Graph.ZoomText")
 				.ColorAndOpacity(FLinearColor(1, 1, 1, 0.30f));
 
-			SAssignNew(CPUCoreUtilizationWidget, STextBlock)
-				.Visibility(EVisibility::HitTestInvisible)
-				.TextStyle(FAppStyle::Get(), "Graph.ZoomText")
-				.ColorAndOpacity(FLinearColor(1, 1, 1, 0.30f));
+			SAssignNew(RenderStatsWidget, SMetaSoundRenderStats);
 		}
 
 		FGraphAppearanceInfo FEditor::GetGraphAppearance() const
diff --git a/Engine/Plugins/Runtime/Metasound/Source/MetasoundEditor/Private/MetasoundEditor.h b/Engine/Plugins/Runtime/Metasound/Source/MetasoundEditor/Private/MetasoundEditor.h
--- a/Engine/Plugins/Runtime/Metasound/Source/MetasoundEditor/Private/MetasoundEditor.h
+++ b/Engine/Plugins/Runtime/Metasound/Source/MetasoundEditor/Private/MetasoundEditor.h
@@ -47,6 +47,7 @@ class FTabManager;
 class SDockableTab;
 class SGraphEditor;
 class SMetasoundPalette;
+class SMetaSoundRenderStats;
 class FSlateRect;
 class IDetailsView;
 class IToolkitHost;
@@ -457,8 +458,8 @@ namespace Metasound
 			TSharedPtr<STextBlock> PlayTimeWidget;
 			double PlayTime = 0.0;
 
-			/** Widget showing render speed that overlays the graph when previewing */
-			TSharedPtr<STextBlock> CPUCoreUtilizationWidget;
+			/** Widget showing render performance information that overlays the graph when previewing */
+			TSharedPtr<SMetaSoundRenderStats> RenderStatsWidget;
 
 			TUniquePtr<FGraphConnectionManager> GraphConnectionManager;
 
diff --git a/Engine/Plugins/Runtime/Metasound/Source/MetasoundEditor/Private/SMetasoundRenderStats.cpp b/Engine/Plugins/Runtime/Metasound/Source/MetasoundEditor/Private/SMetasoundRenderStats.cpp
new file mode 100644
--- /dev/null
+++ b/Engine/Plugins/Runtime/Metasound/Source/MetasoundEditor/Private/SMetasoundRenderStats.cpp
@@ -0,0 +1,80 @@
+// Copyright Epic Games, Inc. All Rights Reserved.
+
+#include "SMetasoundRenderStats.h"
+
+#include "Components/AudioComponent.h"
+#include "Containers/UnrealString.h"
+#include "Editor.h"
+#include "Internationalization/Text.h"
+#include "MetasoundSource.h"
+#include "MetasoundGenerator.h"
+#include "Templates/SharedPointer.h"
+#include "Widgets/DeclarativeSyntaxSupport.h"
+#include "Widgets/SBoxPanel.h"
+#include "Widgets/Text/STextBlock.h"
+
+void SMetaSoundRenderStats::Construct(const FArguments& InArgs)
+{
+	SVerticalBox::Construct(SVerticalBox::FArguments());
+	AddSlot()
+	.HAlign(HAlign_Left)
+	[
+		SAssignNew(RenderStatsCostWidget, STextBlock)
+		.Visibility(EVisibility::HitTestInvisible)
+		.TextStyle(FAppStyle::Get(), "Graph.ZoomText")
+		.ColorAndOpacity(FLinearColor(1, 1, 1, 0.30f))
+	];
+
+	AddSlot()
+	.HAlign(HAlign_Left)
+	[
+		SAssignNew(RenderStatsCPUWidget, STextBlock)
+		.Visibility(EVisibility::HitTestInvisible)
+		.TextStyle(FAppStyle::Get(), "Graph.ZoomText")
+		.ColorAndOpacity(FLinearColor(1, 1, 1, 0.30f))
+	];
+}
+
+void SMetaSoundRenderStats::Update(bool bIsPlaying, const UMetaSoundSource* InSource)
+{
+	using namespace Metasound;
+
+	// Reset maximum values when play restarts
+	if (bIsPlaying && !bPreviousIsPlaying)
+	{
+		MaxRelativeRenderCost = 0.f;
+		MaxCPUCoreUtilization = 0;
+	}
+	bPreviousIsPlaying = bIsPlaying;
+
+	if (RenderStatsCPUWidget.IsValid() && RenderStatsCostWidget.IsValid())
+	{
+		double CPUCoreUtilization = 0;
+		float RelativeRenderCost = 0.f;
+
+		// Find generator for playing preview component. 
+		if (bIsPlaying && InSource)
+		{
+			if (const UAudioComponent* PreviewComponent = GEditor->GetPreviewAudioComponent())
+			{
+				TSharedPtr<FMetasoundGenerator> Generator = InSource->GetGeneratorForAudioComponent(PreviewComponent->GetAudioComponentID()).Pin();
+				if (Generator.IsValid())
+				{
+					// Update render stats
+					CPUCoreUtilization = Generator->GetCPUCoreUtilization();
+					MaxCPUCoreUtilization = FMath::Max(MaxCPUCoreUtilization, CPUCoreUtilization);
+
+					RelativeRenderCost = Generator->GetRelativeRenderCost();
+					MaxRelativeRenderCost = FMath::Max(MaxRelativeRenderCost, RelativeRenderCost);
+				}
+			}
+		}
+
+		// Display updated render stats. 
+		FString CPUCoreUtilizationString = FString::Printf(TEXT("%3.2f%% (%3.2f%% Max) CPU Core"), 100. * CPUCoreUtilization, 100. * MaxCPUCoreUtilization);
+		RenderStatsCPUWidget->SetText(FText::FromString(CPUCoreUtilizationString));
+
+		FString RenderCostString = FString::Printf(TEXT("%3.2f (%3.2f Max) Relative Render Cost"), RelativeRenderCost, MaxRelativeRenderCost);
+		RenderStatsCostWidget->SetText(FText::FromString(RenderCostString));
+	}
+}
diff --git a/Engine/Plugins/Runtime/Metasound/Source/MetasoundEditor/Private/SMetasoundRenderStats.h b/Engine/Plugins/Runtime/Metasound/Source/MetasoundEditor/Private/SMetasoundRenderStats.h
new file mode 100644
--- /dev/null
+++ b/Engine/Plugins/Runtime/Metasound/Source/MetasoundEditor/Private/SMetasoundRenderStats.h
@@ -0,0 +1,30 @@
+// Copyright Epic Games, Inc. All Rights Reserved.
+
+
+#pragma once
+
+#include "Templates/SharedPointer.h"
+#include "Widgets/DeclarativeSyntaxSupport.h"
+#include "Widgets/SBoxPanel.h"
+
+class UMetaSoundSource;
+class STextBlock;
+
+/** Widget for displaying render stats of a previewing MetaSound.. */
+class SMetaSoundRenderStats : public SVerticalBox
+{
+public:
+	SLATE_BEGIN_ARGS( SMetaSoundRenderStats) {};
+	SLATE_END_ARGS()
+
+	void Construct(const FArguments& InArgs);
+	void Update(bool bIsPlaying, const UMetaSoundSource* InSource);
+private:
+	TSharedPtr<STextBlock> RenderStatsCostWidget;
+	TSharedPtr<STextBlock> RenderStatsCPUWidget;
+
+	bool bPreviousIsPlaying = false;
+	double MaxCPUCoreUtilization = 0;
+	float MaxRelativeRenderCost = 0.f;
+};
+
diff --git a/Engine/Plugins/Runtime/Metasound/Source/MetasoundEngine/Private/MetasoundSource.cpp b/Engine/Plugins/Runtime/Metasound/Source/MetasoundEngine/Private/MetasoundSource.cpp
--- a/Engine/Plugins/Runtime/Metasound/Source/MetasoundEngine/Private/MetasoundSource.cpp
+++ b/Engine/Plugins/Runtime/Metasound/Source/MetasoundEngine/Private/MetasoundSource.cpp
@@ -951,6 +951,7 @@ ISoundGeneratorPtr UMetaSoundSource::CreateSoundGenerator(const FSoundGeneratorI
 
 		BuilderSettings.bEnableOperatorRebind = true;
 
+		TSharedRef<FGraphRenderCost> GraphRenderCost = FGraphRenderCost::MakeGraphRenderCost();
 		FMetasoundDynamicGraphGeneratorInitParams InitParams
 		{
 			{
@@ -962,9 +963,10 @@ ISoundGeneratorPtr UMetaSoundSource::CreateSoundGenerator(const FSoundGeneratorI
 				GetOutputAudioChannelOrder(),
 				MoveTemp(InDefaultParameters),
 				bBuildSynchronous,
-				DataChannel
+				DataChannel,
+				GraphRenderCost
 			},
-			DynamicTransactor->CreateTransformQueue(InSettings, Environment) // Create transaction queue
+			DynamicTransactor->CreateTransformQueue(InSettings, Environment, GraphRenderCost) // Create transaction queue
 		};
 		TSharedPtr<FMetasoundDynamicGraphGenerator> DynamicGenerator = MakeShared<FMetasoundDynamicGraphGenerator>(InSettings);
 		DynamicGenerator->Init(MoveTemp(InitParams));
@@ -997,7 +999,7 @@ ISoundGeneratorPtr UMetaSoundSource::CreateSoundGenerator(const FSoundGeneratorI
 				MergePresetOverridesAndSuppliedDefaults(InDefaultParameters, MergedParameters);
 
 				// Update Graph Hierarchy with the asset unique ID (the hierarchy was collapsed by only using the base graph)
-				Environment.SetValue<TArray<FGuid>>(OperatorBuilder::Environment::GraphHierarchy, TArray<FGuid>({ AssetClassID }));
+				Environment.SetValue<TArray<FGuid>>(CoreInterface::Environment::GraphHierarchy, TArray<FGuid>({ AssetClassID }));
 
 				// Create generator.
 				FMetasoundGeneratorInitParams InitParams
@@ -1495,6 +1497,7 @@ Metasound::FMetasoundEnvironment UMetaSoundSource::CreateEnvironment(const FSoun
 
 	FMetasoundEnvironment Environment = CreateEnvironment();
 	Environment.SetValue<bool>(SourceInterface::Environment::IsPreview, InParams.bIsPreviewSound);
+	Environment.SetValue<uint64>(CoreInterface::Environment::InstanceID, InParams.InstanceID);
 	Environment.SetValue<uint64>(SourceInterface::Environment::TransmitterID, InParams.InstanceID);
 	Environment.SetValue<Audio::FDeviceId>(SourceInterface::Environment::DeviceID, InParams.AudioDeviceID);
 	Environment.SetValue<int32>(SourceInterface::Environment::AudioMixerNumOutputFrames, InParams.AudioMixerNumOutputFrames);
@@ -1513,7 +1516,7 @@ Metasound::FMetasoundEnvironment UMetaSoundSource::CreateEnvironment(const Audio
 	using namespace Metasound::Frontend;
 
 	FMetasoundEnvironment Environment = CreateEnvironment();
-	Environment.SetValue<uint64>(SourceInterface::Environment::TransmitterID, InParams.InstanceID);
+	Environment.SetValue<uint64>(CoreInterface::Environment::InstanceID, InParams.InstanceID);
 
 	return Environment;
 }
diff --git a/Engine/Plugins/Runtime/Metasound/Source/MetasoundEngine/Private/MetasoundWavePlayerNode.cpp b/Engine/Plugins/Runtime/Metasound/Source/MetasoundEngine/Private/MetasoundWavePlayerNode.cpp
--- a/Engine/Plugins/Runtime/Metasound/Source/MetasoundEngine/Private/MetasoundWavePlayerNode.cpp
+++ b/Engine/Plugins/Runtime/Metasound/Source/MetasoundEngine/Private/MetasoundWavePlayerNode.cpp
@@ -12,6 +12,7 @@
 #include "MetasoundNodeRegistrationMacro.h"
 #include "MetasoundParamHelper.h"
 #include "MetasoundPrimitives.h"
+#include "MetasoundRenderCost.h"
 #include "MetasoundTrace.h"
 #include "MetasoundTrigger.h"
 #include "MetasoundVertex.h"
@@ -307,6 +308,7 @@ namespace Metasound
 		FBoolReadRef bLoop;
 		FTimeReadRef LoopStartTime;
 		FTimeReadRef LoopDuration;
+		FNodeRenderCost CostReporter;
 	};
 
 	// Maximum decode size in frames. 
@@ -354,6 +356,7 @@ namespace Metasound
 			, LoopPercent(FFloatWriteRef::CreateNew(0.0f))
 			, PlaybackLocation(FFloatWriteRef::CreateNew(0.0f))
 			, PlaybackTime(FTimeWriteRef::CreateNew(0.0))
+			, CostReporter(InArgs.CostReporter)
 		{
 			NumOutputChannels = InArgs.OutputAudioVertices.Num();
 
@@ -436,6 +439,9 @@ namespace Metasound
 			TriggerOnCuePoint->AdvanceBlock();
 			TriggerOnLooped->AdvanceBlock();
 
+			// Reset flag to track render cost.
+			bDidWaveRenderThisBlock = false;
+
 			// Update wave proxy reader with any new looping bounds. 
 			if (WaveProxyReader.IsValid())
 			{
@@ -462,6 +468,11 @@ namespace Metasound
 
 			// Updates output playhead information
 			UpdatePlaybackLocation();
+
+			if (bDidWaveRenderThisBlock)
+			{
+				CostReporter.SetRenderCost(1.f);
+			}
 		}
 
 		void Reset(const IOperator::FResetParams& InParams)
@@ -583,6 +594,9 @@ namespace Metasound
 
 			METASOUND_TRACE_CPUPROFILER_EVENT_SCOPE(Metasound::FWavePlayerOperator::RenderFrameRange);
 
+			// Set flag that this node rendered wave audio this block.
+			bDidWaveRenderThisBlock = true;
+
 			// Assume this is set to true and checked by outside callers
 			check(bIsPlaying);
 
@@ -1053,6 +1067,7 @@ namespace Metasound
 		FTimeWriteRef PlaybackTime;
 		TArray<FAudioBufferWriteRef> OutputAudioBuffers;
 		TArray<FName> OutputAudioBufferVertexNames;
+		FNodeRenderCost CostReporter;
 
 		TUniquePtr<FSoundWaveProxyReader> WaveProxyReader;
 		TUniquePtr<Audio::IConvertDeinterleave> ConvertDeinterleave;
@@ -1071,6 +1086,7 @@ namespace Metasound
 		int32 NumDeinterleaveChannels;
 		bool bOnNearlyDoneTriggeredForWave = false;
 		bool bIsPlaying = false;
+		bool bDidWaveRenderThisBlock = false;
 		// Cached from cvar 
 		int32 DeinterleaveBufferBlockSizeInFrames;
 	};
@@ -1101,7 +1117,8 @@ namespace Metasound
 				Inputs.GetOrCreateDefaultDataReadReference<float>(METASOUND_GET_PARAM_NAME(InputPitchShift), InParams.OperatorSettings),
 				Inputs.GetOrCreateDefaultDataReadReference<bool>(METASOUND_GET_PARAM_NAME(InputLoop), InParams.OperatorSettings),
 				Inputs.GetOrCreateDefaultDataReadReference<FTime>(METASOUND_GET_PARAM_NAME(InputLoopStart), InParams.OperatorSettings),
-				Inputs.GetOrCreateDefaultDataReadReference<FTime>(METASOUND_GET_PARAM_NAME(InputLoopDuration), InParams.OperatorSettings)
+				Inputs.GetOrCreateDefaultDataReadReference<FTime>(METASOUND_GET_PARAM_NAME(InputLoopDuration), InParams.OperatorSettings),
+				FNodeRenderCost(InParams.GraphRenderCost ? InParams.GraphRenderCost->AddNode(InParams.Node.GetInstanceID(), InParams.Environment) : FNodeRenderCost{})
 			};
 
 			return MakeUnique<FWavePlayerOperator>(Args);
diff --git a/Engine/Plugins/Runtime/Metasound/Source/MetasoundEngineTest/Private/EngineTestMetaSoundAutomatedNodeTest.cpp b/Engine/Plugins/Runtime/Metasound/Source/MetasoundEngineTest/Private/EngineTestMetaSoundAutomatedNodeTest.cpp
--- a/Engine/Plugins/Runtime/Metasound/Source/MetasoundEngineTest/Private/EngineTestMetaSoundAutomatedNodeTest.cpp
+++ b/Engine/Plugins/Runtime/Metasound/Source/MetasoundEngineTest/Private/EngineTestMetaSoundAutomatedNodeTest.cpp
@@ -50,11 +50,13 @@ namespace Metasound::EngineTest{
 
 		FMetasoundEnvironment Environment;
 
+		Environment.SetValue<uint64>(CoreInterface::Environment::InstanceID, 0);
+		Environment.SetValue<TArray<FGuid>>(CoreInterface::Environment::GraphHierarchy, TArray<FGuid>({ FGuid() }));
+
 		Environment.SetValue<uint32>(SourceInterface::Environment::SoundUniqueID, 0);
 		Environment.SetValue<bool>(SourceInterface::Environment::IsPreview, false);
 		Environment.SetValue<uint64>(SourceInterface::Environment::TransmitterID, 0);
 		Environment.SetValue<FString>(SourceInterface::Environment::GraphName, TEXT("ENGINE_TEST_REGISTERED_NODES"));
-		Environment.SetValue<TArray<FGuid>>(OperatorBuilder::Environment::GraphHierarchy, TArray<FGuid>({ FGuid() }));
 
 		if (Audio::FMixerDevice* MixerDevice = GetAudioMixerDevice())
 		{
diff --git a/Engine/Plugins/Runtime/Metasound/Source/MetasoundEngineTest/Private/MetasoundDynamicGeneratorTests.cpp b/Engine/Plugins/Runtime/Metasound/Source/MetasoundEngineTest/Private/MetasoundDynamicGeneratorTests.cpp
--- a/Engine/Plugins/Runtime/Metasound/Source/MetasoundEngineTest/Private/MetasoundDynamicGeneratorTests.cpp
+++ b/Engine/Plugins/Runtime/Metasound/Source/MetasoundEngineTest/Private/MetasoundDynamicGeneratorTests.cpp
@@ -40,7 +40,7 @@ namespace Metasound::Test::Generator::Dynamic
 			BuilderSettings.bEnableOperatorRebind = true;
 
 			FMetasoundEnvironment Environment;
-			Environment.SetValue<uint64>(Frontend::SourceInterface::Environment::TransmitterID, 123);
+			Environment.SetValue<uint64>(CoreInterface::Environment::InstanceID, 123);
 			
 			FMetasoundDynamicGraphGeneratorInitParams InitParams
 			{
@@ -54,7 +54,7 @@ namespace Metasound::Test::Generator::Dynamic
 					{},
 					true
 				},
-				Transactor.CreateTransformQueue(OperatorSettings, Environment) // Create transaction queue
+				Transactor.CreateTransformQueue(OperatorSettings, Environment, nullptr) // Create transaction queue
 			};
 			
 			Generator.Init(MoveTemp(InitParams));
@@ -260,4 +260,4 @@ namespace Metasound::Test::Generator::Dynamic
 	}
 }
 
-#endif
\ No newline at end of file
+#endif
diff --git a/Engine/Plugins/Runtime/Metasound/Source/MetasoundEngineTest/Private/NodeTestGraphBuilder.cpp b/Engine/Plugins/Runtime/Metasound/Source/MetasoundEngineTest/Private/NodeTestGraphBuilder.cpp
--- a/Engine/Plugins/Runtime/Metasound/Source/MetasoundEngineTest/Private/NodeTestGraphBuilder.cpp
+++ b/Engine/Plugins/Runtime/Metasound/Source/MetasoundEngineTest/Private/NodeTestGraphBuilder.cpp
@@ -102,7 +102,7 @@ namespace Metasound::Test
 			const FOperatorSettings OperatorSettings{SampleRate, static_cast<float>(SampleRate) / SamplesPerBlock};
 
 			FMetasoundEnvironment Environment;
-			Environment.SetValue<uint64>(SourceInterface::Environment::TransmitterID, 123);
+			Environment.SetValue<uint64>(CoreInterface::Environment::InstanceID, 123);
 			
 			FMetasoundGeneratorInitParams GeneratorInitParams{
 				OperatorSettings,
diff --git a/Engine/Plugins/Runtime/Metasound/Source/MetasoundFrontend/Private/MetasoundFrontendRegistryContainerImpl.cpp b/Engine/Plugins/Runtime/Metasound/Source/MetasoundFrontend/Private/MetasoundFrontendRegistryContainerImpl.cpp
--- a/Engine/Plugins/Runtime/Metasound/Source/MetasoundFrontend/Private/MetasoundFrontendRegistryContainerImpl.cpp
+++ b/Engine/Plugins/Runtime/Metasound/Source/MetasoundFrontend/Private/MetasoundFrontendRegistryContainerImpl.cpp
@@ -136,13 +136,14 @@ namespace Metasound::Frontend
 
 				virtual TUniquePtr<IOperator> CreateOperator(const FBuildOperatorParams& InParams, FBuildResults& OutResults) override
 				{
-					FBuildOperatorParams ForwardParams
+					FBuildOperatorParams ForwardParams 
 					{
 						*Graph,  // Point to correct INode instance
 						InParams.OperatorSettings,
 						InParams.InputData,
 						InParams.Environment,
-						InParams.Builder
+						InParams.Builder,
+						InParams.GraphRenderCost
 					};
 					return GraphFactory->CreateOperator(ForwardParams, OutResults);
 				}
diff --git a/Engine/Plugins/Runtime/Metasound/Source/MetasoundFrontend/Private/MetasoundParameterTransmitter.cpp b/Engine/Plugins/Runtime/Metasound/Source/MetasoundFrontend/Private/MetasoundParameterTransmitter.cpp
--- a/Engine/Plugins/Runtime/Metasound/Source/MetasoundFrontend/Private/MetasoundParameterTransmitter.cpp
+++ b/Engine/Plugins/Runtime/Metasound/Source/MetasoundFrontend/Private/MetasoundParameterTransmitter.cpp
@@ -227,9 +227,9 @@ namespace Metasound
 		using namespace Frontend;
 
 		uint64 InstanceID = -1;
-		if (ensure(InEnvironment.Contains<uint64>(SourceInterface::Environment::TransmitterID)))
+		if (ensure(InEnvironment.Contains<uint64>(CoreInterface::Environment::InstanceID)))
 		{
-			InstanceID = InEnvironment.GetValue<uint64>(SourceInterface::Environment::TransmitterID);
+			InstanceID = InEnvironment.GetValue<uint64>(CoreInterface::Environment::InstanceID);
 		}
 
 		PRAGMA_DISABLE_DEPRECATION_WARNINGS
diff --git a/Engine/Plugins/Runtime/Metasound/Source/MetasoundFrontend/Public/MetasoundArrayRandomNode.h b/Engine/Plugins/Runtime/Metasound/Source/MetasoundFrontend/Public/MetasoundArrayRandomNode.h
--- a/Engine/Plugins/Runtime/Metasound/Source/MetasoundFrontend/Public/MetasoundArrayRandomNode.h
+++ b/Engine/Plugins/Runtime/Metasound/Source/MetasoundFrontend/Public/MetasoundArrayRandomNode.h
@@ -324,9 +324,9 @@ namespace Metasound
 			TriggerOnReset->Reset();
 
 			// Cache shared state id for shared state behavior types that cannot be changed after node init
-			if (InParams.Environment.Contains<TArray<FGuid>>(OperatorBuilder::Environment::GraphHierarchy))
+			if (InParams.Environment.Contains<TArray<FGuid>>(CoreInterface::Environment::GraphHierarchy))
 			{
-				const TArray<FGuid>& GraphHierarchy = InParams.Environment.GetValue<TArray<FGuid>>(OperatorBuilder::Environment::GraphHierarchy);
+				const TArray<FGuid>& GraphHierarchy = InParams.Environment.GetValue<TArray<FGuid>>(CoreInterface::Environment::GraphHierarchy);
 				if (SharedStateBehavior == ESharedStateBehaviorType::SameNode)
 				{
 					check(GraphHierarchy.Num() > 0);
diff --git a/Engine/Plugins/Runtime/Metasound/Source/MetasoundGenerator/Private/MetasoundGenerator.cpp b/Engine/Plugins/Runtime/Metasound/Source/MetasoundGenerator/Private/MetasoundGenerator.cpp
--- a/Engine/Plugins/Runtime/Metasound/Source/MetasoundGenerator/Private/MetasoundGenerator.cpp
+++ b/Engine/Plugins/Runtime/Metasound/Source/MetasoundGenerator/Private/MetasoundGenerator.cpp
@@ -256,6 +256,7 @@ namespace Metasound
 		InParams.AudioOutputNames = {};
 		InParams.DefaultParameters = {};
 		InParams.DataChannel.Reset();
+		InParams.GraphRenderCost.Reset();
 	}
 
 	void FMetasoundDynamicGraphGeneratorInitParams::Reset(FMetasoundDynamicGraphGeneratorInitParams& InParams)
@@ -284,6 +285,7 @@ namespace Metasound
 #else
 		, bDoRuntimeRenderTiming(false)
 #endif // if ENABLE_METASOUND_GENERATOR_RENDER_TIMING
+	    , RelativeRenderCost(1.f)
 	{
 	}
 
@@ -305,7 +307,7 @@ namespace Metasound
 		return OnSetGraph.Add(Delegate);
 	}
 
-	void FMetasoundGenerator::InitBase(const FMetasoundGeneratorInitParams& InInitParams)
+	void FMetasoundGenerator::InitBase(FMetasoundGeneratorInitParams& InInitParams)
 	{
 		MetasoundName = InInitParams.MetaSoundName;
 		NumChannels = InInitParams.AudioOutputNames.Num();
@@ -323,9 +325,20 @@ PRAGMA_DISABLE_DEPRECATION_WARNINGS
 		ParameterPackSendAddress = UMetasoundParameterPack::CreateSendAddressFromEnvironment(InInitParams.Environment);
 PRAGMA_ENABLE_DEPRECATION_WARNINGS
 		ParameterPackReceiver = FDataTransmissionCenter::Get().RegisterNewReceiver<FMetasoundParameterStorageWrapper>(ParameterPackSendAddress, FReceiverInitParams{ OperatorSettings });
-	}
-
 
+		if (InInitParams.GraphRenderCost.IsValid())
+		{
+			// Use the provided render cost if it exists
+			GraphRenderCost = InInitParams.GraphRenderCost;
+		}
+		else
+		{
+			// If no render cost object is provided make one. 
+			GraphRenderCost = FGraphRenderCost::MakeGraphRenderCost();
+			// Add to the init params so that it is available to metasound graph being built.
+			InInitParams.GraphRenderCost = GraphRenderCost;
+		}
+	}
 
 	bool FMetasoundGenerator::RemoveGraphSetCallback(const FDelegateHandle& Handle)
 	{
@@ -629,6 +642,11 @@ PRAGMA_ENABLE_DEPRECATION_WARNINGS
 
 		while (NumSamplesRemaining > 0)
 		{
+			if (GraphRenderCost)
+			{
+				// Set all render node costs to zero to start
+				GraphRenderCost->ResetNodeRenderCosts();
+			}
 			// Create a scoped timed section for the bulk of the processing...
 			{
 				MetasoundGeneratorPrivate::FBlockRenderScope RuntimelBlockRenderScope(RenderTimer.Get());
@@ -691,6 +709,13 @@ PRAGMA_ENABLE_DEPRECATION_WARNINGS
 			{
 				RenderTime = RenderTimer->UpdateCPUCoreUtilization();
 			}
+
+			if (GraphRenderCost)
+			{
+				// Require the cost to be a minimum of 1. 
+				RelativeRenderCost = FMath::Max(1.f, GraphRenderCost->ComputeGraphRenderCost());
+			}
+
 		}
 
 		// If the vertex interface changed, notify listeners
@@ -722,6 +747,11 @@ PRAGMA_ENABLE_DEPRECATION_WARNINGS
 		return RenderTime;
 	}
 
+	float FMetasoundGenerator::GetRelativeRenderCost() const
+	{
+		return RelativeRenderCost;
+	}
+
 	int32 FMetasoundGenerator::FillWithBuffer(const Audio::FAlignedFloatBuffer& InBuffer, float* OutAudio, int32 MaxNumOutputSamples)
 	{
 		int32 InNum = InBuffer.Num();
diff --git a/Engine/Plugins/Runtime/Metasound/Source/MetasoundGenerator/Private/MetasoundGeneratorBuilder.cpp b/Engine/Plugins/Runtime/Metasound/Source/MetasoundGenerator/Private/MetasoundGeneratorBuilder.cpp
--- a/Engine/Plugins/Runtime/Metasound/Source/MetasoundGenerator/Private/MetasoundGeneratorBuilder.cpp
+++ b/Engine/Plugins/Runtime/Metasound/Source/MetasoundGenerator/Private/MetasoundGeneratorBuilder.cpp
@@ -170,7 +170,7 @@ namespace Metasound
 			METASOUND_TRACE_CPUPROFILER_EVENT_SCOPE(GeneratorBuilder::BuildGraphAnalyzer);
 			using namespace Frontend;
 
-			const uint64 InstanceID = InEnvironment.GetValue<uint64>(SourceInterface::Environment::TransmitterID);
+			const uint64 InstanceID = InEnvironment.GetValue<uint64>(CoreInterface::Environment::InstanceID);
 			return MakeUnique<FGraphAnalyzer>(InOperatorSettings, InstanceID, MoveTemp(InInternalDataReferences));
 		}
 
@@ -266,7 +266,7 @@ namespace Metasound
 				return OpAndInputs;
 			}
 			// Create an instance of the new graph operator
-			FBuildGraphOperatorParams BuildParams { *InInitParams.Graph, InOperatorSettings, OpAndInputs.Inputs, InInitParams.Environment };
+			FBuildGraphOperatorParams BuildParams { *InInitParams.Graph, InOperatorSettings, OpAndInputs.Inputs, InInitParams.Environment, InInitParams.GraphRenderCost.Get()};
 			FOperatorBuilder Builder(InInitParams.BuilderSettings);
 			OpAndInputs.Operator = Builder.BuildGraphOperator(BuildParams, OutBuildResults);
 
@@ -287,7 +287,8 @@ namespace Metasound
 					*InInitParams.Graph, 
 					InOperatorSettings, 
 					OpAndInputs.Inputs, 
-					InInitParams.Environment
+					InInitParams.Environment,
+					InInitParams.GraphRenderCost.Get()
 				},
 				MoveTemp(InInitParams.TransformQueue),
 				InCallbacks
diff --git a/Engine/Plugins/Runtime/Metasound/Source/MetasoundGenerator/Public/MetasoundGenerator.h b/Engine/Plugins/Runtime/Metasound/Source/MetasoundGenerator/Public/MetasoundGenerator.h
--- a/Engine/Plugins/Runtime/Metasound/Source/MetasoundGenerator/Public/MetasoundGenerator.h
+++ b/Engine/Plugins/Runtime/Metasound/Source/MetasoundGenerator/Public/MetasoundGenerator.h
@@ -9,6 +9,7 @@
 #include "MetasoundParameterPack.h"
 #include "MetasoundRouter.h"
 #include "MetasoundTrigger.h"
+#include "MetasoundRenderCost.h"
 #include "MetasoundVertex.h"
 #include "MetasoundVertexData.h"
 
@@ -91,6 +92,7 @@ namespace Metasound
 		TArray<FAudioParameter> DefaultParameters;
 		bool bBuildSynchronous = false;
 		TSharedPtr<TSpscQueue<FMetaSoundParameterTransmitter::FParameter>> DataChannel;
+		TSharedPtr<FGraphRenderCost> GraphRenderCost;
 
 		static void Reset(FMetasoundGeneratorInitParams& InParams);
 	};
@@ -230,8 +232,9 @@ namespace Metasound
 
 		//~ Begin FSoundGenerator
 		virtual int32 OnGenerateAudio(float* OutAudio, int32 NumSamples) override;
-		int32 GetDesiredNumSamplesToRenderPerCallback() const override;
-		bool IsFinished() const override;
+		virtual int32 GetDesiredNumSamplesToRenderPerCallback() const override;
+		virtual bool IsFinished() const override;
+		virtual float GetRelativeRenderCost() const override;
 		//~ End FSoundGenerator
 
 		/** Enables the performance timing of the metasound rendering process. You
@@ -243,7 +246,6 @@ namespace Metasound
 		/** Fraction of a single CPU core used to render audio on a scale of 0.0 to 1.0 */
 		double GetCPUCoreUtilization() const;
 
-
 		// Called when a new graph has been "compiled" and set up as this generator's graph.
 		// Note: We don't allow direct assignment to the FOnSetGraph delegate
 		// because we want to give the Delegate an initial immediate callback if the generator 
@@ -261,7 +263,7 @@ namespace Metasound
 
 	protected:
 
-		void InitBase(const FMetasoundGeneratorInitParams& InInitParams);
+		void InitBase(FMetasoundGeneratorInitParams& InInitParams);
 
 
 		/** SetGraph directly sets graph. Callers must ensure that no race conditions exist. */
@@ -372,6 +374,9 @@ namespace Metasound
 		double RenderTime;
 		bool bDoRuntimeRenderTiming;
 		TUniquePtr<MetasoundGeneratorPrivate::FRenderTimer> RenderTimer;
+
+		TSharedPtr<FGraphRenderCost> GraphRenderCost;
+		std::atomic<float> RelativeRenderCost;
 	};
 
 	/** FMetasoundConstGraphGenerator generates audio from a given metasound IOperator
diff --git a/Engine/Plugins/Runtime/Metasound/Source/MetasoundGraphCore/Private/MetasoundDynamicOperatorTransactor.cpp b/Engine/Plugins/Runtime/Metasound/Source/MetasoundGraphCore/Private/MetasoundDynamicOperatorTransactor.cpp
--- a/Engine/Plugins/Runtime/Metasound/Source/MetasoundGraphCore/Private/MetasoundDynamicOperatorTransactor.cpp
+++ b/Engine/Plugins/Runtime/Metasound/Source/MetasoundGraphCore/Private/MetasoundDynamicOperatorTransactor.cpp
@@ -588,9 +588,15 @@ namespace Metasound
 		}
 
 		TSharedRef<TSpscQueue<TUniquePtr<IDynamicOperatorTransform>>> FDynamicOperatorTransactor::CreateTransformQueue(const FOperatorSettings& InOperatorSettings, const FMetasoundEnvironment& InEnvironment)
+		{
+			TSharedPtr<FGraphRenderCost> GraphRenderCost;
+			return CreateTransformQueue(InOperatorSettings, InEnvironment, GraphRenderCost);
+		}
+
+		TSharedRef<TSpscQueue<TUniquePtr<IDynamicOperatorTransform>>> FDynamicOperatorTransactor::CreateTransformQueue(const FOperatorSettings& InOperatorSettings, const FMetasoundEnvironment& InEnvironment, const TSharedPtr<FGraphRenderCost>& InGraphRenderCost)
 		{
 			TSharedRef<TSpscQueue<TUniquePtr<IDynamicOperatorTransform>>> Queue = MakeShared<TSpscQueue<TUniquePtr<IDynamicOperatorTransform>>>();
-			OperatorInfos.Add(FDynamicOperatorInfo{InOperatorSettings, InEnvironment, Queue});
+			OperatorInfos.Add(FDynamicOperatorInfo{InOperatorSettings, InEnvironment, InGraphRenderCost, Queue});
 
 			TMap<FOperatorID, int32> OperatorOrdinals;
 			GraphSorter.GenerateOrdinals(OperatorOrdinals);
@@ -612,7 +618,7 @@ namespace Metasound
 				{
 					if (const int32* Ordinal = OperatorOrdinals.Find(DirectedGraphAlgo::GetOperatorID(GuidAndNode.Value)))
 					{
-						AtomicTransforms.Add(CreateInsertOperatorTransform(*GuidAndNode.Value, *Ordinal, InOperatorSettings, InEnvironment));
+						AtomicTransforms.Add(CreateInsertOperatorTransform(*GuidAndNode.Value, *Ordinal, InOperatorSettings, InEnvironment, InGraphRenderCost.Get()));
 					}
 				}
 			}
@@ -803,9 +809,9 @@ namespace Metasound
 			// Always insert new literal nodes first in execution order.
 			int32 LiteralOrdinal = GraphSorter.InsertOperator(DirectedGraphAlgo::GetOperatorID(LiteralNode), FDynamicGraphIncrementalSorter::EInsertLocation::First);
 
-			auto CreateAddNodeTransform = [&](const FOperatorSettings& InOperatorSettings, const FMetasoundEnvironment& InEnvironment) -> TUniquePtr<IDynamicOperatorTransform>
+			auto CreateAddNodeTransform = [&](const FDynamicOperatorInfo& InInfo) -> TUniquePtr<IDynamicOperatorTransform>
 			{
-				return CreateInsertOperatorTransform(*LiteralNode, LiteralOrdinal, InOperatorSettings, InEnvironment);
+				return CreateInsertOperatorTransform(*LiteralNode, LiteralOrdinal, InInfo.OperatorSettings, InInfo.Environment, InInfo.GraphRenderCost.Get());
 			};
 
 			EnqueueTransformOnOperatorQueues(CreateAddNodeTransform);
@@ -846,9 +852,9 @@ namespace Metasound
 
 			Graph.AddInputDataDestination(*Node, InVertexName);
 
-			auto CreateAddInputTransform = [&](const FOperatorSettings& InOperatorSettings, const FMetasoundEnvironment& InEnvironment) -> TUniquePtr<IDynamicOperatorTransform>
+			auto CreateAddInputTransform = [&](const FDynamicOperatorInfo& InInfo) -> TUniquePtr<IDynamicOperatorTransform>
 			{
-				TOptional<FAnyDataReference> NewDataReference = InFunc(InOperatorSettings, InputVertex.DataTypeName, InDefaultLiteral, ReferenceAccessType);
+				TOptional<FAnyDataReference> NewDataReference = InFunc(InInfo.OperatorSettings, InputVertex.DataTypeName, InDefaultLiteral, ReferenceAccessType);
 				if (NewDataReference.IsSet())
 				{
 					return MakeUnique<FAddInput>(OperatorID, InVertexName, *NewDataReference);
@@ -869,7 +875,7 @@ namespace Metasound
 
 			Graph.RemoveInputDataDestination(InVertexName);
 
-			auto CreateRemoveInputTransform = [&](const FOperatorSettings& InOperatorSettings, const FMetasoundEnvironment& InEnvironment) -> TUniquePtr<IDynamicOperatorTransform>
+			auto CreateRemoveInputTransform = [&](const FDynamicOperatorInfo& InInfo) -> TUniquePtr<IDynamicOperatorTransform>
 			{
 				return MakeUnique<FRemoveInput>(InVertexName);
 			};
@@ -897,7 +903,7 @@ namespace Metasound
 			Graph.AddOutputDataSource(*Node, InVertexName);
 			const FOperatorID OperatorID = DirectedGraphAlgo::GetOperatorID(Node);
 
-			auto CreateAddOutputTransform = [&](const FOperatorSettings& InOperatorSettings, const FMetasoundEnvironment& InEnvironment) -> TUniquePtr<IDynamicOperatorTransform>
+			auto CreateAddOutputTransform = [&](const FDynamicOperatorInfo& InInfo) -> TUniquePtr<IDynamicOperatorTransform>
 			{
 				return MakeUnique<FAddOutput>(OperatorID, InVertexName);
 			};
@@ -911,7 +917,7 @@ namespace Metasound
 
 			Graph.RemoveOutputDataSource(InVertexName);
 
-			auto CreateRemoveOutputTransform = [&](const FOperatorSettings& InOperatorSettings, const FMetasoundEnvironment& InEnvironment) -> TUniquePtr<IDynamicOperatorTransform>
+			auto CreateRemoveOutputTransform = [&](const FDynamicOperatorInfo& InInfo) -> TUniquePtr<IDynamicOperatorTransform>
 			{
 				return MakeUnique<FRemoveOutput>(InVertexName);
 			};
@@ -981,9 +987,9 @@ namespace Metasound
 		
 		void FDynamicOperatorTransactor::EnqueueInsertOperatorTransform(const INode& InNode, int32 InOrdinal)
 		{
-			auto CreateAddNodeTransform = [&](const FOperatorSettings& InOperatorSettings, const FMetasoundEnvironment& InEnvironment) -> TUniquePtr<IDynamicOperatorTransform>
+			auto CreateAddNodeTransform = [&](const FDynamicOperatorInfo& InInfo) -> TUniquePtr<IDynamicOperatorTransform>
 			{
-				return CreateInsertOperatorTransform(InNode, InOrdinal, InOperatorSettings, InEnvironment);
+				return CreateInsertOperatorTransform(InNode, InOrdinal, InInfo.OperatorSettings, InInfo.Environment, InInfo.GraphRenderCost.Get());
 			};
 
 			EnqueueTransformOnOperatorQueues(CreateAddNodeTransform);
@@ -1044,7 +1050,7 @@ namespace Metasound
 			const FOperatorID PriorLiteralOperatorID = DirectedGraphAlgo::GetOperatorID(InPriorLiteralNode);
 
 			// Create transforms for runtime
-			auto CreateAddEdgeTransform = [&](const FOperatorSettings& InOperatorSettings, const FMetasoundEnvironment& InEnvironment) -> TUniquePtr<IDynamicOperatorTransform>
+			auto CreateAddEdgeTransform = [&](const FDynamicOperatorInfo& InInfo) -> TUniquePtr<IDynamicOperatorTransform>
 			{
 				TArray<TUniquePtr<IDynamicOperatorTransform>> AtomicTransforms;
 
@@ -1080,7 +1086,7 @@ namespace Metasound
 			TArrayView<const FVertexName> InputsToFade(&InToVertex, 1);
 			TArrayView<const FVertexName> OutputsToFade;
 
-			auto CreateBeginFadeAndAddEdgeTransform = [&](const FOperatorSettings& InOperatorSettings, const FMetasoundEnvironment& InEnvironment) -> TUniquePtr<IDynamicOperatorTransform>
+			auto CreateBeginFadeAndAddEdgeTransform = [&](const FDynamicOperatorInfo& InInfo) -> TUniquePtr<IDynamicOperatorTransform>
 			{
 				TArray<TUniquePtr<IDynamicOperatorTransform>> AtomicTransforms;
 
@@ -1121,7 +1127,7 @@ namespace Metasound
 		{
 			const FOperatorID OperatorID = DirectedGraphAlgo::GetOperatorID(InNode);
 
-			auto CreateBeginAudioFadeTransform = [&](const FOperatorSettings& InOperatorSettings, const FMetasoundEnvironment& InEnvironment) -> TUniquePtr<IDynamicOperatorTransform>
+			auto CreateBeginAudioFadeTransform = [&](const FDynamicOperatorInfo& InInfo) -> TUniquePtr<IDynamicOperatorTransform>
 			{
 				TArray<TUniquePtr<IDynamicOperatorTransform>> BeginAudioFadeAtomicTransforms;
 
@@ -1142,7 +1148,7 @@ namespace Metasound
 		{
 			const FOperatorID OperatorID = DirectedGraphAlgo::GetOperatorID(InNode);
 
-			auto CreateEndAudioFadeTransform = [&](const FOperatorSettings& InOperatorSettings, const FMetasoundEnvironment& InEnvironment) -> TUniquePtr<IDynamicOperatorTransform>
+			auto CreateEndAudioFadeTransform = [&](const FDynamicOperatorInfo& InInfo) -> TUniquePtr<IDynamicOperatorTransform>
 			{
 				return MakeUnique<FEndAudioFadeTransform>(OperatorID);
 			};
@@ -1154,7 +1160,7 @@ namespace Metasound
 		{
 			const FOperatorID OperatorID = DirectedGraphAlgo::GetOperatorID(InNode);
 
-			auto CreateRemoveNodeTransform = [&OperatorID, &InOperatorsConnectedToInput](const FOperatorSettings& InOperatorSettings, const FMetasoundEnvironment& InEnvironment) -> TUniquePtr<IDynamicOperatorTransform>
+			auto CreateRemoveNodeTransform = [&OperatorID, &InOperatorsConnectedToInput](const FDynamicOperatorInfo& InInfo) -> TUniquePtr<IDynamicOperatorTransform>
 			{
 				return MakeUnique<FRemoveOperator>(OperatorID, InOperatorsConnectedToInput);
 			};
@@ -1169,10 +1175,10 @@ namespace Metasound
 			const FOperatorID ToOperatorID = DirectedGraphAlgo::GetOperatorID(InToNode);
 			const FOperatorID LiteralOperatorID = DirectedGraphAlgo::GetOperatorID(InReplacementLiteralNode);
 
-			auto CreateRemoveEdgeTransform = [&](const FOperatorSettings& InOperatorSettings, const FMetasoundEnvironment& InEnvironment) -> TUniquePtr<IDynamicOperatorTransform>
+			auto CreateRemoveEdgeTransform = [&](const FDynamicOperatorInfo& InInfo) -> TUniquePtr<IDynamicOperatorTransform>
 			{
 				// Add the literal node.
-				TUniquePtr<IDynamicOperatorTransform> AddNodeTransform = CreateInsertOperatorTransform(InReplacementLiteralNode, InLiteralOrdinal, InOperatorSettings, InEnvironment);
+				TUniquePtr<IDynamicOperatorTransform> AddNodeTransform = CreateInsertOperatorTransform(InReplacementLiteralNode, InLiteralOrdinal, InInfo.OperatorSettings, InInfo.Environment, InInfo.GraphRenderCost.Get());
 
 				// Swap prior connection with new connections.
 				TUniquePtr<IDynamicOperatorTransform> ConnectOperatorsTransform = MakeUnique<FSwapOperatorConnection>(FromOperatorID, InFromVertex, LiteralOperatorID, DynamicOperatorTransactorPrivate::LiteralNodeOutputVertexName, ToOperatorID, InToVertex);
@@ -1221,7 +1227,7 @@ namespace Metasound
 			EnqueueEndFadeOperatorTransform(InToNode);
 		}
 
-		TUniquePtr<IDynamicOperatorTransform> FDynamicOperatorTransactor::CreateInsertOperatorTransform(const INode& InNode, int32 InOrdinal, const FOperatorSettings& InOperatorSettings, const FMetasoundEnvironment& InEnvironment) const
+		TUniquePtr<IDynamicOperatorTransform> FDynamicOperatorTransactor::CreateInsertOperatorTransform(const INode& InNode, int32 InOrdinal, const FOperatorSettings& InOperatorSettings, const FMetasoundEnvironment& InEnvironment, FGraphRenderCost* InGraphRenderCost) const
 		{
 			using namespace DynamicOperatorTransactorPrivate;
 
@@ -1233,7 +1239,8 @@ namespace Metasound
 				InOperatorSettings,
 				InterfaceData.GetInputs(),
 				InEnvironment,
-				&OperatorBuilder // Supply an operator builder set to build rebindable inputs to ensure that subgraphs have their data references updated. 
+				&OperatorBuilder, // Supply an operator builder set to build rebindable inputs to ensure that subgraphs have their data references updated. 
+				InGraphRenderCost
 			};
 
 			FBuildResults Results;
@@ -1276,7 +1283,7 @@ namespace Metasound
 				TSharedPtr<TSpscQueue<TUniquePtr<IDynamicOperatorTransform>>> OperatorQueue = OperatorInfoIterator->Queue.Pin();
 				if (OperatorQueue.IsValid())
 				{
-					TUniquePtr<IDynamicOperatorTransform> Transform = InFunc(OperatorInfoIterator->OperatorSettings, OperatorInfoIterator->Environment);
+					TUniquePtr<IDynamicOperatorTransform> Transform = InFunc(*OperatorInfoIterator);
 					if (Transform.IsValid())
 					{
 						OperatorQueue->Enqueue(MoveTemp(Transform));
diff --git a/Engine/Plugins/Runtime/Metasound/Source/MetasoundGraphCore/Private/MetasoundEnvironment.cpp b/Engine/Plugins/Runtime/Metasound/Source/MetasoundGraphCore/Private/MetasoundEnvironment.cpp
--- a/Engine/Plugins/Runtime/Metasound/Source/MetasoundGraphCore/Private/MetasoundEnvironment.cpp
+++ b/Engine/Plugins/Runtime/Metasound/Source/MetasoundGraphCore/Private/MetasoundEnvironment.cpp
@@ -40,4 +40,17 @@ namespace Metasound
 
 		return *this;
 	}
+
+	namespace CoreInterface
+	{
+		namespace Environment
+		{
+			// The InstanceID acts as an external ID for communicating and in and out of MetaSounds. Each MetaSound
+			// has a unique InstanceID
+			const FLazyName InstanceID = "InstanceID";
+			
+			// An array representing the graph hierarchy.
+			const FLazyName GraphHierarchy = "GraphHierarhcy";
+		}
+	}
 }
diff --git a/Engine/Plugins/Runtime/Metasound/Source/MetasoundGraphCore/Private/MetasoundGraph.cpp b/Engine/Plugins/Runtime/Metasound/Source/MetasoundGraphCore/Private/MetasoundGraph.cpp
--- a/Engine/Plugins/Runtime/Metasound/Source/MetasoundGraphCore/Private/MetasoundGraph.cpp
+++ b/Engine/Plugins/Runtime/Metasound/Source/MetasoundGraphCore/Private/MetasoundGraph.cpp
@@ -292,7 +292,7 @@ namespace Metasound
 	{
 		const FGraph& Graph = static_cast<const FGraph&>(InParams.Node);
 
-		FBuildGraphOperatorParams BuildParams { Graph, InParams.OperatorSettings, InParams.InputData, InParams.Environment};
+		FBuildGraphOperatorParams BuildParams = FBuildGraphOperatorParams::FromBuildOperatorParams(Graph, InParams);
 
 		if (nullptr != InParams.Builder)
 		{
diff --git a/Engine/Plugins/Runtime/Metasound/Source/MetasoundGraphCore/Private/MetasoundOperatorBuilder.cpp b/Engine/Plugins/Runtime/Metasound/Source/MetasoundGraphCore/Private/MetasoundOperatorBuilder.cpp
--- a/Engine/Plugins/Runtime/Metasound/Source/MetasoundGraphCore/Private/MetasoundOperatorBuilder.cpp
+++ b/Engine/Plugins/Runtime/Metasound/Source/MetasoundGraphCore/Private/MetasoundOperatorBuilder.cpp
@@ -16,6 +16,7 @@
 #include "MetasoundOperatorInterface.h"
 #include "MetasoundProfilingOperator.h"
 #include "MetasoundRebindableGraphOperator.h"
+#include "MetasoundRenderCost.h"
 #include "MetasoundThreadLocalDebug.h"
 #include "MetasoundTrace.h"
 #include "Templates/PimplPtr.h"
@@ -25,11 +26,6 @@ namespace Metasound
 {
 	namespace OperatorBuilder
 	{
-		namespace Environment
-		{
-			const FName GraphHierarchy = "GraphHierarchy";
-		}
-
 		// Shared context used in the builder to maintain state of current build.
 		struct FBuildContext
 		{
@@ -37,6 +33,7 @@ namespace Metasound
 			const FDirectedGraphAlgoAdapter& AlgoAdapter;
 			const FOperatorSettings& Settings;
 			const FMetasoundEnvironment& Environment;
+			FGraphRenderCost* GraphRenderCost = nullptr;
 			
 			DirectedGraphAlgo::FGraphOperatorData& GraphOperatorData;
 			FBuildResults& Results;
@@ -46,12 +43,14 @@ namespace Metasound
 				const FDirectedGraphAlgoAdapter& InAlgoAdapter,
 				const FOperatorSettings& InSettings,
 				const FMetasoundEnvironment& InEnvironment,
+				FGraphRenderCost* InGraphRenderCost,
 				DirectedGraphAlgo::FGraphOperatorData& InGraphOperatorData,
 				FBuildResults& OutResults)
 			: Graph(InGraph)
 			, AlgoAdapter(InAlgoAdapter)
 			, Settings(InSettings)
 			, Environment(InEnvironment)
+			, GraphRenderCost(InGraphRenderCost)
 			, GraphOperatorData(InGraphOperatorData)
 			, Results(OutResults)
 			{
@@ -209,29 +208,29 @@ namespace Metasound
 
 		// Update environment with current graph hierarchy
 		FMetasoundEnvironment NewEnvironment;
-		if (!InParams.Environment.Contains<TArray<FGuid>>(OperatorBuilder::Environment::GraphHierarchy))
+		if (!InParams.Environment.Contains<TArray<FGuid>>(CoreInterface::Environment::GraphHierarchy))
 		{
 			// Copy old environment and add new environment variable
 			NewEnvironment = InParams.Environment;
 			TArray<FGuid> GraphHierarchy;
 			GraphHierarchy.Emplace(InParams.Graph.GetInstanceID());
-			NewEnvironment.SetValue<TArray<FGuid>>(OperatorBuilder::Environment::GraphHierarchy, GraphHierarchy);
+			NewEnvironment.SetValue<TArray<FGuid>>(CoreInterface::Environment::GraphHierarchy, GraphHierarchy);
 		}
 		else
 		{
 			// Copy and append to existing environment variable (environment variables aren't modifiable in place)
 			for (const auto& Iter : InParams.Environment)
 			{
-				if (Iter.Key == OperatorBuilder::Environment::GraphHierarchy)
+				if (Iter.Key == CoreInterface::Environment::GraphHierarchy)
 				{
 					TArray<FGuid> GraphHierarchy;
-					for (const FGuid Id : InParams.Environment.GetValue<TArray<FGuid>>(OperatorBuilder::Environment::GraphHierarchy))
+					for (const FGuid Id : InParams.Environment.GetValue<TArray<FGuid>>(CoreInterface::Environment::GraphHierarchy))
 					{
 						GraphHierarchy.Emplace(Id);
 					}
 					GraphHierarchy.Emplace(InParams.Graph.GetInstanceID());
 
-					NewEnvironment.SetValue<TArray<FGuid>>(OperatorBuilder::Environment::GraphHierarchy, GraphHierarchy);
+					NewEnvironment.SetValue<TArray<FGuid>>(CoreInterface::Environment::GraphHierarchy, GraphHierarchy);
 				}
 				else
 				{
@@ -241,7 +240,7 @@ namespace Metasound
 			}
 		}
 
-		OperatorBuilder::FBuildContext BuildContext(InParams.Graph, *AlgoAdapter, InParams.OperatorSettings, NewEnvironment, OutGraphOperatorData, OutResults);
+		OperatorBuilder::FBuildContext BuildContext(InParams.Graph, *AlgoAdapter, InParams.OperatorSettings, NewEnvironment, InParams.GraphRenderCost, OutGraphOperatorData, OutResults);
 
 		// Sort the nodes in a valid execution order
 		BuildStatus |= DepthFirstTopologicalSort(BuildContext, OutNodeOrder);
@@ -563,7 +562,7 @@ namespace Metasound
 				METASOUND_TRACE_CPUPROFILER_EVENT_SCOPE_TEXT(*TraceNamePtr);
 #endif // METASOUND_CPUPROFILERTRACE_ENABLED
 
-				FBuildOperatorParams CreateParams{*Node, InOutContext.Settings, OperatorInfo.VertexData.GetInputs(), InOutContext.Environment, this};
+				FBuildOperatorParams CreateParams{*Node, InOutContext.Settings, OperatorInfo.VertexData.GetInputs(), InOutContext.Environment, this, InOutContext.GraphRenderCost};
 				FOperatorFactorySharedRef Factory = Node->GetDefaultOperatorFactory();
 				if (ProfileOperators && Profiling::OperatorShouldBeProfiled(Node->GetMetadata()))
 				{
diff --git a/Engine/Plugins/Runtime/Metasound/Source/MetasoundGraphCore/Private/MetasoundRenderCost.cpp b/Engine/Plugins/Runtime/Metasound/Source/MetasoundGraphCore/Private/MetasoundRenderCost.cpp
new file mode 100644
--- /dev/null
+++ b/Engine/Plugins/Runtime/Metasound/Source/MetasoundGraphCore/Private/MetasoundRenderCost.cpp
@@ -0,0 +1,88 @@
+// Copyright Epic Games, Inc. All Rights Reserved.
+
+#include "MetasoundRenderCost.h"
+
+#include "Containers/Array.h"
+#include "Containers/ArrayView.h"
+#include "Containers/Map.h"
+#include "DSP/FloatArrayMath.h"
+#include "MetasoundEnvironment.h"
+#include "Misc/Guid.h"
+#include "Misc/ScopeLock.h"
+#include "Templates/SharedPointer.h"
+
+namespace Metasound
+{
+	FGraphRenderCost::FGraphRenderCost(EPrivateToken InToken)
+	{
+	}
+
+	TSharedRef<FGraphRenderCost> FGraphRenderCost::MakeGraphRenderCost()
+	{
+		return MakeShared<FGraphRenderCost>(PrivateToken);
+	}
+
+	FNodeRenderCost FGraphRenderCost::AddNode(const FGuid& InNodeInstanceID, const FMetasoundEnvironment& InEnv)
+	{
+		int32 NodeIndex = NodeCosts.Num();
+		NodeCosts.Add(0.f);
+
+		AddNodeHierarchy(InNodeInstanceID, InEnv);
+
+		return FNodeRenderCost(NodeIndex, AsShared());
+	}
+
+	void FGraphRenderCost::ResetNodeRenderCosts()
+	{
+		if (NodeCosts.Num() > 0)
+		{
+			FMemory::Memzero(NodeCosts.GetData(), sizeof(float) * NodeCosts.Num());
+		}
+	}
+
+	void FGraphRenderCost::SetNodeRenderCost(int32 InNodeIndex, float InRenderCost)
+	{
+		NodeCosts[InNodeIndex] = InRenderCost;
+	}
+
+	float FGraphRenderCost::ComputeGraphRenderCost() const
+	{
+		float Cost = 0.f;
+		Audio::ArraySum(NodeCosts, Cost);
+		return Cost;
+	};
+
+#if UE_METASOUNDRENDERCOST_TRACK_NODE_HIERARCHY
+	void FGraphRenderCost::AddNodeHierarchy(const FGuid& InNodeInstanceID, const FMetasoundEnvironment& InEnv)
+	{
+		// Node hierarchies provide a convenient path for finding costly nodes by looking at their graph hierarchy. 
+		TArray<FGuid> NodeHierarchy;
+		if (InEnv.Contains<TArray<FGuid>>(CoreInterface::Environment::GraphHierarchy))
+		{
+			NodeHierarchy = InEnv.GetValue<TArray<FGuid>>(CoreInterface::Environment::GraphHierarchy);
+		}
+		NodeHierarchy.Add(InNodeInstanceID);
+		NodeHierarchies.Add(MoveTemp(NodeHierarchy));
+	}
+#else
+	void FGraphRenderCost::AddNodeHierarchy(const FGuid& InNodeInstanceID, const FMetasoundEnvironment& InEnv)
+	{
+	}
+#endif // UE_METASOUNDRENDERCOST_TRACK_NODE_HIERARCHY
+
+	FNodeRenderCost::FNodeRenderCost(int32 InNodeIndex, TSharedRef<FGraphRenderCost> InGraphRenderCost)
+	: NodeIndex(InNodeIndex)
+	, GraphRenderCost(MoveTemp(InGraphRenderCost))
+	{
+	}
+
+	void FNodeRenderCost::SetRenderCost(float InCost)
+	{
+		if (GraphRenderCost.IsValid())
+		{
+			// Render cost values are stored on an array in the FGraphRenderCost. 
+			GraphRenderCost->SetNodeRenderCost(NodeIndex, InCost);
+		}
+	}
+}
+
diff --git a/Engine/Plugins/Runtime/Metasound/Source/MetasoundGraphCore/Public/MetasoundBuilderInterface.h b/Engine/Plugins/Runtime/Metasound/Source/MetasoundGraphCore/Public/MetasoundBuilderInterface.h
--- a/Engine/Plugins/Runtime/Metasound/Source/MetasoundGraphCore/Public/MetasoundBuilderInterface.h
+++ b/Engine/Plugins/Runtime/Metasound/Source/MetasoundGraphCore/Public/MetasoundBuilderInterface.h
@@ -7,6 +7,7 @@
 #include "MetasoundEnvironment.h"
 #include "MetasoundNodeInterface.h"
 #include "MetasoundOperatorInterface.h"
+#include "MetasoundRenderCost.h"
 #include "Templates/UniquePtr.h"
 
 class FName;
@@ -107,10 +108,13 @@ namespace Metasound
 		/** Pointer to builder actively building graph. */
 		const IOperatorBuilder* Builder = nullptr;
 
+		/** Runtime render cost tracker. */
+		FGraphRenderCost* GraphRenderCost = nullptr;
+
 		/** Implicit conversion to FResetParams for convenience. */
 		operator IOperator::FResetParams() const
 		{
-			return IOperator::FResetParams{OperatorSettings, Environment};
+			return IOperator::FResetParams{OperatorSettings, Environment, GraphRenderCost};
 		}
 	};
 
@@ -147,6 +151,15 @@ namespace Metasound
 
 		/** Environment settings available. */
 		const FMetasoundEnvironment& Environment;
+
+		/** Runtime render cost tracker. */
+		FGraphRenderCost* GraphRenderCost;
+
+		/** Convert operator building params to graph building params */
+		static FBuildGraphOperatorParams FromBuildOperatorParams(const IGraph& InGraph, const FBuildOperatorParams& InParams)
+		{
+			return FBuildGraphOperatorParams{InGraph, InParams.OperatorSettings, InParams.InputData, InParams.Environment, InParams.GraphRenderCost};
+		}
 	};
 
 	/** Convenience template for adding build errors.
diff --git a/Engine/Plugins/Runtime/Metasound/Source/MetasoundGraphCore/Public/MetasoundDynamicOperatorTransactor.h b/Engine/Plugins/Runtime/Metasound/Source/MetasoundGraphCore/Public/MetasoundDynamicOperatorTransactor.h
--- a/Engine/Plugins/Runtime/Metasound/Source/MetasoundGraphCore/Public/MetasoundDynamicOperatorTransactor.h
+++ b/Engine/Plugins/Runtime/Metasound/Source/MetasoundGraphCore/Public/MetasoundDynamicOperatorTransactor.h
@@ -9,6 +9,7 @@
 #include "MetasoundOperatorBuilder.h"
 #include "MetasoundOperatorInterface.h"
 #include "MetasoundNodeInterface.h"
+#include "MetasoundRenderCost.h"
 #include "Templates/Function.h"
 #include "Templates/SharedPointer.h"
 #include "Templates/UniquePtr.h"
@@ -187,7 +188,9 @@ namespace Metasound
 			FDynamicOperatorTransactor(const FGraph& InGraph);
 
 			/** Create a queue for communication with a dynamic operator. */
+			UE_DEPRECATED(5.5, "Replace with CreateTransformQueue overload including FGraphRenderCost")
 			TSharedRef<TSpscQueue<TUniquePtr<IDynamicOperatorTransform>>> CreateTransformQueue(const FOperatorSettings& InOperatorSettings, const FMetasoundEnvironment& InEnvironment);
+			TSharedRef<TSpscQueue<TUniquePtr<IDynamicOperatorTransform>>> CreateTransformQueue(const FOperatorSettings& InOperatorSettings, const FMetasoundEnvironment& InEnvironment, const TSharedPtr<FGraphRenderCost>& InRenderCost);
 
 			/** Add a node to the graph. */
 			void AddNode(const FGuid& InNodeID, TUniquePtr<INode> InNode);
@@ -254,10 +257,17 @@ namespace Metasound
 			void EnqueueFadeAndAddEdgeOperatorTransform(const INode& InFromNode, const FVertexName& InFromVertex, const INode& InToNode, const FVertexName& InToVertex, const INode* InPriorLiteralNode, const TArray<FOrdinalSwap>& InOrdinalUpdates);
 
 			void AddDataEdgeInternal(const INode& InFromNode, const FVertexName& InFromVertex, const FGuid& InToNodeID, const INode& InToNode, const FVertexName& InToVertex);
+			struct FDynamicOperatorInfo
+			{
+				FOperatorSettings OperatorSettings;
+				FMetasoundEnvironment Environment;
+				TSharedPtr<FGraphRenderCost> GraphRenderCost;
+				TWeakPtr<TSpscQueue<TUniquePtr<IDynamicOperatorTransform>>> Queue;
+			};
 
-			using FCreateTransformFunctionRef = TFunctionRef<TUniquePtr<IDynamicOperatorTransform>(const FOperatorSettings& InOperatorSettings, const FMetasoundEnvironment& InEnvironment)>;
+			using FCreateTransformFunctionRef = TFunctionRef<TUniquePtr<IDynamicOperatorTransform>(const FDynamicOperatorInfo& InOperatorInfo)>;
 
-			TUniquePtr<IDynamicOperatorTransform> CreateInsertOperatorTransform(const INode& InNode, int32 InOrdinal, const FOperatorSettings& InOperatorSettings, const FMetasoundEnvironment& InEnvironment) const;
+			TUniquePtr<IDynamicOperatorTransform> CreateInsertOperatorTransform(const INode& InNode, int32 InOrdinal, const FOperatorSettings& InOperatorSettings, const FMetasoundEnvironment& InEnvironment, FGraphRenderCost* InGraphRenderCost) const;
 
 			void EnqueueTransformOnOperatorQueues(FCreateTransformFunctionRef InFunc);
 
@@ -265,12 +275,6 @@ namespace Metasound
 			FGraph Graph;
 			FDynamicGraphIncrementalSorter GraphSorter;
 
-			struct FDynamicOperatorInfo
-			{
-				FOperatorSettings OperatorSettings;
-				FMetasoundEnvironment Environment;
-				TWeakPtr< TSpscQueue< TUniquePtr< IDynamicOperatorTransform >>> Queue;
-			};
 
 			TArray<FDynamicOperatorInfo> OperatorInfos;
 
diff --git a/Engine/Plugins/Runtime/Metasound/Source/MetasoundGraphCore/Public/MetasoundEnvironment.h b/Engine/Plugins/Runtime/Metasound/Source/MetasoundGraphCore/Public/MetasoundEnvironment.h
--- a/Engine/Plugins/Runtime/Metasound/Source/MetasoundGraphCore/Public/MetasoundEnvironment.h
+++ b/Engine/Plugins/Runtime/Metasound/Source/MetasoundGraphCore/Public/MetasoundEnvironment.h
@@ -275,6 +275,19 @@ namespace Metasound
 	private:
 		TSortedMap<FName, TUniquePtr<IMetasoundEnvironmentVariable>, FDefaultAllocator, FNameFastLess> Variables;
 	};
+
+	namespace CoreInterface
+	{
+		namespace Environment
+		{
+			// The InstanceID acts as an external ID for communicating and in and out of MetaSounds. Each MetaSound
+			// has a unique InstanceID
+			METASOUNDGRAPHCORE_API const extern FLazyName InstanceID;
+
+			// An array representing the graph hierarchy.
+			METASOUNDGRAPHCORE_API const extern FLazyName GraphHierarchy;
+		}
+	}
 }
 
 /** Declare basic set of variable types. */
diff --git a/Engine/Plugins/Runtime/Metasound/Source/MetasoundGraphCore/Public/MetasoundOperatorBuilder.h b/Engine/Plugins/Runtime/Metasound/Source/MetasoundGraphCore/Public/MetasoundOperatorBuilder.h
--- a/Engine/Plugins/Runtime/Metasound/Source/MetasoundGraphCore/Public/MetasoundOperatorBuilder.h
+++ b/Engine/Plugins/Runtime/Metasound/Source/MetasoundGraphCore/Public/MetasoundOperatorBuilder.h
@@ -25,10 +25,6 @@ namespace Metasound
 	namespace OperatorBuilder
 	{
 		struct FBuildContext;
-		namespace Environment
-		{
-			METASOUNDGRAPHCORE_API const extern FName GraphHierarchy;
-		}
 	}
 
 	namespace DynamicGraph
diff --git a/Engine/Plugins/Runtime/Metasound/Source/MetasoundGraphCore/Public/MetasoundOperatorInterface.h b/Engine/Plugins/Runtime/Metasound/Source/MetasoundGraphCore/Public/MetasoundOperatorInterface.h
--- a/Engine/Plugins/Runtime/Metasound/Source/MetasoundGraphCore/Public/MetasoundOperatorInterface.h
+++ b/Engine/Plugins/Runtime/Metasound/Source/MetasoundGraphCore/Public/MetasoundOperatorInterface.h
@@ -3,11 +3,11 @@
 #pragma once
 
 #include "MetasoundDataReferenceCollection.h"
+#include "MetasoundRenderCost.h"
 #include "MetasoundVertexData.h"
 
 namespace Metasound
 {
-
 	/** IOperator
 	 *
 	 *  IOperator defines the interface for render time operations.  IOperators are created using an INodeOperatorFactory.
@@ -25,6 +25,9 @@ namespace Metasound
 
 			/** Environment settings available. */
 			const FMetasoundEnvironment& Environment;
+
+			/** Object for tracking runtime render cost of a MetaSound */
+			const FGraphRenderCost* GraphRenderCost = nullptr;
 		};
 
 		virtual ~IOperator() {}
diff --git a/Engine/Plugins/Runtime/Metasound/Source/MetasoundGraphCore/Public/MetasoundRenderCost.h b/Engine/Plugins/Runtime/Metasound/Source/MetasoundGraphCore/Public/MetasoundRenderCost.h
new file mode 100644
--- /dev/null
+++ b/Engine/Plugins/Runtime/Metasound/Source/MetasoundGraphCore/Public/MetasoundRenderCost.h
@@ -0,0 +1,87 @@
+// Copyright Epic Games, Inc. All Rights Reserved.
+
+#pragma once
+
+#include "Containers/Array.h"
+#include "Containers/Map.h"
+#include "MetasoundEnvironment.h"
+#include "Misc/Guid.h"
+#include "Misc/ScopeLock.h"
+#include "Templates/SharedPointer.h"
+
+#ifndef UE_METASOUNDRENDERCOST_TRACK_NODE_HIERARCHY 
+#define UE_METASOUNDRENDERCOST_TRACK_NODE_HIERARCHY (!UE_BUILD_SHIPPING)
+#endif
+
+namespace Metasound
+{
+	class FNodeRenderCost;
+
+	/** FGraphRenderCost represents the accumulated render cost of a graph. Individual
+	 * nodes in a graph can report their render cost through an FNodeRenderCost.
+	 *
+	 * The render cost of each node is added together to determine the graph's render
+	 * cost. 
+	 */
+	class METASOUNDGRAPHCORE_API FGraphRenderCost : public TSharedFromThis<FGraphRenderCost>
+	{
+		// Private token to enforce creation of shared reference. 
+		enum EPrivateToken { PrivateToken };
+		friend class FNodeRenderCost;
+
+	public:
+
+		FGraphRenderCost(EPrivateToken InToken);
+
+		static TSharedRef<FGraphRenderCost> MakeGraphRenderCost();
+
+		/* Add a node to the graph's render cost.
+		 *
+		 * @param InNodeInstanceID - ID of the node to add.
+		 * @param InEnv - The environment used to create the node. 
+		 *
+		 * @return FNodeRenderCost object for reporting the render cost of individual nodes. */
+		FNodeRenderCost AddNode(const FGuid& InNodeInstanceID, const FMetasoundEnvironment& InEnv);
+
+		/** Reset the individual node render costs to zero. */
+		void ResetNodeRenderCosts();
+
+		/** Adds all the individual node render costs and returns the result. */
+		float ComputeGraphRenderCost() const;
+		
+	private:
+		void SetNodeRenderCost(int32 InNodeIndex, float InRenderCost);
+
+		TArray<float> NodeCosts;
+		void AddNodeHierarchy(const FGuid& InNodeInstanceID, const FMetasoundEnvironment& InEnv);
+#if UE_METASOUNDRENDERCOST_TRACK_NODE_HIERARCHY
+		TArray<TArray<FGuid>> NodeHierarchies;
+#endif
+	};
+
+	/** FNodeRenderCost allows individual nodes to report their render cost. 
+	 * 
+	 * FNodeRenderCost should be created with FGraphRenderCost::AddNode(...)
+	 */
+	class METASOUNDGRAPHCORE_API FNodeRenderCost
+	{
+		friend class FGraphRenderCost;
+		FNodeRenderCost(int32 InNodeIndex, TSharedRef<FGraphRenderCost> InGraphRenderCost);
+	public:
+		/** The default constructor is provided for convenience, but creating a
+		 * FNodeRenderCost with the default constructor will allow any method for
+		 * reading or accumulating a nodes render cost. Instead, FNodeRenderCost
+		 * should be created with FGraphRenderCost::AddNode(...) so that it is 
+		 * associated and aggregated within a graph.
+		 */
+		FNodeRenderCost() = default;
+
+		/** Set the render cost of this node. */
+		void SetRenderCost(float InCost);
+
+	private:
+		int32 NodeIndex = INDEX_NONE;
+		TSharedPtr<FGraphRenderCost> GraphRenderCost;
+	};
+}
+
diff --git a/Engine/Source/Runtime/AudioMixer/Private/AudioMixerSource.cpp b/Engine/Source/Runtime/AudioMixer/Private/AudioMixerSource.cpp
--- a/Engine/Source/Runtime/AudioMixer/Private/AudioMixerSource.cpp
+++ b/Engine/Source/Runtime/AudioMixer/Private/AudioMixerSource.cpp
@@ -1064,6 +1064,8 @@ namespace Audio
 
 		UpdateChannelMaps();
 
+		UpdateRelativeRenderCost();
+
 #if ENABLE_AUDIO_DEBUG
 		UpdateCPUCoreUtilization();
 
@@ -1522,6 +1524,15 @@ namespace Audio
 		return 0.0f;
 	}
 
+	float FMixerSource::GetRelativeRenderCost() const
+	{
+		if (MixerSourceVoice)
+		{
+			return MixerSourceVoice->GetRelativeRenderCost();
+		}
+		return 1.0f;
+	}
+
 	void FMixerSource::OnBeginGenerate()
 	{
 	}
@@ -2039,6 +2050,23 @@ namespace Audio
 		bPrevAllowedSpatializationSetting = IsSpatializationCVarEnabled();
 	}
 
+	void FMixerSource::UpdateRelativeRenderCost()
+	{
+		if (MixerSourceVoice)
+		{
+			const float RelativeRenderCost = MixerSourceVoice->GetRelativeRenderCost();
+			check(WaveInstance);
+			WaveInstance->SetRelativeRenderCost(RelativeRenderCost);
+#if ENABLE_AUDIO_DEBUG
+			if (DebugInfo.IsValid())
+			{
+				FScopeLock DebugInfoLock(&DebugInfo->CS);
+				DebugInfo->RelativeRenderCost = RelativeRenderCost;
+			}
+#endif // if ENABLE_AUDIO_DEBUG
+		}
+	}
+
 #if ENABLE_AUDIO_DEBUG
 	void FMixerSource::UpdateCPUCoreUtilization()
 	{
diff --git a/Engine/Source/Runtime/AudioMixer/Private/AudioMixerSource.h b/Engine/Source/Runtime/AudioMixer/Private/AudioMixerSource.h
--- a/Engine/Source/Runtime/AudioMixer/Private/AudioMixerSource.h
+++ b/Engine/Source/Runtime/AudioMixer/Private/AudioMixerSource.h
@@ -51,6 +51,7 @@ namespace Audio
 		virtual float GetPlaybackPercent() const override;
 		virtual int64 GetNumFramesPlayed() const override;
 		virtual float GetEnvelopeValue() const override;
+		virtual float GetRelativeRenderCost() const override;
 		//~ End FSoundSource Interface
 
 		//~ Begin ISourceListener
@@ -89,6 +90,9 @@ namespace Audio
 		/** Updates the channel map of the sound if its a 3d sound.*/
 		void UpdateChannelMaps();
 
+		/** Updates the relative render cost estimate of the playing sound source. */
+		void UpdateRelativeRenderCost();
+
 #if ENABLE_AUDIO_DEBUG
 		void UpdateCPUCoreUtilization();
 #endif // ENABLE_AUDIO_DEBUG
diff --git a/Engine/Source/Runtime/AudioMixer/Private/AudioMixerSourceBuffer.cpp b/Engine/Source/Runtime/AudioMixer/Private/AudioMixerSourceBuffer.cpp
--- a/Engine/Source/Runtime/AudioMixer/Private/AudioMixerSourceBuffer.cpp
+++ b/Engine/Source/Runtime/AudioMixer/Private/AudioMixerSourceBuffer.cpp
@@ -480,6 +480,9 @@ namespace Audio
 					double AudioDuration = static_cast<double>(TaskResult.NumSamplesWritten) / static_cast<double>(FMath::Max(1, NumChannels * SampleRate));
 					UpdateCPUCoreUtilization(TaskResult.CPUDuration, AudioDuration);
 #endif // ENABLE_AUDIO_DEBUG
+
+					// Set the render cost encountered during the last render
+					SetRelativeRenderCost(TaskResult.RelativeRenderCost);
 				}
 				break;
 			}
@@ -606,6 +609,16 @@ namespace Audio
 		return bProcedural && SoundGenerator.IsValid() && SoundGenerator->IsFinished();
 	}
 
+	float FMixerSourceBuffer::GetRelativeRenderCost() const
+	{
+		return RelativeRenderCost.load(std::memory_order_relaxed);
+	}
+
+	void FMixerSourceBuffer::SetRelativeRenderCost(float InRelativeRenderCost)
+	{
+		RelativeRenderCost.store(InRelativeRenderCost, std::memory_order_relaxed);
+	}
+
 #if ENABLE_AUDIO_DEBUG
 	double FMixerSourceBuffer::GetCPUCoreUtilization() const
 	{
diff --git a/Engine/Source/Runtime/AudioMixer/Private/AudioMixerSourceBuffer.h b/Engine/Source/Runtime/AudioMixer/Private/AudioMixerSourceBuffer.h
--- a/Engine/Source/Runtime/AudioMixer/Private/AudioMixerSourceBuffer.h
+++ b/Engine/Source/Runtime/AudioMixer/Private/AudioMixerSourceBuffer.h
@@ -129,6 +129,9 @@ namespace Audio
 		double GetCPUCoreUtilization() const;
 #endif // ENABLE_AUDIO_DEBUG
 
+		// Returns the runtime render cost
+		float GetRelativeRenderCost() const;
+
 	private:
 		FMixerSourceBuffer(FMixerSourceBufferInitArgs& InArgs, TArray<FAudioParameter>&& InDefaultParams);
 
@@ -167,6 +170,9 @@ namespace Audio
 		void UpdateCPUCoreUtilization(double InCPUTime, double InAudioTime);
 #endif // ENABLE_AUDIO_DEBUG
 
+		std::atomic<float> RelativeRenderCost = 1.0f;
+		void SetRelativeRenderCost(float InRelativeRenderCost);
+
 		mutable FCriticalSection SoundWaveCritSec;
 		mutable FCriticalSection DecodeTaskCritSec;
 
diff --git a/Engine/Source/Runtime/AudioMixer/Private/AudioMixerSourceDecode.cpp b/Engine/Source/Runtime/AudioMixer/Private/AudioMixerSourceDecode.cpp
--- a/Engine/Source/Runtime/AudioMixer/Private/AudioMixerSourceDecode.cpp
+++ b/Engine/Source/Runtime/AudioMixer/Private/AudioMixerSourceDecode.cpp
@@ -99,6 +99,7 @@ public:
 					// Generators are responsible to zero memory in case they can't generate the requested amount of samples
 					ProceduralResult.NumSamplesWritten = ProceduralTaskData.SoundGenerator->GetNextBuffer(ProceduralTaskData.AudioData, ProceduralTaskData.NumSamples);
 					ProceduralResult.bIsFinished = ProceduralTaskData.SoundGenerator->IsFinished();
+					ProceduralResult.RelativeRenderCost = ProceduralTaskData.SoundGenerator->GetRelativeRenderCost();
 				}
 				else
 				{
diff --git a/Engine/Source/Runtime/AudioMixer/Private/AudioMixerSourceDecode.h b/Engine/Source/Runtime/AudioMixer/Private/AudioMixerSourceDecode.h
--- a/Engine/Source/Runtime/AudioMixer/Private/AudioMixerSourceDecode.h
+++ b/Engine/Source/Runtime/AudioMixer/Private/AudioMixerSourceDecode.h
@@ -110,6 +110,7 @@ namespace Audio
 	{
 		int32 NumSamplesWritten;
 		bool bIsFinished;
+		float RelativeRenderCost = 1.f;
 
 #if ENABLE_AUDIO_DEBUG
 		double CPUDuration = 0.0;
diff --git a/Engine/Source/Runtime/AudioMixer/Private/AudioMixerSourceManager.cpp b/Engine/Source/Runtime/AudioMixer/Private/AudioMixerSourceManager.cpp
--- a/Engine/Source/Runtime/AudioMixer/Private/AudioMixerSourceManager.cpp
+++ b/Engine/Source/Runtime/AudioMixer/Private/AudioMixerSourceManager.cpp
@@ -394,9 +394,12 @@ namespace Audio
 #if ENABLE_AUDIO_DEBUG
 		GameThreadInfo.CPUCoreUtilization.AddZeroed(NumTotalSources);
 #endif // if ENABLE_AUDIO_DEBUG
+
+		GameThreadInfo.RelativeRenderCost.Reset(NumTotalSources);
 		GameThreadInfo.FreeSourceIndices.Reset(NumTotalSources);
 		for (int32 i = NumTotalSources - 1; i >= 0; --i)
 		{
+			GameThreadInfo.RelativeRenderCost.Add(1.0f);
 			GameThreadInfo.FreeSourceIndices.Add(i);
 		}
 
@@ -1181,6 +1184,11 @@ namespace Audio
 		GameThreadInfo.bIsDebugMode[SourceId] = false;
 #endif
 
+#if ENABLE_AUDIO_DEBUG
+		GameThreadInfo.CPUCoreUtilization[SourceId] = 0.0f;
+#endif 
+
+		GameThreadInfo.RelativeRenderCost[SourceId] = 1.0f;
 		GameThreadInfo.FreeSourceIndices.Push(SourceId);
 
 		AUDIO_MIXER_CHECK(GameThreadInfo.FreeSourceIndices.Contains(SourceId));
@@ -1958,6 +1966,11 @@ namespace Audio
 	}
 #endif // if ENABLE_AUDIO_DEBUG
 
+	float FMixerSourceManager::GetRelativeRenderCost(const int32 SourceId) const
+	{
+		return GameThreadInfo.RelativeRenderCost[SourceId];
+	}
+
 	bool FMixerSourceManager::IsUsingHRTFSpatializer(const int32 SourceId) const
 	{
 		AUDIO_MIXER_CHECK_GAME_THREAD(MixerDevice);
@@ -2019,6 +2032,9 @@ namespace Audio
 					// out as an acceptable race condition given that it is utilized for debug purposes only. 
 					GameThreadInfo.CPUCoreUtilization[SourceId] = SourceInfo.MixerSourceBuffer->GetCPUCoreUtilization();
 #endif // if ENABLE_AUDIO_DEBUG
+
+					GameThreadInfo.RelativeRenderCost[SourceId] = SourceInfo.MixerSourceBuffer->GetRelativeRenderCost();
+
 					SourceInfo.MixerSourceBuffer->OnBufferEnd();
 				}
 			}
diff --git a/Engine/Source/Runtime/AudioMixer/Private/AudioMixerSourceManager.h b/Engine/Source/Runtime/AudioMixer/Private/AudioMixerSourceManager.h
--- a/Engine/Source/Runtime/AudioMixer/Private/AudioMixerSourceManager.h
+++ b/Engine/Source/Runtime/AudioMixer/Private/AudioMixerSourceManager.h
@@ -262,6 +262,7 @@ namespace Audio
 #if ENABLE_AUDIO_DEBUG
 		double GetCPUCoreUtilization(const int32 SourceId) const;
 #endif // ENABLE_AUDIO_DEBUG
+		float GetRelativeRenderCost(const int32 SourceId) const;
 		bool IsUsingHRTFSpatializer(const int32 SourceId) const;
 		bool NeedsSpeakerMap(const int32 SourceId) const;
 		void ComputeNextBlockOfSamples();
@@ -635,6 +636,7 @@ namespace Audio
 #if ENABLE_AUDIO_DEBUG
 			TArray<double> CPUCoreUtilization;
 #endif // #if ENABLE_AUDIO_DEBUG
+			TArray<float> RelativeRenderCost;
 		} GameThreadInfo;
 
 		int32 NumActiveSources;
diff --git a/Engine/Source/Runtime/AudioMixer/Private/AudioMixerSourceVoice.cpp b/Engine/Source/Runtime/AudioMixer/Private/AudioMixerSourceVoice.cpp
--- a/Engine/Source/Runtime/AudioMixer/Private/AudioMixerSourceVoice.cpp
+++ b/Engine/Source/Runtime/AudioMixer/Private/AudioMixerSourceVoice.cpp
@@ -342,6 +342,13 @@ namespace Audio
 	}
 #endif // ENABLE_AUDIO_DEBUG
 
+	float FMixerSourceVoice::GetRelativeRenderCost() const
+	{
+		AUDIO_MIXER_CHECK_GAME_THREAD(MixerDevice);
+
+		return SourceManager->GetRelativeRenderCost(SourceId);
+	}
+
 	void FMixerSourceVoice::MixOutputBuffers(int32 InNumOutputChannels, const float SendLevel, EMixerSourceSubmixSendStage InSubmixSendStage, FAlignedFloatBuffer& OutWetBuffer) const
 	{
 		AUDIO_MIXER_CHECK_AUDIO_PLAT_THREAD(MixerDevice);
diff --git a/Engine/Source/Runtime/AudioMixer/Private/AudioMixerSourceVoice.h b/Engine/Source/Runtime/AudioMixer/Private/AudioMixerSourceVoice.h
--- a/Engine/Source/Runtime/AudioMixer/Private/AudioMixerSourceVoice.h
+++ b/Engine/Source/Runtime/AudioMixer/Private/AudioMixerSourceVoice.h
@@ -121,6 +121,11 @@ namespace Audio
 		double GetCPUCoreUtilization() const;
 #endif // ENABLE_AUDIO_DEBUG
 
+		// Retrieves the current "render cost" of the mixer source voice. Used for debug display and for limiting voice count.
+		// 1.0 is equivalent to a single decoding sound source. 
+		// Useful as a metric more human-centric than CPU utilization.
+		float GetRelativeRenderCost() const;
+
 		// Mixes the dry and wet buffer audio into the given buffers.
 		void MixOutputBuffers(int32 InNumChannels, const float SendLevel, EMixerSourceSubmixSendStage InSubmixSendStage, FAlignedFloatBuffer& OutWetBuffer) const;
 
diff --git a/Engine/Source/Runtime/Engine/Classes/Sound/SoundGenerator.h b/Engine/Source/Runtime/Engine/Classes/Sound/SoundGenerator.h
--- a/Engine/Source/Runtime/Engine/Classes/Sound/SoundGenerator.h
+++ b/Engine/Source/Runtime/Engine/Classes/Sound/SoundGenerator.h
@@ -48,6 +48,13 @@ public:
 	// Optional. Can be overridden to end the sound when generating is finished.
 	virtual bool IsFinished() const { return false; }
 
+	// Return the cost to render this sound generator. Derived classes can optionally
+	// override this method.
+	//
+	// This is called repeatedly to track the cost and managing the maximum number
+	// of concurrently rendering voices. 
+	virtual float GetRelativeRenderCost() const { return 1.f; }
+
 	// Retrieves the next buffer of audio from the generator, called from the audio mixer
 	ENGINE_API int32 GetNextBuffer(float* OutAudio, int32 NumSamples, bool bRequireNumberSamples = false);
 
diff --git a/Engine/Source/Runtime/Engine/Private/Audio.cpp b/Engine/Source/Runtime/Engine/Private/Audio.cpp
--- a/Engine/Source/Runtime/Engine/Private/Audio.cpp
+++ b/Engine/Source/Runtime/Engine/Private/Audio.cpp
@@ -848,6 +848,7 @@ FWaveInstance::FWaveInstance(const UPTRINT InWaveInstanceHash, FActiveSound& InA
 	, OcclusionAttenuation(1.0f)
 	, VolumeMultiplier(1.0f)
 	, EnvelopValue(0.0f)
+	, RelativeRenderCost(1.0f)
 	, EnvelopeFollowerAttackTime(10)
 	, EnvelopeFollowerReleaseTime(100)
 	, Priority(1.0f)
diff --git a/Engine/Source/Runtime/Engine/Private/AudioDebug.cpp b/Engine/Source/Runtime/Engine/Private/AudioDebug.cpp
--- a/Engine/Source/Runtime/Engine/Private/AudioDebug.cpp
+++ b/Engine/Source/Runtime/Engine/Private/AudioDebug.cpp
@@ -2108,16 +2108,19 @@ namespace Audio
 			AActor* SoundOwner = AudioComponent ? AudioComponent->GetOwner() : nullptr;
 
 			double CPUPercent = 0.;
+			float RelativeRenderCost = 1.f;
 			if (WaveInstanceInfo.Key->DebugInfo.IsValid())
 			{
 				FScopeLock Lock(&WaveInstanceInfo.Key->DebugInfo->CS);
 				CPUPercent = 100. * WaveInstanceInfo.Key->DebugInfo->CPUCoreUtilization;
+				RelativeRenderCost = WaveInstanceInfo.Key->DebugInfo->RelativeRenderCost;
 			}
 			
-			FString TheString = *FString::Printf(TEXT("%4i.    %6.2f  CPU:%5.2f%% %s   Owner: %s   SoundClass: %s"),
+			FString TheString = *FString::Printf(TEXT("%4i.    %6.2f  CPU:%5.2f%% RRC:%3.2f %s   Owner: %s   SoundClass: %s"),
 				WaveInstanceInfo.Key->InstanceIndex,
 				WaveInstanceInfo.Key->Volume,
 				CPUPercent,
+				RelativeRenderCost,
 				*WaveInstanceInfo.Key->WaveInstanceName.ToString(),
 				SoundOwner ? *SoundOwner->GetName() : TEXT("None"),
 				*WaveInstanceInfo.Key->SoundClassName.ToString());
diff --git a/Engine/Source/Runtime/Engine/Private/AudioDevice.cpp b/Engine/Source/Runtime/Engine/Private/AudioDevice.cpp
--- a/Engine/Source/Runtime/Engine/Private/AudioDevice.cpp
+++ b/Engine/Source/Runtime/Engine/Private/AudioDevice.cpp
@@ -223,6 +223,14 @@ static FAutoConsoleCommandWithWorldArgsAndOutputDevice GSetCurrentSpatialPluginC
 		})
 	);
 
+
+static int32 EnableRelativeRenderCostVoiceLimitCVar = 0;
+FAutoConsoleVariableRef CVarEnableRelativeRenderCostVoiceLimit(
+	TEXT("au.EnableRelativeRenderCostVoiceLimit"),
+	EnableRelativeRenderCostVoiceLimitCVar,
+	TEXT("Enables or disables using the relative render cost of rendering sources to count toward max channels culling."),
+	ECVF_Default);
+
 #if UE_AUDIO_PROFILERTRACE_ENABLED
 UE_TRACE_EVENT_BEGIN(Audio, VirtualLoopStop)
 	UE_TRACE_EVENT_FIELD(uint32, DeviceId)
@@ -4198,7 +4206,7 @@ int32 FAudioDevice::GetSortedActiveWaveInstances(TArray<FWaveInstance*>& WaveIns
 
 	int32 FirstActiveIndex = 0;
 	// Only need to do the wave instance sort if we have any waves and if our wave instances are greater than our max channels.
-	if (WaveInstances.Num() >= 0)
+	if (WaveInstances.Num() > 0)
 	{
 		// Helper function for "Sort" (higher priority sorts last).
 		struct FCompareFWaveInstanceByPlayPriority
@@ -4214,7 +4222,23 @@ int32 FAudioDevice::GetSortedActiveWaveInstances(TArray<FWaveInstance*>& WaveIns
 
 		// Get the first index that will result in a active source voice
 		int32 CurrentMaxChannels = GetMaxChannels();
-		FirstActiveIndex = FMath::Max(WaveInstances.Num() - CurrentMaxChannels, 0);
+
+		if (EnableRelativeRenderCostVoiceLimitCVar)
+		{
+			// Find the first active index based on the relative render cost estimate of the wave instances
+			FirstActiveIndex = WaveInstances.Num() - 1;
+			float NumActiveChannels = 0.0f;
+			while (NumActiveChannels < (float)CurrentMaxChannels && FirstActiveIndex > 0)
+			{
+				FWaveInstance* WaveInstance = WaveInstances[FirstActiveIndex];
+				NumActiveChannels += WaveInstance->GetRelativeRenderCost();
+				FirstActiveIndex--;
+			}
+		}
+		else
+		{
+			FirstActiveIndex = FMath::Max(WaveInstances.Num() - CurrentMaxChannels, 0);
+		}
 	}
 
 	return FirstActiveIndex;
diff --git a/Engine/Source/Runtime/Engine/Public/Audio.h b/Engine/Source/Runtime/Engine/Public/Audio.h
--- a/Engine/Source/Runtime/Engine/Public/Audio.h
+++ b/Engine/Source/Runtime/Engine/Public/Audio.h
@@ -235,6 +235,9 @@ private:
 	/** The current envelope value of the wave instance. */
 	float EnvelopValue;
 
+	/** The estimated relative render cost of the wave instance. 1.0 is cost of a single decoding sound source. Used for limited the overall voice count. */
+	float RelativeRenderCost;
+
 public:
 	/** The envelope follower attack time in milliseconds. */
 	int32 EnvelopeFollowerAttackTime;
@@ -512,6 +515,12 @@ public:
 	uint32 GetPlayOrder() const { return PlayOrder; }
 
 	friend inline uint32 GetTypeHash(FWaveInstance* A) { return A->PlayOrder; }
+
+	/** Sets the relative render cost of the wave instance. */
+	void SetRelativeRenderCost(float InRelativeRenderCost) { RelativeRenderCost = InRelativeRenderCost; }
+
+	/** Retrieves the relative render cost of wave instance. */
+	float GetRelativeRenderCost() const { return RelativeRenderCost; }
 };
 
 /*-----------------------------------------------------------------------------
@@ -737,6 +746,9 @@ public:
 	/** Returns the source's envelope at the callback block rate. Only implemented in audio mixer. */
 	virtual float GetEnvelopeValue() const { return 0.0f; };
 
+	/** Returns the source's estimated relative render cost (relative to a single decoded sound). Used for debug information and to constrain overall CPU usage. */
+	virtual float GetRelativeRenderCost() const { return 1.0f; }
+
 	ENGINE_API void GetChannelLocations(FVector& Left, FVector&Right) const;
 
 	void NotifyPlaybackData();
@@ -860,6 +872,9 @@ public:
 		/** Fraction of a single CPU core used to render audio. */
 		double CPUCoreUtilization = 0;
 
+		/** Relative cost to render wave. */
+		float RelativeRenderCost = 1.f;
+
 		/** Basic CS so we can pass this around safely. */
 		FCriticalSection CS;
 	};
