UE-170940 - [Workflow Improvement] - MetaSound connections should autoconvert types if possible

diff --git a/Engine/Plugins/Runtime/Metasound/Source/MetasoundEditor/Private/MetasoundEditorGraphSchema.cpp b/Engine/Plugins/Runtime/Metasound/Source/MetasoundEditor/Private/MetasoundEditorGraphSchema.cpp
--- a/Engine/Plugins/Runtime/Metasound/Source/MetasoundEditor/Private/MetasoundEditorGraphSchema.cpp
+++ b/Engine/Plugins/Runtime/Metasound/Source/MetasoundEditor/Private/MetasoundEditorGraphSchema.cpp
@@ -1382,11 +1382,6 @@ const FPinConnectionResponse UMetasoundEditorGraphSchema::CanCreateConnection(co
 		return FPinConnectionResponse(CONNECT_RESPONSE_DISALLOW, LOCTEXT("ConnectionLoop2", "Connection causes loop"));
 	}
 
-	if (InputPin->PinType.PinCategory != OutputPin->PinType.PinCategory)
-	{
-		return FPinConnectionResponse(CONNECT_RESPONSE_DISALLOW, LOCTEXT("ConnectionTypeIncorrect", "Connection pin types do not match"));
-	}
-
 	bool bConnectingNodesWithErrors = false;
 	UEdGraphNode* InputNode = InputPin->GetOwningNode();
 	if (ensure(InputNode))
@@ -1412,11 +1407,10 @@ const FPinConnectionResponse UMetasoundEditorGraphSchema::CanCreateConnection(co
 	const bool bOutputValid = OutputHandle->IsValid();
 	if (bInputValid && bOutputValid)
 	{
-		// TODO: Implement YesWithConverterNode to provide conversion options
 		Frontend::FConnectability Connectability = InputHandle->CanConnectTo(*OutputHandle);
-		if (Connectability.Connectable != Frontend::FConnectability::EConnectable::Yes)
+		if (Connectability.Connectable == Frontend::FConnectability::EConnectable::No)
 		{
-			if ((Frontend::FConnectability::EReason::IncompatibleDataTypes == Connectability.Reason) || (Connectability.Connectable == Frontend::FConnectability::EConnectable::YesWithConverterNode))
+			if (Frontend::FConnectability::EReason::IncompatibleDataTypes == Connectability.Reason)
 			{
 				const FName InputType = InputHandle->GetDataType();
 				const FName OutputType = OutputHandle->GetDataType();
@@ -1432,11 +1426,34 @@ const FPinConnectionResponse UMetasoundEditorGraphSchema::CanCreateConnection(co
 			}
 			else if (Frontend::FConnectability::EReason::IncompatibleAccessTypes == Connectability.Reason)
 			{
-				return FPinConnectionResponse(CONNECT_RESPONSE_DISALLOW, LOCTEXT("ConnectionIncompatibleAccessTypes", "Cannot create connection between incompatible access types. Constructor input pins can only be connected to constructor output pins."));
+				return FPinConnectionResponse(CONNECT_RESPONSE_DISALLOW, LOCTEXT("ConnectionIncompatibleAccessTypes0", "Cannot create connection between incompatible access types. Constructor input pins can only be connected to constructor output pins."));
 			}
 			else
 			{
-				return FPinConnectionResponse(CONNECT_RESPONSE_DISALLOW, LOCTEXT("ConnectionNotAllowed", "Output pin cannot be connected to input pin."));
+				const FName InputType = InputHandle->GetDataType();
+				const FName OutputType = OutputHandle->GetDataType();
+				return FPinConnectionResponse(CONNECT_RESPONSE_DISALLOW, FText::Format(
+					LOCTEXT("ConnectionNotAllowed", "'{0}' is not compatible with '{1}'"),
+					FText::FromName(OutputType),
+					FText::FromName(InputType)
+				));
+			}
+		}
+		else if (Connectability.Connectable == Frontend::FConnectability::EConnectable::YesWithConverterNode)
+		{
+			if (Connectability.PossibleConverterNodeClasses.Num() == 0)
+			{
+				return FPinConnectionResponse(CONNECT_RESPONSE_DISALLOW, LOCTEXT("ConnectionIncompatibleAccessTypes1", "Conversion not supported between these types."));
+			}
+			else
+			{
+				const FName InputType = InputHandle->GetDataType();
+				const FName OutputType = OutputHandle->GetDataType();
+				return FPinConnectionResponse(CONNECT_RESPONSE_MAKE_WITH_CONVERSION_NODE, 
+					FText::Format(LOCTEXT("ConversionSuccess", "Convert {0} to {1}."), 
+					FText::FromName(OutputType),
+					FText::FromName(InputType)
+				));
 			}
 		}
 
@@ -1546,7 +1563,45 @@ bool UMetasoundEditorGraphSchema::TryCreateConnection(UEdGraphPin* PinA, UEdGrap
 		return false;
 	}
 
-	// TODO: Implement YesWithConverterNode with selected conversion option
+	FConstInputHandle InputHandle = FGraphBuilder::GetConstInputHandleFromPin(InputPin);
+	FConstOutputHandle OutputHandle = FGraphBuilder::GetConstOutputHandleFromPin(OutputPin);
+	FConnectability Connectability = InputHandle->CanConnectTo(*OutputHandle);
+	if (Connectability.Connectable == FConnectability::EConnectable::YesWithConverterNode)
+	{
+		UMetasoundEditorGraph* MetaSoundGraph = CastChecked<UMetasoundEditorGraph>(InputPin->GetOwningNode()->GetGraph());
+		UObject& ParentMetaSound = MetaSoundGraph->GetMetasoundChecked();
+
+		if (Connectability.PossibleConverterNodeClasses.Num() == 0)
+		{
+			return false;
+		}
+		FNodeRegistryKey NodeKey = Connectability.PossibleConverterNodeClasses.Last().NodeKey;
+
+		FMetasoundFrontendClassMetadata Metadata;
+		Metadata.SetClassName(NodeKey.ClassName);
+		Metadata.SetType(NodeKey.Type);
+		
+		if (UMetasoundEditorGraphExternalNode* NewGraphNode = FGraphBuilder::AddExternalNode(ParentMetaSound, Metadata, false))
+		{
+			UEdGraphNode* InputNode = InputPin->GetOwningNode();
+			UEdGraphNode* OutputNode = OutputPin->GetOwningNode();
+			
+			FVector2D Location = FVector2D();
+			Location += FVector2D(OutputNode->NodePosX, OutputNode->NodePosY);
+			Location.Y += 40.f;
+
+			NewGraphNode->Modify();
+			NewGraphNode->UpdateFrontendNodeLocation(Location);
+			NewGraphNode->SyncLocationFromFrontendNode();
+
+			SchemaPrivate::TryConnectNewNodeToMatchingDataTypePin(*NewGraphNode, InputPin);
+			SchemaPrivate::TryConnectNewNodeToMatchingDataTypePin(*NewGraphNode, OutputPin);
+
+			return true;
+		}
+
+		return false;
+	}
 
 	// Must mark Metasound object as modified to avoid desync issues ***before*** attempting to create a connection
 	// so that transaction stack observes Frontend changes last if rolled back (i.e. undone).  UEdGraphSchema::TryCreateConnection
