Metasound changes to reduce the number of FNames generated.
Previously the FNodeRegistryKey::ToString() function was the only user of the "FullName" FName and was used for logging. We now build these inline when used as not to pollute the FNamePool.

diff --git a/Engine/Plugins/Runtime/Metasound/Source/MetasoundFrontend/Private/MetasoundFrontendDocument.cpp b/Engine/Plugins/Runtime/Metasound/Source/MetasoundFrontend/Private/MetasoundFrontendDocument.cpp
--- a/Engine/Plugins/Runtime/Metasound/Source/MetasoundFrontend/Private/MetasoundFrontendDocument.cpp
+++ b/Engine/Plugins/Runtime/Metasound/Source/MetasoundFrontend/Private/MetasoundFrontendDocument.cpp
@@ -437,7 +437,14 @@ Metasound::FNodeClassName FMetasoundFrontendClassName::ToNodeClassName() const
 
 FString FMetasoundFrontendClassName::ToString() const
 {
-	return GetFullName().ToString();
+	FNameBuilder NameBuilder;
+	ToString(NameBuilder);
+	return *NameBuilder;
+}
+
+void FMetasoundFrontendClassName::ToString(FNameBuilder& NameBuilder) const
+{
+	Metasound::FNodeClassName::FormatFullName(NameBuilder, Namespace, Name, Variant);
 }
 
 bool FMetasoundFrontendClassName::Parse(const FString& InClassName, FMetasoundFrontendClassName& OutClassName)
diff --git a/Engine/Plugins/Runtime/Metasound/Source/MetasoundFrontend/Private/MetasoundFrontendDocumentController.cpp b/Engine/Plugins/Runtime/Metasound/Source/MetasoundFrontend/Private/MetasoundFrontendDocumentController.cpp
--- a/Engine/Plugins/Runtime/Metasound/Source/MetasoundFrontend/Private/MetasoundFrontendDocumentController.cpp
+++ b/Engine/Plugins/Runtime/Metasound/Source/MetasoundFrontend/Private/MetasoundFrontendDocumentController.cpp
@@ -383,7 +383,7 @@ namespace Metasound
 								UE_LOG(LogMetaSound, Error,
 									TEXT("Cannot add external dependency. No Metasound class found with matching registry key [Key:%s, Name:%s, Version:%s]. Suggested solution \"%s\" by %s."),
 									*Key.ToString(),
-									*InMetadata.GetClassName().GetFullName().ToString(),
+									*InMetadata.GetClassName().ToString(),
 									*InMetadata.GetVersion().ToString(),
 									*InMetadata.GetPromptIfMissing().ToString(),
 									*InMetadata.GetAuthor());
@@ -391,7 +391,7 @@ namespace Metasound
 								UE_LOG(LogMetaSound, Error,
 									TEXT("Cannot add external dependency. No Metasound class found with matching registry key [Key:%s, Name:%s, Version:%s]."),
 									*Key.ToString(),
-									*InMetadata.GetClassName().GetFullName().ToString(),
+									*InMetadata.GetClassName().ToString(),
 									*InMetadata.GetVersion().ToString());
 #endif // !WITH_EDITOR
 							}
@@ -412,7 +412,7 @@ namespace Metasound
 						{
 							UE_LOG(LogMetaSound, Error, TEXT(
 								"Unsupported metasound class type for node: \"%s\" (%s)."),
-								*InMetadata.GetClassName().GetFullName().ToString(),
+								*InMetadata.GetClassName().ToString(),
 								*InMetadata.GetVersion().ToString());
 							checkNoEntry();
 						}
diff --git a/Engine/Plugins/Runtime/Metasound/Source/MetasoundFrontend/Private/MetasoundFrontendGraphController.cpp b/Engine/Plugins/Runtime/Metasound/Source/MetasoundFrontend/Private/MetasoundFrontendGraphController.cpp
--- a/Engine/Plugins/Runtime/Metasound/Source/MetasoundFrontend/Private/MetasoundFrontendGraphController.cpp
+++ b/Engine/Plugins/Runtime/Metasound/Source/MetasoundFrontend/Private/MetasoundFrontendGraphController.cpp
@@ -1599,7 +1599,7 @@ namespace Metasound
 			{
 				if (const FMetasoundFrontendClass* ExistingDependency = OwningDocument->FindClass(InMetadata).Get())
 				{
-					UE_LOG(LogMetaSound, Error, TEXT("Cannot add new subgraph. Metasound class already exists with matching metadata Name: \"%s\", Version %d.%d"), *(ExistingDependency->Metadata.GetClassName().GetFullName().ToString()), ExistingDependency->Metadata.GetVersion().Major, ExistingDependency->Metadata.GetVersion().Minor);
+					UE_LOG(LogMetaSound, Error, TEXT("Cannot add new subgraph. Metasound class already exists with matching metadata Name: \"%s\", Version %d.%d"), *(ExistingDependency->Metadata.GetClassName().ToString()), ExistingDependency->Metadata.GetVersion().Major, ExistingDependency->Metadata.GetVersion().Minor);
 				}
 				else 
 				{
diff --git a/Engine/Plugins/Runtime/Metasound/Source/MetasoundFrontend/Private/MetasoundFrontendRegistryKey.cpp b/Engine/Plugins/Runtime/Metasound/Source/MetasoundFrontend/Private/MetasoundFrontendRegistryKey.cpp
--- a/Engine/Plugins/Runtime/Metasound/Source/MetasoundFrontend/Private/MetasoundFrontendRegistryKey.cpp
+++ b/Engine/Plugins/Runtime/Metasound/Source/MetasoundFrontend/Private/MetasoundFrontendRegistryKey.cpp
@@ -138,7 +138,7 @@ namespace Metasound::Frontend
 		TStringBuilder<128> KeyStringBuilder;
 		KeyStringBuilder.Append(LexToString(Type));
 		KeyStringBuilder.AppendChar('_');
-		KeyStringBuilder.Append(ClassName.GetFullName().ToString());
+		KeyStringBuilder.Append(ClassName.ToString());
 		KeyStringBuilder.AppendChar('_');
 		KeyStringBuilder.Append(FString::FromInt(Version.Major));
 		KeyStringBuilder.AppendChar('.');
diff --git a/Engine/Plugins/Runtime/Metasound/Source/MetasoundFrontend/Public/MetasoundFrontendDocument.h b/Engine/Plugins/Runtime/Metasound/Source/MetasoundFrontend/Public/MetasoundFrontendDocument.h
--- a/Engine/Plugins/Runtime/Metasound/Source/MetasoundFrontend/Public/MetasoundFrontendDocument.h
+++ b/Engine/Plugins/Runtime/Metasound/Source/MetasoundFrontend/Public/MetasoundFrontendDocument.h
@@ -1310,6 +1310,9 @@ struct METASOUNDFRONTEND_API FMetasoundFrontendClassName
 	// Return string version of full name.
 	FString ToString() const;
 
+	// Return a string into an existing FNameBuilder
+	void ToString(FNameBuilder& NameBuilder) const;
+
 	// Parses string into class name.  For deserialization and debug use only.
 	static bool Parse(const FString& InClassName, FMetasoundFrontendClassName& OutClassName);
 
diff --git a/Engine/Plugins/Runtime/Metasound/Source/MetasoundGraphCore/Private/MetasoundDynamicOperatorTransactor.cpp b/Engine/Plugins/Runtime/Metasound/Source/MetasoundGraphCore/Private/MetasoundDynamicOperatorTransactor.cpp
--- a/Engine/Plugins/Runtime/Metasound/Source/MetasoundGraphCore/Private/MetasoundDynamicOperatorTransactor.cpp
+++ b/Engine/Plugins/Runtime/Metasound/Source/MetasoundGraphCore/Private/MetasoundDynamicOperatorTransactor.cpp
@@ -44,7 +44,7 @@ namespace Metasound
 			
 			FString GetDebugNodeNameString(const INode& InNode)
 			{
-				return FString::Printf(TEXT("%s_v%d.%d"), *InNode.GetMetadata().ClassName.GetFullName().ToString(), InNode.GetMetadata().MajorVersion, InNode.GetMetadata().MinorVersion);
+				return FString::Printf(TEXT("%s_v%d.%d"), *InNode.GetMetadata().ClassName.ToString(), InNode.GetMetadata().MajorVersion, InNode.GetMetadata().MinorVersion);
 			}
 
 			FString GetDebugNodeNameString(const FGuid& InNodeID, const INode& InNode)
diff --git a/Engine/Plugins/Runtime/Metasound/Source/MetasoundGraphCore/Private/MetasoundNodeInterface.cpp b/Engine/Plugins/Runtime/Metasound/Source/MetasoundGraphCore/Private/MetasoundNodeInterface.cpp
--- a/Engine/Plugins/Runtime/Metasound/Source/MetasoundGraphCore/Private/MetasoundNodeInterface.cpp
+++ b/Engine/Plugins/Runtime/Metasound/Source/MetasoundGraphCore/Private/MetasoundNodeInterface.cpp
@@ -22,8 +22,6 @@ namespace Metasound
 	, Name(InName)
 	, Variant(InVariant)
 	{
-		ScopedName = FormatScopedName(Namespace, Name);
-		FullName = FormatFullName(Namespace, Name, Variant);
 	}
 
 	/** Namespace of node class. */
@@ -44,35 +42,56 @@ namespace Metasound
 		return Variant;
 	}
 
+	/** The full name of the Node formatted Namespace.Name[.Variant] */
+	const FString FNodeClassName::ToString() const
+	{
+		FNameBuilder Builder;
+		FormatFullName(Builder, Namespace, Name, Variant);
+		return *Builder;
+	}
+
 	/** Namespace and name of the node class. */
-	const FName& FNodeClassName::GetScopedName() const
+	const FName FNodeClassName::GetScopedName() const
 	{
-		return ScopedName;
+		return FormatScopedName(Namespace, Name);
 	}
 
 	/** Namespace, name and variant of the node class. */
-	const FName& FNodeClassName::GetFullName() const
+	const FName FNodeClassName::GetFullName() const
 	{
-		return FullName;
+		return FormatFullName(Namespace, Name, Variant);
 	}
 
 	FName FNodeClassName::FormatFullName(const FName& InNamespace, const FName& InName, const FName& InVariant)
 	{
-		FName ScopedName = FormatScopedName(InNamespace, InName);
+		FNameBuilder Builder;
+		FormatFullName(Builder, InNamespace, InName, InVariant);
+		return *Builder;
+	}
+
+	FName FNodeClassName::FormatScopedName(const FName& InNamespace, const FName& InName)
+	{
+		FNameBuilder Builder;
+		FormatScopedName(Builder, InNamespace, InName);
+		return *Builder;
+	}
+
+	void FNodeClassName::FormatFullName(FNameBuilder& InBuilder, const FName& InNamespace, const FName& InName, const FName& InVariant)
+	{
+		FormatScopedName(InBuilder, InNamespace, InName);
 
 		if (InVariant != NAME_None)
 		{
-			return FName(*FString::Printf(TEXT("%s.%s"), *ScopedName.ToString(), *InVariant.ToString()));
-		}
-		else
-		{
-			return ScopedName;
+			InBuilder.Append(".");
+			InVariant.AppendString(InBuilder);
 		}
 	}
 
-	FName FNodeClassName::FormatScopedName(const FName& InNamespace, const FName& InName)
+	void FNodeClassName::FormatScopedName(FNameBuilder& InBuilder, const FName& InNamespace, const FName& InName)
 	{
-		return FName(*FString::Printf(TEXT("%s.%s"), *InNamespace.ToString(), *InName.ToString()));
+		InNamespace.AppendString(InBuilder);
+		InBuilder.Append(".");
+		InName.AppendString(InBuilder);
 	}
 
 	bool FNodeClassName::IsValid() const
diff --git a/Engine/Plugins/Runtime/Metasound/Source/MetasoundGraphCore/Private/MetasoundOperatorBuilder.cpp b/Engine/Plugins/Runtime/Metasound/Source/MetasoundGraphCore/Private/MetasoundOperatorBuilder.cpp
--- a/Engine/Plugins/Runtime/Metasound/Source/MetasoundGraphCore/Private/MetasoundOperatorBuilder.cpp
+++ b/Engine/Plugins/Runtime/Metasound/Source/MetasoundGraphCore/Private/MetasoundOperatorBuilder.cpp
@@ -518,7 +518,7 @@ namespace Metasound
 			{
 				bFoundUnboundVertex = true;
 				const FNodeClassMetadata& Metadata = InNode.GetMetadata();
-				UE_LOG(LogMetaSound, Warning, TEXT("Operator for node %s v%d.%d contains unbound output vertex %s"), *Metadata.ClassName.GetFullName().ToString(), Metadata.MajorVersion, Metadata.MinorVersion, *OutputBinding.GetVertex().VertexName.ToString());
+				UE_LOG(LogMetaSound, Warning, TEXT("Operator for node %s v%d.%d contains unbound output vertex %s"), *Metadata.ClassName.ToString(), Metadata.MajorVersion, Metadata.MinorVersion, *OutputBinding.GetVertex().VertexName.ToString());
 			}
 		}
 
@@ -557,7 +557,7 @@ namespace Metasound
 				// Use node class name if valid, otherwise (for example graph nodes) use instance name 
 				TStringBuilder<256> TraceNamePtr;
 				const FNodeClassName& NodeClassName = Node->GetMetadata().ClassName;
-				const FName& NodeTraceName = NodeClassName.IsValid() ? NodeClassName.GetFullName() : Node->GetInstanceName();
+				const FString NodeTraceName = NodeClassName.IsValid() ? NodeClassName.ToString() : Node->GetInstanceName().ToString();
 				TraceNamePtr << "Metasound::FOperatorBuilder::CreateOperators::CreateAndBind " << NodeTraceName;
 				METASOUND_TRACE_CPUPROFILER_EVENT_SCOPE_TEXT(*TraceNamePtr);
 #endif // METASOUND_CPUPROFILERTRACE_ENABLED
diff --git a/Engine/Plugins/Runtime/Metasound/Source/MetasoundGraphCore/Private/MetasoundThreadLocalDebug.cpp b/Engine/Plugins/Runtime/Metasound/Source/MetasoundGraphCore/Private/MetasoundThreadLocalDebug.cpp
--- a/Engine/Plugins/Runtime/Metasound/Source/MetasoundGraphCore/Private/MetasoundThreadLocalDebug.cpp
+++ b/Engine/Plugins/Runtime/Metasound/Source/MetasoundGraphCore/Private/MetasoundThreadLocalDebug.cpp
@@ -41,7 +41,7 @@ namespace Metasound
 			{
 				const FNodeClassMetadata& Metadata = ActiveNode->GetMetadata();
 
-				NodeClassNameAndVersion = FString::Format(TEXT("{0} v{1}.{2}"), {Metadata.ClassName.GetFullName().ToString(), Metadata.MajorVersion, Metadata.MinorVersion});
+				NodeClassNameAndVersion = FString::Format(TEXT("{0} v{1}.{2}"), {Metadata.ClassName.ToString(), Metadata.MajorVersion, Metadata.MinorVersion});
 			}
 			else
 			{
diff --git a/Engine/Plugins/Runtime/Metasound/Source/MetasoundGraphCore/Public/MetasoundNodeInterface.h b/Engine/Plugins/Runtime/Metasound/Source/MetasoundGraphCore/Public/MetasoundNodeInterface.h
--- a/Engine/Plugins/Runtime/Metasound/Source/MetasoundGraphCore/Public/MetasoundNodeInterface.h
+++ b/Engine/Plugins/Runtime/Metasound/Source/MetasoundGraphCore/Public/MetasoundNodeInterface.h
@@ -63,19 +63,26 @@ namespace Metasound
 		/** Variant of node class. */
 		const FName& GetVariant() const;
 
+		/** The full name of the Node formatted Namespace.Name and optionally appended with .Variant if valid */
+		const FString ToString() const;
+
 		/** Namespace and name of the node class. */
-		const FName& GetScopedName() const;
+		UE_DEPRECATED(5.4, "Use FormatScopedName functions to generate")
+		const FName GetScopedName() const;
 
 		/** Namespace, name and variant of the node class. */
-		const FName& GetFullName() const;
+		UE_DEPRECATED(5.4, "Use FormatFullName functions to generate")
+		const FName GetFullName() const;
 
 		/** Whether or not this instance of a node class name is a valid name.*/
 		bool IsValid() const;
 
 		static FName FormatFullName(const FName& InNamespace, const FName& InName, const FName& InVariant);
-
 		static FName FormatScopedName(const FName& InNamespace, const FName& InName);
 
+		static void FormatFullName(FNameBuilder& InBuilder, const FName& InNamespace, const FName& InName, const FName& InVariant);
+		static void FormatScopedName(FNameBuilder& InBuilder, const FName& InNamespace, const FName& InName);
+
 		/** Invalid form of node class name (i.e. empty namespace, name, and variant)*/
 		static const FNodeClassName InvalidNodeClassName;
 
@@ -89,8 +96,6 @@ namespace Metasound
 		FName Namespace;
 		FName Name;
 		FName Variant;
-		FName ScopedName;
-		FName FullName;
 	};
 
 	/** Provides metadata for a given node. */
