Fix for null ref crash when adding referenced MetaSound graph in MetaSound editor
#rnx

diff --git a/Engine/Plugins/Runtime/Metasound/Source/MetasoundEditor/Private/MetasoundEditorGraphNode.cpp b/Engine/Plugins/Runtime/Metasound/Source/MetasoundEditor/Private/MetasoundEditorGraphNode.cpp
--- a/Engine/Plugins/Runtime/Metasound/Source/MetasoundEditor/Private/MetasoundEditorGraphNode.cpp
+++ b/Engine/Plugins/Runtime/Metasound/Source/MetasoundEditor/Private/MetasoundEditorGraphNode.cpp
@@ -1102,10 +1102,11 @@ void UMetasoundEditorGraphExternalNode::Validate(Metasound::Editor::FGraphNodeVa
 			const FAssetKey AssetKey(Metadata);
 			if (FMetasoundAssetBase* MetaSoundAsset = IMetaSoundAssetManager::GetChecked().TryLoadAssetFromKey(AssetKey))
 			{
-				const UMetasoundEditorGraph* NodeGraph = CastChecked<UMetasoundEditorGraph>(&MetaSoundAsset->GetGraphChecked());
-				const EMessageSeverity::Type MaxGraphMsg = static_cast<EMessageSeverity::Type>(NodeGraph->GetHighestMessageSeverity());
-				switch (MaxGraphMsg)
+				if (const UMetasoundEditorGraph* NodeGraph = Cast<UMetasoundEditorGraph>(MetaSoundAsset->GetGraph()))
 				{
+					const EMessageSeverity::Type MaxGraphMsg = static_cast<EMessageSeverity::Type>(NodeGraph->GetHighestMessageSeverity());
+					switch (MaxGraphMsg)
+					{
 					case EMessageSeverity::Error:
 					{
 						OutResult.SetMessage(MaxGraphMsg, TEXT("Referenced asset class contains error(s). Check implementation for details."));
@@ -1124,6 +1125,7 @@ void UMetasoundEditorGraphExternalNode::Validate(Metasound::Editor::FGraphNodeVa
 					{
 					}
 					break;
+					}
 				}
 			}
 			break;
