Fix for spamming logs when project has MetaSound assets with duplicate guid keys
- Moved warning to when asset is added to registry rather than every query
- Minor clean-up in asset registry
#rnx

diff --git a/Engine/Plugins/Runtime/Metasound/Source/MetasoundEngine/Private/MetasoundDocumentBuilderRegistry.cpp b/Engine/Plugins/Runtime/Metasound/Source/MetasoundEngine/Private/MetasoundDocumentBuilderRegistry.cpp
--- a/Engine/Plugins/Runtime/Metasound/Source/MetasoundEngine/Private/MetasoundDocumentBuilderRegistry.cpp
+++ b/Engine/Plugins/Runtime/Metasound/Source/MetasoundEngine/Private/MetasoundDocumentBuilderRegistry.cpp
@@ -10,17 +10,23 @@ namespace Metasound::Engine
 {
 	FDocumentBuilderRegistry::~FDocumentBuilderRegistry()
 	{
+		TMultiMap<FMetasoundFrontendClassName, TWeakObjectPtr<UMetaSoundBuilderBase>> BuildersToFinish;
 		FScopeLock Lock(&BuildersCriticalSection);
-		UE_CLOG(!Builders.IsEmpty(), LogMetaSound, Display, TEXT("BuilderRegistry is shutting down with the following %i active builder entries. Forcefully shutting down:"), Builders.Num());
+		{
+			BuildersToFinish = MoveTemp(Builders);
+			Builders.Reset();
+		}
+
+		UE_CLOG(!BuildersToFinish.IsEmpty(), LogMetaSound, Display, TEXT("BuilderRegistry is shutting down with the following %i active builder entries. Forcefully shutting down:"), BuildersToFinish.Num());
 		int32 NumStale = 0;
-		for (const TPair<FMetasoundFrontendClassName, TWeakObjectPtr<UMetaSoundBuilderBase>>& Pair : Builders)
+		for (const TPair<FMetasoundFrontendClassName, TWeakObjectPtr<UMetaSoundBuilderBase>>& Pair : BuildersToFinish)
 		{
 			const bool bIsValid = Pair.Value.IsValid();
 			if (bIsValid)
 			{
 				UE_CLOG(bIsValid, LogMetaSound, Display, TEXT("- %s"), *Pair.Value->GetFullName());
 				constexpr bool bForceUnregister = true;
-				FinishBuilding(Pair.Key, bForceUnregister);
+				FinishBuildingInternal(Pair.Key, bForceUnregister);
 			}
 			else
 			{
@@ -28,7 +34,41 @@ namespace Metasound::Engine
 			}
 		}
 		UE_CLOG(NumStale > 0, LogMetaSound, Display, TEXT("BuilderRegistry is shutting down with %i stale entries"), NumStale);
-		Builders.Reset();
+	}
+
+	void FDocumentBuilderRegistry::AddBuilderInternal(const FMetasoundFrontendClassName& InClassName, UMetaSoundBuilderBase* NewBuilder) const
+	{
+		FScopeLock Lock(&BuildersCriticalSection);
+
+#if !NO_LOGGING
+		const bool bPrintConflicts = Builders.Contains(InClassName);
+#endif // !NO_LOGGING
+
+		Builders.Add(InClassName, NewBuilder);
+
+#if !NO_LOGGING
+		if (bPrintConflicts)
+		{
+			TArray<TWeakObjectPtr<UMetaSoundBuilderBase>> Entries;
+			Builders.MultiFind(InClassName, Entries);
+			UE_LOG(LogMetaSound, Error, TEXT("More than one asset registered with class name '%s'. "
+				"Returning builder that may not be associated with desired object! \n"
+				"This can happen if asset was moved using revision control and original location was revived. \n"
+				"Remove all but one of the following assets and relink a duplicate or copied replacement asset:"),
+				*InClassName.ToString());
+			for (const TWeakObjectPtr<UMetaSoundBuilderBase>& BuilderPtr : Entries)
+			{
+				if (BuilderPtr.IsValid())
+				{
+					UE_LOG(LogMetaSound, Error, TEXT("- %s"), *BuilderPtr->GetConstBuilder().CastDocumentObjectChecked<UObject>().GetPathName());
+				}
+				else
+				{
+					UE_LOG(LogMetaSound, Error, TEXT("- STALE ENTRY (Not removed prior to asset removal)"));
+				}
+			}
+		}
+#endif // !NO_LOGGING
 	}
 
 #if WITH_EDITORONLY_DATA
@@ -96,47 +136,29 @@ namespace Metasound::Engine
 		return nullptr;
 	}
 
-	UMetaSoundBuilderBase* FDocumentBuilderRegistry::FindBuilderObject(const FMetasoundFrontendClassName& ClassName) const
+	UMetaSoundBuilderBase* FDocumentBuilderRegistry::FindBuilderObject(const FMetasoundFrontendClassName& InClassName) const
 	{
 		TArray<TWeakObjectPtr<UMetaSoundBuilderBase>> Entries;
 		{
 			FScopeLock Lock(&BuildersCriticalSection);
-			Builders.MultiFind(ClassName, Entries);
+			Builders.MultiFind(InClassName, Entries);
 		}
 		if (!Entries.IsEmpty())
 		{
-#if !NO_LOGGING
-			if (Entries.Num() > 1)
-			{
-				UE_LOG(LogMetaSound, Error, TEXT("More than one asset registered with class name '%s': returning builder that may not be associated with desired object! This can happen if multiple assets were improperly copied. See following list of assets needing to be resolved:"),
-					*ClassName.ToString());
-				for (const TWeakObjectPtr<UMetaSoundBuilderBase>& BuilderPtr : Entries)
-				{
-					if (BuilderPtr.IsValid())
-					{
-						UE_LOG(LogMetaSound, Error, TEXT("- %s"), *BuilderPtr->GetConstBuilder().CastDocumentObjectChecked<UObject>().GetPathName());
-					}
-					else
-					{
-						UE_LOG(LogMetaSound, Error, TEXT("- STALE"));
-					}
-				}
-			}
-#endif // !NO_LOGGING
 			return Entries.Last().Get();
 		}
 
 		return nullptr;
 	}
 
-	TArray<UMetaSoundBuilderBase*> FDocumentBuilderRegistry::FindBuilderObjects(const FMetasoundFrontendClassName& ClassName) const
+	TArray<UMetaSoundBuilderBase*> FDocumentBuilderRegistry::FindBuilderObjects(const FMetasoundFrontendClassName& InClassName) const
 	{
 		TArray<UMetaSoundBuilderBase*> FoundBuilders;
 		TArray<TWeakObjectPtr<UMetaSoundBuilderBase>> Entries;
 
 		{
 			FScopeLock Lock(&BuildersCriticalSection);
-			Builders.MultiFind(ClassName, Entries);
+			Builders.MultiFind(InClassName, Entries);
 		}
 
 		if (!Entries.IsEmpty())
@@ -163,6 +185,18 @@ namespace Metasound::Engine
 		using namespace Metasound;
 		using namespace Metasound::Engine;
 
+		FinishBuildingInternal(InClassName, bForceUnregisterNodeClass);
+
+		// Still return true in this case as the builder likely has become inaccessible and may be in the "beginning" of destruction,
+		// so entries is still reporting that it was successfully removed.
+		{
+			FScopeLock Lock(&BuildersCriticalSection);
+			return Builders.Remove(InClassName) > 0;
+		}
+	}
+
+	void FDocumentBuilderRegistry::FinishBuildingInternal(const FMetasoundFrontendClassName& InClassName, bool bForceUnregisterNodeClass) const
+	{
 		TArray<UMetaSoundBuilderBase*> FoundBuilders = FindBuilderObjects(InClassName);
 		for (UMetaSoundBuilderBase* Builder : FoundBuilders)
 		{
@@ -190,13 +224,6 @@ namespace Metasound::Engine
 				DocBuilder.FinishBuilding();
 			}
 		}
-
-		// Still return true in this case as the builder likely has become inaccessible and may be in the "beginning" of destruction,
-		// so entries is still reporting that it was successfully removed.
-		{
-			FScopeLock Lock(&BuildersCriticalSection);
-			return Builders.Remove(InClassName) > 0;
-		}
 	}
 
 	bool FDocumentBuilderRegistry::ReloadBuilder(const FMetasoundFrontendClassName& InClassName) const
diff --git a/Engine/Plugins/Runtime/Metasound/Source/MetasoundEngine/Public/MetasoundDocumentBuilderRegistry.h b/Engine/Plugins/Runtime/Metasound/Source/MetasoundEngine/Public/MetasoundDocumentBuilderRegistry.h
--- a/Engine/Plugins/Runtime/Metasound/Source/MetasoundEngine/Public/MetasoundDocumentBuilderRegistry.h
+++ b/Engine/Plugins/Runtime/Metasound/Source/MetasoundEngine/Public/MetasoundDocumentBuilderRegistry.h
@@ -68,16 +68,15 @@ namespace Metasound::Engine
 				NewBuilder = CastChecked<UMetaSoundBuilderBase>(NewObject<UObject>(&InMetaSoundObject, &DocInterface->GetBuilderUClass()));
 				FMetaSoundFrontendDocumentBuilder& BuilderRef = NewBuilder->GetBuilder();
 				BuilderRef = FMetaSoundFrontendDocumentBuilder(DocInterface);
-
 				const FMetasoundFrontendDocument& Document = DocInterface->GetConstDocument();
-				const FMetasoundFrontendClassName& FullClassName = Document.RootGraph.Metadata.GetClassName();
-				if (!FullClassName.IsValid())
+				const FMetasoundFrontendClassName& ClassName = Document.RootGraph.Metadata.GetClassName();
+				if (!ClassName.IsValid())
 				{
 					BuilderRef.InitDocument();
 				}
 
-				checkf(FullClassName.IsValid(), TEXT("Document initialization must result in a valid class name being generated"));
-				Builders.Add(FullClassName, NewBuilder);
+				checkf(ClassName.IsValid(), TEXT("Document initialization must result in a valid class name being generated"));
+				AddBuilderInternal(ClassName, NewBuilder);
 			}
 
 			return *CastChecked<BuilderClass>(NewBuilder);
@@ -97,9 +96,13 @@ namespace Metasound::Engine
 		virtual bool FinishBuilding(const FMetasoundFrontendClassName& InClassName, bool bForceUnregisterNodeClass = false) const override;
 
 		UMetaSoundBuilderBase* FindBuilderObject(TScriptInterface<const IMetaSoundDocumentInterface> MetaSound) const;
-		UMetaSoundBuilderBase* FindBuilderObject(const FMetasoundFrontendClassName& ClassName) const;
-		TArray<UMetaSoundBuilderBase*> FindBuilderObjects(const FMetasoundFrontendClassName& ClassName) const;
+		UMetaSoundBuilderBase* FindBuilderObject(const FMetasoundFrontendClassName& InClassName) const;
+		TArray<UMetaSoundBuilderBase*> FindBuilderObjects(const FMetasoundFrontendClassName& InClassName) const;
 
 		bool ReloadBuilder(const FMetasoundFrontendClassName& InClassName) const override;
+
+	private:
+		void AddBuilderInternal(const FMetasoundFrontendClassName& InClassName, UMetaSoundBuilderBase* NewBuilder) const;
+		void FinishBuildingInternal(const FMetasoundFrontendClassName& InClassName, bool bForceUnregisterNodeClass) const;
 	};
 } // namespace Metasound::Engine
\ No newline at end of file
