[MetaSound]  'hide unconnected pins' action to nodes
#rnx

diff --git a/Engine/Plugins/Runtime/Metasound/Source/MetasoundEditor/Private/MetasoundEditor.cpp b/Engine/Plugins/Runtime/Metasound/Source/MetasoundEditor/Private/MetasoundEditor.cpp
--- a/Engine/Plugins/Runtime/Metasound/Source/MetasoundEditor/Private/MetasoundEditor.cpp
+++ b/Engine/Plugins/Runtime/Metasound/Source/MetasoundEditor/Private/MetasoundEditor.cpp
@@ -2164,6 +2164,12 @@ namespace Metasound
 					FExecuteAction::CreateLambda([this] { RenameSelectedNode(); }),
 					FCanExecuteAction::CreateLambda([this]() { return CanRenameSelectedNodes(); }));
 
+				GraphEditorCommands->MapAction(FGraphEditorCommands::Get().HideNoConnectionPins,
+					FExecuteAction::CreateSP(this, &FEditor::HideUnconnectedPins));				
+				
+				GraphEditorCommands->MapAction(FGraphEditorCommands::Get().ShowAllPins,
+					FExecuteAction::CreateSP(this, &FEditor::ShowUnconnectedPins));
+
 				// Alignment Commands
 				GraphEditorCommands->MapAction(FGraphEditorCommands::Get().AlignNodesTop,
 					FExecuteAction::CreateLambda([this]() { MetasoundGraphEditor->OnAlignTop(); }));
@@ -2917,6 +2923,30 @@ namespace Metasound
 			}
 		}
 
+		void FEditor::HideUnconnectedPins()
+		{			
+			const FGraphPanelSelectionSet& SelectedNodes = MetasoundGraphEditor->GetSelectedNodes();
+			for (UObject* Object : SelectedNodes)
+			{
+				if (UMetasoundEditorGraphExternalNode* ExternalNode = Cast<UMetasoundEditorGraphExternalNode>(Object))
+				{
+					ExternalNode->HideUnconnectedPins(true);
+				}	
+			}
+		}
+
+		void FEditor::ShowUnconnectedPins()
+		{
+			const FGraphPanelSelectionSet& SelectedNodes = MetasoundGraphEditor->GetSelectedNodes();
+			for (UObject* Object : SelectedNodes)
+			{
+				if (UMetasoundEditorGraphExternalNode* ExternalNode = Cast<UMetasoundEditorGraphExternalNode>(Object))
+				{
+					ExternalNode->HideUnconnectedPins(false);
+				}
+			}
+		}
+
 		bool FEditor::CanPasteNodes()
 		{
 			if (!IsGraphEditable())
diff --git a/Engine/Plugins/Runtime/Metasound/Source/MetasoundEditor/Private/MetasoundEditor.h b/Engine/Plugins/Runtime/Metasound/Source/MetasoundEditor/Private/MetasoundEditor.h
--- a/Engine/Plugins/Runtime/Metasound/Source/MetasoundEditor/Private/MetasoundEditor.h
+++ b/Engine/Plugins/Runtime/Metasound/Source/MetasoundEditor/Private/MetasoundEditor.h
@@ -345,6 +345,12 @@ namespace Metasound
 			/** Find selected node from Graph*/
 			void FindSelectedNodeInGraph();
 
+			/**Hide pins without connection*/
+			void HideUnconnectedPins();
+			
+			/**Show pins without connection*/
+			void ShowUnconnectedPins();
+
 		private:
 			void RemoveInvalidSelection();
 
diff --git a/Engine/Plugins/Runtime/Metasound/Source/MetasoundEditor/Private/MetasoundEditorGraphBuilder.cpp b/Engine/Plugins/Runtime/Metasound/Source/MetasoundEditor/Private/MetasoundEditorGraphBuilder.cpp
--- a/Engine/Plugins/Runtime/Metasound/Source/MetasoundEditor/Private/MetasoundEditorGraphBuilder.cpp
+++ b/Engine/Plugins/Runtime/Metasound/Source/MetasoundEditor/Private/MetasoundEditorGraphBuilder.cpp
@@ -1734,7 +1734,16 @@ namespace Metasound
 			// Pin ToolTips are no longer cached on pins, and are instead dynamically generated via UMetasoundEditorGraphNode::GetPinHoverText
 			InPin.PinToolTip = { };
 			InPin.bAdvancedView = InMetadata.bIsAdvancedDisplay;
-			if (InPin.bAdvancedView)
+
+			//Are the Unconnected hidden for the node
+			bool bUnconnectedHidden = false;
+			UMetasoundEditorGraphNode* MetasoundGraphNode = Cast<UMetasoundEditorGraphNode>(InPin.GetOwningNode());
+			{
+				FMetasoundFrontendNode FrontendNode = MetasoundGraphNode->GetFrontendNodeChecked();
+				bUnconnectedHidden = FrontendNode.Style.bUnconnectedPinsHidden;
+			}
+
+			if (InPin.bAdvancedView || bUnconnectedHidden)
 			{
 				UEdGraphNode* OwningNode = InPin.GetOwningNode();
 				check(OwningNode);
diff --git a/Engine/Plugins/Runtime/Metasound/Source/MetasoundEditor/Private/MetasoundEditorGraphNode.cpp b/Engine/Plugins/Runtime/Metasound/Source/MetasoundEditor/Private/MetasoundEditorGraphNode.cpp
--- a/Engine/Plugins/Runtime/Metasound/Source/MetasoundEditor/Private/MetasoundEditorGraphNode.cpp
+++ b/Engine/Plugins/Runtime/Metasound/Source/MetasoundEditor/Private/MetasoundEditorGraphNode.cpp
@@ -1350,6 +1350,52 @@ void UMetasoundEditorGraphExternalNode::Validate(Metasound::Editor::FGraphNodeVa
 #endif // WITH_EDITOR
 }
 
+void UMetasoundEditorGraphExternalNode::HideUnconnectedPins(const bool InHidePins)
+{
+	using namespace Metasound::Frontend;
+	using namespace Metasound::Editor;
+
+	UObject& MetaSound = GetMetasoundChecked();
+	MetaSound.Modify();
+
+	if (const FMetasoundFrontendNode* FrontendNode = GetFrontendNode())
+	{
+		FMetaSoundFrontendDocumentBuilder& Builder = IDocumentBuilderRegistry::GetChecked().FindOrBeginBuilding(&MetaSound);	
+		FMetasoundFrontendNodeStyle Style = Builder.GetNodeStyle(GetNodeID());
+		Style.bUnconnectedPinsHidden = InHidePins;
+		Builder.SetNodeStyle(GetNodeID(), Style);
+	}
+
+	if (!InHidePins)
+	{
+		bool bIsAdvancedView = false;
+		for (const UEdGraphPin* Pin : Pins)
+		{
+			if (Pin && Pin->bAdvancedView)
+			{
+				bIsAdvancedView = true;			
+			}
+		}
+
+		if (!bIsAdvancedView)
+		{
+			AdvancedPinDisplay = ENodeAdvancedPins::NoPins;
+		}
+	}
+	else
+	{
+		AdvancedPinDisplay = ENodeAdvancedPins::Hidden;
+	}
+
+	if (TSharedPtr<FEditor> MetaSoundEditor = FGraphBuilder::GetEditorForMetasound(GetMetasoundChecked()))
+	{
+		if (TSharedPtr<SGraphEditor> GraphEditor = MetaSoundEditor->GetGraphEditor())
+		{
+			GraphEditor->RefreshNode(*this);
+		}
+	}
+}
+
 const FMetasoundEditorGraphNodeBreadcrumb& UMetasoundEditorGraphExternalNode::GetBreadcrumb() const
 {
 	return Breadcrumb;
diff --git a/Engine/Plugins/Runtime/Metasound/Source/MetasoundEditor/Private/MetasoundEditorGraphSchema.cpp b/Engine/Plugins/Runtime/Metasound/Source/MetasoundEditor/Private/MetasoundEditorGraphSchema.cpp
--- a/Engine/Plugins/Runtime/Metasound/Source/MetasoundEditor/Private/MetasoundEditorGraphSchema.cpp
+++ b/Engine/Plugins/Runtime/Metasound/Source/MetasoundEditor/Private/MetasoundEditorGraphSchema.cpp
@@ -242,7 +242,7 @@ PRAGMA_ENABLE_DEPRECATION_WARNINGS
 				}
 			}
 
-			void SelectNodeInEditor(UMetasoundEditorGraph& InMetaSoundGraph, UMetasoundEditorGraphNode& InNode)
+			void SelecteNodeInEditor(UMetasoundEditorGraph& InMetaSoundGraph, UMetasoundEditorGraphNode& InNode)
 			{
 				TSharedPtr<FEditor> MetasoundEditor = FGraphBuilder::GetEditorForGraph(InMetaSoundGraph);
 				if (MetasoundEditor.IsValid())
@@ -251,16 +251,6 @@ PRAGMA_ENABLE_DEPRECATION_WARNINGS
 				}
 			}
 
-			void SelectNodeInEditorForRename(UMetasoundEditorGraph& InMetaSoundGraph, UMetasoundEditorGraphNode& InNode)
-			{
-				TSharedPtr<FEditor> MetasoundEditor = FGraphBuilder::GetEditorForGraph(InMetaSoundGraph);
-				if (MetasoundEditor.IsValid())
-				{
-					MetasoundEditor->ClearSelectionAndSelectNode(&InNode);
-					MetasoundEditor->SetDelayedRename();
-				}
-			}
-
 			UEdGraphNode* PromoteToVariable(const FName BaseName, UEdGraphPin& FromPin, const FName DataType, const FMetasoundFrontendClass& InVariableClass, const FVector2D& InLocation, bool bSelectNode)
 			{
 				using namespace Frontend;
@@ -326,7 +316,7 @@ PRAGMA_ENABLE_DEPRECATION_WARNINGS
 							if (ensure(SchemaPrivate::TryConnectNewNodeToMatchingDataTypePin(*EdGraphNode, &FromPin)))
 							{
 								FGraphBuilder::RegisterGraphWithFrontend(ParentMetaSound);
-								SelectNodeInEditorForRename(*MetaSoundGraph, *NewGraphNode);
+								SelecteNodeInEditor(*MetaSoundGraph, *NewGraphNode);
 								return EdGraphNode;
 							}
 						}
@@ -476,7 +466,7 @@ UEdGraphNode* FMetasoundGraphSchemaAction_NewNode::PerformAction(UEdGraph* Paren
 		NewGraphNode->UpdateFrontendNodeLocation(Location);
 		NewGraphNode->SyncLocationFromFrontendNode();
 		SchemaPrivate::TryConnectNewNodeToMatchingDataTypePin(*NewGraphNode, FromPin);
-		SchemaPrivate::SelectNodeInEditor(*MetaSoundGraph, *NewGraphNode);
+		SchemaPrivate::SelecteNodeInEditor(*MetaSoundGraph, *NewGraphNode);
 		return NewGraphNode;
 	}
 
@@ -593,7 +583,7 @@ UEdGraphNode* FMetasoundGraphSchemaAction_PromoteToInput::PerformAction(UEdGraph
 				if (ensure(SchemaPrivate::TryConnectNewNodeToMatchingDataTypePin(*EdGraphNode, FromPin)))
 				{
 					FGraphBuilder::RegisterGraphWithFrontend(ParentMetasound);
-					SchemaPrivate::SelectNodeInEditorForRename(*MetasoundGraph, *NewGraphNode);
+					SchemaPrivate::SelecteNodeInEditor(*MetasoundGraph, *NewGraphNode);
 					return EdGraphNode;
 				}
 			}
@@ -803,7 +793,7 @@ UEdGraphNode* FMetasoundGraphSchemaAction_PromoteToOutput::PerformAction(UEdGrap
 				if (ensure(SchemaPrivate::TryConnectNewNodeToMatchingDataTypePin(*EdGraphNode, FromPin)))
 				{
 					FGraphBuilder::RegisterGraphWithFrontend(ParentMetasound);
-					SchemaPrivate::SelectNodeInEditorForRename(*MetasoundGraph, *NewGraphNode);
+					SchemaPrivate::SelecteNodeInEditor(*MetasoundGraph, *NewGraphNode);
 					return EdGraphNode;
 				}
 			}
@@ -1324,6 +1314,8 @@ void UMetasoundEditorGraphSchema::GetContextMenuActions(class UToolMenu* Menu, c
 			LOCTEXT("FindReferencesInGraph_Tooltip", "Find References to the selected Node in the current Graph"),
 			FSlateIcon());
 
+		FToolMenuSection& OrganizationSection = Menu->AddSection("MetasoundGraphSchemaNodeActionsOrganization", LOCTEXT("NodeActionsOrganizationMenuHeader", "Organization"));
+
 		// Only display update ability if node is of type external
 		// and node registry is reporting a major update is available.
 		if (const UMetasoundEditorGraphExternalNode* ExternalNode = Cast<UMetasoundEditorGraphExternalNode>(Context->Node))
@@ -1340,6 +1332,18 @@ void UMetasoundEditorGraphSchema::GetContextMenuActions(class UToolMenu* Menu, c
 			{
 				Section.AddMenuEntry(FEditorCommands::Get().UpdateNodeClass);
 			}
+
+			FMetasoundFrontendNodeStyle Style = NodeHandle->GetNodeStyle();		
+			if (Style.bUnconnectedPinsHidden)
+			{
+				OrganizationSection.AddMenuEntry(FGraphEditorCommands::Get().ShowAllPins, LOCTEXT("ShowUnconnectedPins", "Show Unconnected Pins"),
+					LOCTEXT("ShowUnconnectedPins_Tooltip", "Shows all pins with no connection"),
+					FSlateIcon());
+			}
+			else
+			{
+				OrganizationSection.AddMenuEntry(FGraphEditorCommands::Get().HideNoConnectionPins);
+			}
 		}
 	}
 
diff --git a/Engine/Plugins/Runtime/Metasound/Source/MetasoundEditor/Private/SMetasoundGraphNode.cpp b/Engine/Plugins/Runtime/Metasound/Source/MetasoundEditor/Private/SMetasoundGraphNode.cpp
--- a/Engine/Plugins/Runtime/Metasound/Source/MetasoundEditor/Private/SMetasoundGraphNode.cpp
+++ b/Engine/Plugins/Runtime/Metasound/Source/MetasoundEditor/Private/SMetasoundGraphNode.cpp
@@ -366,6 +366,11 @@ namespace Metasound
 
 		void SMetaSoundGraphNode::CreateStandardPinWidget(UEdGraphPin* InPin)
 		{
+			//Set pin hidden if the node has unconnected pins hidden
+			const FMetasoundFrontendNode& FrontendNode = GetMetaSoundNode().GetFrontendNodeChecked();
+			FMetasoundFrontendNodeStyle Style = FrontendNode.Style;
+			InPin->SafeSetHidden(Style.bUnconnectedPinsHidden);
+
 			const bool bShowPin = ShouldPinBeHidden(InPin);
 			if (bShowPin)
 			{
@@ -498,6 +503,18 @@ namespace Metasound
 			}
 		}
 
+		void SMetaSoundGraphNode::OnAdvancedViewChanged(const ECheckBoxState NewCheckedState)
+		{
+			if (NewCheckedState == ECheckBoxState::Checked)
+			{
+				if (UMetasoundEditorGraphExternalNode* ExternalNode = Cast<UMetasoundEditorGraphExternalNode>(&GetMetaSoundNode()))
+				{
+					ExternalNode->HideUnconnectedPins(false);
+				}
+			}
+			SGraphNode::OnAdvancedViewChanged(NewCheckedState);
+		}
+
 		FLinearColor SMetaSoundGraphNode::GetNodeTitleColorOverride() const
 		{
 			FLinearColor ReturnTitleColor = GraphNode->IsDeprecated() ? FLinearColor::Red : GetNodeObj()->GetNodeTitleColor();
diff --git a/Engine/Plugins/Runtime/Metasound/Source/MetasoundEditor/Private/SMetasoundGraphNode.h b/Engine/Plugins/Runtime/Metasound/Source/MetasoundEditor/Private/SMetasoundGraphNode.h
--- a/Engine/Plugins/Runtime/Metasound/Source/MetasoundEditor/Private/SMetasoundGraphNode.h
+++ b/Engine/Plugins/Runtime/Metasound/Source/MetasoundEditor/Private/SMetasoundGraphNode.h
@@ -58,6 +58,7 @@ namespace Metasound::Editor
 		virtual void GetOverlayBrushes(bool bSelected, const FVector2D WidgetSize, TArray<FOverlayBrushInfo>& Brushes) const override;
 		virtual void OnCommentBubbleToggled(bool bInCommentBubbleVisible) override;
 		virtual void OnCommentTextCommitted(const FText& NewComment, ETextCommit::Type CommitInfo) override;
+		virtual void OnAdvancedViewChanged(const ECheckBoxState NewCheckedState) override;
 
 		FLinearColor GetNodeTitleColorOverride() const;
 
diff --git a/Engine/Plugins/Runtime/Metasound/Source/MetasoundEditor/Public/MetasoundEditorGraphNode.h b/Engine/Plugins/Runtime/Metasound/Source/MetasoundEditor/Public/MetasoundEditorGraphNode.h
--- a/Engine/Plugins/Runtime/Metasound/Source/MetasoundEditor/Public/MetasoundEditorGraphNode.h
+++ b/Engine/Plugins/Runtime/Metasound/Source/MetasoundEditor/Public/MetasoundEditorGraphNode.h
@@ -303,6 +303,8 @@ public:
 	// Validates node and returns whether or not the node is valid.
 	virtual void Validate(Metasound::Editor::FGraphNodeValidationResult& OutResult) override;
 
+	/**Set Unconnected Pins hidden/Unhidden*/
+	void HideUnconnectedPins(const bool InHidePins);
 
 protected:
 	// Friended to enable mutation of Frontend NodeID & direct breadcrumb access
diff --git a/Engine/Plugins/Runtime/Metasound/Source/MetasoundFrontend/Private/MetasoundFrontendDocumentBuilder.cpp b/Engine/Plugins/Runtime/Metasound/Source/MetasoundFrontend/Private/MetasoundFrontendDocumentBuilder.cpp
--- a/Engine/Plugins/Runtime/Metasound/Source/MetasoundFrontend/Private/MetasoundFrontendDocumentBuilder.cpp
+++ b/Engine/Plugins/Runtime/Metasound/Source/MetasoundFrontend/Private/MetasoundFrontendDocumentBuilder.cpp
@@ -3493,6 +3493,27 @@ bool FMetaSoundFrontendDocumentBuilder::SetNodeLocation(const FGuid& InNodeID, c
 
 	return false;
 }
+const FMetasoundFrontendNodeStyle& FMetaSoundFrontendDocumentBuilder::GetNodeStyle(const FGuid& InNodeID)
+{
+	if (const FMetasoundFrontendNode* Node = FindNodeInternal(InNodeID))
+	{
+		return Node->Style;
+	}
+
+	static const FMetasoundFrontendNodeStyle Invalid;
+	return Invalid;
+}
+
+bool FMetaSoundFrontendDocumentBuilder::SetNodeStyle(const FGuid& InNodeID, const FMetasoundFrontendNodeStyle& InNodeStyle)
+{
+	if (FMetasoundFrontendNode* Node = FindNodeInternal(InNodeID))
+	{
+		Node->Style = InNodeStyle;
+		return true;
+	}
+
+	return false;
+}
 #endif // WITH_EDITOR
 
 void FMetaSoundFrontendDocumentBuilder::SetVersionNumber(const FMetasoundFrontendVersionNumber& InDocumentVersionNumber)
diff --git a/Engine/Plugins/Runtime/Metasound/Source/MetasoundFrontend/Public/MetasoundFrontendDocument.h b/Engine/Plugins/Runtime/Metasound/Source/MetasoundFrontend/Public/MetasoundFrontendDocument.h
--- a/Engine/Plugins/Runtime/Metasound/Source/MetasoundFrontend/Public/MetasoundFrontendDocument.h
+++ b/Engine/Plugins/Runtime/Metasound/Source/MetasoundFrontend/Public/MetasoundFrontendDocument.h
@@ -562,6 +562,11 @@ struct METASOUNDFRONTEND_API FMetasoundFrontendNodeStyle
 
 	UPROPERTY()
 	bool bIsPrivate = false;
+	
+	//Whether or not Unconnected pins are hidden
+	UPROPERTY()
+	bool bUnconnectedPinsHidden = false;
+
 #endif // WITH_EDITORONLY_DATA
 };
 
diff --git a/Engine/Plugins/Runtime/Metasound/Source/MetasoundFrontend/Public/MetasoundFrontendDocumentBuilder.h b/Engine/Plugins/Runtime/Metasound/Source/MetasoundFrontend/Public/MetasoundFrontendDocumentBuilder.h
--- a/Engine/Plugins/Runtime/Metasound/Source/MetasoundFrontend/Public/MetasoundFrontendDocumentBuilder.h
+++ b/Engine/Plugins/Runtime/Metasound/Source/MetasoundFrontend/Public/MetasoundFrontendDocumentBuilder.h
@@ -368,6 +368,13 @@ public:
 	// Sets the editor-only node location of a node with the given ID to the provided location.
 	// Returns true if the node was found and the location was updated, false if not.
 	bool SetNodeLocation(const FGuid& InNodeID, const FVector2D& InLocation, const FGuid* InLocationGuid = nullptr);
+
+	// Sets the editor-only style of a node with the given ID.
+	bool SetNodeStyle(const FGuid& InNodeID, const FMetasoundFrontendNodeStyle& InNodeStyle);
+
+	// Gets the editor-only style of a node with the given ID.
+	const FMetasoundFrontendNodeStyle& GetNodeStyle(const FGuid& InNodeID);
+
 #endif // WITH_EDITOR
 
 	bool SetNodeInputDefault(const FGuid& InNodeID, const FGuid& InVertexID, const FMetasoundFrontendLiteral& InLiteral);
