More protections around accessing the wrong builder via the MetaSound builder registry by supplying the TopLevelPath of the asset if available when the registry has multiple conflicting entries due to bad content.
#rnx

diff --git a/Engine/Plugins/Runtime/Metasound/Source/MetasoundEngine/Private/MetasoundAssetSubsystem.cpp b/Engine/Plugins/Runtime/Metasound/Source/MetasoundEngine/Private/MetasoundAssetSubsystem.cpp
--- a/Engine/Plugins/Runtime/Metasound/Source/MetasoundEngine/Private/MetasoundAssetSubsystem.cpp
+++ b/Engine/Plugins/Runtime/Metasound/Source/MetasoundEngine/Private/MetasoundAssetSubsystem.cpp
@@ -626,14 +626,14 @@ namespace Metasound::Engine
 		const FMetasoundFrontendDocument& Document = DocInterface->GetConstDocument();
 		const FMetasoundFrontendClassMetadata& Metadata = Document.RootGraph.Metadata;
 
+		const FTopLevelAssetPath AssetPath(&InObject);
 		if (IDocumentBuilderRegistry* BuilderRegistry = IDocumentBuilderRegistry::Get())
 		{
 			constexpr bool bForceUnregister = true;
-			BuilderRegistry->FinishBuilding(Metadata.GetClassName(), bForceUnregister);
+			BuilderRegistry->FinishBuilding(Metadata.GetClassName(), AssetPath, bForceUnregister);
 		}
 
 		const FAssetKey AssetKey(Metadata.GetClassName(), Metadata.GetVersion());
-		const FTopLevelAssetPath AssetPath(&InObject);
 		AssetSubsystemPrivate::RemovePath(PathMap, AssetKey, AssetPath);
 	}
 
@@ -642,16 +642,16 @@ namespace Metasound::Engine
 		using namespace Frontend;
 
 		FNodeClassInfo ClassInfo;
+		const FTopLevelAssetPath AssetPath(InAssetData.PackageName, InAssetData.AssetName);
 		if (ensureAlways(AssetSubsystemPrivate::GetAssetClassInfo(InAssetData, ClassInfo)))
 		{
 			if (IDocumentBuilderRegistry* BuilderRegistry = IDocumentBuilderRegistry::Get())
 			{
 				constexpr bool bForceUnregister = true;
-				BuilderRegistry->FinishBuilding(ClassInfo.ClassName, bForceUnregister);
+				BuilderRegistry->FinishBuilding(ClassInfo.ClassName, AssetPath, bForceUnregister);
 			}
 
 			const FAssetKey AssetKey(ClassInfo.ClassName, ClassInfo.Version);
-			const FTopLevelAssetPath AssetPath(InAssetData.PackageName, InAssetData.AssetName);
 			AssetSubsystemPrivate::RemovePath(PathMap, AssetKey, AssetPath);
 		}
 	}
diff --git a/Engine/Plugins/Runtime/Metasound/Source/MetasoundEngine/Private/MetasoundBuilderSubsystem.cpp b/Engine/Plugins/Runtime/Metasound/Source/MetasoundEngine/Private/MetasoundBuilderSubsystem.cpp
--- a/Engine/Plugins/Runtime/Metasound/Source/MetasoundEngine/Private/MetasoundBuilderSubsystem.cpp
+++ b/Engine/Plugins/Runtime/Metasound/Source/MetasoundEngine/Private/MetasoundBuilderSubsystem.cpp
@@ -1012,15 +1012,7 @@ UMetaSoundBuilderBase* UMetaSoundBuilderSubsystem::FindBuilderOfDocument(TScript
 {
 	using namespace Metasound::Engine;
 
-	if (const UObject* MetaSoundObject = InMetaSound.GetObject())
-	{
-		const FMetasoundFrontendDocument& Document = InMetaSound->GetConstDocument();
-		const FMetasoundFrontendClassName& InClassName = Document.RootGraph.Metadata.GetClassName();
-
-		return FDocumentBuilderRegistry::GetChecked().FindBuilderObject(InClassName);
-	}
-
-	return nullptr;
+	return FDocumentBuilderRegistry::GetChecked().FindBuilderObject(InMetaSound);
 }
 
 UMetaSoundPatchBuilder* UMetaSoundBuilderSubsystem::FindPatchBuilder(FName BuilderName)
diff --git a/Engine/Plugins/Runtime/Metasound/Source/MetasoundEngine/Private/MetasoundDocumentBuilderRegistry.cpp b/Engine/Plugins/Runtime/Metasound/Source/MetasoundEngine/Private/MetasoundDocumentBuilderRegistry.cpp
--- a/Engine/Plugins/Runtime/Metasound/Source/MetasoundEngine/Private/MetasoundDocumentBuilderRegistry.cpp
+++ b/Engine/Plugins/Runtime/Metasound/Source/MetasoundEngine/Private/MetasoundDocumentBuilderRegistry.cpp
@@ -26,7 +26,7 @@ namespace Metasound::Engine
 			{
 				UE_CLOG(bIsValid, LogMetaSound, Display, TEXT("- %s"), *Pair.Value->GetFullName());
 				constexpr bool bForceUnregister = true;
-				FinishBuildingInternal(Pair.Key, bForceUnregister);
+				FinishBuildingInternal(*Pair.Value.Get(), bForceUnregister);
 			}
 			else
 			{
@@ -91,9 +91,9 @@ namespace Metasound::Engine
 		return nullptr;
 	}
 
-	FMetaSoundFrontendDocumentBuilder* FDocumentBuilderRegistry::FindBuilder(const FMetasoundFrontendClassName& InClassName) const
+	FMetaSoundFrontendDocumentBuilder* FDocumentBuilderRegistry::FindBuilder(const FMetasoundFrontendClassName& InClassName, const FTopLevelAssetPath& AssetPath) const
 	{
-		if (UMetaSoundBuilderBase* Builder = FindBuilderObject(InClassName))
+		if (UMetaSoundBuilderBase* Builder = FindBuilderObject(InClassName, AssetPath))
 		{
 			return &Builder->GetBuilder();
 		}
@@ -136,19 +136,47 @@ namespace Metasound::Engine
 		return nullptr;
 	}
 
-	UMetaSoundBuilderBase* FDocumentBuilderRegistry::FindBuilderObject(const FMetasoundFrontendClassName& InClassName) const
+	UMetaSoundBuilderBase* FDocumentBuilderRegistry::FindBuilderObject(const FMetasoundFrontendClassName& InClassName, const FTopLevelAssetPath& AssetPath) const
 	{
 		TArray<TWeakObjectPtr<UMetaSoundBuilderBase>> Entries;
 		{
 			FScopeLock Lock(&BuildersCriticalSection);
 			Builders.MultiFind(InClassName, Entries);
 		}
+
+		TWeakObjectPtr<UMetaSoundBuilderBase> FoundEntry;
 		if (!Entries.IsEmpty())
 		{
-			return Entries.Last().Get();
+			if (AssetPath.IsValid())
+			{
+				auto AssetPathMatches = [&AssetPath](const TWeakObjectPtr<UMetaSoundBuilderBase>& BuilderPtr)
+				{
+					if (UMetaSoundBuilderBase* Builder = BuilderPtr.Get())
+					{
+						const FMetaSoundFrontendDocumentBuilder& DocBuilder = BuilderPtr->GetConstBuilder();
+						const UObject& DocObject = DocBuilder.CastDocumentObjectChecked<UObject>();
+						FTopLevelAssetPath ObjectPath;
+						if (ObjectPath.TrySetPath(&DocObject))
+						{
+							return AssetPath == ObjectPath;
+						}
+					}
+					return false;
+				};
+
+				if (TWeakObjectPtr<UMetaSoundBuilderBase>* WeakPtr = Entries.FindByPredicate(AssetPathMatches))
+				{
+					FoundEntry = *WeakPtr;
+				}
+
+			}
+			else
+			{
+				FoundEntry = Entries.Last();
+			}
 		}
 
-		return nullptr;
+		return FoundEntry.Get();
 	}
 
 	TArray<UMetaSoundBuilderBase*> FDocumentBuilderRegistry::FindBuilderObjects(const FMetasoundFrontendClassName& InClassName) const
@@ -185,55 +213,72 @@ namespace Metasound::Engine
 		using namespace Metasound;
 		using namespace Metasound::Engine;
 
-		FinishBuildingInternal(InClassName, bForceUnregisterNodeClass);
+		bool bBuildersRemoved = false;
+		TArray<UMetaSoundBuilderBase*> FoundBuilders = FindBuilderObjects(InClassName);
+		for (UMetaSoundBuilderBase* Builder : FoundBuilders)
+		{
+			FinishBuildingInternal(*Builder, bForceUnregisterNodeClass);
+
+			FScopeLock Lock(&BuildersCriticalSection);
+			bBuildersRemoved |= Builders.Remove(InClassName) > 0;
+		}
+
+		return bBuildersRemoved;
+	}
 
-		// Still return true in this case as the builder likely has become inaccessible and may be in the "beginning" of destruction,
-		// so entries is still reporting that it was successfully removed.
+	bool FDocumentBuilderRegistry::FinishBuilding(const FMetasoundFrontendClassName& InClassName, const FTopLevelAssetPath& AssetPath, bool bForceUnregisterNodeClass) const
+	{
+		using namespace Metasound;
+		using namespace Metasound::Engine;
+
+		if (UMetaSoundBuilderBase* Builder = FindBuilderObject(InClassName, AssetPath))
 		{
+			FinishBuildingInternal(*Builder, bForceUnregisterNodeClass);
+
 			FScopeLock Lock(&BuildersCriticalSection);
-			return Builders.Remove(InClassName) > 0;
+			return Builders.RemoveSingle(InClassName, TWeakObjectPtr<UMetaSoundBuilderBase>(Builder)) > 0;
 		}
+
+		return false;
 	}
 
-	void FDocumentBuilderRegistry::FinishBuildingInternal(const FMetasoundFrontendClassName& InClassName, bool bForceUnregisterNodeClass) const
+	void FDocumentBuilderRegistry::FinishBuildingInternal(UMetaSoundBuilderBase& Builder, bool bForceUnregisterNodeClass) const
 	{
-		TArray<UMetaSoundBuilderBase*> FoundBuilders = FindBuilderObjects(InClassName);
-		for (UMetaSoundBuilderBase* Builder : FoundBuilders)
+		// If the builder has applied transactions to its document object that are not mirrored in the frontend registry,
+		// unregister version in registry. This will ensure that future requests for the builder's associated asset will
+		// register a fresh version from the object as the transaction history is intrinsically lost once this builder
+		// is destroyed. It is also possible that the DocBuilder's underlying object can be invalid if object was force
+		// deleted, so validity check is necessary.
+		FMetaSoundFrontendDocumentBuilder& DocBuilder = Builder.GetBuilder();
+		if (DocBuilder.IsValid())
 		{
-			// If the builder has applied transactions to its document object that are not mirrored in the frontend registry,
-			// unregister version in registry. This will ensure that future requests for the builder's associated asset will
-			// register a fresh version from the object as the transaction history is intrinsically lost once this builder
-			// is destroyed. It is also possible that the DocBuilder's underlying object can be invalid if object was force
-			// deleted, so validity check is necessary.
-			FMetaSoundFrontendDocumentBuilder& DocBuilder = Builder->GetBuilder();
-			if (DocBuilder.IsValid())
+			if (!IsRunningCookCommandlet())
 			{
-				if (!IsRunningCookCommandlet())
+				const int32 TransactionCount = DocBuilder.GetTransactionCount();
+				const int32 LastTransactionRegistered = Builder.GetLastTransactionRegistered();
+				if (bForceUnregisterNodeClass || LastTransactionRegistered != TransactionCount)
 				{
-					const int32 TransactionCount = DocBuilder.GetTransactionCount();
-					const int32 LastTransactionRegistered = Builder->GetLastTransactionRegistered();
-					if (bForceUnregisterNodeClass || LastTransactionRegistered != TransactionCount)
+					UObject& MetaSound = DocBuilder.CastDocumentObjectChecked<UObject>();
+					if (FMetasoundAssetBase* MetaSoundAsset = IMetasoundUObjectRegistry::Get().GetObjectAsAssetBase(&MetaSound))
 					{
-						UObject& MetaSound = DocBuilder.CastDocumentObjectChecked<UObject>();
-						if (FMetasoundAssetBase* MetaSoundAsset = IMetasoundUObjectRegistry::Get().GetObjectAsAssetBase(&MetaSound))
-						{
-							MetaSoundAsset->UnregisterGraphWithFrontend();
-						}
+						MetaSoundAsset->UnregisterGraphWithFrontend();
 					}
 				}
-				DocBuilder.FinishBuilding();
 			}
+			DocBuilder.FinishBuilding();
 		}
 	}
 
 	bool FDocumentBuilderRegistry::ReloadBuilder(const FMetasoundFrontendClassName& InClassName) const
 	{
-		if (UMetaSoundBuilderBase* Builder = FindBuilderObject(InClassName))
+		bool bReloaded = false;
+		TArray<UMetaSoundBuilderBase*> ClassBuilders = FindBuilderObjects(InClassName);
+		for (UMetaSoundBuilderBase* Builder : ClassBuilders)
 		{
 			Builder->Reload();
-			return true;
+			bReloaded = true;
 		}
 
-		return false;
+		return bReloaded;
 	}
 } // namespace Metasound::Engine
\ No newline at end of file
diff --git a/Engine/Plugins/Runtime/Metasound/Source/MetasoundEngine/Public/MetasoundDocumentBuilderRegistry.h b/Engine/Plugins/Runtime/Metasound/Source/MetasoundEngine/Public/MetasoundDocumentBuilderRegistry.h
--- a/Engine/Plugins/Runtime/Metasound/Source/MetasoundEngine/Public/MetasoundDocumentBuilderRegistry.h
+++ b/Engine/Plugins/Runtime/Metasound/Source/MetasoundEngine/Public/MetasoundDocumentBuilderRegistry.h
@@ -89,20 +89,29 @@ namespace Metasound::Engine
 #endif // WITH_EDITORONLY_DATA
 
 		virtual FMetaSoundFrontendDocumentBuilder* FindBuilder(TScriptInterface<IMetaSoundDocumentInterface> MetaSound) const override;
-		virtual FMetaSoundFrontendDocumentBuilder* FindBuilder(const FMetasoundFrontendClassName& InClassName) const override;
+		virtual FMetaSoundFrontendDocumentBuilder* FindBuilder(const FMetasoundFrontendClassName& InClassName, const FTopLevelAssetPath& AssetPath) const override;
 
 		virtual FMetaSoundFrontendDocumentBuilder* FindOutermostBuilder(const UObject& InSubObject) const override;
 
 		virtual bool FinishBuilding(const FMetasoundFrontendClassName& InClassName, bool bForceUnregisterNodeClass = false) const override;
+		virtual bool FinishBuilding(const FMetasoundFrontendClassName& InClassName, const FTopLevelAssetPath& AssetPath, bool bForceUnregisterNodeClass = false) const override;
 
+		// Returns the builder object associated with the given MetaSound asset if one is registered and active.
 		UMetaSoundBuilderBase* FindBuilderObject(TScriptInterface<const IMetaSoundDocumentInterface> MetaSound) const;
-		UMetaSoundBuilderBase* FindBuilderObject(const FMetasoundFrontendClassName& InClassName) const;
+
+		// Returns the builder object associated with the given ClassName if one is registered and active.
+		// Optionally, if provided the AssetPath and there is a conflict (i.e. more than one asset is registered
+		// with a given ClassName), will return the one with the provided AssetPath.  Otherwise, will arbitrarily
+		// return one.
+		UMetaSoundBuilderBase* FindBuilderObject(const FMetasoundFrontendClassName& InClassName, const FTopLevelAssetPath& AssetPath) const;
+
+		// Returns all builder objects registered and active associated with the given ClassName.
 		TArray<UMetaSoundBuilderBase*> FindBuilderObjects(const FMetasoundFrontendClassName& InClassName) const;
 
 		bool ReloadBuilder(const FMetasoundFrontendClassName& InClassName) const override;
 
 	private:
 		void AddBuilderInternal(const FMetasoundFrontendClassName& InClassName, UMetaSoundBuilderBase* NewBuilder) const;
-		void FinishBuildingInternal(const FMetasoundFrontendClassName& InClassName, bool bForceUnregisterNodeClass) const;
+		void FinishBuildingInternal(UMetaSoundBuilderBase& Builder, bool bForceUnregisterNodeClass) const;
 	};
 } // namespace Metasound::Engine
\ No newline at end of file
diff --git a/Engine/Plugins/Runtime/Metasound/Source/MetasoundFrontend/Private/MetasoundFrontendDocumentAccessPtr.cpp b/Engine/Plugins/Runtime/Metasound/Source/MetasoundFrontend/Private/MetasoundFrontendDocumentAccessPtr.cpp
--- a/Engine/Plugins/Runtime/Metasound/Source/MetasoundFrontend/Private/MetasoundFrontendDocumentAccessPtr.cpp
+++ b/Engine/Plugins/Runtime/Metasound/Source/MetasoundFrontend/Private/MetasoundFrontendDocumentAccessPtr.cpp
@@ -16,7 +16,7 @@ namespace Metasound
 			{
 				if (IDocumentBuilderRegistry* BuilderRegistry = IDocumentBuilderRegistry::Get())
 				{
-					if (FMetaSoundFrontendDocumentBuilder* Builder = BuilderRegistry->FindBuilder(InGraphClass.Metadata.GetClassName()))
+					if (FMetaSoundFrontendDocumentBuilder* Builder = BuilderRegistry->FindBuilder(InGraphClass.Metadata.GetClassName(), { }))
 					{
 						return Builder->FindBuildGraphChecked();
 					}
@@ -29,7 +29,7 @@ namespace Metasound
 			{
 				if (IDocumentBuilderRegistry* BuilderRegistry = IDocumentBuilderRegistry::Get())
 				{
-					if (FMetaSoundFrontendDocumentBuilder* Builder = BuilderRegistry->FindBuilder(InGraphClass.Metadata.GetClassName()))
+					if (FMetaSoundFrontendDocumentBuilder* Builder = BuilderRegistry->FindBuilder(InGraphClass.Metadata.GetClassName(), { }))
 					{
 						return Builder->FindConstBuildGraphChecked();
 					}
diff --git a/Engine/Plugins/Runtime/Metasound/Source/MetasoundFrontend/Private/MetasoundFrontendGraphController.cpp b/Engine/Plugins/Runtime/Metasound/Source/MetasoundFrontend/Private/MetasoundFrontendGraphController.cpp
--- a/Engine/Plugins/Runtime/Metasound/Source/MetasoundFrontend/Private/MetasoundFrontendGraphController.cpp
+++ b/Engine/Plugins/Runtime/Metasound/Source/MetasoundFrontend/Private/MetasoundFrontendGraphController.cpp
@@ -32,7 +32,7 @@ namespace Metasound
 			// Registry can be null in test builds for legacy controller implementation, so doesn't use GetChecked
 			if (IDocumentBuilderRegistry* BuilderRegistry = IDocumentBuilderRegistry::Get())
 			{
-				if (const FMetaSoundFrontendDocumentBuilder* Builder = BuilderRegistry->FindBuilder(InGraphClass.Metadata.GetClassName()))
+				if (const FMetaSoundFrontendDocumentBuilder* Builder = BuilderRegistry->FindBuilder(InGraphClass.Metadata.GetClassName(), { }))
 				{
 					return Builder->FindConstBuildGraphChecked();
 				}
@@ -46,7 +46,7 @@ namespace Metasound
 			// Registry can be null in test builds for legacy controller implementation, so doesn't use GetChecked
 			if (IDocumentBuilderRegistry* BuilderRegistry = IDocumentBuilderRegistry::Get())
 			{
-				if (FMetaSoundFrontendDocumentBuilder* Builder = BuilderRegistry->FindBuilder(InGraphClass.Metadata.GetClassName()))
+				if (FMetaSoundFrontendDocumentBuilder* Builder = BuilderRegistry->FindBuilder(InGraphClass.Metadata.GetClassName(), { }))
 				{
 					return Builder->FindBuildGraphChecked();
 				}
diff --git a/Engine/Plugins/Runtime/Metasound/Source/MetasoundFrontend/Private/MetasoundFrontendTransform.cpp b/Engine/Plugins/Runtime/Metasound/Source/MetasoundFrontend/Private/MetasoundFrontendTransform.cpp
--- a/Engine/Plugins/Runtime/Metasound/Source/MetasoundFrontend/Private/MetasoundFrontendTransform.cpp
+++ b/Engine/Plugins/Runtime/Metasound/Source/MetasoundFrontend/Private/MetasoundFrontendTransform.cpp
@@ -407,18 +407,6 @@ namespace Metasound
 
 		bool FModifyRootGraphInterfaces::Transform(FMetasoundFrontendDocument& InOutDocument) const
 		{
-			{
-				// Mutation of a document via the soft deprecated access ptr/controller system is not tracked by
-				// the builder registry, so the document cache is invalidated here. It is discouraged to mutate
-				// documents using controllers at this point as it disconnects delegates applied at the object level
-				// (i.e. disconnecting changes to any MetaSound instances being auditioned).
-				const FMetasoundFrontendClassName& ClassName = InOutDocument.RootGraph.Metadata.GetClassName();
-				if (FMetaSoundFrontendDocumentBuilder* Builder = IDocumentBuilderRegistry::GetChecked().FindBuilder(ClassName))
-				{
-					Builder->Reload();
-				}
-			}
-
 			FDocumentAccessPtr DocAccessPtr = MakeAccessPtr<FDocumentAccessPtr>(InOutDocument.AccessPoint, InOutDocument);
 			return Transform(FDocumentController::CreateDocumentHandle(DocAccessPtr));
 		}
diff --git a/Engine/Plugins/Runtime/Metasound/Source/MetasoundFrontend/Public/MetasoundDocumentInterface.h b/Engine/Plugins/Runtime/Metasound/Source/MetasoundFrontend/Public/MetasoundDocumentInterface.h
--- a/Engine/Plugins/Runtime/Metasound/Source/MetasoundFrontend/Public/MetasoundDocumentInterface.h
+++ b/Engine/Plugins/Runtime/Metasound/Source/MetasoundFrontend/Public/MetasoundDocumentInterface.h
@@ -79,7 +79,7 @@ namespace Metasound::Frontend
 		virtual ~IDocumentBuilderRegistry() = default;
 
 		virtual FMetaSoundFrontendDocumentBuilder* FindBuilder(TScriptInterface<IMetaSoundDocumentInterface> MetaSound) const = 0;
-		virtual FMetaSoundFrontendDocumentBuilder* FindBuilder(const FMetasoundFrontendClassName& ClassName) const = 0;
+		virtual FMetaSoundFrontendDocumentBuilder* FindBuilder(const FMetasoundFrontendClassName& ClassName, const FTopLevelAssetPath& AssetPath) const = 0;
 		virtual FMetaSoundFrontendDocumentBuilder* FindOutermostBuilder(const UObject& InSubObject) const = 0;
 
 #if WITH_EDITORONLY_DATA
@@ -94,6 +94,7 @@ namespace Metasound::Frontend
 		// Removes builder from registry, clearing any cached builder state. (Optionally) forces unregistration from the Frontend Node Class Registry
 		// (If the builder has outstanding transactions, unregistration from the Node Class Registry will occur regardless).
 		virtual bool FinishBuilding(const FMetasoundFrontendClassName& InClassName, bool bForceUnregisterNodeClass = false) const = 0;
+		virtual bool FinishBuilding(const FMetasoundFrontendClassName& InClassName, const FTopLevelAssetPath& InAssetPath, bool bForceUnregisterNodeClass = false) const = 0;
 
 		UE_DEPRECATED(5.5, "Document cache can now be invalidated by retrieving an asset builder and calling 'ReloadBuilder'")
 		virtual void InvalidateDocumentCache(const FMetasoundFrontendClassName& InClassName) const { }
