Fixed cook issues ([Backout] - CL34114021)
- Fixed regression with referencing graphs not updating
#tests AudioQA -cook

Original CL Desc
-----------------------------------------------------------------
[Backout] - CL34112568 due to autotest error.
#rnx
Original CL Desc
-----------------------------------------------------------------
More protections around accessing the wrong builder via the MetaSound builder registry by supplying the TopLevelPath of the asset if available when the registry has multiple conflicting entries due to bad content.
#rnx

diff --git a/Engine/Plugins/Runtime/Metasound/Source/MetasoundEditor/Private/MetasoundAssetDefinitions.cpp b/Engine/Plugins/Runtime/Metasound/Source/MetasoundEditor/Private/MetasoundAssetDefinitions.cpp
--- a/Engine/Plugins/Runtime/Metasound/Source/MetasoundEditor/Private/MetasoundAssetDefinitions.cpp
+++ b/Engine/Plugins/Runtime/Metasound/Source/MetasoundEditor/Private/MetasoundAssetDefinitions.cpp
@@ -328,6 +328,8 @@ namespace MenuExtension_MetaSoundSourceTemplate
 	{
 		if (const UContentBrowserAssetContextMenuContext* Context = UContentBrowserAssetContextMenuContext::FindContextWithAssets(MenuContext))
 		{
+			using namespace Metasound::Editor;
+
 			for (TClass* ReferencedMetaSound : Context->LoadSelectedObjects<TClass>())
 			{
 				FString PackagePath;
@@ -345,8 +347,15 @@ namespace MenuExtension_MetaSoundSourceTemplate
 				}
 
 				UMetaSoundEditorSubsystem& MetaSoundEditorSubsystem = UMetaSoundEditorSubsystem::GetChecked();
-				MetaSoundEditorSubsystem.BuildToAsset(&Builder, MetaSoundEditorSubsystem.GetDefaultAuthor(), AssetName, FPackageName::GetLongPackagePath(PackagePath), BuilderResult);
-				if (BuilderResult != EMetaSoundBuilderResult::Succeeded)
+				TScriptInterface<IMetaSoundDocumentInterface> NewMetaSound = MetaSoundEditorSubsystem.BuildToAsset(&Builder, MetaSoundEditorSubsystem.GetDefaultAuthor(), AssetName, FPackageName::GetLongPackagePath(PackagePath), BuilderResult);
+				if (BuilderResult == EMetaSoundBuilderResult::Succeeded)
+				{
+					if (ensure(NewMetaSound))
+					{
+						FGraphBuilder::RegisterGraphWithFrontend(*NewMetaSound.GetObject());
+					}
+				}
+				else
 				{
 					UE_LOG(LogMetaSound, Error, TEXT("Error building to asset when creating preset '%s'"), *AssetName);
 				}
@@ -354,7 +363,8 @@ namespace MenuExtension_MetaSoundSourceTemplate
 		}
 	}
 
- 	static FDelayedAutoRegisterHelper DelayedAutoRegister(EDelayedRegisterRunPhase::EndOfEngineInit, [] {
+ 	static FDelayedAutoRegisterHelper DelayedAutoRegister(EDelayedRegisterRunPhase::EndOfEngineInit, []
+	{
  		UToolMenus::RegisterStartupCallback(FSimpleMulticastDelegate::FDelegate::CreateLambda([]()
  		{
  			FToolMenuOwnerScoped OwnerScoped(UE_MODULE_NAME);
diff --git a/Engine/Plugins/Runtime/Metasound/Source/MetasoundEditor/Private/MetasoundEditorGraphBuilder.cpp b/Engine/Plugins/Runtime/Metasound/Source/MetasoundEditor/Private/MetasoundEditorGraphBuilder.cpp
--- a/Engine/Plugins/Runtime/Metasound/Source/MetasoundEditor/Private/MetasoundEditorGraphBuilder.cpp
+++ b/Engine/Plugins/Runtime/Metasound/Source/MetasoundEditor/Private/MetasoundEditorGraphBuilder.cpp
@@ -480,13 +480,6 @@ namespace Metasound
 					NewGraphNode.bIsClassNative = FMetasoundFrontendRegistryContainer::Get()->IsNodeNative(RegistryKey);
 					NewGraphNode.NodeID = InNodeID;
 					NewGraphNode.ClassName = InMetadata.GetClassName();
-
-					if (!NewGraphNode.bIsClassNative)
-					{
-						FMetasoundAssetBase* MetaSoundAsset = IMetasoundUObjectRegistry::Get().GetObjectAsAssetBase(MetasoundGraph.GetMetasound());
-						check(MetaSoundAsset);
-						MetaSoundAsset->RebuildReferencedAssetClasses();
-					}
 				};
 
 				return GraphBuilderPrivate::AddNode<UMetasoundEditorGraphExternalNode>(InMetaSound, InitNodeFunc, bInSelectNewNode);
@@ -497,11 +490,22 @@ namespace Metasound
 
 		UMetasoundEditorGraphExternalNode* FGraphBuilder::AddExternalNode(UObject& InMetaSound, const FMetasoundFrontendClassMetadata& InMetadata, bool bInSelectNewNode)
 		{
-			FMetasoundAssetBase* MetaSoundAsset = IMetasoundUObjectRegistry::Get().GetObjectAsAssetBase(&InMetaSound);
-			check(MetaSoundAsset);
+			using namespace Frontend;
+
+			FMetaSoundFrontendDocumentBuilder& Builder = IDocumentBuilderRegistry::GetChecked().FindOrBeginBuilding(&InMetaSound);
+			if (const FMetasoundFrontendNode* NewNode = Builder.AddNodeByClassName(InMetadata.GetClassName(), InMetadata.GetVersion().Major))
+			{
+				FMetasoundAssetBase* MetaSoundAsset = IMetasoundUObjectRegistry::Get().GetObjectAsAssetBase(&InMetaSound);
+				check(MetaSoundAsset);
+				MetaSoundAsset->GetDocumentChecked().Metadata.ModifyContext.AddNodeIDModified(NewNode->GetID());
+				const FMetasoundFrontendClass* Dependency = Builder.FindDependency(NewNode->ClassID);
+				if (ensure(Dependency))
+				{
+					return AddExternalNode(InMetaSound, NewNode->GetID(), Dependency->Metadata, bInSelectNewNode);
+				}
+			}
 
-			Frontend::FConstNodeHandle NodeHandle = MetaSoundAsset->GetRootGraphHandle()->AddNode(InMetadata);
-			return AddExternalNode(InMetaSound, NodeHandle->GetID(), NodeHandle->GetClassMetadata(), bInSelectNewNode);
+			return nullptr;
 		}
 
 		Frontend::FNodeHandle FGraphBuilder::AddExternalNodeHandle(UObject& InMetaSound, const FMetasoundFrontendClassName& InClassName)
diff --git a/Engine/Plugins/Runtime/Metasound/Source/MetasoundEditor/Private/MetasoundEditorSubsystem.cpp b/Engine/Plugins/Runtime/Metasound/Source/MetasoundEditor/Private/MetasoundEditorSubsystem.cpp
--- a/Engine/Plugins/Runtime/Metasound/Source/MetasoundEditor/Private/MetasoundEditorSubsystem.cpp
+++ b/Engine/Plugins/Runtime/Metasound/Source/MetasoundEditor/Private/MetasoundEditorSubsystem.cpp
@@ -91,6 +91,9 @@ TScriptInterface<IMetaSoundDocumentInterface> UMetaSoundEditorSubsystem::BuildTo
 			constexpr bool bForceNodeCreation = true;
 			NewDocBuilder.InjectInputTemplateNodes(bForceNodeCreation, InjectResult);
 
+			FMetasoundAssetBase& Asset = NewDocBuilder.GetBuilder().GetMetasoundAsset();
+			Asset.RebuildReferencedAssetClasses();
+
 			if (!bWasRooted)
 			{
 				InBuilder->RemoveFromRoot();
diff --git a/Engine/Plugins/Runtime/Metasound/Source/MetasoundEngine/Private/MetasoundAssetSubsystem.cpp b/Engine/Plugins/Runtime/Metasound/Source/MetasoundEngine/Private/MetasoundAssetSubsystem.cpp
--- a/Engine/Plugins/Runtime/Metasound/Source/MetasoundEngine/Private/MetasoundAssetSubsystem.cpp
+++ b/Engine/Plugins/Runtime/Metasound/Source/MetasoundEngine/Private/MetasoundAssetSubsystem.cpp
@@ -626,14 +626,14 @@ namespace Metasound::Engine
 		const FMetasoundFrontendDocument& Document = DocInterface->GetConstDocument();
 		const FMetasoundFrontendClassMetadata& Metadata = Document.RootGraph.Metadata;
 
+		const FTopLevelAssetPath AssetPath(&InObject);
 		if (IDocumentBuilderRegistry* BuilderRegistry = IDocumentBuilderRegistry::Get())
 		{
 			constexpr bool bForceUnregister = true;
-			BuilderRegistry->FinishBuilding(Metadata.GetClassName(), bForceUnregister);
+			BuilderRegistry->FinishBuilding(Metadata.GetClassName(), AssetPath, bForceUnregister);
 		}
 
 		const FAssetKey AssetKey(Metadata.GetClassName(), Metadata.GetVersion());
-		const FTopLevelAssetPath AssetPath(&InObject);
 		AssetSubsystemPrivate::RemovePath(PathMap, AssetKey, AssetPath);
 	}
 
@@ -642,16 +642,16 @@ namespace Metasound::Engine
 		using namespace Frontend;
 
 		FNodeClassInfo ClassInfo;
+		const FTopLevelAssetPath AssetPath(InAssetData.PackageName, InAssetData.AssetName);
 		if (ensureAlways(AssetSubsystemPrivate::GetAssetClassInfo(InAssetData, ClassInfo)))
 		{
 			if (IDocumentBuilderRegistry* BuilderRegistry = IDocumentBuilderRegistry::Get())
 			{
 				constexpr bool bForceUnregister = true;
-				BuilderRegistry->FinishBuilding(ClassInfo.ClassName, bForceUnregister);
+				BuilderRegistry->FinishBuilding(ClassInfo.ClassName, AssetPath, bForceUnregister);
 			}
 
 			const FAssetKey AssetKey(ClassInfo.ClassName, ClassInfo.Version);
-			const FTopLevelAssetPath AssetPath(InAssetData.PackageName, InAssetData.AssetName);
 			AssetSubsystemPrivate::RemovePath(PathMap, AssetKey, AssetPath);
 		}
 	}
diff --git a/Engine/Plugins/Runtime/Metasound/Source/MetasoundEngine/Private/MetasoundBuilderSubsystem.cpp b/Engine/Plugins/Runtime/Metasound/Source/MetasoundEngine/Private/MetasoundBuilderSubsystem.cpp
--- a/Engine/Plugins/Runtime/Metasound/Source/MetasoundEngine/Private/MetasoundBuilderSubsystem.cpp
+++ b/Engine/Plugins/Runtime/Metasound/Source/MetasoundEngine/Private/MetasoundBuilderSubsystem.cpp
@@ -1012,15 +1012,7 @@ UMetaSoundBuilderBase* UMetaSoundBuilderSubsystem::FindBuilderOfDocument(TScript
 {
 	using namespace Metasound::Engine;
 
-	if (const UObject* MetaSoundObject = InMetaSound.GetObject())
-	{
-		const FMetasoundFrontendDocument& Document = InMetaSound->GetConstDocument();
-		const FMetasoundFrontendClassName& InClassName = Document.RootGraph.Metadata.GetClassName();
-
-		return FDocumentBuilderRegistry::GetChecked().FindBuilderObject(InClassName);
-	}
-
-	return nullptr;
+	return FDocumentBuilderRegistry::GetChecked().FindBuilderObject(InMetaSound);
 }
 
 UMetaSoundPatchBuilder* UMetaSoundBuilderSubsystem::FindPatchBuilder(FName BuilderName)
diff --git a/Engine/Plugins/Runtime/Metasound/Source/MetasoundEngine/Private/MetasoundDocumentBuilderRegistry.cpp b/Engine/Plugins/Runtime/Metasound/Source/MetasoundEngine/Private/MetasoundDocumentBuilderRegistry.cpp
--- a/Engine/Plugins/Runtime/Metasound/Source/MetasoundEngine/Private/MetasoundDocumentBuilderRegistry.cpp
+++ b/Engine/Plugins/Runtime/Metasound/Source/MetasoundEngine/Private/MetasoundDocumentBuilderRegistry.cpp
@@ -26,7 +26,7 @@ namespace Metasound::Engine
 			{
 				UE_CLOG(bIsValid, LogMetaSound, Display, TEXT("- %s"), *Pair.Value->GetFullName());
 				constexpr bool bForceUnregister = true;
-				FinishBuildingInternal(Pair.Key, bForceUnregister);
+				FinishBuildingInternal(*Pair.Value.Get(), bForceUnregister);
 			}
 			else
 			{
@@ -91,9 +91,9 @@ namespace Metasound::Engine
 		return nullptr;
 	}
 
-	FMetaSoundFrontendDocumentBuilder* FDocumentBuilderRegistry::FindBuilder(const FMetasoundFrontendClassName& InClassName) const
+	FMetaSoundFrontendDocumentBuilder* FDocumentBuilderRegistry::FindBuilder(const FMetasoundFrontendClassName& InClassName, const FTopLevelAssetPath& AssetPath) const
 	{
-		if (UMetaSoundBuilderBase* Builder = FindBuilderObject(InClassName))
+		if (UMetaSoundBuilderBase* Builder = FindBuilderObject(InClassName, AssetPath))
 		{
 			return &Builder->GetBuilder();
 		}
@@ -103,6 +103,7 @@ namespace Metasound::Engine
 
 	UMetaSoundBuilderBase* FDocumentBuilderRegistry::FindBuilderObject(TScriptInterface<const IMetaSoundDocumentInterface> MetaSound) const
 	{
+		UMetaSoundBuilderBase* FoundEntry = nullptr;
 		if (const UObject* MetaSoundObject = MetaSound.GetObject())
 		{
 			const FMetasoundFrontendDocument& Document = MetaSound->GetConstDocument();
@@ -114,41 +115,77 @@ namespace Metasound::Engine
 				Builders.MultiFind(ClassName, Entries);
 			}
 
-			if (!Entries.IsEmpty())
+			for (const TWeakObjectPtr<UMetaSoundBuilderBase>& BuilderPtr : Entries)
 			{
-				auto IsMetaSound = [MetaSoundObject = MetaSound.GetObject()](const TWeakObjectPtr<UMetaSoundBuilderBase>& BuilderPtr)
+				if (UMetaSoundBuilderBase* Builder = BuilderPtr.Get())
 				{
-					if (BuilderPtr.IsValid())
+					// Can be invalid if look-up is called during asset removal/destruction or the entry was
+					// prematurely "finished". Only return invalid entry if builder asset path cannot be
+					// matched as this is likely the destroyed entry associated with the provided AssetPath.
+					const FMetaSoundFrontendDocumentBuilder& DocBuilder = Builder->GetConstBuilder();
+					if (DocBuilder.IsValid())
 					{
 						UObject& TestMetaSound = BuilderPtr->GetConstBuilder().CastDocumentObjectChecked<UObject>();
-						return &TestMetaSound == MetaSoundObject;
+						if (&TestMetaSound == MetaSoundObject)
+						{
+							FoundEntry = Builder;
+							break;
+						}
+					}
+					else
+					{
+						FoundEntry = Builder;
 					}
-					return false;
-				};
-				const TWeakObjectPtr<UMetaSoundBuilderBase>* Builder = Entries.FindByPredicate(IsMetaSound);
-				if (Builder)
-				{
-					return Builder->Get();
 				}
 			}
 		}
 
-		return nullptr;
+		return FoundEntry;
 	}
 
-	UMetaSoundBuilderBase* FDocumentBuilderRegistry::FindBuilderObject(const FMetasoundFrontendClassName& InClassName) const
+	UMetaSoundBuilderBase* FDocumentBuilderRegistry::FindBuilderObject(const FMetasoundFrontendClassName& InClassName, const FTopLevelAssetPath& AssetPath) const
 	{
 		TArray<TWeakObjectPtr<UMetaSoundBuilderBase>> Entries;
 		{
 			FScopeLock Lock(&BuildersCriticalSection);
 			Builders.MultiFind(InClassName, Entries);
 		}
-		if (!Entries.IsEmpty())
+
+		UMetaSoundBuilderBase* FoundEntry = nullptr;
+		for (const TWeakObjectPtr<UMetaSoundBuilderBase>& BuilderPtr : Entries)
 		{
-			return Entries.Last().Get();
+			if (UMetaSoundBuilderBase* Builder = BuilderPtr.Get())
+			{
+				const FMetaSoundFrontendDocumentBuilder& DocBuilder = Builder->GetConstBuilder();
+
+				// Can be invalid if look-up is called during asset removal/destruction or the entry was
+				// prematurely "finished". Only return invalid entry if builder asset path cannot be
+				// matched as this is likely the destroyed entry associated with the provided AssetPath.
+				if (DocBuilder.IsValid())
+				{
+					const UObject& DocObject = DocBuilder.CastDocumentObjectChecked<UObject>();
+					FTopLevelAssetPath ObjectPath;
+					if (ObjectPath.TrySetPath(&DocObject))
+					{
+						if (AssetPath == ObjectPath)
+						{
+							FoundEntry = Builder;
+							break;
+						}
+					}
+					else
+					{
+						FoundEntry = Builder;
+					}
+				}
+				else
+				{
+					FoundEntry = Builder;
+				}
+			}
 		}
 
-		return nullptr;
+		return FoundEntry;
 	}
 
 	TArray<UMetaSoundBuilderBase*> FDocumentBuilderRegistry::FindBuilderObjects(const FMetasoundFrontendClassName& InClassName) const
@@ -185,55 +222,72 @@ namespace Metasound::Engine
 		using namespace Metasound;
 		using namespace Metasound::Engine;
 
-		FinishBuildingInternal(InClassName, bForceUnregisterNodeClass);
-
-		// Still return true in this case as the builder likely has become inaccessible and may be in the "beginning" of destruction,
-		// so entries is still reporting that it was successfully removed.
+		bool bBuildersRemoved = false;
+		TArray<UMetaSoundBuilderBase*> FoundBuilders = FindBuilderObjects(InClassName);
+		for (UMetaSoundBuilderBase* Builder : FoundBuilders)
 		{
+			FinishBuildingInternal(*Builder, bForceUnregisterNodeClass);
+
 			FScopeLock Lock(&BuildersCriticalSection);
-			return Builders.Remove(InClassName) > 0;
+			bBuildersRemoved |= Builders.Remove(InClassName) > 0;
 		}
+
+		return bBuildersRemoved;
 	}
 
-	void FDocumentBuilderRegistry::FinishBuildingInternal(const FMetasoundFrontendClassName& InClassName, bool bForceUnregisterNodeClass) const
+	bool FDocumentBuilderRegistry::FinishBuilding(const FMetasoundFrontendClassName& InClassName, const FTopLevelAssetPath& AssetPath, bool bForceUnregisterNodeClass) const
 	{
-		TArray<UMetaSoundBuilderBase*> FoundBuilders = FindBuilderObjects(InClassName);
-		for (UMetaSoundBuilderBase* Builder : FoundBuilders)
+		using namespace Metasound;
+		using namespace Metasound::Engine;
+
+		TWeakObjectPtr<UMetaSoundBuilderBase> BuilderPtr;
+		if (UMetaSoundBuilderBase* Builder = FindBuilderObject(InClassName, AssetPath))
+		{
+			FinishBuildingInternal(*Builder, bForceUnregisterNodeClass);
+			BuilderPtr = TWeakObjectPtr<UMetaSoundBuilderBase>(Builder);
+		}
+
+		FScopeLock Lock(&BuildersCriticalSection);
+		return Builders.RemoveSingle(InClassName, BuilderPtr) > 0;
+	}
+
+	void FDocumentBuilderRegistry::FinishBuildingInternal(UMetaSoundBuilderBase& Builder, bool bForceUnregisterNodeClass) const
+	{
+		// If the builder has applied transactions to its document object that are not mirrored in the frontend registry,
+		// unregister version in registry. This will ensure that future requests for the builder's associated asset will
+		// register a fresh version from the object as the transaction history is intrinsically lost once this builder
+		// is destroyed. It is also possible that the DocBuilder's underlying object can be invalid if object was force
+		// deleted, so validity check is necessary.
+		FMetaSoundFrontendDocumentBuilder& DocBuilder = Builder.GetBuilder();
+		if (DocBuilder.IsValid())
 		{
-			// If the builder has applied transactions to its document object that are not mirrored in the frontend registry,
-			// unregister version in registry. This will ensure that future requests for the builder's associated asset will
-			// register a fresh version from the object as the transaction history is intrinsically lost once this builder
-			// is destroyed. It is also possible that the DocBuilder's underlying object can be invalid if object was force
-			// deleted, so validity check is necessary.
-			FMetaSoundFrontendDocumentBuilder& DocBuilder = Builder->GetBuilder();
-			if (DocBuilder.IsValid())
+			if (!IsRunningCookCommandlet())
 			{
-				if (!IsRunningCookCommandlet())
+				const int32 TransactionCount = DocBuilder.GetTransactionCount();
+				const int32 LastTransactionRegistered = Builder.GetLastTransactionRegistered();
+				if (bForceUnregisterNodeClass || LastTransactionRegistered != TransactionCount)
 				{
-					const int32 TransactionCount = DocBuilder.GetTransactionCount();
-					const int32 LastTransactionRegistered = Builder->GetLastTransactionRegistered();
-					if (bForceUnregisterNodeClass || LastTransactionRegistered != TransactionCount)
+					UObject& MetaSound = DocBuilder.CastDocumentObjectChecked<UObject>();
+					if (FMetasoundAssetBase* MetaSoundAsset = IMetasoundUObjectRegistry::Get().GetObjectAsAssetBase(&MetaSound))
 					{
-						UObject& MetaSound = DocBuilder.CastDocumentObjectChecked<UObject>();
-						if (FMetasoundAssetBase* MetaSoundAsset = IMetasoundUObjectRegistry::Get().GetObjectAsAssetBase(&MetaSound))
-						{
-							MetaSoundAsset->UnregisterGraphWithFrontend();
-						}
+						MetaSoundAsset->UnregisterGraphWithFrontend();
 					}
 				}
-				DocBuilder.FinishBuilding();
 			}
+			DocBuilder.FinishBuilding();
 		}
 	}
 
 	bool FDocumentBuilderRegistry::ReloadBuilder(const FMetasoundFrontendClassName& InClassName) const
 	{
-		if (UMetaSoundBuilderBase* Builder = FindBuilderObject(InClassName))
+		bool bReloaded = false;
+		TArray<UMetaSoundBuilderBase*> ClassBuilders = FindBuilderObjects(InClassName);
+		for (UMetaSoundBuilderBase* Builder : ClassBuilders)
 		{
 			Builder->Reload();
-			return true;
+			bReloaded = true;
 		}
 
-		return false;
+		return bReloaded;
 	}
 } // namespace Metasound::Engine
\ No newline at end of file
diff --git a/Engine/Plugins/Runtime/Metasound/Source/MetasoundEngine/Public/MetasoundDocumentBuilderRegistry.h b/Engine/Plugins/Runtime/Metasound/Source/MetasoundEngine/Public/MetasoundDocumentBuilderRegistry.h
--- a/Engine/Plugins/Runtime/Metasound/Source/MetasoundEngine/Public/MetasoundDocumentBuilderRegistry.h
+++ b/Engine/Plugins/Runtime/Metasound/Source/MetasoundEngine/Public/MetasoundDocumentBuilderRegistry.h
@@ -89,20 +89,29 @@ namespace Metasound::Engine
 #endif // WITH_EDITORONLY_DATA
 
 		virtual FMetaSoundFrontendDocumentBuilder* FindBuilder(TScriptInterface<IMetaSoundDocumentInterface> MetaSound) const override;
-		virtual FMetaSoundFrontendDocumentBuilder* FindBuilder(const FMetasoundFrontendClassName& InClassName) const override;
+		virtual FMetaSoundFrontendDocumentBuilder* FindBuilder(const FMetasoundFrontendClassName& InClassName, const FTopLevelAssetPath& AssetPath) const override;
 
 		virtual FMetaSoundFrontendDocumentBuilder* FindOutermostBuilder(const UObject& InSubObject) const override;
 
 		virtual bool FinishBuilding(const FMetasoundFrontendClassName& InClassName, bool bForceUnregisterNodeClass = false) const override;
+		virtual bool FinishBuilding(const FMetasoundFrontendClassName& InClassName, const FTopLevelAssetPath& AssetPath, bool bForceUnregisterNodeClass = false) const override;
 
+		// Returns the builder object associated with the given MetaSound asset if one is registered and active.
 		UMetaSoundBuilderBase* FindBuilderObject(TScriptInterface<const IMetaSoundDocumentInterface> MetaSound) const;
-		UMetaSoundBuilderBase* FindBuilderObject(const FMetasoundFrontendClassName& InClassName) const;
+
+		// Returns the builder object associated with the given ClassName if one is registered and active.
+		// Optionally, if provided the AssetPath and there is a conflict (i.e. more than one asset is registered
+		// with a given ClassName), will return the one with the provided AssetPath.  Otherwise, will arbitrarily
+		// return one.
+		UMetaSoundBuilderBase* FindBuilderObject(const FMetasoundFrontendClassName& InClassName, const FTopLevelAssetPath& AssetPath) const;
+
+		// Returns all builder objects registered and active associated with the given ClassName.
 		TArray<UMetaSoundBuilderBase*> FindBuilderObjects(const FMetasoundFrontendClassName& InClassName) const;
 
 		bool ReloadBuilder(const FMetasoundFrontendClassName& InClassName) const override;
 
 	private:
 		void AddBuilderInternal(const FMetasoundFrontendClassName& InClassName, UMetaSoundBuilderBase* NewBuilder) const;
-		void FinishBuildingInternal(const FMetasoundFrontendClassName& InClassName, bool bForceUnregisterNodeClass) const;
+		void FinishBuildingInternal(UMetaSoundBuilderBase& Builder, bool bForceUnregisterNodeClass) const;
 	};
 } // namespace Metasound::Engine
\ No newline at end of file
diff --git a/Engine/Plugins/Runtime/Metasound/Source/MetasoundFrontend/Private/MetasoundAssetBase.cpp b/Engine/Plugins/Runtime/Metasound/Source/MetasoundFrontend/Private/MetasoundAssetBase.cpp
--- a/Engine/Plugins/Runtime/Metasound/Source/MetasoundFrontend/Private/MetasoundAssetBase.cpp
+++ b/Engine/Plugins/Runtime/Metasound/Source/MetasoundFrontend/Private/MetasoundAssetBase.cpp
@@ -883,11 +883,9 @@ void FMetasoundAssetBase::RebuildReferencedAssetClasses()
 	using namespace Metasound::Frontend;
 
 	IMetaSoundAssetManager& AssetManager = IMetaSoundAssetManager::GetChecked();
-	if (AssetManager.AddAssetReferences(*this))
-	{
-		TSet<IMetaSoundAssetManager::FAssetInfo> ReferencedAssetClasses = AssetManager.GetReferencedAssetClasses(*this);
-		SetReferencedAssetClasses(MoveTemp(ReferencedAssetClasses));
-	}
+	AssetManager.AddAssetReferences(*this);
+	TSet<IMetaSoundAssetManager::FAssetInfo> ReferencedAssetClasses = AssetManager.GetReferencedAssetClasses(*this);
+	SetReferencedAssetClasses(MoveTemp(ReferencedAssetClasses));
 }
 #endif // WITH_EDITOR
 
diff --git a/Engine/Plugins/Runtime/Metasound/Source/MetasoundFrontend/Private/MetasoundFrontendDocumentAccessPtr.cpp b/Engine/Plugins/Runtime/Metasound/Source/MetasoundFrontend/Private/MetasoundFrontendDocumentAccessPtr.cpp
--- a/Engine/Plugins/Runtime/Metasound/Source/MetasoundFrontend/Private/MetasoundFrontendDocumentAccessPtr.cpp
+++ b/Engine/Plugins/Runtime/Metasound/Source/MetasoundFrontend/Private/MetasoundFrontendDocumentAccessPtr.cpp
@@ -16,7 +16,7 @@ namespace Metasound
 			{
 				if (IDocumentBuilderRegistry* BuilderRegistry = IDocumentBuilderRegistry::Get())
 				{
-					if (FMetaSoundFrontendDocumentBuilder* Builder = BuilderRegistry->FindBuilder(InGraphClass.Metadata.GetClassName()))
+					if (FMetaSoundFrontendDocumentBuilder* Builder = BuilderRegistry->FindBuilder(InGraphClass.Metadata.GetClassName(), { }))
 					{
 						return Builder->FindBuildGraphChecked();
 					}
@@ -29,7 +29,7 @@ namespace Metasound
 			{
 				if (IDocumentBuilderRegistry* BuilderRegistry = IDocumentBuilderRegistry::Get())
 				{
-					if (FMetaSoundFrontendDocumentBuilder* Builder = BuilderRegistry->FindBuilder(InGraphClass.Metadata.GetClassName()))
+					if (FMetaSoundFrontendDocumentBuilder* Builder = BuilderRegistry->FindBuilder(InGraphClass.Metadata.GetClassName(), { }))
 					{
 						return Builder->FindConstBuildGraphChecked();
 					}
diff --git a/Engine/Plugins/Runtime/Metasound/Source/MetasoundFrontend/Private/MetasoundFrontendGraphController.cpp b/Engine/Plugins/Runtime/Metasound/Source/MetasoundFrontend/Private/MetasoundFrontendGraphController.cpp
--- a/Engine/Plugins/Runtime/Metasound/Source/MetasoundFrontend/Private/MetasoundFrontendGraphController.cpp
+++ b/Engine/Plugins/Runtime/Metasound/Source/MetasoundFrontend/Private/MetasoundFrontendGraphController.cpp
@@ -32,7 +32,7 @@ namespace Metasound
 			// Registry can be null in test builds for legacy controller implementation, so doesn't use GetChecked
 			if (IDocumentBuilderRegistry* BuilderRegistry = IDocumentBuilderRegistry::Get())
 			{
-				if (const FMetaSoundFrontendDocumentBuilder* Builder = BuilderRegistry->FindBuilder(InGraphClass.Metadata.GetClassName()))
+				if (const FMetaSoundFrontendDocumentBuilder* Builder = BuilderRegistry->FindBuilder(InGraphClass.Metadata.GetClassName(), { }))
 				{
 					return Builder->FindConstBuildGraphChecked();
 				}
@@ -46,7 +46,7 @@ namespace Metasound
 			// Registry can be null in test builds for legacy controller implementation, so doesn't use GetChecked
 			if (IDocumentBuilderRegistry* BuilderRegistry = IDocumentBuilderRegistry::Get())
 			{
-				if (FMetaSoundFrontendDocumentBuilder* Builder = BuilderRegistry->FindBuilder(InGraphClass.Metadata.GetClassName()))
+				if (FMetaSoundFrontendDocumentBuilder* Builder = BuilderRegistry->FindBuilder(InGraphClass.Metadata.GetClassName(), { }))
 				{
 					return Builder->FindBuildGraphChecked();
 				}
diff --git a/Engine/Plugins/Runtime/Metasound/Source/MetasoundFrontend/Private/MetasoundFrontendTransform.cpp b/Engine/Plugins/Runtime/Metasound/Source/MetasoundFrontend/Private/MetasoundFrontendTransform.cpp
--- a/Engine/Plugins/Runtime/Metasound/Source/MetasoundFrontend/Private/MetasoundFrontendTransform.cpp
+++ b/Engine/Plugins/Runtime/Metasound/Source/MetasoundFrontend/Private/MetasoundFrontendTransform.cpp
@@ -407,18 +407,6 @@ namespace Metasound
 
 		bool FModifyRootGraphInterfaces::Transform(FMetasoundFrontendDocument& InOutDocument) const
 		{
-			{
-				// Mutation of a document via the soft deprecated access ptr/controller system is not tracked by
-				// the builder registry, so the document cache is invalidated here. It is discouraged to mutate
-				// documents using controllers at this point as it disconnects delegates applied at the object level
-				// (i.e. disconnecting changes to any MetaSound instances being auditioned).
-				const FMetasoundFrontendClassName& ClassName = InOutDocument.RootGraph.Metadata.GetClassName();
-				if (FMetaSoundFrontendDocumentBuilder* Builder = IDocumentBuilderRegistry::GetChecked().FindBuilder(ClassName))
-				{
-					Builder->Reload();
-				}
-			}
-
 			FDocumentAccessPtr DocAccessPtr = MakeAccessPtr<FDocumentAccessPtr>(InOutDocument.AccessPoint, InOutDocument);
 			return Transform(FDocumentController::CreateDocumentHandle(DocAccessPtr));
 		}
diff --git a/Engine/Plugins/Runtime/Metasound/Source/MetasoundFrontend/Public/MetasoundDocumentInterface.h b/Engine/Plugins/Runtime/Metasound/Source/MetasoundFrontend/Public/MetasoundDocumentInterface.h
--- a/Engine/Plugins/Runtime/Metasound/Source/MetasoundFrontend/Public/MetasoundDocumentInterface.h
+++ b/Engine/Plugins/Runtime/Metasound/Source/MetasoundFrontend/Public/MetasoundDocumentInterface.h
@@ -79,7 +79,7 @@ namespace Metasound::Frontend
 		virtual ~IDocumentBuilderRegistry() = default;
 
 		virtual FMetaSoundFrontendDocumentBuilder* FindBuilder(TScriptInterface<IMetaSoundDocumentInterface> MetaSound) const = 0;
-		virtual FMetaSoundFrontendDocumentBuilder* FindBuilder(const FMetasoundFrontendClassName& ClassName) const = 0;
+		virtual FMetaSoundFrontendDocumentBuilder* FindBuilder(const FMetasoundFrontendClassName& ClassName, const FTopLevelAssetPath& AssetPath) const = 0;
 		virtual FMetaSoundFrontendDocumentBuilder* FindOutermostBuilder(const UObject& InSubObject) const = 0;
 
 #if WITH_EDITORONLY_DATA
@@ -94,6 +94,7 @@ namespace Metasound::Frontend
 		// Removes builder from registry, clearing any cached builder state. (Optionally) forces unregistration from the Frontend Node Class Registry
 		// (If the builder has outstanding transactions, unregistration from the Node Class Registry will occur regardless).
 		virtual bool FinishBuilding(const FMetasoundFrontendClassName& InClassName, bool bForceUnregisterNodeClass = false) const = 0;
+		virtual bool FinishBuilding(const FMetasoundFrontendClassName& InClassName, const FTopLevelAssetPath& InAssetPath, bool bForceUnregisterNodeClass = false) const = 0;
 
 		UE_DEPRECATED(5.5, "Document cache can now be invalidated by retrieving an asset builder and calling 'ReloadBuilder'")
 		virtual void InvalidateDocumentCache(const FMetasoundFrontendClassName& InClassName) const { }
