Stopgap disablement of error that is likely cruft with new look-up logic
#rnx

diff --git a/Engine/Plugins/Runtime/Metasound/Source/MetasoundEngine/Private/MetasoundDocumentBuilderRegistry.cpp b/Engine/Plugins/Runtime/Metasound/Source/MetasoundEngine/Private/MetasoundDocumentBuilderRegistry.cpp
--- a/Engine/Plugins/Runtime/Metasound/Source/MetasoundEngine/Private/MetasoundDocumentBuilderRegistry.cpp
+++ b/Engine/Plugins/Runtime/Metasound/Source/MetasoundEngine/Private/MetasoundDocumentBuilderRegistry.cpp
@@ -40,35 +40,35 @@ namespace Metasound::Engine
 	{
 		FScopeLock Lock(&BuildersCriticalSection);
 
-#if !NO_LOGGING
-		const bool bPrintConflicts = Builders.Contains(InClassName);
-#endif // !NO_LOGGING
+// #if !NO_LOGGING
+// 		const bool bPrintConflicts = Builders.Contains(InClassName);
+// #endif // !NO_LOGGING
 
 		Builders.Add(InClassName, NewBuilder);
 
-#if !NO_LOGGING
-		if (bPrintConflicts)
-		{
-			TArray<TWeakObjectPtr<UMetaSoundBuilderBase>> Entries;
-			Builders.MultiFind(InClassName, Entries);
-			UE_LOG(LogMetaSound, Error, TEXT("More than one asset registered with class name '%s'. "
-				"Returning builder that may not be associated with desired object! \n"
-				"This can happen if asset was moved using revision control and original location was revived. \n"
-				"Remove all but one of the following assets and relink a duplicate or copied replacement asset:"),
-				*InClassName.ToString());
-			for (const TWeakObjectPtr<UMetaSoundBuilderBase>& BuilderPtr : Entries)
-			{
-				if (BuilderPtr.IsValid())
-				{
-					UE_LOG(LogMetaSound, Error, TEXT("- %s"), *BuilderPtr->GetConstBuilder().CastDocumentObjectChecked<UObject>().GetPathName());
-				}
-				else
-				{
-					UE_LOG(LogMetaSound, Error, TEXT("- STALE ENTRY (Not removed prior to asset removal)"));
-				}
-			}
-		}
-#endif // !NO_LOGGING
+// #if !NO_LOGGING
+// 		if (bPrintConflicts)
+// 		{
+// 			TArray<TWeakObjectPtr<UMetaSoundBuilderBase>> Entries;
+// 			Builders.MultiFind(InClassName, Entries);
+// 			UE_LOG(LogMetaSound, Error, TEXT("More than one asset registered with class name '%s'. "
+// 				"Look-up may return builder that is not associated with desired object! \n"
+// 				"This can happen if asset was moved using revision control and original location was revived. \n"
+// 				"Remove all but one of the following assets and relink a duplicate or copied replacement asset:"),
+// 				*InClassName.ToString());
+// 			for (const TWeakObjectPtr<UMetaSoundBuilderBase>& BuilderPtr : Entries)
+// 			{
+// 				if (BuilderPtr.IsValid())
+// 				{
+// 					UE_LOG(LogMetaSound, Error, TEXT("- %s"), *BuilderPtr->GetConstBuilder().CastDocumentObjectChecked<UObject>().GetPathName());
+// 				}
+// 				else
+// 				{
+// 					UE_LOG(LogMetaSound, Error, TEXT("- STALE ENTRY (Not removed prior to asset removal)"));
+// 				}
+// 			}
+// 		}
+// #endif // !NO_LOGGING
 	}
 
 #if WITH_EDITORONLY_DATA
