UE-216526 - [MetaSounds] - MetaSound Asset Registration runs continuously in a template project

diff --git a/Engine/Plugins/Runtime/Metasound/Source/MetasoundEditor/Private/MetasoundEditor.cpp b/Engine/Plugins/Runtime/Metasound/Source/MetasoundEditor/Private/MetasoundEditor.cpp
--- a/Engine/Plugins/Runtime/Metasound/Source/MetasoundEditor/Private/MetasoundEditor.cpp
+++ b/Engine/Plugins/Runtime/Metasound/Source/MetasoundEditor/Private/MetasoundEditor.cpp
@@ -2227,7 +2227,8 @@ namespace Metasound
 				.TextStyle(FAppStyle::Get(), "Graph.ZoomText")
 				.ColorAndOpacity(FLinearColor(1, 1, 1, 0.30f));
 
-			SAssignNew(RenderStatsWidget, SMetaSoundRenderStats);
+			SAssignNew(RenderStatsWidget, SMetaSoundRenderStats)
+				.Visibility(EVisibility::HitTestInvisible);
 		}
 
 		FGraphAppearanceInfo FEditor::GetGraphAppearance() const
@@ -3349,7 +3350,9 @@ namespace Metasound
 			if (bPrimingRegistry)
 			{
 				IMetasoundEditorModule& MetaSoundEditorModule = FModuleManager::GetModuleChecked<IMetasoundEditorModule>("MetaSoundEditor");
-				if (MetaSoundEditorModule.GetAssetRegistryPrimeStatus() == EAssetPrimeStatus::Complete)
+				EAssetPrimeStatus PrimeStatus = MetaSoundEditorModule.GetAssetRegistryPrimeStatus();
+				EAssetScanStatus ScanStatus = MetaSoundEditorModule.GetAssetRegistryScanStatus();
+				if (PrimeStatus == EAssetPrimeStatus::Complete || (PrimeStatus == EAssetPrimeStatus::InProgress && ScanStatus == EAssetScanStatus::Complete))
 				{
 					bPrimingRegistry = false;
 					NotifyAssetPrimeComplete();
diff --git a/Engine/Plugins/Runtime/Metasound/Source/MetasoundEditor/Private/MetasoundEditorModule.cpp b/Engine/Plugins/Runtime/Metasound/Source/MetasoundEditor/Private/MetasoundEditorModule.cpp
--- a/Engine/Plugins/Runtime/Metasound/Source/MetasoundEditor/Private/MetasoundEditorModule.cpp
+++ b/Engine/Plugins/Runtime/Metasound/Source/MetasoundEditor/Private/MetasoundEditorModule.cpp
@@ -712,6 +712,11 @@ namespace Metasound
 				return AssetPrimeStatus;
 			}
 
+			virtual EAssetScanStatus GetAssetRegistryScanStatus() const override
+			{
+				return AssetScanStatus;
+			}
+
 			virtual void RegisterExplicitProxyClass(const UClass& InClass) override
 			{
 				using namespace Metasound::Frontend;
diff --git a/Engine/Plugins/Runtime/Metasound/Source/MetasoundEditor/Public/MetasoundEditorModule.h b/Engine/Plugins/Runtime/Metasound/Source/MetasoundEditor/Public/MetasoundEditorModule.h
--- a/Engine/Plugins/Runtime/Metasound/Source/MetasoundEditor/Public/MetasoundEditorModule.h
+++ b/Engine/Plugins/Runtime/Metasound/Source/MetasoundEditor/Public/MetasoundEditorModule.h
@@ -103,6 +103,7 @@ namespace Metasound
 			// loaded) & registers them if not already registered with the MetaSound Class Registry.
 			virtual void PrimeAssetRegistryAsync() = 0;
 			virtual EAssetPrimeStatus GetAssetRegistryPrimeStatus() const = 0;
+			virtual EAssetScanStatus GetAssetRegistryScanStatus() const = 0;
 
 			virtual TUniquePtr<FMetasoundDefaultLiteralCustomizationBase> CreateMemberDefaultLiteralCustomization(UClass& InClass, IDetailCategoryBuilder& DefaultCategoryBuilder) const = 0;
 
