Trivial: Protection against edge case where a builder can go out of scope on asset editor init and graph edit state is checked
#rnx

diff --git a/Engine/Plugins/Runtime/Metasound/Source/MetasoundEditor/Private/MetasoundEditor.cpp b/Engine/Plugins/Runtime/Metasound/Source/MetasoundEditor/Private/MetasoundEditor.cpp
--- a/Engine/Plugins/Runtime/Metasound/Source/MetasoundEditor/Private/MetasoundEditor.cpp
+++ b/Engine/Plugins/Runtime/Metasound/Source/MetasoundEditor/Private/MetasoundEditor.cpp
@@ -3524,8 +3524,12 @@ namespace Metasound
 
 			if (Builder.IsValid())
 			{
-				const FMetasoundFrontendGraph& Graph = Builder->GetConstBuilder().FindConstBuildGraphChecked();
-				return Graph.Style.bIsGraphEditable;
+				const FMetaSoundFrontendDocumentBuilder& DocBuilder = Builder->GetConstBuilder();
+				if (DocBuilder.IsValid())
+				{
+					const FMetasoundFrontendGraph& Graph = DocBuilder.FindConstBuildGraphChecked();
+					return Graph.Style.bIsGraphEditable;
+				}
 			}
 
 			return false;
