Trivial: Don't set DisplayName on asset rename
#rnx

diff --git a/Engine/Plugins/Runtime/Metasound/Source/MetasoundEngine/Private/MetasoundAssetSubsystem.cpp b/Engine/Plugins/Runtime/Metasound/Source/MetasoundEngine/Private/MetasoundAssetSubsystem.cpp
--- a/Engine/Plugins/Runtime/Metasound/Source/MetasoundEngine/Private/MetasoundAssetSubsystem.cpp
+++ b/Engine/Plugins/Runtime/Metasound/Source/MetasoundEngine/Private/MetasoundAssetSubsystem.cpp
@@ -679,15 +679,6 @@ namespace Metasound::Engine
 			const FTopLevelAssetPath OldPath(InOldObjectPath);
 			AssetSubsystemPrivate::RemovePath(PathMap, AssetKey, OldPath);
 
-#if WITH_EDITOR
-			IDocumentBuilderRegistry* BuilderRegistry = IDocumentBuilderRegistry::Get();
-			if (BuilderRegistry)
-			{
-				FMetaSoundFrontendDocumentBuilder& Builder = BuilderRegistry->FindOrBeginBuilding(MetaSoundAsset->GetOwningAsset());
-				Builder.SetDisplayName(FText());
-			}
-#endif // WITH_EDITOR
-
 			if (ClassInfo.AssetClassID.IsValid())
 			{
 				if (AssetKey.IsValid())
