Fix up finishbuilding bug leaving stale entry around during some forms of builder destruction
#rnx

diff --git a/Engine/Plugins/Runtime/Metasound/Source/MetasoundEngine/Private/MetasoundDocumentBuilderRegistry.cpp b/Engine/Plugins/Runtime/Metasound/Source/MetasoundEngine/Private/MetasoundDocumentBuilderRegistry.cpp
--- a/Engine/Plugins/Runtime/Metasound/Source/MetasoundEngine/Private/MetasoundDocumentBuilderRegistry.cpp
+++ b/Engine/Plugins/Runtime/Metasound/Source/MetasoundEngine/Private/MetasoundDocumentBuilderRegistry.cpp
@@ -222,17 +222,15 @@ namespace Metasound::Engine
 		using namespace Metasound;
 		using namespace Metasound::Engine;
 
-		bool bBuildersRemoved = false;
 		TArray<UMetaSoundBuilderBase*> FoundBuilders = FindBuilderObjects(InClassName);
 		for (UMetaSoundBuilderBase* Builder : FoundBuilders)
 		{
 			FinishBuildingInternal(*Builder, bForceUnregisterNodeClass);
 
-			FScopeLock Lock(&BuildersCriticalSection);
-			bBuildersRemoved |= Builders.Remove(InClassName) > 0;
 		}
 
-		return bBuildersRemoved;
+		FScopeLock Lock(&BuildersCriticalSection);
+		return Builders.Remove(InClassName) > 0;
 	}
 
 	bool FDocumentBuilderRegistry::FinishBuilding(const FMetasoundFrontendClassName& InClassName, const FTopLevelAssetPath& AssetPath, bool bForceUnregisterNodeClass) const
