Trivial: fix negation allowing user to delete interface members
#rnx

diff --git a/Engine/Plugins/Runtime/Metasound/Source/MetasoundEditor/Private/MetasoundEditorGraphNode.cpp b/Engine/Plugins/Runtime/Metasound/Source/MetasoundEditor/Private/MetasoundEditorGraphNode.cpp
--- a/Engine/Plugins/Runtime/Metasound/Source/MetasoundEditor/Private/MetasoundEditorGraphNode.cpp
+++ b/Engine/Plugins/Runtime/Metasound/Source/MetasoundEditor/Private/MetasoundEditorGraphNode.cpp
@@ -979,11 +979,10 @@ bool UMetasoundEditorGraphOutputNode::CanUserDeleteNode() const
 {
 	if (const UMetasoundEditorGraphVertex* Vertex = Cast<UMetasoundEditorGraphVertex>(GetMember()))
 	{
-		return Vertex->IsInterfaceMember();
+		return !Vertex->IsInterfaceMember();
 	}
 	
-	return true;	
-
+	return true;
 }
 
 void UMetasoundEditorGraphOutputNode::SetNodeID(FGuid InNodeID)
