Fix non-standard include path
#rnx

diff --git a/Engine/Plugins/Runtime/Metasound/Source/MetasoundEngine/Private/Metasound.cpp b/Engine/Plugins/Runtime/Metasound/Source/MetasoundEngine/Private/Metasound.cpp
--- a/Engine/Plugins/Runtime/Metasound/Source/MetasoundEngine/Private/Metasound.cpp
+++ b/Engine/Plugins/Runtime/Metasound/Source/MetasoundEngine/Private/Metasound.cpp
@@ -5,6 +5,7 @@
 #include "Internationalization/Text.h"
 #include "Logging/TokenizedMessage.h"
 #include "MetasoundAssetBase.h"
+#include "MetasoundAssetManager.h"
 #include "MetasoundAudioFormats.h"
 #include "MetasoundBuilderSubsystem.h"
 #include "MetasoundDocumentInterface.h"
@@ -34,7 +35,6 @@
 #if WITH_EDITORONLY_DATA
 #include "EdGraph/EdGraph.h"
 #endif // WITH_EDITORONLY_DATA
-#include "../../MetasoundFrontend/Public/MetasoundAssetManager.h"
 
 #define LOCTEXT_NAMESPACE "MetaSound"
 
