Fix for registration not being called after conforming an object to document
#rnx

diff --git a/Engine/Plugins/Runtime/Metasound/Source/MetasoundEngine/Private/MetasoundSource.cpp b/Engine/Plugins/Runtime/Metasound/Source/MetasoundEngine/Private/MetasoundSource.cpp
--- a/Engine/Plugins/Runtime/Metasound/Source/MetasoundEngine/Private/MetasoundSource.cpp
+++ b/Engine/Plugins/Runtime/Metasound/Source/MetasoundEngine/Private/MetasoundSource.cpp
@@ -433,9 +433,9 @@ void UMetaSoundSource::PostEditChangeOutputFormat()
 
 		// Use the editor form of register to ensure other editors'
 		// MetaSounds are auto-updated if they are referencing this graph.
-		if (Graph)
+		if (EditorGraph)
 		{
-			Graph->RegisterGraphWithFrontend();
+			EditorGraph->RegisterGraphWithFrontend();
 		}
 		MarkMetasoundDocumentDirty();
 	}
