UE-171622
MetaSound inputs and variables should have right click (or CTRL+D) duplicate option
#wip
#rnx

diff --git a/Engine/Plugins/Runtime/Metasound/Source/MetasoundEditor/Private/MetasoundEditor.cpp b/Engine/Plugins/Runtime/Metasound/Source/MetasoundEditor/Private/MetasoundEditor.cpp
--- a/Engine/Plugins/Runtime/Metasound/Source/MetasoundEditor/Private/MetasoundEditor.cpp
+++ b/Engine/Plugins/Runtime/Metasound/Source/MetasoundEditor/Private/MetasoundEditor.cpp
@@ -1620,6 +1620,10 @@ namespace Metasound
 			ToolkitCommands->MapAction(FGenericCommands::Get().Rename,
 				FExecuteAction::CreateSP(this, &FEditor::RenameSelectedInterfaceItem),
 				FCanExecuteAction::CreateSP(this, &FEditor::CanRenameSelectedInterfaceItems));
+			
+			ToolkitCommands->MapAction(FGenericCommands::Get().Duplicate,
+				FExecuteAction::CreateSP(this, &FEditor::DuplicateSelectedMemberItems),
+				FCanExecuteAction::CreateSP(this, &FEditor::CanDuplicateSelectedMemberItems));
 
 			ToolkitCommands->MapAction(
 				FEditorCommands::Get().UpdateNodeClass,
@@ -2817,6 +2821,149 @@ namespace Metasound
 			}
 		}
 
+		bool FEditor::CanDuplicateSelectedMemberItems() const
+		{
+			if (!IsGraphEditable())
+			{
+				return false;
+			}
+
+			if (!GraphMembersMenu.IsValid())
+			{
+				return false;
+			}
+
+			TArray<TSharedPtr<FEdGraphSchemaAction>> Actions;
+			GraphMembersMenu->GetSelectedActions(Actions);
+
+			if (Actions.IsEmpty())
+			{
+				return false;
+			}
+
+			for (const TSharedPtr<FEdGraphSchemaAction>& Action : Actions)
+			{
+				TSharedPtr<FMetasoundGraphMemberSchemaAction> MetasoundAction = StaticCastSharedPtr<FMetasoundGraphMemberSchemaAction>(Action);
+				if (MetasoundAction.IsValid())
+				{						
+					if (const UMetasoundEditorGraphVertex* GraphVertex = Cast<UMetasoundEditorGraphVertex>(MetasoundAction->GetGraphMember()))
+					{
+						if (GraphVertex->IsInterfaceMember())
+						{
+							return false;
+						}
+					}
+				}
+			}		
+
+			return true;
+		}
+
+		void FEditor::DuplicateSelectedMemberItems()
+		{
+			using namespace Frontend;
+
+			if (!GraphMembersMenu.IsValid())
+			{
+				return;
+			}
+
+			TArray<TSharedPtr<FEdGraphSchemaAction>> Actions;
+			GraphMembersMenu->GetSelectedActions(Actions);
+
+			if (Actions.IsEmpty())
+			{
+				return;
+			}
+				
+			UMetasoundEditorGraph& Graph = GetMetaSoundGraphChecked();
+
+			TArray<TObjectPtr<UObject>> SelectedObjects;
+			FName NameToSelect;
+
+			for (const TSharedPtr<FEdGraphSchemaAction>& Action : Actions)
+			{
+				TSharedPtr<FMetasoundGraphMemberSchemaAction> MetasoundAction = StaticCastSharedPtr<FMetasoundGraphMemberSchemaAction>(Action);
+				if (!MetasoundAction.IsValid())
+				{
+					continue;
+				}
+
+				if (const UMetasoundEditorGraphMember* SourceGraphMember = MetasoundAction->GetGraphMember())
+				{								
+					const FScopedTransaction Transaction(TEXT(""), LOCTEXT("MetaSoundEditorDuplicateMember", "Duplicate MetaSound Member"), Metasound);
+					Metasound->Modify();
+
+					UMetasoundEditorGraphMember* NewGraphMember = nullptr;
+										
+					//Duplicate the Sources NodeHandle and add a new member from it
+					if (const UMetasoundEditorGraphVariable* SourceGraphVariable = Cast<UMetasoundEditorGraphVariable>(SourceGraphMember))
+					{
+						FConstVariableHandle VariableHandle = FGraphBuilder::DuplicateVariableHandle(Graph.GetMetasoundChecked(), SourceGraphVariable->GetConstVariableHandle());
+						if (ensure(VariableHandle->IsValid()))
+						{							
+							NewGraphMember = Graph.FindOrAddVariable(VariableHandle);
+						}
+					}
+					else if (const UMetasoundEditorGraphVertex* SourceGraphVertex = Cast<UMetasoundEditorGraphVertex>(SourceGraphMember))
+					{
+						const FName SourceMemberName = SourceGraphVertex->GetMemberName();
+						const EMetasoundFrontendClassType ClassType = SourceGraphVertex->GetClassType();
+						
+						FMetaSoundFrontendDocumentBuilder& DocumentBuilder = IDocumentBuilderRegistry::GetChecked().FindOrBeginBuilding(&Graph.GetMetasoundChecked());
+						const FName Name = FGraphBuilder::GenerateUniqueNameByClassType(Graph.GetMetasoundChecked(), ClassType, SourceMemberName.ToString());
+
+						if (ClassType == EMetasoundFrontendClassType::Input)
+						{
+							if (const FMetasoundFrontendClassInput* SourceInput = DocumentBuilder.FindGraphInput(SourceMemberName))
+							{							
+								if (const FMetasoundFrontendNode* FrontendNode = DocumentBuilder.DuplicateGraphInput(*SourceInput, SourceInput->DefaultLiteral, Name))
+								{
+									FGraphBuilder::SynchronizeGraphMembers(DocumentBuilder, Graph);
+									NewGraphMember = Graph.FindInput(FrontendNode->Name);
+								}
+							}
+						}
+						else if (ClassType == EMetasoundFrontendClassType::Output)
+						{
+							if (const FMetasoundFrontendClassOutput* SourceOutput = DocumentBuilder.FindGraphOutput(SourceMemberName))
+							{
+								if (const FMetasoundFrontendNode* FrontendNode = DocumentBuilder.DuplicateGraphOutput(*SourceOutput, Name))
+								{
+									FGraphBuilder::SynchronizeGraphMembers(DocumentBuilder, Graph);
+									NewGraphMember = Graph.FindOutput(FrontendNode->Name);
+								}
+							}
+						}
+					}
+
+					//Duplicate the literal from the SourceGraphMember to the NewGraphMember added
+					if (NewGraphMember)
+					{
+						if (UMetaSoundEditorSubsystem* EditorSubsystem = GEditor->GetEditorSubsystem<UMetaSoundEditorSubsystem>())
+						{
+							FMetaSoundFrontendDocumentBuilder& DocumentBuilder = IDocumentBuilderRegistry::GetChecked().FindOrBeginBuilding(Metasound);
+							TSubclassOf<UMetasoundEditorGraphMemberDefaultLiteral> SubClass = SourceGraphMember->GetLiteral()->GetClass();
+							EditorSubsystem->BindMemberMetadata(DocumentBuilder, *NewGraphMember, SubClass, SourceGraphMember->GetLiteral());
+
+							NameToSelect = NewGraphMember->GetMemberName();
+							SelectedObjects.Add(NewGraphMember);
+						}
+					}	
+				}
+			}	
+
+			if (GraphMembersMenu.IsValid())
+			{
+				GraphMembersMenu->RefreshAllActions(true);
+				if (!NameToSelect.IsNone())
+				{
+					GraphMembersMenu->SelectItemByName(NameToSelect);
+					SetSelection(SelectedObjects);
+				}
+			}
+		}
+
 		void FEditor::RefreshDetails()
 		{
 			TRACE_CPUPROFILER_EVENT_SCOPE(Metasound::Editor::FEditor::RefreshDetails);
@@ -3285,6 +3432,7 @@ namespace Metasound
 
 			MenuBuilder.AddMenuEntry(FGenericCommands::Get().Delete);
 			MenuBuilder.AddMenuEntry(FGenericCommands::Get().Rename);
+			MenuBuilder.AddMenuEntry(FGenericCommands::Get().Duplicate);
 
 			MenuBuilder.AddMenuEntry(
 				LOCTEXT("JumpToNodesMenuEntry", "Jump to Node(s) in Graph"),
diff --git a/Engine/Plugins/Runtime/Metasound/Source/MetasoundEditor/Private/MetasoundEditor.h b/Engine/Plugins/Runtime/Metasound/Source/MetasoundEditor/Private/MetasoundEditor.h
--- a/Engine/Plugins/Runtime/Metasound/Source/MetasoundEditor/Private/MetasoundEditor.h
+++ b/Engine/Plugins/Runtime/Metasound/Source/MetasoundEditor/Private/MetasoundEditor.h
@@ -319,6 +319,12 @@ namespace Metasound
 			/** Rename selected interface item. */
 			void RenameSelectedInterfaceItem();
 
+			/** Whether the currently selected Member item(s) can be duplicated. */
+			bool CanDuplicateSelectedMemberItems() const;
+
+			/** Duplicate selected Member items. */
+			void DuplicateSelectedMemberItems();
+
 			/** Whether there are nodes to jump to for the currently selected interface item. */
 			bool CanJumpToNodesForSelectedInterfaceItem() const;
 
diff --git a/Engine/Plugins/Runtime/Metasound/Source/MetasoundEditor/Private/MetasoundEditorGraphBuilder.cpp b/Engine/Plugins/Runtime/Metasound/Source/MetasoundEditor/Private/MetasoundEditorGraphBuilder.cpp
--- a/Engine/Plugins/Runtime/Metasound/Source/MetasoundEditor/Private/MetasoundEditorGraphBuilder.cpp
+++ b/Engine/Plugins/Runtime/Metasound/Source/MetasoundEditor/Private/MetasoundEditorGraphBuilder.cpp
@@ -1451,7 +1451,7 @@ namespace Metasound
 			return *MetaSoundAsset;
 		}
 
-		Frontend::FVariableHandle FGraphBuilder::AddVariableHandle(UObject& InMetaSound, const FName& InTypeName)
+		Frontend::FVariableHandle FGraphBuilder::AddVariableHandle(UObject& InMetaSound, const FName& InTypeName, const FName* InNameBase)
 		{
 			using namespace Frontend;
 
@@ -1460,10 +1460,7 @@ namespace Metasound
 
 			FGraphHandle FrontendGraph = MetaSoundAsset->GetRootGraphHandle();
 
-			FText BaseDisplayName = LOCTEXT("VariableDefaultDisplayName", "Variable");
-
-			FString BaseName = BaseDisplayName.ToString();
-			FName VariableName = GenerateUniqueVariableName(FrontendGraph, BaseName);
+			FName VariableName = GenerateUniqueVariableName(FrontendGraph, InNameBase ? InNameBase->ToString() : LOCTEXT("VariableDefaultDisplayName", "Variable").ToString());
 			FVariableHandle Variable = FrontendGraph->AddVariable(InTypeName);
 
 			Variable->SetDisplayName(FText::GetEmpty());
@@ -1472,6 +1469,15 @@ namespace Metasound
 			return Variable;
 		}
 
+		Frontend::FVariableHandle FGraphBuilder::DuplicateVariableHandle(UObject& InMetaSound, const Frontend::FConstVariableHandle& InHandle)
+		{
+			Frontend::FVariableHandle NewHandle = AddVariableHandle(InMetaSound, InHandle->GetDataType(), &InHandle->GetName());
+			NewHandle->SetDisplayName(FText::GetEmpty());
+			NewHandle->SetDescription(InHandle->GetDescription());
+			NewHandle->SetLiteral(InHandle->GetLiteral());
+			return NewHandle;
+		}
+
 		Frontend::FNodeHandle FGraphBuilder::AddVariableNodeHandle(UObject& InMetaSound, const FGuid& InVariableID, const Metasound::FNodeClassName& InVariableNodeClassName, UMetasoundEditorGraphVariableNode* InVariableNode)
 		{
 			using namespace Frontend;
diff --git a/Engine/Plugins/Runtime/Metasound/Source/MetasoundEditor/Private/MetasoundEditorGraphBuilder.h b/Engine/Plugins/Runtime/Metasound/Source/MetasoundEditor/Private/MetasoundEditorGraphBuilder.h
--- a/Engine/Plugins/Runtime/Metasound/Source/MetasoundEditor/Private/MetasoundEditorGraphBuilder.h
+++ b/Engine/Plugins/Runtime/Metasound/Source/MetasoundEditor/Private/MetasoundEditorGraphBuilder.h
@@ -166,9 +166,13 @@ namespace Metasound
 			//
 			// @param InMetaSound - FMetasoundAssetBase derived UObject.
 			// @param InTypeName - Data type of variable.
-			//
+			// @param InNameBase - (optional) Name of the Variable, if not given default naming is used.
+			// 
 			// @return The added frontend variable handle. On error, the returned handle is invalid.
-			static Frontend::FVariableHandle AddVariableHandle(UObject& InMetaSound, const FName& InTypeName);
+			static Frontend::FVariableHandle AddVariableHandle(UObject& InMetaSound, const FName& InTypeName, const FName* InNameBase = nullptr);
+
+			// Generates Duplicate FVariableHandle from the given handle.
+			static Frontend::FVariableHandle DuplicateVariableHandle(UObject& InMetaSound, const Frontend::FConstVariableHandle& InHandle);
 
 			// Adds a frontend variable node to root graph using the supplied node class name.
 			//
diff --git a/Engine/Plugins/Runtime/Metasound/Source/MetasoundFrontend/Private/MetasoundFrontendDocumentBuilder.cpp b/Engine/Plugins/Runtime/Metasound/Source/MetasoundFrontend/Private/MetasoundFrontendDocumentBuilder.cpp
--- a/Engine/Plugins/Runtime/Metasound/Source/MetasoundFrontend/Private/MetasoundFrontendDocumentBuilder.cpp
+++ b/Engine/Plugins/Runtime/Metasound/Source/MetasoundFrontend/Private/MetasoundFrontendDocumentBuilder.cpp
@@ -1052,6 +1052,25 @@ const FMetasoundFrontendNode* FMetaSoundFrontendDocumentBuilder::AddGraphInput(c
 	return nullptr;
 }
 
+const FMetasoundFrontendNode* FMetaSoundFrontendDocumentBuilder::DuplicateGraphInput(const FMetasoundFrontendClassInput& InClassInput, FMetasoundFrontendLiteral DefaultValue, const FName InName)
+{
+	using namespace Metasound::Frontend;
+
+	FDocumentIDGenerator& IDGenerator = FDocumentIDGenerator::Get();
+	const FMetasoundFrontendDocument& Doc = GetConstDocumentChecked();
+		
+	FMetasoundFrontendClassInput ClassInput = InClassInput;
+	ClassInput.NodeID = IDGenerator.CreateNodeID(Doc);
+	ClassInput.VertexID = IDGenerator.CreateVertexID(Doc);
+#if WITH_EDITORONLY_DATA
+	ClassInput.Metadata.SetDisplayName(FText::GetEmpty());
+#endif // WITH_EDITORONLY_DATA
+	ClassInput.Name = InName;
+	ClassInput.DefaultLiteral = static_cast<FMetasoundFrontendLiteral>(DefaultValue);
+
+	return AddGraphInput(ClassInput);
+}
+
 const FMetasoundFrontendNode* FMetaSoundFrontendDocumentBuilder::AddGraphOutput(const FMetasoundFrontendClassOutput& InClassOutput)
 {
 	using namespace Metasound::Frontend;
@@ -1111,6 +1130,24 @@ const FMetasoundFrontendNode* FMetaSoundFrontendDocumentBuilder::AddGraphOutput(
 	return nullptr;
 }
 
+const FMetasoundFrontendNode* FMetaSoundFrontendDocumentBuilder::DuplicateGraphOutput(const FMetasoundFrontendClassOutput& InClassOutput, const FName InName)
+{
+	using namespace Metasound::Frontend;
+
+	FDocumentIDGenerator& IDGenerator = FDocumentIDGenerator::Get();
+	const FMetasoundFrontendDocument& Doc = GetConstDocumentChecked();
+
+	FMetasoundFrontendClassOutput ClassOutput = InClassOutput;
+	ClassOutput.NodeID = IDGenerator.CreateNodeID(Doc);
+	ClassOutput.VertexID = IDGenerator.CreateVertexID(Doc);
+#if WITH_EDITORONLY_DATA
+	ClassOutput.Metadata.SetDisplayName(FText::GetEmpty());
+#endif // WITH_EDITORONLY_DATA
+	ClassOutput.Name = InName;
+	
+	return AddGraphOutput(ClassOutput);
+}
+
 bool FMetaSoundFrontendDocumentBuilder::AddInterface(FName InterfaceName)
 {
 	using namespace Metasound::Frontend;
diff --git a/Engine/Plugins/Runtime/Metasound/Source/MetasoundFrontend/Public/MetasoundFrontendDocumentBuilder.h b/Engine/Plugins/Runtime/Metasound/Source/MetasoundFrontend/Public/MetasoundFrontendDocumentBuilder.h
--- a/Engine/Plugins/Runtime/Metasound/Source/MetasoundFrontend/Public/MetasoundFrontendDocumentBuilder.h
+++ b/Engine/Plugins/Runtime/Metasound/Source/MetasoundFrontend/Public/MetasoundFrontendDocumentBuilder.h
@@ -160,8 +160,10 @@ public:
 	bool AddEdgesFromMatchingInterfaceNodeOutputsToGraphOutputs(const FGuid& InNodeID, TArray<const FMetasoundFrontendEdge*>& OutEdgesCreated, bool bReplaceExistingConnections = true);
 	bool AddEdgesFromMatchingInterfaceNodeInputsToGraphInputs(const FGuid& InNodeID, TArray<const FMetasoundFrontendEdge*>& OutEdgesCreated, bool bReplaceExistingConnections = true);
 	const FMetasoundFrontendNode* AddGraphInput(const FMetasoundFrontendClassInput& InClassInput);
+	const FMetasoundFrontendNode* DuplicateGraphInput(const FMetasoundFrontendClassInput& InClassInput, FMetasoundFrontendLiteral DefaultValue, const FName InName);
 	const FMetasoundFrontendNode* AddGraphNode(const FMetasoundFrontendGraphClass& InClass, FGuid InNodeID = FGuid::NewGuid());
 	const FMetasoundFrontendNode* AddGraphOutput(const FMetasoundFrontendClassOutput& InClassOutput);
+	const FMetasoundFrontendNode* DuplicateGraphOutput(const FMetasoundFrontendClassOutput& InClassOutput, const FName InName);
 
 	bool AddInterface(FName InterfaceName);
 
