- Protect against invalid builder or MetaSound object in MetaSound Asset Editor
- Clean-up multiple ways MetaSound is tracked in Asset Editor (ex. FGCObject in ToolKit overload, internal to ToolKit, StrongObjectPtr, etc.)
    - Now just uses EditingObject and GetMetasoundObject() wraps it safely
- Reordered & cleaned asset editor Init logic to accomidate condensing references
#rnx

diff --git a/Engine/Plugins/Runtime/Metasound/Source/MetasoundEditor/Private/MetasoundDetailCustomization.cpp b/Engine/Plugins/Runtime/Metasound/Source/MetasoundEditor/Private/MetasoundDetailCustomization.cpp
--- a/Engine/Plugins/Runtime/Metasound/Source/MetasoundEditor/Private/MetasoundDetailCustomization.cpp
+++ b/Engine/Plugins/Runtime/Metasound/Source/MetasoundEditor/Private/MetasoundDetailCustomization.cpp
@@ -14,6 +14,7 @@
 #include "Interfaces/MetasoundFrontendInterfaceRegistry.h"
 #include "MetasoundAssetBase.h"
 #include "MetasoundBuilderSubsystem.h"
+#include "MetasoundDocumentBuilderRegistry.h"
 #include "MetasoundEditor.h"
 #include "MetasoundEditorGraphBuilder.h"
 #include "MetasoundEditorSettings.h"
@@ -44,534 +45,542 @@
 #define LOCTEXT_NAMESPACE "MetaSoundEditor"
 
 
-namespace Metasound
+namespace Metasound::Editor
 {
-	namespace Editor
+	FName BuildChildPath(const FString& InBasePath, FName InPropertyName)
 	{
-		FName BuildChildPath(const FString& InBasePath, FName InPropertyName)
-		{
-			return FName(InBasePath + TEXT(".") + InPropertyName.ToString());
-		}
+		return FName(InBasePath + TEXT(".") + InPropertyName.ToString());
+	}
+
+	FName BuildChildPath(const FName& InBasePath, FName InPropertyName)
+	{
+		return FName(InBasePath.ToString() + TEXT(".") + InPropertyName.ToString());
+	}
 
-		FName BuildChildPath(const FName& InBasePath, FName InPropertyName)
+	UObject* FMetaSoundDetailCustomizationBase::GetMetaSound() const
+	{
+		if (Builder.IsValid())
 		{
-			return FName(InBasePath.ToString() + TEXT(".") + InPropertyName.ToString());
+			const FMetaSoundFrontendDocumentBuilder& DocBuilder = Builder->GetBuilder();
+			if (DocBuilder.IsValid())
+			{
+				return &DocBuilder.CastDocumentObjectChecked<UObject>();
+			}
 		}
 
-		FMetasoundDetailCustomization::FMetasoundDetailCustomization(FName InDocumentPropertyName)
-			: IDetailCustomization()
-			, DocumentPropertyName(InDocumentPropertyName)
+		return nullptr;
+	}
+
+	void FMetaSoundDetailCustomizationBase::InitBuilder(UObject& MetaSound)
+	{
+		using namespace Engine;
+		Builder.Reset(&FDocumentBuilderRegistry::GetChecked().FindOrBeginBuilding(MetaSound));
+
+		IsGraphEditableAttribute = TAttribute<bool>::Create([this]()
 		{
-			IsGraphEditableAttribute = TAttribute<bool>::Create([this]()
+			using namespace Engine;
+
+			if (Builder.IsValid())
 			{
-				using namespace Metasound;
-				using namespace Metasound::Frontend;
-				if (const FMetasoundAssetBase* MetaSoundAsset = IMetasoundUObjectRegistry::Get().GetObjectAsAssetBase(MetaSound.Get()))
+				const FMetaSoundFrontendDocumentBuilder& DocBuilder = Builder->GetConstBuilder();
+				if (DocBuilder.IsValid())
 				{
-					FConstGraphHandle GraphHandle = MetaSoundAsset->GetRootGraphHandle();
-					return GraphHandle->GetGraphStyle().bIsGraphEditable;
+					return DocBuilder.FindConstBuildGraphChecked().Style.bIsGraphEditable;
 				}
+			}
 
-				return false;
-			});
-		}
+			return false;
+		});
+	}
 
-		FName FMetasoundDetailCustomization::GetInterfaceVersionsPath() const
-		{
-			return Metasound::Editor::BuildChildPath(DocumentPropertyName, GET_MEMBER_NAME_CHECKED(FMetasoundFrontendDocument, Interfaces));
-		}
+	FMetasoundDetailCustomization::FMetasoundDetailCustomization(FName InDocumentPropertyName)
+		: DocumentPropertyName(InDocumentPropertyName)
+	{
+	}
 
-		FName FMetasoundDetailCustomization::GetMetadataRootClassPath() const
-		{
-			return Metasound::Editor::BuildChildPath(DocumentPropertyName, GET_MEMBER_NAME_CHECKED(FMetasoundFrontendDocument, RootGraph));
-		}
+	FName FMetasoundDetailCustomization::GetInterfaceVersionsPath() const
+	{
+		return BuildChildPath(DocumentPropertyName, GET_MEMBER_NAME_CHECKED(FMetasoundFrontendDocument, Interfaces));
+	}
 
-		FName FMetasoundDetailCustomization::GetMetadataPropertyPath() const
-		{
-			const FName RootClass = FName(GetMetadataRootClassPath());
-			return Metasound::Editor::BuildChildPath(RootClass, GET_MEMBER_NAME_CHECKED(FMetasoundFrontendClass, Metadata));
-		}
+	FName FMetasoundDetailCustomization::GetMetadataRootClassPath() const
+	{
+		return BuildChildPath(DocumentPropertyName, GET_MEMBER_NAME_CHECKED(FMetasoundFrontendDocument, RootGraph));
+	}
 
-		void FMetasoundDetailCustomization::CustomizeDetails(IDetailLayoutBuilder& DetailLayout)
-		{
-			using namespace Frontend;
+	FName FMetasoundDetailCustomization::GetMetadataPropertyPath() const
+	{
+		const FName RootClass = FName(GetMetadataRootClassPath());
+		return BuildChildPath(RootClass, GET_MEMBER_NAME_CHECKED(FMetasoundFrontendClass, Metadata));
+	}
 
-			TArray<TWeakObjectPtr<UObject>> Objects;
-			DetailLayout.GetObjectsBeingCustomized(Objects);
+	void FMetasoundDetailCustomization::CustomizeDetails(IDetailLayoutBuilder& DetailLayout)
+	{
+		using namespace Frontend;
 
-			// Only support modifying a single MetaSound at a time (Multiple
-			// MetaSound editing will be covered most likely by separate tool).
-			if (Objects.Num() > 1)
-			{
-				return;
-			}
-			MetaSound = Objects.Last();
-			if (!MetaSound.IsValid())
-			{
-				return;
-			}
+		TArray<TWeakObjectPtr<UObject>> Objects;
+		DetailLayout.GetObjectsBeingCustomized(Objects);
 
-			TWeakObjectPtr<UMetaSoundSource> MetaSoundSource = Cast<UMetaSoundSource>(MetaSound.Get());
+		// Only support modifying a single MetaSound at a time (Multiple
+		// MetaSound editing will be covered most likely by separate tool).
+		if (Objects.Num() > 1 || !Objects.Last().IsValid())
+		{
+			return;
+		}
+
+		UObject& MetaSound = *Objects.Last().Get();
+		InitBuilder(MetaSound);
+		TWeakObjectPtr<UMetaSoundSource> MetaSoundSource = Cast<UMetaSoundSource>(&MetaSound);
 
-			// MetaSound patches don't have source settings, so view MetaSound settings by default 
-			EMetasoundActiveDetailView DetailsView = EMetasoundActiveDetailView::Metasound;
-			if (MetaSoundSource.IsValid())
+		// MetaSound patches don't have source settings, so view MetaSound settings by default 
+		EMetasoundActiveDetailView DetailsView = EMetasoundActiveDetailView::Metasound;
+		if (MetaSoundSource.IsValid())
+		{
+			// Show source settings by default unless previously set
+			DetailsView = EMetasoundActiveDetailView::General;
+			if (const UMetasoundEditorSettings* EditorSettings = GetDefault<UMetasoundEditorSettings>())
 			{
-				// Show source settings by default unless previously set
-				DetailsView = EMetasoundActiveDetailView::General;
-				if (const UMetasoundEditorSettings* EditorSettings = GetDefault<UMetasoundEditorSettings>())
-				{
-					DetailsView = EditorSettings->DetailView;
-				}
+				DetailsView = EditorSettings->DetailView;
 			}
+		}
 
-			switch (DetailsView)
+		switch (DetailsView)
+		{
+			case EMetasoundActiveDetailView::Metasound:
 			{
-				case EMetasoundActiveDetailView::Metasound:
+				IDetailCategoryBuilder& GeneralCategoryBuilder = DetailLayout.EditCategory("MetaSound");
+				const FName AuthorPropertyPath = BuildChildPath(GetMetadataPropertyPath(), FMetasoundFrontendClassMetadata::GetAuthorPropertyName());
+				const FName CategoryHierarchyPropertyPath = BuildChildPath(GetMetadataPropertyPath(), FMetasoundFrontendClassMetadata::GetCategoryHierarchyPropertyName());
+				const FName DescPropertyPath = BuildChildPath(GetMetadataPropertyPath(), FMetasoundFrontendClassMetadata::GetDescriptionPropertyName());
+				const FName DisplayNamePropertyPath = BuildChildPath(GetMetadataPropertyPath(), FMetasoundFrontendClassMetadata::GetDisplayNamePropertyName());
+				const FName KeywordsPropertyPath = BuildChildPath(GetMetadataPropertyPath(), FMetasoundFrontendClassMetadata::GetKeywordsPropertyName());
+				const FName IsDeprecatedPropertyPath = BuildChildPath(GetMetadataPropertyPath(), FMetasoundFrontendClassMetadata::GetIsDeprecatedPropertyName());
+
+				const FName ClassNamePropertyPath = BuildChildPath(GetMetadataPropertyPath(), FMetasoundFrontendClassMetadata::GetClassNamePropertyName());
+				const FName ClassNameNamePropertyPath = BuildChildPath(ClassNamePropertyPath, GET_MEMBER_NAME_CHECKED(FMetasoundFrontendClassName, Name));
+
+				const FName VersionPropertyPath = BuildChildPath(GetMetadataPropertyPath(), FMetasoundFrontendClassMetadata::GetVersionPropertyName());
+				const FName MajorVersionPropertyPath = BuildChildPath(VersionPropertyPath, GET_MEMBER_NAME_CHECKED(FMetasoundFrontendVersionNumber, Major));
+				const FName MinorVersionPropertyPath = BuildChildPath(VersionPropertyPath, GET_MEMBER_NAME_CHECKED(FMetasoundFrontendVersionNumber, Minor));
+
+				const FName InterfaceVersionsPropertyPath = GetInterfaceVersionsPath();
+
+				TSharedPtr<IPropertyHandle> AuthorHandle = DetailLayout.GetProperty(AuthorPropertyPath);
+				TSharedPtr<IPropertyHandle> CategoryHierarchyHandle = DetailLayout.GetProperty(CategoryHierarchyPropertyPath);
+				TSharedPtr<IPropertyHandle> ClassNameHandle = DetailLayout.GetProperty(ClassNameNamePropertyPath);
+				TSharedPtr<IPropertyHandle> DisplayNameHandle = DetailLayout.GetProperty(DisplayNamePropertyPath);
+				TSharedPtr<IPropertyHandle> DescHandle = DetailLayout.GetProperty(DescPropertyPath);
+				TSharedPtr<IPropertyHandle> KeywordsHandle = DetailLayout.GetProperty(KeywordsPropertyPath);
+				TSharedPtr<IPropertyHandle> IsDeprecatedHandle = DetailLayout.GetProperty(IsDeprecatedPropertyPath);
+				TSharedPtr<IPropertyHandle> InterfaceVersionsHandle = DetailLayout.GetProperty(InterfaceVersionsPropertyPath);
+				TSharedPtr<IPropertyHandle> MajorVersionHandle = DetailLayout.GetProperty(MajorVersionPropertyPath);
+				TSharedPtr<IPropertyHandle> MinorVersionHandle = DetailLayout.GetProperty(MinorVersionPropertyPath);
+
+				// Invalid for UMetaSounds
+				TSharedPtr<IPropertyHandle> OutputFormat = DetailLayout.GetProperty(GET_MEMBER_NAME_CHECKED(UMetaSoundSource, OutputFormat));
+				if (OutputFormat.IsValid())
 				{
-					IDetailCategoryBuilder& GeneralCategoryBuilder = DetailLayout.EditCategory("MetaSound");
-					const FName AuthorPropertyPath = BuildChildPath(GetMetadataPropertyPath(), FMetasoundFrontendClassMetadata::GetAuthorPropertyName());
-					const FName CategoryHierarchyPropertyPath = BuildChildPath(GetMetadataPropertyPath(), FMetasoundFrontendClassMetadata::GetCategoryHierarchyPropertyName());
-					const FName DescPropertyPath = BuildChildPath(GetMetadataPropertyPath(), FMetasoundFrontendClassMetadata::GetDescriptionPropertyName());
-					const FName DisplayNamePropertyPath = BuildChildPath(GetMetadataPropertyPath(), FMetasoundFrontendClassMetadata::GetDisplayNamePropertyName());
-					const FName KeywordsPropertyPath = BuildChildPath(GetMetadataPropertyPath(), FMetasoundFrontendClassMetadata::GetKeywordsPropertyName());
-					const FName IsDeprecatedPropertyPath = BuildChildPath(GetMetadataPropertyPath(), FMetasoundFrontendClassMetadata::GetIsDeprecatedPropertyName());
-
-					const FName ClassNamePropertyPath = BuildChildPath(GetMetadataPropertyPath(), FMetasoundFrontendClassMetadata::GetClassNamePropertyName());
-					const FName ClassNameNamePropertyPath = BuildChildPath(ClassNamePropertyPath, GET_MEMBER_NAME_CHECKED(FMetasoundFrontendClassName, Name));
-
-					const FName VersionPropertyPath = BuildChildPath(GetMetadataPropertyPath(), FMetasoundFrontendClassMetadata::GetVersionPropertyName());
-					const FName MajorVersionPropertyPath = BuildChildPath(VersionPropertyPath, GET_MEMBER_NAME_CHECKED(FMetasoundFrontendVersionNumber, Major));
-					const FName MinorVersionPropertyPath = BuildChildPath(VersionPropertyPath, GET_MEMBER_NAME_CHECKED(FMetasoundFrontendVersionNumber, Minor));
-
-					const FName InterfaceVersionsPropertyPath = GetInterfaceVersionsPath();
-
-					TSharedPtr<IPropertyHandle> AuthorHandle = DetailLayout.GetProperty(AuthorPropertyPath);
-					TSharedPtr<IPropertyHandle> CategoryHierarchyHandle = DetailLayout.GetProperty(CategoryHierarchyPropertyPath);
-					TSharedPtr<IPropertyHandle> ClassNameHandle = DetailLayout.GetProperty(ClassNameNamePropertyPath);
-					TSharedPtr<IPropertyHandle> DisplayNameHandle = DetailLayout.GetProperty(DisplayNamePropertyPath);
-					TSharedPtr<IPropertyHandle> DescHandle = DetailLayout.GetProperty(DescPropertyPath);
-					TSharedPtr<IPropertyHandle> KeywordsHandle = DetailLayout.GetProperty(KeywordsPropertyPath);
-					TSharedPtr<IPropertyHandle> IsDeprecatedHandle = DetailLayout.GetProperty(IsDeprecatedPropertyPath);
-					TSharedPtr<IPropertyHandle> InterfaceVersionsHandle = DetailLayout.GetProperty(InterfaceVersionsPropertyPath);
-					TSharedPtr<IPropertyHandle> MajorVersionHandle = DetailLayout.GetProperty(MajorVersionPropertyPath);
-					TSharedPtr<IPropertyHandle> MinorVersionHandle = DetailLayout.GetProperty(MinorVersionPropertyPath);
-
-					// Invalid for UMetaSounds
-					TSharedPtr<IPropertyHandle> OutputFormat = DetailLayout.GetProperty(GET_MEMBER_NAME_CHECKED(UMetaSoundSource, OutputFormat));
-					if (OutputFormat.IsValid())
+					if (MetaSoundSource.IsValid())
 					{
-						if (MetaSoundSource.IsValid())
+						OutputFormat->SetOnPropertyValuePreChange(FSimpleDelegate::CreateLambda([Source = MetaSoundSource]()
 						{
-							OutputFormat->SetOnPropertyValuePreChange(FSimpleDelegate::CreateLambda([Source = MetaSoundSource]()
-							{
-								if (Source.IsValid())
-								{
-									TSharedPtr<FEditor> ParentEditor = FGraphBuilder::GetEditorForMetasound(*Source.Get());
-									if (ParentEditor.IsValid())
-									{
-										ParentEditor->Stop();
-									};
-								}
-							}));
-
-							OutputFormat->SetOnPropertyValueChanged(FSimpleDelegate::CreateLambda([Source = MetaSoundSource]()
+							if (Source.IsValid())
 							{
-								if (Source.IsValid())
+								TSharedPtr<FEditor> ParentEditor = FGraphBuilder::GetEditorForMetasound(*Source.Get());
+								if (ParentEditor.IsValid())
 								{
-									TSharedPtr<FEditor> ParentEditor = FGraphBuilder::GetEditorForMetasound(*Source.Get());
-									if (ParentEditor.IsValid())
-									{
-										ParentEditor->CreateAnalyzers();
-									};
-								}
-							}));
-						}
-
-						TSharedRef<SWidget> OutputFormatValueWidget = OutputFormat->CreatePropertyValueWidget();
-						OutputFormatValueWidget->SetEnabled(IsGraphEditableAttribute);
-
-						static const FText OutputFormatName = LOCTEXT("MetasoundOutputFormatPropertyName", "Output Format");
-						GeneralCategoryBuilder.AddCustomRow(OutputFormatName)
-						.NameContent()
-						[
-							OutputFormat->CreatePropertyNameWidget()
-						]
-						.ValueContent()
-						[
-							OutputFormatValueWidget
-						];
-
-						OutputFormat->MarkHiddenByCustomization();
-					}
+									ParentEditor->Stop();
+								};
+							}
+						}));
 
-					// Updates FText properties on open editors if required
-					{
-						FSimpleDelegate RegisterOnChange = FSimpleDelegate::CreateLambda([this]()
+						OutputFormat->SetOnPropertyValueChanged(FSimpleDelegate::CreateLambda([Source = MetaSoundSource]()
+						{
+							if (Source.IsValid())
 							{
-								if (MetaSound.IsValid())
+								TSharedPtr<FEditor> ParentEditor = FGraphBuilder::GetEditorForMetasound(*Source.Get());
+								if (ParentEditor.IsValid())
 								{
-									if (FMetasoundAssetBase* MetaSoundAsset = IMetasoundUObjectRegistry::Get().GetObjectAsAssetBase(MetaSound.Get()))
-									{
-										MetaSoundAsset->GetDocumentChecked().RootGraph.Style.UpdateChangeID();
-									}
-									constexpr bool bForceViewSynchronization = true;
-									FGraphBuilder::RegisterGraphWithFrontend(*MetaSound.Get(), bForceViewSynchronization);
-								}
-							});
-						AuthorHandle->SetOnChildPropertyValueChanged(RegisterOnChange);
-						DescHandle->SetOnPropertyValueChanged(RegisterOnChange);
-						DisplayNameHandle->SetOnPropertyValueChanged(RegisterOnChange);
-						KeywordsHandle->SetOnPropertyValueChanged(RegisterOnChange);
-						KeywordsHandle->SetOnChildPropertyValueChanged(RegisterOnChange);
-						IsDeprecatedHandle->SetOnPropertyValueChanged(RegisterOnChange);
+									ParentEditor->CreateAnalyzers(*Source.Get());
+								};
+							}
+						}));
 					}
 
-					GeneralCategoryBuilder.AddProperty(DisplayNameHandle);
-					GeneralCategoryBuilder.AddProperty(DescHandle);
-					GeneralCategoryBuilder.AddProperty(AuthorHandle);
-					GeneralCategoryBuilder.AddProperty(IsDeprecatedHandle);
-					GeneralCategoryBuilder.AddProperty(MajorVersionHandle);
-					GeneralCategoryBuilder.AddProperty(MinorVersionHandle);
+					TSharedRef<SWidget> OutputFormatValueWidget = OutputFormat->CreatePropertyValueWidget();
+					OutputFormatValueWidget->SetEnabled(IsGraphEditableAttribute);
 
-					static const FText ClassGuidName = LOCTEXT("MetasoundClassGuidPropertyName", "Class Guid");
-					GeneralCategoryBuilder.AddCustomRow(ClassGuidName).NameContent()
+					static const FText OutputFormatName = LOCTEXT("MetasoundOutputFormatPropertyName", "Output Format");
+					GeneralCategoryBuilder.AddCustomRow(OutputFormatName)
+					.NameContent()
 					[
-						SNew(SHorizontalBox)
-						+ SHorizontalBox::Slot()
-						.VAlign(VAlign_Center)
-						[
-							SNew(STextBlock)
-							.Text(ClassGuidName)
-							.Font(FAppStyle::GetFontStyle(TEXT("PropertyWindow.NormalFont")))
-						]
+						OutputFormat->CreatePropertyNameWidget()
 					]
 					.ValueContent()
 					[
-						ClassNameHandle->CreatePropertyValueWidget()
+						OutputFormatValueWidget
 					];
-					GeneralCategoryBuilder.AddProperty(CategoryHierarchyHandle);
-					GeneralCategoryBuilder.AddProperty(KeywordsHandle);
-
-					DetailLayout.HideCategory("Attenuation");
-					DetailLayout.HideCategory("Developer");
-					DetailLayout.HideCategory("Effects");
-					DetailLayout.HideCategory("Loading");
-					DetailLayout.HideCategory("Modulation");
-					DetailLayout.HideCategory("Sound");
-					DetailLayout.HideCategory("Voice Management");
-				}
-				break;
-
-				case EMetasoundActiveDetailView::General:
-				default:
-					DetailLayout.HideCategory("MetaSound");
 
-					TArray<TSharedRef<IPropertyHandle>>DeveloperProperties;
-					TArray<TSharedRef<IPropertyHandle>>SoundProperties;
-
-					DetailLayout.EditCategory("Sound")
-						.GetDefaultProperties(SoundProperties);
-					DetailLayout.EditCategory("Developer")
-						.GetDefaultProperties(DeveloperProperties);
+					OutputFormat->MarkHiddenByCustomization();
+				}
 
-					auto HideProperties = [](const TSet<FName>& PropsToHide, const TArray<TSharedRef<IPropertyHandle>>& Properties)
+				// Updates FText properties on open editors if required
+				{
+					FSimpleDelegate RegisterOnChange = FSimpleDelegate::CreateLambda([this]()
 					{
-						for (TSharedRef<IPropertyHandle> Property : Properties)
+						if (Builder.IsValid())
 						{
-							if (PropsToHide.Contains(Property->GetProperty()->GetFName()))
+							FMetaSoundFrontendDocumentBuilder& DocBuilder = Builder->GetBuilder();
+							if (DocBuilder.IsValid())
 							{
-								Property->MarkHiddenByCustomization();
+								DocBuilder.GetConstDocumentChecked().RootGraph.Style.UpdateChangeID();
 							}
+							constexpr bool bForceViewSynchronization = true;
+							FGraphBuilder::RegisterGraphWithFrontend(DocBuilder.CastDocumentObjectChecked<UObject>(), bForceViewSynchronization);
 						}
-					};
+					});
+					AuthorHandle->SetOnChildPropertyValueChanged(RegisterOnChange);
+					DescHandle->SetOnPropertyValueChanged(RegisterOnChange);
+					DisplayNameHandle->SetOnPropertyValueChanged(RegisterOnChange);
+					KeywordsHandle->SetOnPropertyValueChanged(RegisterOnChange);
+					KeywordsHandle->SetOnChildPropertyValueChanged(RegisterOnChange);
+					IsDeprecatedHandle->SetOnPropertyValueChanged(RegisterOnChange);
+				}
 
-					static const TSet<FName> SoundPropsToHide =
-					{
-						GET_MEMBER_NAME_CHECKED(USoundWave, bLooping),
-						GET_MEMBER_NAME_CHECKED(USoundWave, SoundGroup)
-					};
-					HideProperties(SoundPropsToHide, SoundProperties);
+				GeneralCategoryBuilder.AddProperty(DisplayNameHandle);
+				GeneralCategoryBuilder.AddProperty(DescHandle);
+				GeneralCategoryBuilder.AddProperty(AuthorHandle);
+				GeneralCategoryBuilder.AddProperty(IsDeprecatedHandle);
+				GeneralCategoryBuilder.AddProperty(MajorVersionHandle);
+				GeneralCategoryBuilder.AddProperty(MinorVersionHandle);
 
-					static const TSet<FName> DeveloperPropsToHide =
+				static const FText ClassGuidName = LOCTEXT("MetasoundClassGuidPropertyName", "Class Guid");
+				GeneralCategoryBuilder.AddCustomRow(ClassGuidName).NameContent()
+				[
+					SNew(SHorizontalBox)
+					+ SHorizontalBox::Slot()
+					.VAlign(VAlign_Center)
+					[
+						SNew(STextBlock)
+						.Text(ClassGuidName)
+						.Font(FAppStyle::GetFontStyle(TEXT("PropertyWindow.NormalFont")))
+					]
+				]
+				.ValueContent()
+				[
+					ClassNameHandle->CreatePropertyValueWidget()
+				];
+				GeneralCategoryBuilder.AddProperty(CategoryHierarchyHandle);
+				GeneralCategoryBuilder.AddProperty(KeywordsHandle);
+
+				DetailLayout.HideCategory("Attenuation");
+				DetailLayout.HideCategory("Developer");
+				DetailLayout.HideCategory("Effects");
+				DetailLayout.HideCategory("Loading");
+				DetailLayout.HideCategory("Modulation");
+				DetailLayout.HideCategory("Sound");
+				DetailLayout.HideCategory("Voice Management");
+			}
+			break;
+
+			case EMetasoundActiveDetailView::General:
+			default:
+				DetailLayout.HideCategory("MetaSound");
+
+				TArray<TSharedRef<IPropertyHandle>>DeveloperProperties;
+				TArray<TSharedRef<IPropertyHandle>>SoundProperties;
+
+				DetailLayout.EditCategory("Sound")
+					.GetDefaultProperties(SoundProperties);
+				DetailLayout.EditCategory("Developer")
+					.GetDefaultProperties(DeveloperProperties);
+
+				auto HideProperties = [](const TSet<FName>& PropsToHide, const TArray<TSharedRef<IPropertyHandle>>& Properties)
+				{
+					for (TSharedRef<IPropertyHandle> Property : Properties)
 					{
-						GET_MEMBER_NAME_CHECKED(USoundBase, Duration),
-						GET_MEMBER_NAME_CHECKED(USoundBase, MaxDistance),
-						GET_MEMBER_NAME_CHECKED(USoundBase, TotalSamples)
-					};
-					HideProperties(DeveloperPropsToHide, DeveloperProperties);
+						if (PropsToHide.Contains(Property->GetProperty()->GetFName()))
+						{
+							Property->MarkHiddenByCustomization();
+						}
+					}
+				};
 
-					break;
-			}
+				static const TSet<FName> SoundPropsToHide =
+				{
+					GET_MEMBER_NAME_CHECKED(USoundWave, bLooping),
+					GET_MEMBER_NAME_CHECKED(USoundWave, SoundGroup)
+				};
+				HideProperties(SoundPropsToHide, SoundProperties);
 
-			// Hack to hide parent structs for nested metadata properties
-			DetailLayout.HideCategory("CustomView");
-
-			DetailLayout.HideCategory("Analysis");
-			DetailLayout.HideCategory("Curves");
-			DetailLayout.HideCategory("File Path");
-			DetailLayout.HideCategory("Format");
-			DetailLayout.HideCategory("Info");
-			DetailLayout.HideCategory("Loading");
-			DetailLayout.HideCategory("Playback");
-			DetailLayout.HideCategory("Subtitles");
-			DetailLayout.HideCategory("Waveform Processing");
+				static const TSet<FName> DeveloperPropsToHide =
+				{
+					GET_MEMBER_NAME_CHECKED(USoundBase, Duration),
+					GET_MEMBER_NAME_CHECKED(USoundBase, MaxDistance),
+					GET_MEMBER_NAME_CHECKED(USoundBase, TotalSamples)
+				};
+				HideProperties(DeveloperPropsToHide, DeveloperProperties);
+
+				break;
 		}
 
-		FMetasoundInterfacesDetailCustomization::FMetasoundInterfacesDetailCustomization()
+		// Hack to hide parent structs for nested metadata properties
+		DetailLayout.HideCategory("CustomView");
+
+		DetailLayout.HideCategory("Analysis");
+		DetailLayout.HideCategory("Curves");
+		DetailLayout.HideCategory("File Path");
+		DetailLayout.HideCategory("Format");
+		DetailLayout.HideCategory("Info");
+		DetailLayout.HideCategory("Loading");
+		DetailLayout.HideCategory("Playback");
+		DetailLayout.HideCategory("Subtitles");
+		DetailLayout.HideCategory("Waveform Processing");
+	}
+
+	void FMetasoundInterfacesDetailCustomization::CustomizeDetails(IDetailLayoutBuilder& DetailLayout)
+	{
+		TArray<TWeakObjectPtr<UObject>> Objects;
+		DetailLayout.GetObjectsBeingCustomized(Objects);
+
+		// Only support modifying a single MetaSound at a time (Multiple
+		// MetaSound editing will be covered most likely by separate tool).
+		if (Objects.Num() > 1)
 		{
-			IsGraphEditableAttribute = TAttribute<bool>::Create([this]()
+			return;
+		}
+		if (UMetasoundInterfacesView* InterfacesView = CastChecked<UMetasoundInterfacesView>(Objects.Last()))
+		{
+			if (UObject* MetaSound = InterfacesView->GetMetasound().Get())
 			{
-				using namespace Metasound;
-				using namespace Metasound::Frontend;
-				if (const FMetasoundAssetBase* MetaSoundAsset = IMetasoundUObjectRegistry::Get().GetObjectAsAssetBase(MetaSound.Get()))
-				{
-					FConstGraphHandle GraphHandle = MetaSoundAsset->GetRootGraphHandle();
-					return GraphHandle->GetGraphStyle().bIsGraphEditable;
-				}
-
-				return false;
-			});
+				InitBuilder(*MetaSound);
+			}
 		}
 
-		void FMetasoundInterfacesDetailCustomization::CustomizeDetails(IDetailLayoutBuilder& DetailLayout)
-		{
-			TArray<TWeakObjectPtr<UObject>> Objects;
-			DetailLayout.GetObjectsBeingCustomized(Objects);
+		UpdateInterfaceNames();
 
-			// Only support modifying a single MetaSound at a time (Multiple
-			// MetaSound editing will be covered most likely by separate tool).
-			if (Objects.Num() > 1)
+		SAssignNew(InterfaceComboBox, SSearchableComboBox)
+			.OptionsSource(&AddableInterfaceNames)
+			.OnGenerateWidget_Lambda([](TSharedPtr<FString> InItem)
 			{
-				return;
-			}
-			if (UMetasoundInterfacesView* InterfacesView = CastChecked<UMetasoundInterfacesView>(Objects.Last()))
+				return SNew(STextBlock)
+					.Text(FText::FromString(*InItem));
+			})
+			.OnSelectionChanged_Lambda([this](TSharedPtr<FString> NameToAdd, ESelectInfo::Type InSelectInfo)
 			{
-				MetaSound = InterfacesView->GetMetasound();
-			}
-
-			UpdateInterfaceNames();
+				using namespace Metasound;
+				using namespace Metasound::Frontend;
 
-			SAssignNew(InterfaceComboBox, SSearchableComboBox)
-				.OptionsSource(&AddableInterfaceNames)
-				.OnGenerateWidget_Lambda([](TSharedPtr<FString> InItem)
+				if (Builder.IsValid() && InSelectInfo != ESelectInfo::OnNavigation)
 				{
-					return SNew(STextBlock)
-						.Text(FText::FromString(*InItem));
-				})
-				.OnSelectionChanged_Lambda([this](TSharedPtr<FString> NameToAdd, ESelectInfo::Type InSelectInfo)
+					FMetaSoundFrontendDocumentBuilder& DocBuilder = Builder->GetBuilder();
+					UObject& MetaSound = DocBuilder.CastDocumentObjectChecked<UObject>();
+					FMetasoundFrontendInterface InterfaceToAdd;
+					const FName InterfaceName { *NameToAdd.Get() };
+					if (ensure(ISearchEngine::Get().FindInterfaceWithHighestVersion(InterfaceName, InterfaceToAdd)))
+					{
+						const FScopedTransaction Transaction(FText::Format(LOCTEXT("AddInterfaceTransactionFormat", "Add MetaSound Interface '{0}'"), FText::FromString(InterfaceToAdd.Version.ToString())));
+						MetaSound.Modify();
+						FModifyInterfaceOptions Options({ }, { InterfaceToAdd });
+						Options.bSetDefaultNodeLocations = false; // Don't automatically add nodes to ed graph
+						DocBuilder.ModifyInterfaces(MoveTemp(Options));
+					}
+
+					UpdateInterfaceNames();
+					InterfaceComboBox->RefreshOptions();
+					FGraphBuilder::RegisterGraphWithFrontend(MetaSound);
+				}
+			})
+			.Content()
+			[
+				SNew(STextBlock)
+				.Text(LOCTEXT("UpdateInterfaceAction", "Add Interface..."))
+				.IsEnabled(IsGraphEditableAttribute)
+			];
+
+		TSharedRef<SWidget> InterfaceUtilities = SNew(SHorizontalBox)
+			+ SHorizontalBox::Slot()
+			.Padding(2.0f)
+			.HAlign(HAlign_Left)
+			.VAlign(VAlign_Center)
+			.AutoWidth()
+			[
+				InterfaceComboBox->AsShared()
+			]
+		+ SHorizontalBox::Slot()
+			.Padding(2.0f)
+			.HAlign(HAlign_Left)
+			.VAlign(VAlign_Center)
+			.AutoWidth()
+			[
+				PropertyCustomizationHelpers::MakeDeleteButton(FSimpleDelegate::CreateLambda([this]()
 				{
-					using namespace Metasound;
-					using namespace Metasound::Frontend;
+					using namespace Frontend;
 
-					FMetasoundAssetBase* MetaSoundAsset = IMetasoundUObjectRegistry::Get().GetObjectAsAssetBase(MetaSound.Get());
-					if (!ensure(MetaSoundAsset))
+					if (!Builder.IsValid())
 					{
 						return;
 					}
 
-					if (InSelectInfo != ESelectInfo::OnNavigation)
+					FMetaSoundFrontendDocumentBuilder& DocBuilder = Builder->GetBuilder();
+					if (!DocBuilder.IsValid())
 					{
-						FMetasoundFrontendInterface InterfaceToAdd;
-						const FName InterfaceName { *NameToAdd.Get() };
-						if (ensure(ISearchEngine::Get().FindInterfaceWithHighestVersion(InterfaceName, InterfaceToAdd)))
-						{
-							const FScopedTransaction Transaction(FText::Format(LOCTEXT("AddInterfaceTransactionFormat", "Add MetaSound Interface '{0}'"), FText::FromString(InterfaceToAdd.Version.ToString())));
-							MetaSound.Get()->Modify();
-							MetaSoundAsset->GetGraphChecked().Modify();
-							FMetaSoundFrontendDocumentBuilder& Builder = IDocumentBuilderRegistry::GetChecked().FindOrBeginBuilding(MetaSound.Get());
-							FModifyInterfaceOptions Options({ }, { InterfaceToAdd });
-							Options.bSetDefaultNodeLocations = false; // Don't automatically add nodes to ed graph
-							Builder.ModifyInterfaces(MoveTemp(Options));
-						}
-
-						UpdateInterfaceNames();
-						InterfaceComboBox->RefreshOptions();
-						FGraphBuilder::RegisterGraphWithFrontend(*MetaSound.Get());
+						return;
 					}
-				})
-				.Content()
-				[
-					SNew(STextBlock)
-					.Text(LOCTEXT("UpdateInterfaceAction", "Add Interface..."))
-					.IsEnabled(IsGraphEditableAttribute)
-				];
 
-			TSharedRef<SWidget> InterfaceUtilities = SNew(SHorizontalBox)
-				+ SHorizontalBox::Slot()
-				.Padding(2.0f)
-				.HAlign(HAlign_Left)
-				.VAlign(VAlign_Center)
-				.AutoWidth()
-				[
-					InterfaceComboBox->AsShared()
-				]
-			+ SHorizontalBox::Slot()
-				.Padding(2.0f)
-				.HAlign(HAlign_Left)
-				.VAlign(VAlign_Center)
-				.AutoWidth()
-				[
-					PropertyCustomizationHelpers::MakeDeleteButton(FSimpleDelegate::CreateLambda([this]()
+					TArray<FMetasoundFrontendInterface> ImplementedInterfaces;
+					Algo::Transform(ImplementedInterfaceNames, ImplementedInterfaces, [](const FName& Name)
 					{
-						using namespace Frontend;
+						FMetasoundFrontendInterface Interface;
+						ISearchEngine::Get().FindInterfaceWithHighestVersion(Name, Interface);
+						return Interface;
+					});
 
-						FMetasoundAssetBase* MetaSoundAsset = IMetasoundUObjectRegistry::Get().GetObjectAsAssetBase(MetaSound.Get());
-						if (!ensure(MetaSoundAsset))
-						{
-							return;
-						}
-
-						TArray<FMetasoundFrontendInterface> ImplementedInterfaces;
-						Algo::Transform(ImplementedInterfaceNames, ImplementedInterfaces, [](const FName& Name)
-						{
-							FMetasoundFrontendInterface Interface;
-							ISearchEngine::Get().FindInterfaceWithHighestVersion(Name, Interface);
-							return Interface;
-						});
+					UObject& MetaSound = DocBuilder.CastDocumentObjectChecked<UObject>();
+					{
+						const FScopedTransaction Transaction(LOCTEXT("RemoveAllInterfacesTransaction", "Remove All MetaSound Interfaces"));
+						MetaSound.Modify();
+						FModifyInterfaceOptions Options(ImplementedInterfaces, { });
+						Options.bSetDefaultNodeLocations = false; // Don't automatically add nodes to ed graph
+						DocBuilder.ModifyInterfaces(MoveTemp(Options));
+					}
 
-						{
-							const FScopedTransaction Transaction(LOCTEXT("RemoveAllInterfacesTransaction", "Remove All MetaSound Interfaces"));
-							MetaSound.Get()->Modify();
-							MetaSoundAsset->GetGraphChecked().Modify();
+					UpdateInterfaceNames();
+					InterfaceComboBox->RefreshOptions();
+					FGraphBuilder::RegisterGraphWithFrontend(MetaSound);
 
-							FDocumentHandle DocumentHandle = MetaSoundAsset->GetDocumentHandle();
-							FModifyRootGraphInterfaces({ ImplementedInterfaces }, { }).Transform(DocumentHandle);
-						}
+				}), LOCTEXT("RemoveInterfaceTooltip1", "Removes all interfaces from the given MetaSound."))
+			];
+		InterfaceUtilities->SetEnabled(IsGraphEditableAttribute);
 
-						UpdateInterfaceNames();
-						InterfaceComboBox->RefreshOptions();
-						FGraphBuilder::RegisterGraphWithFrontend(*MetaSound.Get());
+		const FText HeaderName = LOCTEXT("InterfacesGroupDisplayName", "Interfaces");
+		IDetailCategoryBuilder& InterfaceCategory = DetailLayout.EditCategory("Interfaces", HeaderName);
 
-					}), LOCTEXT("RemoveInterfaceTooltip1", "Removes all interfaces from the given MetaSound."))
-				];
-			InterfaceUtilities->SetEnabled(IsGraphEditableAttribute);
+		InterfaceCategory.AddCustomRow(HeaderName)
+		[
+			InterfaceUtilities
+		];
 
-			const FText HeaderName = LOCTEXT("InterfacesGroupDisplayName", "Interfaces");
-			IDetailCategoryBuilder& InterfaceCategory = DetailLayout.EditCategory("Interfaces", HeaderName);
+		auto CreateInterfaceEntryWidget = [&](FName InInterfaceName) -> TSharedRef<SWidget>
+		{
+			using namespace Frontend;
 
-			InterfaceCategory.AddCustomRow(HeaderName)
-			[
-				InterfaceUtilities
-			];
+			FMetasoundFrontendInterface InterfaceEntry;
+			if (!ensure(ISearchEngine::Get().FindInterfaceWithHighestVersion(InInterfaceName, InterfaceEntry)))
+			{
+				return SNullWidget::NullWidget;
+			}
 
-			auto CreateInterfaceEntryWidget = [&](FName InInterfaceName) -> TSharedRef<SWidget>
+			TSharedRef<SWidget> RemoveButtonWidget = PropertyCustomizationHelpers::MakeDeleteButton(FSimpleDelegate::CreateLambda([this, InInterfaceName, InterfaceEntry]()
 			{
 				using namespace Frontend;
 
-				FMetasoundFrontendInterface InterfaceEntry;
-				if (!ensure(ISearchEngine::Get().FindInterfaceWithHighestVersion(InInterfaceName, InterfaceEntry)))
+				if (!Builder.IsValid())
 				{
-					return SNullWidget::NullWidget;
+					return;
 				}
 
-				TSharedRef<SWidget> RemoveButtonWidget = PropertyCustomizationHelpers::MakeDeleteButton(FSimpleDelegate::CreateLambda([this, InInterfaceName, InterfaceEntry]()
+				FMetaSoundFrontendDocumentBuilder& DocBuilder = Builder->GetBuilder();
+				if (!DocBuilder.IsValid())
 				{
-					using namespace Frontend;
-
-					FMetasoundAssetBase* MetaSoundAsset = IMetasoundUObjectRegistry::Get().GetObjectAsAssetBase(MetaSound.Get());
-					if (!ensure(MetaSoundAsset))
-					{
-						return;
-					}
-
-					{
-						const FScopedTransaction Transaction(FText::Format(LOCTEXT("RemoveInterfaceTransactionFormat", "Remove MetaSound Interface '{0}'"), FText::FromString(InterfaceEntry.Version.ToString())));
-						MetaSound.Get()->Modify();
-						MetaSoundAsset->GetGraphChecked().Modify();
+					return;
+				}
 
-						FDocumentHandle DocumentHandle = MetaSoundAsset->GetDocumentHandle();
-						FModifyRootGraphInterfaces({ InterfaceEntry }, { }).Transform(DocumentHandle);
-					}
+				UObject& MetaSound = DocBuilder.CastDocumentObjectChecked<UObject>();
+				{
+					const FScopedTransaction Transaction(FText::Format(LOCTEXT("RemoveInterfaceTransactionFormat", "Remove MetaSound Interface '{0}'"), FText::FromString(InterfaceEntry.Version.ToString())));
+					MetaSound.Modify();
+					FModifyInterfaceOptions Options({ InterfaceEntry }, {});
+					Options.bSetDefaultNodeLocations = false; // Don't automatically add nodes to ed graph
+					DocBuilder.ModifyInterfaces(MoveTemp(Options));
+				}
 
-					UpdateInterfaceNames();
-					InterfaceComboBox->RefreshOptions();
-					FGraphBuilder::RegisterGraphWithFrontend(*MetaSound.Get());
+				UpdateInterfaceNames();
+				InterfaceComboBox->RefreshOptions();
+				FGraphBuilder::RegisterGraphWithFrontend(MetaSound);
 
-				}), LOCTEXT("RemoveInterfaceTooltip2", "Removes the associated interface from the MetaSound."));
+			}), LOCTEXT("RemoveInterfaceTooltip2", "Removes the associated interface from the MetaSound."));
 
-				TSharedRef<SWidget> EntryWidget = SNew(SHorizontalBox)
-					+ SHorizontalBox::Slot()
-					.Padding(2.0f)
-					.HAlign(HAlign_Center)
-					.VAlign(VAlign_Center)
-					.AutoWidth()
-					[
-						SNew(STextBlock)
-						.Text(FText::FromName(InterfaceEntry.Version.Name))
-					]
+			TSharedRef<SWidget> EntryWidget = SNew(SHorizontalBox)
 				+ SHorizontalBox::Slot()
-					.Padding(2.0f)
-					.HAlign(HAlign_Center)
-					.VAlign(VAlign_Center)
-					.AutoWidth()
-					[
-						RemoveButtonWidget
-					];
-
-				EntryWidget->SetEnabled(IsGraphEditableAttribute);
-				return EntryWidget;
-			};
-
-			TArray<FName> InterfaceNames = ImplementedInterfaceNames.Array();
-			InterfaceNames.Sort([](const FName& A, const FName& B) { return A.LexicalLess(B); });
-			for (const FName& InterfaceName : InterfaceNames)
-			{
-				InterfaceCategory.AddCustomRow(FText::FromName(InterfaceName))
+				.Padding(2.0f)
+				.HAlign(HAlign_Center)
+				.VAlign(VAlign_Center)
+				.AutoWidth()
+				[
+					SNew(STextBlock)
+					.Text(FText::FromName(InterfaceEntry.Version.Name))
+				]
+			+ SHorizontalBox::Slot()
+				.Padding(2.0f)
+				.HAlign(HAlign_Center)
+				.VAlign(VAlign_Center)
+				.AutoWidth()
 				[
-					CreateInterfaceEntryWidget(InterfaceName)
+					RemoveButtonWidget
 				];
-			}
-		}
 
-		void FMetasoundInterfacesDetailCustomization::UpdateInterfaceNames()
+			EntryWidget->SetEnabled(IsGraphEditableAttribute);
+			return EntryWidget;
+		};
+
+		TArray<FName> InterfaceNames = ImplementedInterfaceNames.Array();
+		InterfaceNames.Sort([](const FName& A, const FName& B) { return A.LexicalLess(B); });
+		for (const FName& InterfaceName : InterfaceNames)
 		{
-			AddableInterfaceNames.Reset();
-			ImplementedInterfaceNames.Reset();
+			InterfaceCategory.AddCustomRow(FText::FromName(InterfaceName))
+			[
+				CreateInterfaceEntryWidget(InterfaceName)
+			];
+		}
+	}
 
-			const UObject* MetaSoundObject = MetaSound.Get();
-			if (const FMetasoundAssetBase* MetaSoundAsset = IMetasoundUObjectRegistry::Get().GetObjectAsAssetBase(MetaSoundObject))
+	void FMetasoundInterfacesDetailCustomization::UpdateInterfaceNames()
+	{
+		AddableInterfaceNames.Reset();
+		ImplementedInterfaceNames.Reset();
+
+		if (const UObject* MetaSoundObject = GetMetaSound())
+		{
+			auto GetVersionName = [](const FMetasoundFrontendVersion& Version) { return Version.Name; };
+			const UClass* MetaSoundClass = MetaSoundObject->GetClass();
+			auto CanAddOrRemoveInterface = [ClassName = MetaSoundClass->GetClassPathName()](const FMetasoundFrontendVersion& Version)
 			{
-				auto GetVersionName = [](const FMetasoundFrontendVersion& Version) { return Version.Name; };
-				const UClass* MetaSoundClass = MetaSoundObject->GetClass();
-				auto CanAddOrRemoveInterface = [ClassName = MetaSoundClass->GetClassPathName()](const FMetasoundFrontendVersion& Version)
-				{
-					using namespace Frontend;
+				using namespace Frontend;
 
-					const FInterfaceRegistryKey Key = GetInterfaceRegistryKey(Version);
-					if (const IInterfaceRegistryEntry* Entry = IInterfaceRegistry::Get().FindInterfaceRegistryEntry(Key))
+				const FInterfaceRegistryKey Key = GetInterfaceRegistryKey(Version);
+				if (const IInterfaceRegistryEntry* Entry = IInterfaceRegistry::Get().FindInterfaceRegistryEntry(Key))
+				{
+					if (const FMetasoundFrontendInterfaceUClassOptions* Options = Entry->GetInterface().FindClassOptions(ClassName))
 					{
-						if (const FMetasoundFrontendInterfaceUClassOptions* Options = Entry->GetInterface().FindClassOptions(ClassName))
-						{
-							return Options->bIsModifiable;
-						}
-
-						// If no options are found for the given class, interface is modifiable by default.
-						return true;
+						return Options->bIsModifiable;
 					}
 
-					return false;
-				};
+					// If no options are found for the given class, interface is modifiable by default.
+					return true;
+				}
 
-				const TSet<FMetasoundFrontendVersion>& ImplementedInterfaces = MetaSoundAsset->GetConstDocumentChecked().Interfaces;
-				Algo::TransformIf(ImplementedInterfaces, ImplementedInterfaceNames, CanAddOrRemoveInterface, GetVersionName);
+				return false;
+			};
 
-				TArray<FMetasoundFrontendInterface> Interfaces = Frontend::ISearchEngine::Get().FindAllInterfaces();
-				for (const FMetasoundFrontendInterface& Interface : Interfaces)
+			const TSet<FMetasoundFrontendVersion>& ImplementedInterfaces = Builder->GetBuilder().GetConstDocumentChecked().Interfaces;
+			Algo::TransformIf(ImplementedInterfaces, ImplementedInterfaceNames, CanAddOrRemoveInterface, GetVersionName);
+
+			TArray<FMetasoundFrontendInterface> Interfaces = Frontend::ISearchEngine::Get().FindAllInterfaces();
+			for (const FMetasoundFrontendInterface& Interface : Interfaces)
+			{
+				if (!ImplementedInterfaceNames.Contains(Interface.Version.Name))
 				{
-					if (!ImplementedInterfaceNames.Contains(Interface.Version.Name))
+					if (CanAddOrRemoveInterface(Interface.Version))
 					{
-						if (CanAddOrRemoveInterface(Interface.Version))
-						{
-							FString Name = Interface.Version.Name.ToString();
-							AddableInterfaceNames.Add(MakeShared<FString>(MoveTemp(Name)));
-						}
+						FString Name = Interface.Version.Name.ToString();
+						AddableInterfaceNames.Add(MakeShared<FString>(MoveTemp(Name)));
 					}
 				}
-
-				AddableInterfaceNames.Sort([](const TSharedPtr<FString>& A, const TSharedPtr<FString>& B) { return A->Compare(*B) < 0; });
 			}
+
+			AddableInterfaceNames.Sort([](const TSharedPtr<FString>& A, const TSharedPtr<FString>& B) { return A->Compare(*B) < 0; });
 		}
-	} // namespace Editor
-} // namespace Metasound
+	}
+} // namespace Metasound::Editor
 #undef LOCTEXT_NAMESPACE
diff --git a/Engine/Plugins/Runtime/Metasound/Source/MetasoundEditor/Private/MetasoundDetailCustomization.h b/Engine/Plugins/Runtime/Metasound/Source/MetasoundEditor/Private/MetasoundDetailCustomization.h
--- a/Engine/Plugins/Runtime/Metasound/Source/MetasoundEditor/Private/MetasoundDetailCustomization.h
+++ b/Engine/Plugins/Runtime/Metasound/Source/MetasoundEditor/Private/MetasoundDetailCustomization.h
@@ -12,16 +12,31 @@
 // Forward Declarations
 class FPropertyRestriction;
 class IDetailLayoutBuilder;
+class UMetaSoundBuilderBase;
 struct FPointerEvent;
 
 namespace Metasound
 {
 	namespace Editor
 	{
-		class FMetasoundDetailCustomization : public IDetailCustomization
+		class FMetaSoundDetailCustomizationBase : public IDetailCustomization
+		{
+		public:
+			virtual ~FMetaSoundDetailCustomizationBase() = default;
+
+		protected:
+			UObject* GetMetaSound() const;
+			void InitBuilder(UObject& MetaSound);
+
+			TAttribute<bool> IsGraphEditableAttribute;
+			TStrongObjectPtr<UMetaSoundBuilderBase> Builder;
+		};
+
+		class FMetasoundDetailCustomization : public FMetaSoundDetailCustomizationBase
 		{
 		public:
 			FMetasoundDetailCustomization(FName InDocumentPropertyName);
+			virtual ~FMetasoundDetailCustomization() = default;
 
 			// IDetailCustomization interface
 			virtual void CustomizeDetails(IDetailLayoutBuilder& DetailLayout) override;
@@ -32,18 +47,12 @@ namespace Metasound
 			FName GetMetadataRootClassPath() const;
 			FName GetMetadataPropertyPath() const;
 
-			TAttribute<bool> IsGraphEditableAttribute;
-
-			TWeakObjectPtr<UObject> MetaSound;
-
 			FName DocumentPropertyName;
 		};
 
-		class FMetasoundInterfacesDetailCustomization : public IDetailCustomization
+		class FMetasoundInterfacesDetailCustomization : public FMetaSoundDetailCustomizationBase
 		{
 		public:
-			FMetasoundInterfacesDetailCustomization();
-
 			// IDetailCustomization interface
 			virtual void CustomizeDetails(IDetailLayoutBuilder& DetailLayout) override;
 			// End of IDetailCustomization interface
@@ -56,8 +65,6 @@ namespace Metasound
 			TSet<FName> ImplementedInterfaceNames;
 			TSharedPtr<SSearchableComboBox> InterfaceComboBox;
 			TAttribute<bool> IsGraphEditableAttribute;
-
-			TWeakObjectPtr<UObject> MetaSound;
 		};
 	} // namespace Editor
 } // namespace Metasound
\ No newline at end of file
diff --git a/Engine/Plugins/Runtime/Metasound/Source/MetasoundEditor/Private/MetasoundEditor.cpp b/Engine/Plugins/Runtime/Metasound/Source/MetasoundEditor/Private/MetasoundEditor.cpp
--- a/Engine/Plugins/Runtime/Metasound/Source/MetasoundEditor/Private/MetasoundEditor.cpp
+++ b/Engine/Plugins/Runtime/Metasound/Source/MetasoundEditor/Private/MetasoundEditor.cpp
@@ -113,12 +113,18 @@ namespace Metasound
 
 		public:
 			UEdGraph* Graph = nullptr;
+			TWeakObjectPtr<UMetaSoundBuilderBase> Builder;
 
 			void SetMemberID(const FGuid& InID)
 			{
 				MemberID = InID;
 			}
 
+			void SetBuilder(UMetaSoundBuilderBase& InBuilder)
+			{
+				Builder = &InBuilder;
+			}
+
 			FMetasoundGraphMemberSchemaAction()
 				: FEdGraphSchemaAction()
 			{
@@ -144,15 +150,6 @@ namespace Metasound
 				return NAME_None;
 			}
 
-			Frontend::FConstGraphHandle GetConstGraphHandle() const
-			{
-				const UObject* Object = CastChecked<const UMetasoundEditorGraph>(Graph)->GetMetasound();
-				const FMetasoundAssetBase* MetasoundAsset = IMetasoundUObjectRegistry::Get().GetObjectAsAssetBase(Object);
-				check(MetasoundAsset);
-
-				return MetasoundAsset->GetRootGraphHandle();
-			}
-
 			// FEdGraphSchemaAction interface
 			virtual bool IsParentable() const override
 			{
@@ -166,11 +163,40 @@ namespace Metasound
 
 			virtual int32 GetReorderIndexInContainer() const override
 			{
-				TArray<Frontend::FConstNodeHandle> InputHandles = GetConstGraphHandle()->GetConstInputNodes();
-				return InputHandles.IndexOfByPredicate([this](const Frontend::FConstNodeHandle& NodeHandle)
+				if (Builder.IsValid())
 				{
-					return NodeHandle->GetID() == MemberID;
-				});
+					if (UMetasoundEditorGraphMember* Member = GetGraphMember())
+					{
+						const FMetaSoundFrontendDocumentBuilder& DocBuilder = Builder->GetBuilder();
+						if (Member->IsA<UMetasoundEditorGraphVertex>())
+						{
+							auto FindVertexWithID = ([this](const FMetasoundFrontendClassVertex& Vertex)
+							{
+								return Vertex.NodeID == MemberID;
+							});
+							const FMetasoundFrontendDocument& Document = DocBuilder.GetConstDocumentChecked();
+							const FMetasoundFrontendClassInterface& Interface = Document.RootGraph.Interface;
+							if (Member->IsA<UMetasoundEditorGraphInput>())
+							{
+								return Interface.Inputs.IndexOfByPredicate(FindVertexWithID);
+							}
+
+							if (Member->IsA<UMetasoundEditorGraphOutput>())
+							{
+								return Interface.Outputs.IndexOfByPredicate(FindVertexWithID);
+							}
+						}
+						else if (Member->IsA<UMetasoundEditorGraphVariable>())
+						{
+							auto FindVariableWithID = [this](const FMetasoundFrontendVariable& Variable)
+							{
+								return Variable.ID == MemberID;
+							};
+							return DocBuilder.FindConstBuildGraphChecked().Variables.IndexOfByPredicate(FindVariableWithID);
+						}
+					}
+				}
+				return INDEX_NONE;
 			}
 
 			virtual bool ReorderToBeforeAction(TSharedRef<FEdGraphSchemaAction> OtherAction) override
@@ -835,7 +861,7 @@ namespace Metasound
 
 		bool FEditor::IsPlaying() const
 		{
-			if (Metasound)
+			if (UObject* MetaSound = GetMetasoundObject())
 			{
 				if (const UAudioComponent* PreviewComponent = GEditor->GetPreviewAudioComponent())
 				{
@@ -843,7 +869,7 @@ namespace Metasound
 					{
 						if (const USoundBase* Sound = PreviewComponent->Sound)
 						{
-							return Sound->GetUniqueID() == Metasound->GetUniqueID();
+							return Sound->GetUniqueID() == MetaSound->GetUniqueID();
 						}
 					}
 				}
@@ -881,14 +907,13 @@ namespace Metasound
 			}
 
 			// Support undo/redo
-			Metasound = ObjectToEdit;
-			Metasound->SetFlags(RF_Transactional);
+			ObjectToEdit->SetFlags(RF_Transactional);
 
 			// Typically sounds are versioned on load of the asset. There are certain instances where an asset is not versioned on reload.
 			// This forces versioning the document on load prior to the editor synchronizing and building the editor graph if an asset is
 			// reloaded while the asset editor was open.
 			Builder.Reset(&Engine::FDocumentBuilderRegistry::GetChecked().FindOrBeginBuilding(*ObjectToEdit));
-			if (FMetasoundAssetBase* MetaSoundAsset = IMetasoundUObjectRegistry::Get().GetObjectAsAssetBase(Metasound))
+			if (FMetasoundAssetBase* MetaSoundAsset = IMetasoundUObjectRegistry::Get().GetObjectAsAssetBase(ObjectToEdit))
 			{
 				if (MetaSoundAsset->VersionAsset(Builder->GetBuilder()))
 				{
@@ -905,6 +930,7 @@ namespace Metasound
 
 			FGraphEditorCommands::Register();
 			FEditorCommands::Register();
+			BindGraphCommands();
 
 			// If sound was already playing in the editor (ex. from ContentBrowser),
 			// restart to synchronize visual state of editor (ex. volume meter analysis
@@ -913,92 +939,95 @@ namespace Metasound
 			const bool bRestartSound = IsPlaying();
 			if (!bRestartSound)
 			{
-				FGraphBuilder::RegisterGraphWithFrontend(*Metasound);
+				FGraphBuilder::RegisterGraphWithFrontend(*ObjectToEdit);
 			}
 
-			RefreshEditorContext();
-
-			BindGraphCommands();
-			CreateInternalWidgets();
-			CreateAnalyzers();
+			RefreshEditorContext(*ObjectToEdit);
+			CreateInternalWidgets(*ObjectToEdit);
 
-			// Has to be run after widgets are initialized to properly display
-			if (bPrimingRegistry)
+			if (UMetaSoundSource* MetaSoundSource = Cast<UMetaSoundSource>(ObjectToEdit))
 			{
-				NotifyAssetPrimeInProgress();
+				CreateAnalyzers(*MetaSoundSource);
 			}
 
 			const TSharedRef<FTabManager::FLayout> StandaloneDefaultLayout = FTabManager::NewLayout("Standalone_MetasoundEditor_Layout_v11")
-				->AddArea
-				(
-					FTabManager::NewPrimaryArea()
-					->SetOrientation(Orient_Vertical)
-					->Split(FTabManager::NewSplitter()
-						->SetOrientation(Orient_Horizontal)
+			->AddArea
+			(
+				FTabManager::NewPrimaryArea()
+				->SetOrientation(Orient_Vertical)
+				->Split(FTabManager::NewSplitter()
+					->SetOrientation(Orient_Horizontal)
+					->Split
+					(
+						FTabManager::NewSplitter()
+						->SetSizeCoefficient(0.15f)
+						->SetOrientation(Orient_Vertical)
 						->Split
 						(
-							FTabManager::NewSplitter()
-							->SetSizeCoefficient(0.15f)
-							->SetOrientation(Orient_Vertical)
-							->Split
-							(
-								FTabManager::NewStack()
-								->SetSizeCoefficient(0.25f)
-								->SetHideTabWell(false)
-								->AddTab(TabFactory::Names::Members, ETabState::OpenedTab)
-							)
-							->Split
-							(
-								FTabManager::NewStack()
-								->SetSizeCoefficient(0.1f)
-								->SetHideTabWell(true)
-								->AddTab(TabFactory::Names::Interfaces, ETabState::OpenedTab)
-							)
-							->Split
-							(
-								FTabManager::NewStack()
-								->SetSizeCoefficient(0.50f)
-								->SetHideTabWell(false)
-								->AddTab(TabFactory::Names::Details, ETabState::OpenedTab)
-							)
+							FTabManager::NewStack()
+							->SetSizeCoefficient(0.25f)
+							->SetHideTabWell(false)
+							->AddTab(TabFactory::Names::Members, ETabState::OpenedTab)
 						)
 						->Split
 						(
-							FTabManager::NewSplitter()
-							->SetSizeCoefficient(0.77f)
-							->SetOrientation(Orient_Vertical)
-							->Split
-							(
-								FTabManager::NewStack()
-								->SetSizeCoefficient(0.8f)
-								->SetHideTabWell(true)
-								->AddTab(TabFactory::Names::GraphCanvas, ETabState::OpenedTab)
-							)
-							->Split
-							(
-								FTabManager::NewStack()
-								->SetSizeCoefficient(0.2f)
-								->SetHideTabWell(true)
-								->AddTab(TabFactory::Names::Find, ETabState::OpenedTab)
-							)
+							FTabManager::NewStack()
+							->SetSizeCoefficient(0.1f)
+							->SetHideTabWell(true)
+							->AddTab(TabFactory::Names::Interfaces, ETabState::OpenedTab)
 						)
-
 						->Split
 						(
 							FTabManager::NewStack()
-							->SetSizeCoefficient(0.08f)
+							->SetSizeCoefficient(0.50f)
+							->SetHideTabWell(false)
+							->AddTab(TabFactory::Names::Details, ETabState::OpenedTab)
+						)
+					)
+					->Split
+					(
+						FTabManager::NewSplitter()
+						->SetSizeCoefficient(0.77f)
+						->SetOrientation(Orient_Vertical)
+						->Split
+						(
+							FTabManager::NewStack()
+							->SetSizeCoefficient(0.8f)
 							->SetHideTabWell(true)
-							->AddTab(TabFactory::Names::Analyzers, ETabState::OpenedTab)
+							->AddTab(TabFactory::Names::GraphCanvas, ETabState::OpenedTab)
+						)
+						->Split
+						(
+							FTabManager::NewStack()
+							->SetSizeCoefficient(0.2f)
+							->SetHideTabWell(true)
+							->AddTab(TabFactory::Names::Find, ETabState::OpenedTab)
 						)
 					)
-				);
+
+					->Split
+					(
+						FTabManager::NewStack()
+						->SetSizeCoefficient(0.08f)
+						->SetHideTabWell(true)
+						->AddTab(TabFactory::Names::Analyzers, ETabState::OpenedTab)
+					)
+				)
+			);
 
 			constexpr bool bCreateDefaultStandaloneMenu = true;
 			constexpr bool bCreateDefaultToolbar = true;
 			constexpr bool bToolbarFocusable = false;
 			constexpr bool bUseSmallToolbarIcons = true;
+
 			FAssetEditorToolkit::InitAssetEditor(Mode, InitToolkitHost, TEXT("MetasoundEditorApp"), StandaloneDefaultLayout, bCreateDefaultStandaloneMenu, bCreateDefaultToolbar, ObjectToEdit, bToolbarFocusable, bUseSmallToolbarIcons);
 
+			// Has to be run after widgets are initialized to properly display
+			if (bPrimingRegistry)
+			{
+				NotifyAssetPrimeInProgress();
+			}
+
 			ExtendToolbarInternal();
 			RegenerateMenusAndToolbars();
 
@@ -1014,7 +1043,24 @@ namespace Metasound
 
 		UObject* FEditor::GetMetasoundObject() const
 		{
-			return Metasound;
+			if (HasEditingObject())
+			{
+				return GetEditingObject();
+			}
+
+			// During init, editing object isn't yet set by underlying EditorToolkit::Init.
+			// If it hasn't been cached off, use the builder's pointer which is set
+			// early in editor initialization.
+			if (Builder.IsValid())
+			{
+				const FMetaSoundFrontendDocumentBuilder& DocBuilder = Builder->GetBuilder();
+				if (DocBuilder.IsValid())
+				{
+					return &DocBuilder.CastDocumentObjectChecked<UObject>();
+				}
+			}
+
+			return nullptr;
 		}
 
 		void FEditor::SetSelection(const TArray<UObject*>& SelectedObjects)
@@ -1061,16 +1107,17 @@ namespace Metasound
 
 		FLinearColor FEditor::GetWorldCentricTabColorScale() const
 		{
-			if (const ISlateStyle* MetasoundStyle = FSlateStyleRegistry::FindSlateStyle("MetaSoundStyle"))
+			if (const ISlateStyle* MetaSoundStyle = FSlateStyleRegistry::FindSlateStyle("MetaSoundStyle"))
 			{
-				if (UMetaSoundSource* MetaSoundSource = Cast<UMetaSoundSource>(Metasound))
+				UObject* MetaSound = GetMetasoundObject();
+				if (UMetaSoundSource* MetaSoundSource = Cast<UMetaSoundSource>(MetaSound))
 				{
-					return MetasoundStyle->GetColor("MetaSoundSource.Color");
+					return MetaSoundStyle->GetColor("MetaSoundSource.Color");
 				}
 
-				if (UMetaSoundPatch* MetaSoundPatch = Cast<UMetaSoundPatch>(Metasound))
+				if (UMetaSoundPatch* MetaSoundPatch = Cast<UMetaSoundPatch>(MetaSound))
 				{
-					return MetasoundStyle->GetColor("MetaSoundPatch.Color");
+					return MetaSoundStyle->GetColor("MetaSoundPatch.Color");
 				}
 			}
 
@@ -1086,16 +1133,17 @@ namespace Metasound
 			}
 			else
 			{
-				if (UMetaSoundSource* MetaSoundSource = Cast<UMetaSoundSource>(Metasound))
+				UObject* MetaSound = GetMetasoundObject();
+				if (UMetaSoundSource* MetaSoundSource = Cast<UMetaSoundSource>(MetaSound))
 				{
 					IconName += TEXT(".MetasoundSource");
 				}
-				else if (UMetaSoundPatch* MetaSoundPatch = Cast<UMetaSoundPatch>(Metasound))
+				else if (UMetaSoundPatch* MetaSoundPatch = Cast<UMetaSoundPatch>(MetaSound))
 				{
 					IconName += TEXT(".MetasoundPatch");
 				}
 
-				const FMetasoundAssetBase* MetaSoundAsset = IMetasoundUObjectRegistry::Get().GetObjectAsAssetBase(Metasound);
+				const FMetasoundAssetBase* MetaSoundAsset = IMetasoundUObjectRegistry::Get().GetObjectAsAssetBase(MetaSound);
 				check(MetaSoundAsset);
 				if (MetaSoundAsset->GetConstDocumentChecked().RootGraph.PresetOptions.bIsPreset)
 				{
@@ -1110,21 +1158,24 @@ namespace Metasound
 
 		FLinearColor FEditor::GetDefaultTabColor() const
 		{
-			if (IsPlaying())
+			if (UObject* MetaSound = GetMetasoundObject())
 			{
-				if (const ISlateStyle* MetasoundStyle = FSlateStyleRegistry::FindSlateStyle("MetaSoundStyle"))
+				if (IsPlaying())
 				{
-					if (UMetaSoundSource* MetaSoundSource = Cast<UMetaSoundSource>(Metasound))
+					if (const ISlateStyle* MetasoundStyle = FSlateStyleRegistry::FindSlateStyle("MetaSoundStyle"))
 					{
-						return MetasoundStyle->GetColor("MetaSoundSource.Color");
-					}
+						if (UMetaSoundSource* MetaSoundSource = Cast<UMetaSoundSource>(MetaSound))
+						{
+							return MetasoundStyle->GetColor("MetaSoundSource.Color");
+						}
 
-					if (UMetaSoundPatch* MetaSoundPatch = Cast<UMetaSoundPatch>(Metasound))
-					{
-						return MetasoundStyle->GetColor("MetaSoundPatch.Color");
+						if (UMetaSoundPatch* MetaSoundPatch = Cast<UMetaSoundPatch>(MetaSound))
+						{
+							return MetasoundStyle->GetColor("MetaSoundPatch.Color");
+						}
 					}
 				}
-			}
+			}	
 
 			return FAssetEditorToolkit::GetDefaultTabColor();
 		}
@@ -1134,11 +1185,6 @@ namespace Metasound
 			return FEditor::EditorName;
 		}
 
-		void FEditor::AddReferencedObjects(FReferenceCollector& Collector)
-		{
-			Collector.AddReferencedObject(Metasound);
-		}
-
 		void FEditor::PostUndo(bool bSuccess)
 		{
 			if (MetasoundGraphEditor.IsValid())
@@ -1207,23 +1253,23 @@ namespace Metasound
 			if (MetasoundGraphEditor.IsValid())
 			{
 				UMetasoundEditorGraph& MetaSoundGraph = GetMetaSoundGraphChecked();
-				FMetasoundAssetBase* MetaSoundAsset = IMetasoundUObjectRegistry::Get().GetObjectAsAssetBase(Metasound);
-				check(MetaSoundAsset);
-
-				if (MetaSoundAsset->GetVersionedOnLoad())
+				if (FMetasoundAssetBase* MetaSoundAsset = IMetasoundUObjectRegistry::Get().GetObjectAsAssetBase(GetMetasoundObject()))
 				{
-					MetaSoundAsset->ClearVersionedOnLoad();
+					if (MetaSoundAsset->GetVersionedOnLoad())
+					{
+						MetaSoundAsset->ClearVersionedOnLoad();
 
-					const FString VersionString = MetaSoundAsset->GetConstDocumentChecked().Metadata.Version.Number.ToString();
-					FText Msg = FText::Format(LOCTEXT("MetaSoundDocumentVersioned", "Document versioned to '{0}' on load."), FText::FromString(VersionString));
-					FNotificationInfo Info(Msg);
-					Info.bFireAndForget = true;
-					Info.bUseSuccessFailIcons = false;
-					Info.ExpireDuration = 5.0f;
+						const FString VersionString = MetaSoundAsset->GetConstDocumentChecked().Metadata.Version.Number.ToString();
+						FText Msg = FText::Format(LOCTEXT("MetaSoundDocumentVersioned", "Document versioned to '{0}' on load."), FText::FromString(VersionString));
+						FNotificationInfo Info(Msg);
+						Info.bFireAndForget = true;
+						Info.bUseSuccessFailIcons = false;
+						Info.ExpireDuration = 5.0f;
 
-					MetasoundGraphEditor->AddNotification(Info, false /* bSuccess */);
+						MetasoundGraphEditor->AddNotification(Info, false /* bSuccess */);
 
-					MetaSoundAsset->MarkMetasoundDocumentDirty();
+						MetaSoundAsset->MarkMetasoundDocumentDirty();
+					}
 				}
 			}
 		}
@@ -1255,9 +1301,9 @@ namespace Metasound
 				// If a property change event occurs outside of the metasound UEdGraph and results in the metasound document changing,
 				// then the document and the UEdGraph need to be synchronized. There may be a better trigger for this call to reduce
 				// the number of times the graph is synchronized.
-				if (Metasound)
+				if (UObject* MetaSound = GetMetasoundObject())
 				{
-					if (FMetasoundAssetBase* Asset = IMetasoundUObjectRegistry::Get().GetObjectAsAssetBase(Metasound))
+					if (FMetasoundAssetBase* Asset = IMetasoundUObjectRegistry::Get().GetObjectAsAssetBase(MetaSound))
 					{
 						Asset->GetModifyContext().SetDocumentModified();
 					}
@@ -1265,9 +1311,9 @@ namespace Metasound
 			}
 		}
 
-		void FEditor::CreateInternalWidgets()
+		void FEditor::CreateInternalWidgets(UObject& MetaSound)
 		{
-			CreateGraphEditorWidget();
+			CreateGraphEditorWidget(MetaSound);
 
 			FDetailsViewArgs Args;
 			Args.bHideSelectionTip = true;
@@ -1300,12 +1346,12 @@ namespace Metasound
 			MetasoundDetails->SetExtensionHandler(MakeShared<FMetaSoundNodeExtensionHandler>());
 
 			// Set details selection to the MetaSound's source settings 
-			SetSelection({ Metasound });
+			SetSelection({ &MetaSound });
 			InterfacesDetails = PropertyModule.CreateDetailView(Args);
 			if (InterfacesDetails.IsValid())
 			{
 				InterfacesView = TStrongObjectPtr(NewObject<UMetasoundInterfacesView>());
-				InterfacesView->SetMetasound(Metasound);
+				InterfacesView->SetMetasound(&MetaSound);
 				const TArray<UObject*> InterfacesViewObj{ InterfacesView.Get() };
 
 				InterfacesDetails->SetObjects(InterfacesViewObj);
@@ -1323,96 +1369,93 @@ namespace Metasound
 			return false;
 		}
 
-		void FEditor::CreateAnalyzers()
+		void FEditor::CreateAnalyzers(UMetaSoundSource& MetaSoundSource)
 		{
-			if (UMetaSoundSource* MetaSoundSource = Cast<UMetaSoundSource>(Metasound))
+			if (ensure(GEditor))
 			{
-				if (ensure(GEditor))
-				{
-					const Audio::FDeviceId AudioDeviceId = GEditor->GetMainAudioDeviceID();
-					constexpr UAudioBus* DefaultBus = nullptr;
+				const Audio::FDeviceId AudioDeviceId = GEditor->GetMainAudioDeviceID();
+				constexpr UAudioBus* DefaultBus = nullptr;
 
-					if (!OutputMeter.IsValid())
+				if (!OutputMeter.IsValid())
+				{
+					const UMetasoundEditorSettings* EditorSettings = GetDefault<UMetasoundEditorSettings>();
+					check(EditorSettings)
+					const bool bUseAudioMaterialWidgets = EditorSettings->bUseAudioMaterialWidgets;
+					if (bUseAudioMaterialWidgets)
 					{
-						const UMetasoundEditorSettings* EditorSettings = GetDefault<UMetasoundEditorSettings>();
-						check(EditorSettings)
-						const bool bUseAudioMaterialWidgets = EditorSettings->bUseAudioMaterialWidgets;
-						if (bUseAudioMaterialWidgets)
+						const FAudioMaterialMeterStyle* MeterStyle = EditorSettings->GetMeterStyle();
+						if (ensureMsgf(MeterStyle, TEXT("Failed to find MaterialMeterStyle when attempting to build MetaSound Editor output meter. Falling back to default non-material meter.")))
 						{
-							const FAudioMaterialMeterStyle* MeterStyle = EditorSettings->GetMeterStyle();
-							if (ensureMsgf(MeterStyle, TEXT("Failed to find MaterialMeterStyle when attempting to build MetaSound Editor output meter. Falling back to default non-material meter.")))
-							{
-								OutputMeter = MakeShared<AudioWidgets::FAudioMeter>(MetaSoundSource->NumChannels, AudioDeviceId, *MeterStyle, DefaultBus);
-							}
-						}
-
-						if (!OutputMeter.IsValid())
-						{
-							OutputMeter = MakeShared<AudioWidgets::FAudioMeter>(MetaSoundSource->NumChannels, AudioDeviceId, DefaultBus);
+							OutputMeter = MakeShared<AudioWidgets::FAudioMeter>(MetaSoundSource.NumChannels, AudioDeviceId, *MeterStyle, DefaultBus);
 						}
 					}
-					else if (OutputMeter->GetAudioBus()->GetNumChannels() != MetaSoundSource->NumChannels)
+
+					if (!OutputMeter.IsValid())
 					{
-						OutputMeter->Init(MetaSoundSource->NumChannels, AudioDeviceId, DefaultBus);
+						OutputMeter = MakeShared<AudioWidgets::FAudioMeter>(MetaSoundSource.NumChannels, AudioDeviceId, DefaultBus);
 					}
+				}
+				else if (OutputMeter->GetAudioBus()->GetNumChannels() != MetaSoundSource.NumChannels)
+				{
+					OutputMeter->Init(MetaSoundSource.NumChannels, AudioDeviceId, DefaultBus);
+				}
 
-					const uint32 MetaSoundNumChannels = static_cast<uint32>(MetaSoundSource->NumChannels);
-
-					// Init Oscilloscope
-					constexpr float OscilloscopeTimeWindowMs     = 10.0f;
-					constexpr float OscilloscopeMaxTimeWindowMs  = 10.0f;
-					constexpr float OscilloscopeAnalysisPeriodMs = 10.0f;
-					constexpr EAudioPanelLayoutType OscilloscopePanelLayoutType = EAudioPanelLayoutType::Basic;
+				const uint32 MetaSoundNumChannels = static_cast<uint32>(MetaSoundSource.NumChannels);
 
-					if (!OutputOscilloscope.IsValid())
-					{
-						OutputOscilloscope = MakeShared<AudioWidgets::FAudioOscilloscope>(AudioDeviceId,
-							MetaSoundNumChannels,
-							OscilloscopeTimeWindowMs,
-							OscilloscopeMaxTimeWindowMs,
-							OscilloscopeAnalysisPeriodMs,
-							OscilloscopePanelLayoutType);
-					}
-					else if (OutputOscilloscope->GetAudioBus()->GetNumChannels() != MetaSoundSource->NumChannels)
-					{
-						OutputOscilloscope->CreateAudioBus(MetaSoundNumChannels);
-						OutputOscilloscope->CreateDataProvider(AudioDeviceId, OscilloscopeTimeWindowMs, OscilloscopeMaxTimeWindowMs, OscilloscopeAnalysisPeriodMs, OscilloscopePanelLayoutType);
-						OutputOscilloscope->CreateOscilloscopeWidget(MetaSoundNumChannels, OscilloscopePanelLayoutType);
-					}
+				// Init Oscilloscope
+				constexpr float OscilloscopeTimeWindowMs     = 10.0f;
+				constexpr float OscilloscopeMaxTimeWindowMs  = 10.0f;
+				constexpr float OscilloscopeAnalysisPeriodMs = 10.0f;
+				constexpr EAudioPanelLayoutType OscilloscopePanelLayoutType = EAudioPanelLayoutType::Basic;
 
-					// Init Vectorscope
-					constexpr float VectorscopeTimeWindowMs     = 30.0f;
-					constexpr float VectorscopeMaxTimeWindowMs  = 30.0f;
-					constexpr float VectorscopeAnalysisPeriodMs = 10.0f;
-					constexpr EAudioPanelLayoutType VectorscopePanelLayoutType = EAudioPanelLayoutType::Basic;
+				if (!OutputOscilloscope.IsValid())
+				{
+					OutputOscilloscope = MakeShared<AudioWidgets::FAudioOscilloscope>(AudioDeviceId,
+						MetaSoundNumChannels,
+						OscilloscopeTimeWindowMs,
+						OscilloscopeMaxTimeWindowMs,
+						OscilloscopeAnalysisPeriodMs,
+						OscilloscopePanelLayoutType);
+				}
+				else if (OutputOscilloscope->GetAudioBus()->GetNumChannels() != MetaSoundSource.NumChannels)
+				{
+					OutputOscilloscope->CreateAudioBus(MetaSoundNumChannels);
+					OutputOscilloscope->CreateDataProvider(AudioDeviceId, OscilloscopeTimeWindowMs, OscilloscopeMaxTimeWindowMs, OscilloscopeAnalysisPeriodMs, OscilloscopePanelLayoutType);
+					OutputOscilloscope->CreateOscilloscopeWidget(MetaSoundNumChannels, OscilloscopePanelLayoutType);
+				}
 
-					if (!OutputVectorscope.IsValid())
-					{
-						OutputVectorscope = MakeShared<AudioWidgets::FAudioVectorscope>(AudioDeviceId,
-							MetaSoundNumChannels,
-							VectorscopeTimeWindowMs,
-							VectorscopeMaxTimeWindowMs,
-							VectorscopeAnalysisPeriodMs,
-							VectorscopePanelLayoutType);
-					}
-					else if (OutputVectorscope->GetAudioBus()->GetNumChannels() != MetaSoundSource->NumChannels)
-					{
-						OutputVectorscope->CreateAudioBus(MetaSoundNumChannels);
-						OutputVectorscope->CreateDataProvider(AudioDeviceId, VectorscopeTimeWindowMs, VectorscopeMaxTimeWindowMs, VectorscopeAnalysisPeriodMs);
-						OutputVectorscope->CreateVectorscopeWidget(VectorscopePanelLayoutType);
-					}
+				// Init Vectorscope
+				constexpr float VectorscopeTimeWindowMs     = 30.0f;
+				constexpr float VectorscopeMaxTimeWindowMs  = 30.0f;
+				constexpr float VectorscopeAnalysisPeriodMs = 10.0f;
+				constexpr EAudioPanelLayoutType VectorscopePanelLayoutType = EAudioPanelLayoutType::Basic;
 
-					if (!OutputSpectrumAnalyzer.IsValid())
-					{
-						OutputSpectrumAnalyzer = MakeShared<AudioWidgets::FAudioSpectrumAnalyzer>(MetaSoundSource->NumChannels, AudioDeviceId);
-					}
-					else if (OutputSpectrumAnalyzer->GetAudioBus()->GetNumChannels() != MetaSoundSource->NumChannels)
-					{
-						OutputSpectrumAnalyzer->Init(MetaSoundSource->NumChannels, AudioDeviceId, nullptr);
-					}
+				if (!OutputVectorscope.IsValid())
+				{
+					OutputVectorscope = MakeShared<AudioWidgets::FAudioVectorscope>(AudioDeviceId,
+						MetaSoundNumChannels,
+						VectorscopeTimeWindowMs,
+						VectorscopeMaxTimeWindowMs,
+						VectorscopeAnalysisPeriodMs,
+						VectorscopePanelLayoutType);
+				}
+				else if (OutputVectorscope->GetAudioBus()->GetNumChannels() != MetaSoundSource.NumChannels)
+				{
+					OutputVectorscope->CreateAudioBus(MetaSoundNumChannels);
+					OutputVectorscope->CreateDataProvider(AudioDeviceId, VectorscopeTimeWindowMs, VectorscopeMaxTimeWindowMs, VectorscopeAnalysisPeriodMs);
+					OutputVectorscope->CreateVectorscopeWidget(VectorscopePanelLayoutType);
+				}
 
-					return;
+				if (!OutputSpectrumAnalyzer.IsValid())
+				{
+					OutputSpectrumAnalyzer = MakeShared<AudioWidgets::FAudioSpectrumAnalyzer>(MetaSoundSource.NumChannels, AudioDeviceId);
+				}
+				else if (OutputSpectrumAnalyzer->GetAudioBus()->GetNumChannels() != MetaSoundSource.NumChannels)
+				{
+					OutputSpectrumAnalyzer->Init(MetaSoundSource.NumChannels, AudioDeviceId, nullptr);
 				}
+
+				return;
 			}
 
 			DestroyAnalyzers();
@@ -1476,7 +1519,7 @@ namespace Metasound
 
 					ToolbarBuilder.BeginSection("Settings");
 					{
-						if (Metasound->IsA<USoundBase>())
+						if (IsAuditionable())
 						{
 							ToolbarBuilder.AddToolBarButton(
 								FEditorCommands::Get().EditSourceSettings,
@@ -1499,7 +1542,7 @@ namespace Metasound
 					}
 					ToolbarBuilder.EndSection();
 
-					if (Metasound->IsA<USoundBase>())
+					if (IsAuditionable())
 					{
 						ToolbarBuilder.BeginSection("Audition");
 						{
@@ -1637,7 +1680,13 @@ namespace Metasound
 		void FEditor::Import()
 		{
 			// TODO: Prompt OFD and provide path from user
-			const FString InputPath = FPaths::ProjectIntermediateDir() / TEXT("MetaSounds") + FPaths::ChangeExtension(Metasound->GetPathName(), FMetasoundAssetBase::FileExtension);
+			UObject* MetaSound = GetMetasoundObject();
+			if (!MetaSound)
+			{
+				return;
+			}
+
+			const FString InputPath = FPaths::ProjectIntermediateDir() / TEXT("MetaSounds") + FPaths::ChangeExtension(MetaSound->GetPathName(), FMetasoundAssetBase::FileExtension);
 			
 			// TODO: use the same directory as the currently open MetaSound
 			const FString OutputPath = FString("/Game/ImportedMetaSound/GeneratedMetaSound");
@@ -1684,14 +1733,17 @@ namespace Metasound
 
 		void FEditor::Export()
 		{
-			FMetasoundAssetBase* MetasoundAsset = IMetasoundUObjectRegistry::Get().GetObjectAsAssetBase(Metasound);
-			check(MetasoundAsset);
+			if (UObject* MetaSound = GetMetasoundObject())
+			{
+				FMetasoundAssetBase* MetaSoundAsset = IMetasoundUObjectRegistry::Get().GetObjectAsAssetBase(MetaSound);
+				check(MetaSoundAsset);
 
-			static const FString MetasoundExtension(TEXT(".metasound"));
+				static const FString MetasoundExtension(TEXT(".metasound"));
 
-			// TODO: We could just make this an object.
-			const FString Path = FPaths::ProjectSavedDir() / TEXT("MetaSounds") + FPaths::ChangeExtension(Metasound->GetPathName(), MetasoundExtension);
-			MetasoundAsset->GetDocumentHandle()->ExportToJSONAsset(Path);
+				// TODO: We could just make this an object.
+				const FString Path = FPaths::ProjectSavedDir() / TEXT("MetaSounds") + FPaths::ChangeExtension(MetaSound->GetPathName(), MetasoundExtension);
+				MetaSoundAsset->GetDocumentHandle()->ExportToJSONAsset(Path);
+			}	
 		}
 
 		FText FEditor::GetGraphStatusDescription() const
@@ -1796,7 +1848,7 @@ namespace Metasound
 
 		void FEditor::Play()
 		{
-			if (USoundBase* MetasoundToPlay = Cast<USoundBase>(Metasound))
+			if (USoundBase* MetaSoundToPlay = Cast<USoundBase>(GetMetasoundObject()))
 			{
 				HighestMessageSeverity = GetMetaSoundGraphChecked().GetHighestMessageSeverity();
 				if (HighestMessageSeverity == EMessageSeverity::Error)
@@ -1807,11 +1859,11 @@ namespace Metasound
 				// Even though the MetaSoundSource will attempt to register via InitResources
 				// later in this execution (and deeper in the stack), this call forces
 				// re-registering to make sure everything is up-to-date.
-				FGraphBuilder::RegisterGraphWithFrontend(*Metasound);
+				FGraphBuilder::RegisterGraphWithFrontend(*MetaSoundToPlay);
 
 				// Set the send to the audio bus that is used for analyzing the metasound output
 				check(GEditor);
-				if (UAudioComponent* PreviewComp = GEditor->PlayPreviewSound(MetasoundToPlay))
+				if (UAudioComponent* PreviewComp = GEditor->PlayPreviewSound(MetaSoundToPlay))
 				{
 					PlayTime = 0.0;
 
@@ -1841,8 +1893,8 @@ namespace Metasound
 						PreviewComp->SetAudioBusSendPostEffect(AudioBus, 1.0f);
 					}
 
-					FMetasoundAssetBase* MetasoundAsset = IMetasoundUObjectRegistry::Get().GetObjectAsAssetBase(Metasound);
-					check(MetasoundAsset);
+					FMetasoundAssetBase* MetaSoundAsset = IMetasoundUObjectRegistry::Get().GetObjectAsAssetBase(MetaSoundToPlay);
+					check(MetaSoundAsset);
 
 					FAudioDevice* AudioDevice = PreviewComp->GetAudioDevice();
 					check(AudioDevice);
@@ -1852,9 +1904,9 @@ namespace Metasound
 					const uint32 PlayOrder = PreviewComp->GetLastPlayOrder();
 					const uint64 TransmitterID = Audio::GetTransmitterID(PreviewComp->GetAudioComponentID(), 0, PlayOrder);
 
-					UMetaSoundSource* Source = CastChecked<UMetaSoundSource>(Metasound);
+					UMetaSoundSource* Source = CastChecked<UMetaSoundSource>(MetaSoundToPlay);
 					GraphConnectionManager = MakeUnique<FGraphConnectionManager>(
-						*MetasoundAsset, *PreviewComp, TransmitterID, Source->GetOperatorSettings(DeviceSampleRate));
+						*MetaSoundAsset, *PreviewComp, TransmitterID, Source->GetOperatorSettings(DeviceSampleRate));
 				}
 
 				MetasoundGraphEditor->RegisterActiveTimer(0.0f,
@@ -1864,7 +1916,7 @@ namespace Metasound
 						if (RenderStatsWidget.IsValid())
 						{
 							check(IsInGameThread());
-							RenderStatsWidget->Update(bIsPlaying, Cast<const UMetaSoundSource>(Metasound));
+							RenderStatsWidget->Update(bIsPlaying, Cast<const UMetaSoundSource>(GetMetasoundObject()));
 						}
 
 						if (bIsPlaying)
@@ -1931,20 +1983,18 @@ namespace Metasound
 
 		void FEditor::SetPreviewID(uint32 InPreviewID)
 		{
-			if (!Metasound)
+			if (HasEditingObject())
 			{
-				return;
+				GetMetaSoundGraphChecked().SetPreviewID(InPreviewID);
 			}
-
-			GetMetaSoundGraphChecked().SetPreviewID(InPreviewID);
 		}
 
 		UMetasoundEditorGraph& FEditor::GetMetaSoundGraphChecked()
 		{
-			FMetasoundAssetBase* MetasoundAsset = IMetasoundUObjectRegistry::Get().GetObjectAsAssetBase(Metasound);
-			check(MetasoundAsset);
+			FMetasoundAssetBase* MetaSoundAsset = IMetasoundUObjectRegistry::Get().GetObjectAsAssetBase(GetMetasoundObject());
+			check(MetaSoundAsset);
 
-			UEdGraph* Graph = MetasoundAsset->GetGraph();
+			UEdGraph* Graph = MetaSoundAsset->GetGraph();
 			check(Graph);
 
 			return *CastChecked<UMetasoundEditorGraph>(Graph);
@@ -2035,7 +2085,7 @@ namespace Metasound
 			// if swapping from one object-level edit mode to the other
 			// (ex. Metasound Settings to General Settings)
 			SetSelection({ });
-			SetSelection({ Metasound });
+			SetSelection({ GetMetasoundObject() });
 		}
 
 		void FEditor::ConvertFromPreset()
@@ -2044,7 +2094,7 @@ namespace Metasound
 
 			check(GEditor);
 
-			if (Metasound)
+			if (Builder.IsValid())
 			{
 				EMetaSoundBuilderResult Result = EMetaSoundBuilderResult::Failed;
 				Builder->ConvertFromPreset(Result);
@@ -2090,7 +2140,7 @@ namespace Metasound
 
 			if (!ObjectsToSync.Num())
 			{
-				ObjectsToSync.Add(Metasound);
+				ObjectsToSync.Add(GetMetasoundObject());
 			}
 
 			check(GEditor);
@@ -2119,7 +2169,7 @@ namespace Metasound
 			}
 		}
 
-		void FEditor::CreateGraphEditorWidget()
+		void FEditor::CreateGraphEditorWidget(UObject& MetaSound)
 		{
 			if (!GraphEditorCommands.IsValid())
 			{
@@ -2131,7 +2181,7 @@ namespace Metasound
 				GraphEditorCommands->MapAction(FEditorCommands::Get().EditMetasoundSettings,
 					FExecuteAction::CreateSP(this, &FEditor::EditMetasoundSettings));
 
-				if (Metasound->IsA<UMetaSoundSource>())
+				if (MetaSound.IsA<UMetaSoundSource>())
 				{
 					GraphEditorCommands->MapAction(FEditorCommands::Get().EditSourceSettings,
 						FExecuteAction::CreateSP(this, &FEditor::EditSourceSettings));
@@ -2214,15 +2264,15 @@ namespace Metasound
 			GraphEvents.OnSelectionChanged = SGraphEditor::FOnSelectionChanged::CreateSP(this, &FEditor::OnSelectedNodesChanged);
 			GraphEvents.OnTextCommitted = FOnNodeTextCommitted::CreateSP(this, &FEditor::OnNodeTitleCommitted);
 
-			FMetasoundAssetBase* MetasoundAsset = IMetasoundUObjectRegistry::Get().GetObjectAsAssetBase(Metasound);
-			check(MetasoundAsset);
+			FMetasoundAssetBase* MetaSoundAsset = IMetasoundUObjectRegistry::Get().GetObjectAsAssetBase(&MetaSound);
+			check(MetaSoundAsset);
 
 			SAssignNew(MetasoundGraphEditor, SGraphEditor)
 				.AdditionalCommands(GraphEditorCommands)
 				.Appearance(this, &FEditor::GetGraphAppearance)
 				.AutoExpandActionMenu(true)
 				.GraphEvents(GraphEvents)
-				.GraphToEdit(MetasoundAsset->GetGraph())
+				.GraphToEdit(MetaSoundAsset->GetGraph())
 				.IsEditable(this, &FEditor::IsGraphEditable)
 				.ShowGraphStateOverlay(false);
 
@@ -2239,11 +2289,11 @@ namespace Metasound
 		{
 			FGraphAppearanceInfo AppearanceInfo;
 
-			if (Metasound)
+			if (UObject* MetaSound = GetMetasoundObject())
 			{
-				const FMetasoundAssetBase* MetasoundAsset = IMetasoundUObjectRegistry::Get().GetObjectAsAssetBase(Metasound);
-				check(MetasoundAsset);
-				AppearanceInfo.CornerText = MetasoundAsset->GetDisplayName();
+				const FMetasoundAssetBase* MetaSoundAsset = IMetasoundUObjectRegistry::Get().GetObjectAsAssetBase(MetaSound);
+				check(MetaSoundAsset);
+				AppearanceInfo.CornerText = MetaSoundAsset->GetDisplayName();
 			}
 
 			return AppearanceInfo;
@@ -2292,7 +2342,16 @@ namespace Metasound
 		void FEditor::DeleteInterfaceItem(TSharedPtr<FMetasoundGraphMemberSchemaAction> ActionToDelete)
 		{
 			using namespace Metasound::Frontend;
-			check(Metasound);
+			if (!Builder.IsValid())
+			{
+				return;
+			}
+
+			UObject* MetaSound = GetMetasoundObject();
+			if (!MetaSound)
+			{
+				return;
+			}	
 
 			UMetasoundEditorGraphMember* GraphMember = ActionToDelete->GetGraphMember();
 			if (ensure(GraphMember))
@@ -2303,7 +2362,7 @@ namespace Metasound
 
 				{
 					const FScopedTransaction Transaction(LOCTEXT("MetaSoundEditorDeleteSelectedMember", "Delete MetaSound Graph Member"));
-					Metasound->Modify();
+					MetaSound->Modify();
 					Graph.Modify();
 					GraphMember->Modify();
 
@@ -2344,7 +2403,7 @@ namespace Metasound
 				}
 			}
 
-			FGraphBuilder::RegisterGraphWithFrontend(*Metasound);
+			FGraphBuilder::RegisterGraphWithFrontend(*MetaSound);
 		}
 
 		void FEditor::DeleteSelected()
@@ -2370,12 +2429,15 @@ namespace Metasound
 			const FGraphPanelSelectionSet SelectedNodes = MetasoundGraphEditor->GetSelectedNodes();
 			MetasoundGraphEditor->ClearSelectionSet();
 
-			FMetasoundAssetBase* MetasoundAsset = IMetasoundUObjectRegistry::Get().GetObjectAsAssetBase(Metasound);
-			check(MetasoundAsset);
+			UObject* MetaSound = GetMetasoundObject();
+			if (!MetaSound)
+			{
+				return;
+			}
 
 			const FScopedTransaction Transaction(LOCTEXT("MetaSoundEditorDeleteSelectedNode2", "Delete Selected MetaSound Node(s)"));
-			check(Metasound);
-			Metasound->Modify();
+			check(MetaSound);
+			MetaSound->Modify();
 			UEdGraph* Graph = MetasoundGraphEditor->GetCurrentGraph();
 			check(Graph);
 			Graph->Modify();
@@ -2530,10 +2592,10 @@ namespace Metasound
 			FString NodeString;
 			FEdGraphUtilities::ExportNodesToText(SelectedNodes, NodeString);
 
-			FMetasoundAssetBase* MetasoundAsset = IMetasoundUObjectRegistry::Get().GetObjectAsAssetBase(Metasound);
-			check(MetasoundAsset);
+			FMetasoundAssetBase* MetaSoundAsset = IMetasoundUObjectRegistry::Get().GetObjectAsAssetBase(GetMetasoundObject());
+			check(MetaSoundAsset);
 
-			UEdGraph* Graph = MetasoundAsset->GetGraph();
+			UEdGraph* Graph = MetaSoundAsset->GetGraph();
 			if (!Graph)
 			{
 				return false;
@@ -2644,7 +2706,7 @@ namespace Metasound
 			}
 
 			FDocumentPasteNotifications Notifications;
-			TArray<UEdGraphNode*> PastedNodes = FDocumentClipboardUtils::PasteClipboardString(InTransactionText, NodeTextToPaste, Location, *Metasound, Notifications);
+			TArray<UEdGraphNode*> PastedNodes = FDocumentClipboardUtils::PasteClipboardString(InTransactionText, NodeTextToPaste, Location, *GetMetasoundObject(), Notifications);
 
 			// Clear the selection set (newly pasted stuff will be selected)
 			if (!PastedNodes.IsEmpty())
@@ -2863,6 +2925,12 @@ namespace Metasound
 		{
 			using namespace Frontend;
 
+			UObject* MetaSound = GetMetasoundObject();
+			if (!MetaSound)
+			{
+				return;
+			}
+
 			if (!GraphMembersMenu.IsValid())
 			{
 				return;
@@ -2875,7 +2943,7 @@ namespace Metasound
 			{
 				return;
 			}
-				
+
 			UMetasoundEditorGraph& Graph = GetMetaSoundGraphChecked();
 
 			TArray<TObjectPtr<UObject>> SelectedObjects;
@@ -2890,18 +2958,18 @@ namespace Metasound
 				}
 
 				if (const UMetasoundEditorGraphMember* SourceGraphMember = MetasoundAction->GetGraphMember())
-				{								
-					const FScopedTransaction Transaction(TEXT(""), LOCTEXT("MetaSoundEditorDuplicateMember", "Duplicate MetaSound Member"), Metasound);
-					Metasound->Modify();
+				{
+					const FScopedTransaction Transaction(TEXT(""), LOCTEXT("MetaSoundEditorDuplicateMember", "Duplicate MetaSound Member"), MetaSound);
+					MetaSound->Modify();
 
 					UMetasoundEditorGraphMember* NewGraphMember = nullptr;
-										
+
 					//Duplicate the Sources NodeHandle and add a new member from it
 					if (const UMetasoundEditorGraphVariable* SourceGraphVariable = Cast<UMetasoundEditorGraphVariable>(SourceGraphMember))
 					{
 						FConstVariableHandle VariableHandle = FGraphBuilder::DuplicateVariableHandle(Graph.GetMetasoundChecked(), SourceGraphVariable->GetConstVariableHandle());
 						if (ensure(VariableHandle->IsValid()))
-						{							
+						{
 							NewGraphMember = Graph.FindOrAddVariable(VariableHandle);
 						}
 					}
@@ -2916,7 +2984,7 @@ namespace Metasound
 						if (ClassType == EMetasoundFrontendClassType::Input)
 						{
 							if (const FMetasoundFrontendClassInput* SourceInput = DocumentBuilder.FindGraphInput(SourceMemberName))
-							{							
+							{
 								if (const FMetasoundFrontendNode* FrontendNode = DocumentBuilder.DuplicateGraphInput(*SourceInput, SourceInput->DefaultLiteral, Name))
 								{
 									FGraphBuilder::SynchronizeGraphMembers(DocumentBuilder, Graph);
@@ -2942,16 +3010,16 @@ namespace Metasound
 					{
 						if (UMetaSoundEditorSubsystem* EditorSubsystem = GEditor->GetEditorSubsystem<UMetaSoundEditorSubsystem>())
 						{
-							FMetaSoundFrontendDocumentBuilder& DocumentBuilder = IDocumentBuilderRegistry::GetChecked().FindOrBeginBuilding(Metasound);
+							FMetaSoundFrontendDocumentBuilder& DocumentBuilder = IDocumentBuilderRegistry::GetChecked().FindOrBeginBuilding(MetaSound);
 							TSubclassOf<UMetasoundEditorGraphMemberDefaultLiteral> SubClass = SourceGraphMember->GetLiteral()->GetClass();
 							EditorSubsystem->BindMemberMetadata(DocumentBuilder, *NewGraphMember, SubClass, SourceGraphMember->GetLiteral());
 
 							NameToSelect = NewGraphMember->GetMemberName();
 							SelectedObjects.Add(NewGraphMember);
 						}
-					}	
+					}
 				}
-			}	
+			}
 
 			if (GraphMembersMenu.IsValid())
 			{
@@ -3019,8 +3087,9 @@ namespace Metasound
 			using namespace Metasound::Frontend;
 
 			const FScopedTransaction Transaction(LOCTEXT("NodeVersionUpdate", "Update MetaSound Node(s) Class(es)"));
-			check(Metasound);
-			Metasound->Modify();
+			UObject* MetaSound = GetMetasoundObject();
+			check(MetaSound);
+			MetaSound->Modify();
 
 			UMetasoundEditorGraph& Graph = GetMetaSoundGraphChecked();
 			Graph.Modify();
@@ -3100,11 +3169,14 @@ namespace Metasound
 
 		void FEditor::CollectAllActions(FGraphActionListBuilderBase& OutAllActions)
 		{
-			const FMetasoundAssetBase* MetasoundAsset = IMetasoundUObjectRegistry::Get().GetObjectAsAssetBase(Metasound);
-			check(MetasoundAsset);
+			// Uses the builder rather than the local edit object as it may not be set
+			// initially when loading the editor prior to init call on the underlying AssetToolKit.
+			if (!Builder.IsValid())
+			{
+				return;
+			}
 
-			UMetasoundEditorGraph& EdGraph = GetMetaSoundGraphChecked();
-			Frontend::FConstGraphHandle FrontendGraph = MetasoundAsset->GetRootGraphHandle();
+			const FMetaSoundFrontendDocumentBuilder& DocBuilder = Builder->GetBuilder();
 
 			auto GetMemberCategory = [](FName InFullCategoryName)
 			{
@@ -3122,56 +3194,100 @@ namespace Metasound
 				return FText::FromString(CategoryString);
 			};
 
-			FrontendGraph->IterateConstNodes([this, &GetMemberCategory, &EdGraph, ActionList = &OutAllActions](const Frontend::FConstNodeHandle& Input)
+			struct FAddActionParams
 			{
-				constexpr bool bIncludeNamespace = false;
-				const FText Tooltip = Input->GetDescription();
-				const FText MenuDesc = FGraphBuilder::GetDisplayName(*Input, bIncludeNamespace);
-				const FGuid NodeID = Input->GetID();
-				const FText Category = GetMemberCategory(Input->GetNodeName());
+				const FName FullName;
+				const FText Tooltip;
+				const FText MenuDesc;
+				const ENodeSection Section;
+				const FGuid MemberID;
+			};
 
-				TSharedPtr<FMetasoundGraphMemberSchemaAction> NewFuncAction = MakeShared<FMetasoundGraphMemberSchemaAction>(Category, MenuDesc, Tooltip, 1, ENodeSection::Inputs);
+			const FMetasoundAssetBase& AssetBase = DocBuilder.GetMetasoundAsset();
+			UEdGraph& EdGraph = AssetBase.GetGraphChecked();
+			auto AddMemberAction = [&](const FAddActionParams& Params)
+			{
+				const FText Category = GetMemberCategory(Params.FullName);
+				TSharedPtr<FMetasoundGraphMemberSchemaAction> NewFuncAction = MakeShared<FMetasoundGraphMemberSchemaAction>(
+					Category,
+					Params.MenuDesc,
+					Params.Tooltip,
+					1, /* Grouping */
+					Params.Section);
 				NewFuncAction->Graph = &EdGraph;
-				NewFuncAction->SetMemberID(NodeID);
-
-				ActionList->AddAction(NewFuncAction);
-			}, EMetasoundFrontendClassType::Input);
+				NewFuncAction->SetMemberID(Params.MemberID);
+				NewFuncAction->SetBuilder(*Builder.Get());
+				OutAllActions.AddAction(NewFuncAction);
+			};
 
-			FrontendGraph->IterateConstNodes([this, &GetMemberCategory, &EdGraph, ActionList = &OutAllActions](const Frontend::FConstNodeHandle& Output)
+			auto GetVertexDisplayName = [&](const FMetasoundFrontendNode& Node)
 			{
-				constexpr bool bIncludeNamespace = false;
-
-				const FText Tooltip = Output->GetDescription();
-				const FText MenuDesc = FGraphBuilder::GetDisplayName(*Output, bIncludeNamespace);
-				const FGuid NodeID = Output->GetID();
-				const FText Category = GetMemberCategory(Output->GetNodeName());
+				if (const FMetasoundFrontendClass* Class = DocBuilder.FindDependency(Node.ClassID))
+				{
+					constexpr bool bDisplayNamespace = false;
+					return FGraphBuilder::GetDisplayName(Class->Metadata, Node.Name, bDisplayNamespace);
+				}
 
-				TSharedPtr<FMetasoundGraphMemberSchemaAction> NewFuncAction = MakeShared<FMetasoundGraphMemberSchemaAction>(Category, MenuDesc, Tooltip, 1, ENodeSection::Outputs);
-				NewFuncAction->Graph = &EdGraph;
-				NewFuncAction->SetMemberID(NodeID);
-				ActionList->AddAction(NewFuncAction);
-			}, EMetasoundFrontendClassType::Output);
+				return FText();
+			};
 
-			TArray<Frontend::FConstVariableHandle> Variables = FrontendGraph->GetVariables();
-			for (const Frontend::FConstVariableHandle& Variable : Variables)
+			for (const FMetasoundFrontendClassInput& Input : DocBuilder.GetConstDocumentChecked().RootGraph.Interface.Inputs)
 			{
-				const FText MenuDesc = FGraphBuilder::GetDisplayName(*Variable);
-				const FGuid VariableID = Variable->GetID();
-				const FText Category = GetMemberCategory(Variable->GetName());
+				if (const FMetasoundFrontendNode* Node = DocBuilder.FindGraphInputNode(Input.Name))
+				{
+					AddMemberAction(FAddActionParams
+					{
+						Input.Name,							// FullName
+						Input.Metadata.GetDescription(),	// Tooltip
+						GetVertexDisplayName(*Node),		// MenuDesc
+						ENodeSection::Inputs,				// Section
+						Node->GetID()						// MemberID
+					});
+				}
+			}
 
-				TSharedPtr<FMetasoundGraphMemberSchemaAction> NewFuncAction = MakeShared<FMetasoundGraphMemberSchemaAction>(Category, MenuDesc, FText::GetEmpty(), 1, ENodeSection::Variables);
-				NewFuncAction->Graph = &EdGraph;
-				NewFuncAction->SetMemberID(VariableID);
-				OutAllActions.AddAction(NewFuncAction);
+			const FMetasoundFrontendGraphClass& RootGraph = DocBuilder.GetConstDocumentChecked().RootGraph;
+			for (const FMetasoundFrontendClassOutput& Output : RootGraph.Interface.Outputs)
+			{
+				if (const FMetasoundFrontendNode* Node = DocBuilder.FindGraphOutputNode(Output.Name))
+				{
+					AddMemberAction(FAddActionParams
+					{
+						Output.Name,						// FullName
+						Output.Metadata.GetDescription(),	// Tooltip
+						GetVertexDisplayName(*Node),		// MenuDesc
+						ENodeSection::Outputs,				// Section
+						Node->GetID()						// MemberID
+					});
+				}
 			}
 
+			const FMetasoundFrontendGraph& Graph = DocBuilder.FindConstBuildGraphChecked();
+			for (const FMetasoundFrontendVariable& Variable : Graph.Variables)
+			{
+				FText MenuDesc = Variable.DisplayName;
+				if (MenuDesc.IsEmptyOrWhitespace())
+				{
+					FName Namespace;
+					FName ParameterName;
+					Audio::FParameterPath::SplitName(Variable.Name, Namespace, ParameterName);
+					MenuDesc = FText::FromName(ParameterName);
+				}
+
+				AddMemberAction(FAddActionParams
+				{
+					Variable.Name,				// FullName
+					Variable.Description,		// Tooltip
+					MenuDesc,					// MenuDesc
+					ENodeSection::Variables,	// Section
+					Variable.ID					 // MemberID
+				});
+			}
 		}
 
 		void FEditor::CollectStaticSections(TArray<int32>& StaticSectionIDs)
 		{
-			Frontend::FConstDocumentHandle DocumentHandle = GetMetaSoundGraphChecked().GetDocumentHandle();
-			const FMetasoundFrontendGraphClass& RootGraphClass = DocumentHandle->GetRootGraphClass();
-			const bool bIsPreset = RootGraphClass.PresetOptions.bIsPreset;
+			const bool bIsPreset = Builder.IsValid() ? Builder->IsPreset() : false;
 
 			for (int32 i = 0; i < static_cast<int32>(ENodeSection::COUNT); ++i)
 			{
@@ -3490,7 +3606,8 @@ namespace Metasound
 
 		void FEditor::Tick(float DeltaTime)
 		{
-			if (!Metasound)
+			UObject* MetaSound = GetMetasoundObject();
+			if (!MetaSound)
 			{
 				return;
 			}
@@ -3507,36 +3624,46 @@ namespace Metasound
 				}
 			}
 
-			RefreshEditorContext();
+			RefreshEditorContext(*MetaSound);
 
 			GraphConnectionManager->Update(DeltaTime);
 		}
 
-		void FEditor::RefreshEditorContext()
+		void FEditor::RefreshEditorContext(UObject& MetaSound)
 		{
 			TRACE_CPUPROFILER_EVENT_SCOPE(Metasound::Editor::FEditor::RefreshEditorContext);
 
-			check(Builder.IsValid());
+			if (!Builder.IsValid())
+			{
+				return;
+			}
+
+			const FMetaSoundFrontendDocumentBuilder& DocBuilder = Builder->GetConstBuilder();
+			if (!DocBuilder.IsValid())
+			{
+				return;
+			}
+
 			UMetasoundEditorGraph* Graph = nullptr;
-			FGraphBuilder::BindEditorGraph(Builder->GetConstBuilder(), &Graph);
+			FGraphBuilder::BindEditorGraph(DocBuilder, &Graph);
 			check(Graph);
 
-			const bool bSynchronizedGraph = FGraphBuilder::SynchronizeGraph(Builder->GetConstBuilder(), *Graph);
+			const bool bSynchronizedGraph = FGraphBuilder::SynchronizeGraph(DocBuilder, *Graph);
 
-			FMetasoundAssetBase* MetasoundAsset = IMetasoundUObjectRegistry::Get().GetObjectAsAssetBase(Metasound);
-			check(MetasoundAsset);
+			FMetasoundAssetBase* MetaSoundAsset = IMetasoundUObjectRegistry::Get().GetObjectAsAssetBase(&MetaSound);
+			check(MetaSoundAsset);
 
 			bool bNodesSelected = false;
 
 			// Capture after synchronizing as the modification state may be modified therein
-			const FMetasoundFrontendDocumentModifyContext& ModifyContext = MetasoundAsset->GetConstModifyContext();
+			const FMetasoundFrontendDocumentModifyContext& ModifyContext = MetaSoundAsset->GetConstModifyContext();
 			const bool bForceRefreshViews = ModifyContext.GetForceRefreshViews();
 			const TSet<FName>& InterfacesModified = ModifyContext.GetInterfacesModified();
 			const TSet<FGuid>& MembersModified = ModifyContext.GetMemberIDsModified();
 			const TSet<FGuid>& NodesModified = ModifyContext.GetNodeIDsModified();
 			if (bSynchronizedGraph || bForceRefreshViews || !InterfacesModified.IsEmpty() || !NodesModified.IsEmpty() || !MembersModified.IsEmpty())
 			{
-				FGraphValidationResults Results = FGraphBuilder::ValidateGraph(*Metasound);
+				FGraphValidationResults Results = FGraphBuilder::ValidateGraph(MetaSound);
 
 				for (const FGraphNodeValidationResult& Result : Results.GetResults())
 				{
@@ -3614,11 +3741,11 @@ namespace Metasound
 					RefreshInterfaces();
 				}
 
-				HighestMessageSeverity = GetMetaSoundGraphChecked().GetHighestMessageSeverity();
+				HighestMessageSeverity = Graph->GetHighestMessageSeverity();
 
 				// Modify data has been observed both from synchronization & by
 				// updating views by this point, so full reset is completed here.
-				MetasoundAsset->GetModifyContext().Reset();
+				MetaSoundAsset->GetModifyContext().Reset();
 
 				if (!Selection.IsEmpty())
 				{
@@ -3666,6 +3793,16 @@ namespace Metasound
 			return FText::GetEmpty();
 		}
 
+		bool FEditor::IsAuditionable() const
+		{
+			if (const UObject* MetaSound = GetMetasoundObject())
+			{
+				return MetaSound->IsA<USoundBase>();
+			}
+
+			return false;
+		}
+
 		bool FEditor::IsGraphEditable() const
 		{
 			using namespace Metasound::Editor;
@@ -3742,7 +3879,8 @@ namespace Metasound
 
 		FReply FEditor::OnAddButtonClickedOnSection(int32 InSectionID)
 		{
-			if (!Metasound)
+			UObject* MetaSound = GetMetasoundObject();
+			if (!MetaSound)
 			{
 				return FReply::Unhandled();
 			}
@@ -3757,13 +3895,13 @@ namespace Metasound
 				case ENodeSection::Inputs:
 				{
 					const FScopedTransaction Transaction(LOCTEXT("AddInputNode", "Add MetaSound Input"));
-					Metasound->Modify();
+					MetaSound->Modify();
 
 					const FName DataTypeName = GetMetasoundDataTypeName<float>();
 					FCreateNodeVertexParams VertexParams;
 					VertexParams.DataType = DataTypeName;
 
-					Frontend::FNodeHandle NodeHandle = FGraphBuilder::AddInputNodeHandle(*Metasound, VertexParams);
+					Frontend::FNodeHandle NodeHandle = FGraphBuilder::AddInputNodeHandle(*MetaSound, VertexParams);
 					if (ensure(NodeHandle->IsValid()))
 					{
 						NameToSelect = NodeHandle->GetNodeName();
@@ -3779,14 +3917,14 @@ namespace Metasound
 
 				case ENodeSection::Outputs:
 				{
-					const FScopedTransaction Transaction(TEXT(""), LOCTEXT("AddOutputNode", "Add MetaSound Output"), Metasound);
-					Metasound->Modify();
+					const FScopedTransaction Transaction(TEXT(""), LOCTEXT("AddOutputNode", "Add MetaSound Output"), MetaSound);
+					MetaSound->Modify();
 
 					const FName DataTypeName = GetMetasoundDataTypeName<float>();
 					FCreateNodeVertexParams VertexParams;
 					VertexParams.DataType = DataTypeName;
 
-					Frontend::FNodeHandle NodeHandle = FGraphBuilder::AddOutputNodeHandle(*Metasound, VertexParams);
+					Frontend::FNodeHandle NodeHandle = FGraphBuilder::AddOutputNodeHandle(*MetaSound, VertexParams);
 					if (ensure(NodeHandle->IsValid()))
 					{
 						NameToSelect = NodeHandle->GetNodeName();
@@ -3802,12 +3940,12 @@ namespace Metasound
 
 				case ENodeSection::Variables:
 				{
-					const FScopedTransaction Transaction(TEXT(""), LOCTEXT("AddVariableNode", "Add MetaSound Variable"), Metasound);
-					Metasound->Modify();
+					const FScopedTransaction Transaction(TEXT(""), LOCTEXT("AddVariableNode", "Add MetaSound Variable"), MetaSound);
+					MetaSound->Modify();
 
 					const FName DataTypeName = GetMetasoundDataTypeName<float>();
 					
-					Frontend::FVariableHandle FrontendVariable = FGraphBuilder::AddVariableHandle(*Metasound, DataTypeName);
+					Frontend::FVariableHandle FrontendVariable = FGraphBuilder::AddVariableHandle(*MetaSound, DataTypeName);
 					if (ensure(FrontendVariable->IsValid()))
 					{
 						TObjectPtr<UMetasoundEditorGraphVariable> EditorVariable = Graph.FindOrAddVariable(FrontendVariable);
@@ -3824,7 +3962,7 @@ namespace Metasound
 				return FReply::Unhandled();
 			}
 
-			FGraphBuilder::RegisterGraphWithFrontend(*Metasound, true);
+			FGraphBuilder::RegisterGraphWithFrontend(*MetaSound, true);
 
 			if (GraphMembersMenu.IsValid())
 			{
diff --git a/Engine/Plugins/Runtime/Metasound/Source/MetasoundEditor/Private/MetasoundEditor.h b/Engine/Plugins/Runtime/Metasound/Source/MetasoundEditor/Private/MetasoundEditor.h
--- a/Engine/Plugins/Runtime/Metasound/Source/MetasoundEditor/Private/MetasoundEditor.h
+++ b/Engine/Plugins/Runtime/Metasound/Source/MetasoundEditor/Private/MetasoundEditor.h
@@ -51,6 +51,7 @@ class SMetaSoundRenderStats;
 class FSlateRect;
 class IDetailsView;
 class IToolkitHost;
+
 namespace Metasound::Editor
 {
 	class SFindInMetasound;
@@ -59,6 +60,7 @@ class SVerticalBox;
 class UAudioComponent;
 class UEdGraphNode;
 class UMetaSoundPatch;
+class UMetaSoundSource;
 class UMetasoundEditorGraph;
 class UMetasoundEditorGraphNode;
 
@@ -99,7 +101,7 @@ namespace Metasound
 			COUNT
 		};
 
-		class FEditor : public IMetasoundEditor, public FGCObject, public FNotifyHook, public FEditorUndoClient, public FTickableEditorObject
+		class FEditor : public IMetasoundEditor, public FNotifyHook, public FEditorUndoClient, public FTickableEditorObject
 		{
 		public:
 			static const FName EditorName;
@@ -139,13 +141,6 @@ namespace Metasound
 				return FString(TEXT("working-with-audio/sound-sources/meta-sounds"));
 			}
 
-			/** FGCObject interface */
-			virtual void AddReferencedObjects(FReferenceCollector& Collector) override;
-			virtual FString GetReferencerName() const override
-			{
-				return TEXT("Metasound::Editor::FEditor");
-			}
-
 			/** FEditorUndoClient Interface */
 			virtual void PostUndo(bool bSuccess) override;
 			virtual void PostRedo(bool bSuccess) override { PostUndo(bSuccess); }
@@ -186,6 +181,9 @@ namespace Metasound
 			/** Updates selected node classes to highest class found in the MetaSound Class Registry. */
 			void UpdateSelectedNodeClasses();
 
+			/** Whether or not MetaSound can be auditioned */
+			bool IsAuditionable() const;
+
 			/* Whether the displayed graph is marked as editable */
 			bool IsGraphEditable() const;
 
@@ -197,7 +195,7 @@ namespace Metasound
 			}
 
 			/** Creates analyzers */
-			void CreateAnalyzers();
+			void CreateAnalyzers(UMetaSoundSource& MetaSoundSource);
 
 			/** Destroys analyzers */
 			void DestroyAnalyzers();
@@ -343,7 +341,7 @@ namespace Metasound
 			/** Called to redo the last undone action */
 			void RedoGraphAction();
 
-			void RefreshEditorContext();
+			void RefreshEditorContext(UObject& MetaSound);
 
 			/** Show and focus the Find in MetaSound tab. */
 			void ShowFindInMetaSound();
@@ -362,7 +360,7 @@ namespace Metasound
 			virtual void NotifyPostChange(const FPropertyChangedEvent& PropertyChangedEvent, FProperty* PropertyThatChanged) override;
 
 			/** Creates all internal widgets for the tabs to point at */
-			void CreateInternalWidgets();
+			void CreateInternalWidgets(UObject& MetaSound);
 
 			/** Builds the toolbar widget for the Metasound editor */
 			void ExtendToolbarInternal();
@@ -414,7 +412,7 @@ namespace Metasound
 			void OnCreateComment();
 
 			/** Create new graph editor widget */
-			void CreateGraphEditorWidget();
+			void CreateGraphEditorWidget(UObject& MetaSound);
 
 			TSharedPtr<SWidget> BuildAnalyzerWidget() const;
 
@@ -472,10 +470,6 @@ namespace Metasound
 			/** Command list for this editor */
 			TSharedPtr<FUICommandList> GraphEditorCommands;
 
-			/** TODO: Deprecate in favor of just accessing the MetaSound document via
-			  * the builder below */
-			TObjectPtr<UObject> Metasound;
-
 			/** Pointer to builder being actively used to mutate MetaSound asset */
 			TStrongObjectPtr<UMetaSoundBuilderBase> Builder;
 
diff --git a/Engine/Plugins/Runtime/Metasound/Source/MetasoundFrontend/Public/MetasoundFrontendDocument.h b/Engine/Plugins/Runtime/Metasound/Source/MetasoundFrontend/Public/MetasoundFrontendDocument.h
--- a/Engine/Plugins/Runtime/Metasound/Source/MetasoundFrontend/Public/MetasoundFrontendDocument.h
+++ b/Engine/Plugins/Runtime/Metasound/Source/MetasoundFrontend/Public/MetasoundFrontendDocument.h
@@ -1571,7 +1571,7 @@ struct FMetasoundFrontendClassStyle
 	static FMetasoundFrontendClassStyle GenerateClassStyle(const Metasound::FNodeDisplayStyle& InNodeDisplayStyle);
 
 	// Editor only ID that allows for pumping view to reflect changes to class.
-	void UpdateChangeID()
+	void UpdateChangeID() const
 	{
 		ChangeID = FGuid::NewGuid();
 	}
@@ -1582,8 +1582,9 @@ struct FMetasoundFrontendClassStyle
 	}
 
 private:
+	// TODO: Deprecate this change behavior in favor of using the builder API transaction counters
 	UPROPERTY(Transient)
-	FGuid ChangeID;
+	mutable FGuid ChangeID;
 #endif // WITH_EDITORONLY_DATA
 };
 
