- Move duplication to re-use registered builder at different points of duplication process as multiple builders per class name in the builder registry its now supported
- Give asset builders proper names based off the MetaSound asset being built for easier debugging
- Add log supression during duplication from spamming user as brief existence of multiple builders with shared class name is expected behavior (Leaving disabled until cook can be verified log is not being hit)
#rnx

diff --git a/Engine/Plugins/Runtime/Metasound/Source/MetasoundEngine/Private/Metasound.cpp b/Engine/Plugins/Runtime/Metasound/Source/MetasoundEngine/Private/Metasound.cpp
--- a/Engine/Plugins/Runtime/Metasound/Source/MetasoundEngine/Private/Metasound.cpp
+++ b/Engine/Plugins/Runtime/Metasound/Source/MetasoundEngine/Private/Metasound.cpp
@@ -103,6 +103,12 @@ const FMetasoundFrontendDocument& UMetaSoundPatch::GetConstDocument() const
 }
 
 #if WITH_EDITOR
+void UMetaSoundPatch::PreDuplicate(FObjectDuplicationParameters& DupParams)
+{
+	Super::PreDuplicate(DupParams);
+	Metasound::Engine::FAssetHelper::PreDuplicate(this, DupParams);
+}
+
 void UMetaSoundPatch::PostDuplicate(EDuplicateMode::Type InDuplicateMode)
 {
 	Super::PostDuplicate(InDuplicateMode);
diff --git a/Engine/Plugins/Runtime/Metasound/Source/MetasoundEngine/Private/MetasoundDocumentBuilderRegistry.cpp b/Engine/Plugins/Runtime/Metasound/Source/MetasoundEngine/Private/MetasoundDocumentBuilderRegistry.cpp
--- a/Engine/Plugins/Runtime/Metasound/Source/MetasoundEngine/Private/MetasoundDocumentBuilderRegistry.cpp
+++ b/Engine/Plugins/Runtime/Metasound/Source/MetasoundEngine/Private/MetasoundDocumentBuilderRegistry.cpp
@@ -41,36 +41,57 @@ namespace Metasound::Engine
 		FScopeLock Lock(&BuildersCriticalSection);
 
 // #if !NO_LOGGING
-// 		const bool bPrintConflicts = Builders.Contains(InClassName);
+//		bool bLogDuplicateEntries = CanPostEventLog(ELogEvent::DuplicateEntries, ELogVerbosity::Error);
+//		if (bLogDuplicateEntries)
+//		{
+//			bLogDuplicateEntries = Builders.Contains(InClassName);
+//		}
 // #endif // !NO_LOGGING
 
 		Builders.Add(InClassName, NewBuilder);
 
 // #if !NO_LOGGING
-// 		if (bPrintConflicts)
+// 		if (bLogDuplicateEntries)
 // 		{
 // 			TArray<TWeakObjectPtr<UMetaSoundBuilderBase>> Entries;
 // 			Builders.MultiFind(InClassName, Entries);
-// 			UE_LOG(LogMetaSound, Error, TEXT("More than one asset registered with class name '%s'. "
-// 				"Look-up may return builder that is not associated with desired object! \n"
-// 				"This can happen if asset was moved using revision control and original location was revived. \n"
-// 				"Remove all but one of the following assets and relink a duplicate or copied replacement asset:"),
-// 				*InClassName.ToString());
-// 			for (const TWeakObjectPtr<UMetaSoundBuilderBase>& BuilderPtr : Entries)
+// 
+// 			// Don't print stale entries as during cook and some editor asset actions,
+// 			// these may be removed after a new valid builder is created.  If stale
+// 			// entries leak, they will show up on registry logging upon destruction.
+// 			Entries.RemoveAllSwap([](const TWeakObjectPtr<UMetaSoundBuilderBase>& Builder) { return !Builder.IsValid(); });
+// 
+// 			if (!Entries.IsEmpty())
 // 			{
-// 				if (BuilderPtr.IsValid())
+// 				UE_LOG(LogMetaSound, Error, TEXT("More than one asset registered with class name '%s'. "
+// 					"Look-up may return builder that is not associated with desired object! \n"
+// 					"This can happen if asset was moved using revision control and original location was revived. \n"
+// 					"Remove all but one of the following assets and relink a duplicate or copied replacement asset:"),
+// 					*InClassName.ToString());
+// 				for (const TWeakObjectPtr<UMetaSoundBuilderBase>& BuilderPtr : Entries)
 // 				{
 // 					UE_LOG(LogMetaSound, Error, TEXT("- %s"), *BuilderPtr->GetConstBuilder().CastDocumentObjectChecked<UObject>().GetPathName());
 // 				}
-// 				else
-// 				{
-// 					UE_LOG(LogMetaSound, Error, TEXT("- STALE ENTRY (Not removed prior to asset removal)"));
-// 				}
 // 			}
 // 		}
 // #endif // !NO_LOGGING
 	}
 
+	bool FDocumentBuilderRegistry::CanPostEventLog(ELogEvent Event, ELogVerbosity::Type Verbosity) const
+	{
+#if NO_LOGGING
+		return false;
+#else // !NO_LOGGING
+
+		if (const ELogVerbosity::Type* SetVerbosity = EventLogVerbosity.Find(Event))
+		{
+			return *SetVerbosity >= Verbosity;
+		}
+
+		return true;
+#endif // !NO_LOGGING
+	}
+
 #if WITH_EDITORONLY_DATA
 	FMetaSoundFrontendDocumentBuilder& FDocumentBuilderRegistry::FindOrBeginBuilding(TScriptInterface<IMetaSoundDocumentInterface> MetaSound)
 	{
@@ -288,4 +309,9 @@ namespace Metasound::Engine
 
 		return bReloaded;
 	}
+
+	void FDocumentBuilderRegistry::SetEventLogVerbosity(ELogEvent Event, ELogVerbosity::Type Verbosity)
+	{
+		EventLogVerbosity.FindOrAdd(Event) = Verbosity;
+	}
 } // namespace Metasound::Engine
\ No newline at end of file
diff --git a/Engine/Plugins/Runtime/Metasound/Source/MetasoundEngine/Private/MetasoundEngineAsset.h b/Engine/Plugins/Runtime/Metasound/Source/MetasoundEngine/Private/MetasoundEngineAsset.h
--- a/Engine/Plugins/Runtime/Metasound/Source/MetasoundEngine/Private/MetasoundEngineAsset.h
+++ b/Engine/Plugins/Runtime/Metasound/Source/MetasoundEngine/Private/MetasoundEngineAsset.h
@@ -29,20 +29,32 @@ namespace Metasound::Engine
 	struct FAssetHelper
 	{
 #if WITH_EDITOR
+		static void PreDuplicate(TScriptInterface<IMetaSoundDocumentInterface> MetaSound, FObjectDuplicationParameters& DupParams)
+		{
+			FDocumentBuilderRegistry::GetChecked().SetEventLogVerbosity(FDocumentBuilderRegistry::ELogEvent::DuplicateEntries, ELogVerbosity::NoLogging);
+		}
+
 		static void PostDuplicate(TScriptInterface<IMetaSoundDocumentInterface> MetaSound, EDuplicateMode::Type InDuplicateMode, FGuid& OutAssetClassID)
 		{
+			using namespace Engine;
 			using namespace Frontend;
 
 			if (InDuplicateMode == EDuplicateMode::Normal)
 			{
-				// Uses a bespoke builder instead of a registered one to avoid class name collision
-				// that can result in accessing the wrong pre-existing builder for the pre-existing
-				// MetaSound the provided MetaSound was duplicated from. Being that this is a new
-				// MetaSound, it is perfectly safe and should not result in multiple builders existing
-				// for the given asset.
-				FMetaSoundFrontendDocumentBuilder DuplicateBuilder(MetaSound);
-				const FMetasoundFrontendClassName NewName = DuplicateBuilder.GenerateNewClassName();
+				UObject* MetaSoundObject = MetaSound.GetObject();
+				check(MetaSoundObject);
+
+				FDocumentBuilderRegistry& BuilderRegistry = FDocumentBuilderRegistry::GetChecked();
+				UMetaSoundBuilderBase& DuplicateBuilder = BuilderRegistry.FindOrBeginBuilding(*MetaSoundObject);
+
+				FMetaSoundFrontendDocumentBuilder& DocBuilder = DuplicateBuilder.GetBuilder();
+				const FMetasoundFrontendClassName DuplicateName = DocBuilder.GetConstDocumentChecked().RootGraph.Metadata.GetClassName();
+				const FMetasoundFrontendClassName NewName = DocBuilder.GenerateNewClassName();
 				ensureAlwaysMsgf(IMetaSoundAssetManager::GetChecked().TryGetAssetIDFromClassName(NewName, OutAssetClassID), TEXT("Failed to retrieve newly duplicated MetaSoundClassName AssetID"));
+
+				constexpr bool bForceUnregisterNodeClass = true;
+				BuilderRegistry.FinishBuilding(DuplicateName, MetaSound->GetAssetPathChecked(), bForceUnregisterNodeClass);
+				BuilderRegistry.SetEventLogVerbosity(FDocumentBuilderRegistry::ELogEvent::DuplicateEntries, ELogVerbosity::All);
 			}
 		}
 
diff --git a/Engine/Plugins/Runtime/Metasound/Source/MetasoundEngine/Private/MetasoundSource.cpp b/Engine/Plugins/Runtime/Metasound/Source/MetasoundEngine/Private/MetasoundSource.cpp
--- a/Engine/Plugins/Runtime/Metasound/Source/MetasoundEngine/Private/MetasoundSource.cpp
+++ b/Engine/Plugins/Runtime/Metasound/Source/MetasoundEngine/Private/MetasoundSource.cpp
@@ -361,6 +361,12 @@ void UMetaSoundSource::PostEditUndo()
 	Metasound::Engine::FAssetHelper::PostEditUndo(*this);
 }
 
+void UMetaSoundSource::PreDuplicate(FObjectDuplicationParameters& DupParams)
+{
+	Super::PreDuplicate(DupParams);
+	Metasound::Engine::FAssetHelper::PreDuplicate(this, DupParams);
+}
+
 void UMetaSoundSource::PostDuplicate(EDuplicateMode::Type InDuplicateMode)
 {
 	Super::PostDuplicate(InDuplicateMode);
diff --git a/Engine/Plugins/Runtime/Metasound/Source/MetasoundEngine/Public/Metasound.h b/Engine/Plugins/Runtime/Metasound/Source/MetasoundEngine/Public/Metasound.h
--- a/Engine/Plugins/Runtime/Metasound/Source/MetasoundEngine/Public/Metasound.h
+++ b/Engine/Plugins/Runtime/Metasound/Source/MetasoundEngine/Public/Metasound.h
@@ -150,6 +150,7 @@ public:
 	virtual const FMetasoundFrontendDocument& GetConstDocument() const override;
 
 #if WITH_EDITOR
+	virtual void PreDuplicate(FObjectDuplicationParameters& DupParams) override;
 	virtual void PostDuplicate(EDuplicateMode::Type DuplicateMode) override;
 	virtual void PostEditUndo() override;
 #endif // WITH_EDITOR
diff --git a/Engine/Plugins/Runtime/Metasound/Source/MetasoundEngine/Public/MetasoundDocumentBuilderRegistry.h b/Engine/Plugins/Runtime/Metasound/Source/MetasoundEngine/Public/MetasoundDocumentBuilderRegistry.h
--- a/Engine/Plugins/Runtime/Metasound/Source/MetasoundEngine/Public/MetasoundDocumentBuilderRegistry.h
+++ b/Engine/Plugins/Runtime/Metasound/Source/MetasoundEngine/Public/MetasoundDocumentBuilderRegistry.h
@@ -29,6 +29,11 @@ namespace Metasound::Engine
 			return static_cast<FDocumentBuilderRegistry&>(IDocumentBuilderRegistry::GetChecked());
 		}
 
+		enum class ELogEvent : uint8
+		{
+			DuplicateEntries
+		};
+
 		template <typename BuilderClass>
 		BuilderClass& CreateTransientBuilder(FName BuilderName = FName())
 		{
@@ -62,10 +67,16 @@ namespace Metasound::Engine
 				return *CastChecked<BuilderClass>(Builder);
 			}
 
+			FNameBuilder BuilderName;
+			BuilderName.Append(InMetaSoundObject.GetName());
+			BuilderName.Append(TEXT("_Builder"));
+			const UClass& BuilderUClass = DocInterface->GetBuilderUClass();
+			const FName NewName = MakeUniqueObjectName(nullptr, &BuilderUClass, FName(*BuilderName));
+
 			TObjectPtr<UMetaSoundBuilderBase> NewBuilder;
 			{
 				FScopeLock Lock(&BuildersCriticalSection);
-				NewBuilder = CastChecked<UMetaSoundBuilderBase>(NewObject<UObject>(&InMetaSoundObject, &DocInterface->GetBuilderUClass()));
+				NewBuilder = CastChecked<UMetaSoundBuilderBase>(NewObject<UObject>(&InMetaSoundObject, &BuilderUClass, NewName));
 				FMetaSoundFrontendDocumentBuilder& BuilderRef = NewBuilder->GetBuilder();
 				BuilderRef = FMetaSoundFrontendDocumentBuilder(DocInterface);
 				const FMetasoundFrontendDocument& Document = DocInterface->GetConstDocument();
@@ -110,8 +121,13 @@ namespace Metasound::Engine
 
 		bool ReloadBuilder(const FMetasoundFrontendClassName& InClassName) const override;
 
+		void SetEventLogVerbosity(ELogEvent Event, ELogVerbosity::Type Verbosity);
+
 	private:
 		void AddBuilderInternal(const FMetasoundFrontendClassName& InClassName, UMetaSoundBuilderBase* NewBuilder) const;
+		bool CanPostEventLog(ELogEvent Event, ELogVerbosity::Type Verbosity) const;
 		void FinishBuildingInternal(UMetaSoundBuilderBase& Builder, bool bForceUnregisterNodeClass) const;
+
+		TSortedMap<ELogEvent, ELogVerbosity::Type> EventLogVerbosity;
 	};
 } // namespace Metasound::Engine
\ No newline at end of file
diff --git a/Engine/Plugins/Runtime/Metasound/Source/MetasoundEngine/Public/MetasoundSource.h b/Engine/Plugins/Runtime/Metasound/Source/MetasoundEngine/Public/MetasoundSource.h
--- a/Engine/Plugins/Runtime/Metasound/Source/MetasoundEngine/Public/MetasoundSource.h
+++ b/Engine/Plugins/Runtime/Metasound/Source/MetasoundEngine/Public/MetasoundSource.h
@@ -211,10 +211,9 @@ public:
 		return false;
 	}
 
+	virtual void PreDuplicate(FObjectDuplicationParameters& DupParams) override;
 	virtual void PostDuplicate(EDuplicateMode::Type InDuplicateMode) override;
-
 	virtual void PostEditChangeProperty(FPropertyChangedEvent& InEvent) override;
-
 	virtual bool CanEditChange(const FProperty* InProperty) const override;
 
 private:
