Fix issue with returning false if Access/DataType is already the correct value when being swapped
#rnx

diff --git a/Engine/Plugins/Runtime/Metasound/Source/MetasoundFrontend/Private/MetasoundFrontendDocumentBuilder.cpp b/Engine/Plugins/Runtime/Metasound/Source/MetasoundFrontend/Private/MetasoundFrontendDocumentBuilder.cpp
--- a/Engine/Plugins/Runtime/Metasound/Source/MetasoundFrontend/Private/MetasoundFrontendDocumentBuilder.cpp
+++ b/Engine/Plugins/Runtime/Metasound/Source/MetasoundFrontend/Private/MetasoundFrontendDocumentBuilder.cpp
@@ -3145,47 +3145,48 @@ bool FMetaSoundFrontendDocumentBuilder::SetGraphInputAccessType(FName InputName,
 		FMetasoundFrontendDocument& Document = GetDocumentChecked();
 		FMetasoundFrontendClassInput& GraphInput = Document.RootGraph.Interface.Inputs[*Index];
 
-		if (GraphInput.AccessType != AccessType)
+		if (GraphInput.AccessType == AccessType)
 		{
-			GraphInput.AccessType = AccessType;
-			if (AccessType == EMetasoundFrontendVertexAccessType::Reference)
+			return true;
+		}
+
+		GraphInput.AccessType = AccessType;
+		if (AccessType == EMetasoundFrontendVertexAccessType::Reference)
+		{
+			const IDocumentGraphNodeCache& NodeCache = DocumentCache->GetNodeCache();
+			if (const int32* NodeIndex = NodeCache.FindNodeIndex(GraphInput.NodeID))
 			{
-				const IDocumentGraphNodeCache& NodeCache = DocumentCache->GetNodeCache();
-				if (const int32* NodeIndex = NodeCache.FindNodeIndex(GraphInput.NodeID))
+				FMetasoundFrontendNode& Node = FindBuildGraphChecked().Nodes[*NodeIndex];
+				const FMetasoundFrontendVertex& NodeOutput = Node.Interface.Outputs.Last();
+				IterateNodesConnectedWithVertex({ GraphInput.NodeID, NodeOutput.VertexID }, [this, &AccessType, &Document](const FMetasoundFrontendEdge& Edge, FMetasoundFrontendNode& ConnectedNode)
 				{
-					FMetasoundFrontendNode& Node = FindBuildGraphChecked().Nodes[*NodeIndex];
-					const FMetasoundFrontendVertex& NodeOutput = Node.Interface.Outputs.Last();
-					IterateNodesConnectedWithVertex({ GraphInput.NodeID, NodeOutput.VertexID }, [this, &AccessType, &Document](const FMetasoundFrontendEdge& Edge, FMetasoundFrontendNode& ConnectedNode)
+					if (const FMetasoundFrontendClass* ConnectedNodeClass = FindDependency(ConnectedNode.ClassID))
 					{
-						if (const FMetasoundFrontendClass* ConnectedNodeClass = FindDependency(ConnectedNode.ClassID))
+						// If connected to an input template node, disconnect the template node from other nodes as the data type is
+						// about to be mismatched.  Otherwise, direct connection to other nodes (i.e. at runtime when template
+						// nodes aren't injected) forcefully remove to avoid data type mismatch.
+						if (ConnectedNodeClass->Metadata.GetClassName() == FInputNodeTemplate::ClassName)
 						{
-							// If connected to an input template node, disconnect the template node from other nodes as the data type is
-							// about to be mismatched.  Otherwise, direct connection to other nodes (i.e. at runtime when template
-							// nodes aren't injected) forcefully remove to avoid data type mismatch.
-							if (ConnectedNodeClass->Metadata.GetClassName() == FInputNodeTemplate::ClassName)
-							{
-								const FMetasoundFrontendVertex& ConnectedNodeOutput = ConnectedNode.Interface.Outputs.Last();
-								IterateNodesConnectedWithVertex({ Edge.ToNodeID, ConnectedNodeOutput.VertexID }, [this, &AccessType](const FMetasoundFrontendEdge& TempEdge, FMetasoundFrontendNode&)
-								{
-									const EMetasoundFrontendVertexAccessType ConnectedAccessType = GetNodeInputAccessType(TempEdge.ToNodeID, TempEdge.ToVertexID);
-									if (!FMetasoundFrontendClassVertex::CanConnectVertexAccessTypes(AccessType, ConnectedAccessType))
-									{
-										RemoveEdgeToNodeInput(TempEdge.ToNodeID, TempEdge.ToVertexID);
-									}
-								});
-							}
-							else
+							const FMetasoundFrontendVertex& ConnectedNodeOutput = ConnectedNode.Interface.Outputs.Last();
+							IterateNodesConnectedWithVertex({ Edge.ToNodeID, ConnectedNodeOutput.VertexID }, [this, &AccessType](const FMetasoundFrontendEdge& TempEdge, FMetasoundFrontendNode&)
 							{
-								const EMetasoundFrontendVertexAccessType ConnectedAccessType = GetNodeInputAccessType(Edge.ToNodeID, Edge.ToVertexID);
+								const EMetasoundFrontendVertexAccessType ConnectedAccessType = GetNodeInputAccessType(TempEdge.ToNodeID, TempEdge.ToVertexID);
 								if (!FMetasoundFrontendClassVertex::CanConnectVertexAccessTypes(AccessType, ConnectedAccessType))
 								{
-									RemoveEdgeToNodeInput(Edge.ToNodeID, Edge.ToVertexID);
+									RemoveEdgeToNodeInput(TempEdge.ToNodeID, TempEdge.ToVertexID);
 								}
+							});
+						}
+						else
+						{
+							const EMetasoundFrontendVertexAccessType ConnectedAccessType = GetNodeInputAccessType(Edge.ToNodeID, Edge.ToVertexID);
+							if (!FMetasoundFrontendClassVertex::CanConnectVertexAccessTypes(AccessType, ConnectedAccessType))
+							{
+								RemoveEdgeToNodeInput(Edge.ToNodeID, Edge.ToVertexID);
 							}
 						}
-					});
-
-				}
+					}
+				});
 			}
 
 			const bool bNodeConformed = ConformGraphInputNodeToClass(GraphInput);
@@ -3213,53 +3214,55 @@ bool FMetaSoundFrontendDocumentBuilder::SetGraphInputDataType(FName InputName, F
 		{
 			FMetasoundFrontendDocument& Document = GetDocumentChecked();
 			FMetasoundFrontendClassInput& GraphInput = Document.RootGraph.Interface.Inputs[*Index];
-			if (GraphInput.TypeName != DataType)
+			if (GraphInput.TypeName == DataType)
 			{
-				FMetasoundFrontendLiteral DefaultLiteral;
-				DefaultLiteral.SetFromLiteral(IDataTypeRegistry::Get().CreateDefaultLiteral(DataType));
-				GraphInput.DefaultLiteral = DefaultLiteral;
-				GraphInput.TypeName = DataType;
+				return true;
+			}
+
+			FMetasoundFrontendLiteral DefaultLiteral;
+			DefaultLiteral.SetFromLiteral(IDataTypeRegistry::Get().CreateDefaultLiteral(DataType));
+			GraphInput.DefaultLiteral = DefaultLiteral;
+			GraphInput.TypeName = DataType;
 
-				const IDocumentGraphNodeCache& NodeCache = DocumentCache->GetNodeCache();
-				if (const int32* NodeIndex = NodeCache.FindNodeIndex(GraphInput.NodeID))
+			const IDocumentGraphNodeCache& NodeCache = DocumentCache->GetNodeCache();
+			if (const int32* NodeIndex = NodeCache.FindNodeIndex(GraphInput.NodeID))
+			{
+				FMetasoundFrontendNode& Node = FindBuildGraphChecked().Nodes[*NodeIndex];
+				FMetasoundFrontendVertex& NodeOutput = Node.Interface.Outputs.Last();
+				IterateNodesConnectedWithVertex({ GraphInput.NodeID, NodeOutput.VertexID }, [this, &DataType](const FMetasoundFrontendEdge& Edge, FMetasoundFrontendNode& ConnectedNode)
 				{
-					FMetasoundFrontendNode& Node = FindBuildGraphChecked().Nodes[*NodeIndex];
-					FMetasoundFrontendVertex& NodeOutput = Node.Interface.Outputs.Last();
-					IterateNodesConnectedWithVertex({ GraphInput.NodeID, NodeOutput.VertexID }, [this, &DataType](const FMetasoundFrontendEdge& Edge, FMetasoundFrontendNode& ConnectedNode)
+					const FMetasoundFrontendClass* ConnectedNodeClass = FindDependency(ConnectedNode.ClassID);
+					if (ensure(ConnectedNodeClass))
 					{
-						const FMetasoundFrontendClass* ConnectedNodeClass = FindDependency(ConnectedNode.ClassID);
-						if (ensure(ConnectedNodeClass))
+						// If connected to an input template node, disconnect the template node from other nodes as the data type is
+						// about to be mismatched.  Otherwise, direct connection to other nodes (i.e. at runtime when template
+						// nodes aren't injected) forcefully remove to avoid data type mismatch.
+						if (ConnectedNodeClass->Metadata.GetClassName() == FInputNodeTemplate::ClassName)
 						{
-							// If connected to an input template node, disconnect the template node from other nodes as the data type is
-							// about to be mismatched.  Otherwise, direct connection to other nodes (i.e. at runtime when template
-							// nodes aren't injected) forcefully remove to avoid data type mismatch.
-							if (ConnectedNodeClass->Metadata.GetClassName() == FInputNodeTemplate::ClassName)
-							{
-								RemoveEdgesFromNodeOutput(Edge.ToNodeID, ConnectedNode.Interface.Outputs.Last().VertexID);
-								ConnectedNode.Interface.Inputs.Last().TypeName = DataType;
-								ConnectedNode.Interface.Outputs.Last().TypeName = DataType;
-							}
-							else
-							{
-								RemoveEdgeToNodeInput(Edge.ToNodeID, Edge.ToVertexID);
-							}
+							RemoveEdgesFromNodeOutput(Edge.ToNodeID, ConnectedNode.Interface.Outputs.Last().VertexID);
+							ConnectedNode.Interface.Inputs.Last().TypeName = DataType;
+							ConnectedNode.Interface.Outputs.Last().TypeName = DataType;
 						}
-					});
-				}
+						else
+						{
+							RemoveEdgeToNodeInput(Edge.ToNodeID, Edge.ToVertexID);
+						}
+					}
+				});
+			}
 
-				const bool bNodeConformed = ConformGraphInputNodeToClass(GraphInput);
-				if (bNodeConformed)
-				{
-					RemoveUnusedDependencies();
+			const bool bNodeConformed = ConformGraphInputNodeToClass(GraphInput);
+			if (bNodeConformed)
+			{
+				RemoveUnusedDependencies();
 
 #if WITH_EDITORONLY_DATA
-					ClearMemberMetadata(GraphInput.NodeID);
-					Document.Metadata.ModifyContext.AddMemberIDModified(GraphInput.NodeID);
-					Document.Metadata.ModifyContext.AddNodeIDModified(GraphInput.NodeID);
+				ClearMemberMetadata(GraphInput.NodeID);
+				Document.Metadata.ModifyContext.AddMemberIDModified(GraphInput.NodeID);
+				Document.Metadata.ModifyContext.AddNodeIDModified(GraphInput.NodeID);
 #endif // WITH_EDITORONLY_DATA
 
-					return true;
-				}
+				return true;
 			}
 		}
 	}
@@ -3332,40 +3335,42 @@ bool FMetaSoundFrontendDocumentBuilder::SetGraphOutputAccessType(FName OutputNam
 	{
 		FMetasoundFrontendDocument& Document = GetDocumentChecked();
 		FMetasoundFrontendClassOutput& GraphOutput = Document.RootGraph.Interface.Outputs[*Index];
-		if (GraphOutput.AccessType != AccessType)
+		if (GraphOutput.AccessType == AccessType)
+		{
+			return true;
+		}
+
+		GraphOutput.AccessType = AccessType;
+		if (AccessType == EMetasoundFrontendVertexAccessType::Value)
 		{
-			GraphOutput.AccessType = AccessType;
-			if (AccessType == EMetasoundFrontendVertexAccessType::Value)
+			const IDocumentGraphNodeCache& NodeCache = DocumentCache->GetNodeCache();
+			if (const int32* NodeIndex = NodeCache.FindNodeIndex(GraphOutput.NodeID))
 			{
-				const IDocumentGraphNodeCache& NodeCache = DocumentCache->GetNodeCache();
-				if (const int32* NodeIndex = NodeCache.FindNodeIndex(GraphOutput.NodeID))
+				FMetasoundFrontendNode& Node = FindBuildGraphChecked().Nodes[*NodeIndex];
+				const FMetasoundFrontendVertex& NodeInput = Node.Interface.Inputs.Last();
+				IterateNodesConnectedWithVertex({ GraphOutput.NodeID, NodeInput.VertexID }, [this, &AccessType](const FMetasoundFrontendEdge& Edge, FMetasoundFrontendNode& ConnectedNode)
 				{
-					FMetasoundFrontendNode& Node = FindBuildGraphChecked().Nodes[*NodeIndex];
-					const FMetasoundFrontendVertex& NodeInput = Node.Interface.Inputs.Last();
-					IterateNodesConnectedWithVertex({ GraphOutput.NodeID, NodeInput.VertexID }, [this, &AccessType](const FMetasoundFrontendEdge& Edge, FMetasoundFrontendNode& ConnectedNode)
+					if (const FMetasoundFrontendClass* ConnectedNodeClass = FindDependency(ConnectedNode.ClassID))
 					{
-						if (const FMetasoundFrontendClass* ConnectedNodeClass = FindDependency(ConnectedNode.ClassID))
+						const FMetasoundFrontendVertex& ConnectedNodeOutput = ConnectedNode.Interface.Outputs.Last();
+						const EMetasoundFrontendVertexAccessType ConnectedAccessType = GetNodeOutputAccessType(ConnectedNode.GetID(), ConnectedNodeOutput.VertexID);
+						if (!FMetasoundFrontendClassVertex::CanConnectVertexAccessTypes(ConnectedAccessType, AccessType))
 						{
-							const FMetasoundFrontendVertex& ConnectedNodeOutput = ConnectedNode.Interface.Outputs.Last();
-							const EMetasoundFrontendVertexAccessType ConnectedAccessType = GetNodeOutputAccessType(ConnectedNode.GetID(), ConnectedNodeOutput.VertexID);
-							if (!FMetasoundFrontendClassVertex::CanConnectVertexAccessTypes(ConnectedAccessType, AccessType))
-							{
-								RemoveEdgeToNodeInput(Edge.ToNodeID, Edge.ToVertexID);
-							}
+							RemoveEdgeToNodeInput(Edge.ToNodeID, Edge.ToVertexID);
 						}
-					});
-				}
+					}
+				});
 			}
+		}
 
-			const bool bNodeConformed = ConformGraphOutputNodeToClass(GraphOutput);
-			if (bNodeConformed)
-			{
-	#if WITH_EDITORONLY_DATA
-				Document.Metadata.ModifyContext.AddMemberIDModified(GraphOutput.NodeID);
-	#endif // WITH_EDITORONLY_DATA
+		const bool bNodeConformed = ConformGraphOutputNodeToClass(GraphOutput);
+		if (bNodeConformed)
+		{
+#if WITH_EDITORONLY_DATA
+			Document.Metadata.ModifyContext.AddMemberIDModified(GraphOutput.NodeID);
+#endif // WITH_EDITORONLY_DATA
 
-				return true;
-			}
+			return true;
 		}
 	}
 
@@ -3382,36 +3387,38 @@ bool FMetaSoundFrontendDocumentBuilder::SetGraphOutputDataType(FName OutputName,
 		{
 			FMetasoundFrontendDocument& Document = GetDocumentChecked();
 			FMetasoundFrontendClassOutput& GraphOutput = Document.RootGraph.Interface.Outputs[*Index];
-			if (GraphOutput.TypeName != DataType)
+			if (GraphOutput.TypeName == DataType)
 			{
-				GraphOutput.TypeName = DataType;
+				return true;
+			}
 
-				const IDocumentGraphNodeCache& NodeCache = DocumentCache->GetNodeCache();
-				if (const int32* NodeIndex = NodeCache.FindNodeIndex(GraphOutput.NodeID))
-				{
-					FMetasoundFrontendNode& Node = FindBuildGraphChecked().Nodes[*NodeIndex];
+			GraphOutput.TypeName = DataType;
+
+			const IDocumentGraphNodeCache& NodeCache = DocumentCache->GetNodeCache();
+			if (const int32* NodeIndex = NodeCache.FindNodeIndex(GraphOutput.NodeID))
+			{
+				FMetasoundFrontendNode& Node = FindBuildGraphChecked().Nodes[*NodeIndex];
 
-					FMetasoundFrontendLiteral DefaultLiteral;
-					DefaultLiteral.SetFromLiteral(IDataTypeRegistry::Get().CreateDefaultLiteral(DataType));
-					FMetasoundFrontendVertex& NodeInput = Node.Interface.Inputs.Last();
-					Node.InputLiterals = { FMetasoundFrontendVertexLiteral { NodeInput.VertexID, MoveTemp(DefaultLiteral) } };
+				FMetasoundFrontendLiteral DefaultLiteral;
+				DefaultLiteral.SetFromLiteral(IDataTypeRegistry::Get().CreateDefaultLiteral(DataType));
+				FMetasoundFrontendVertex& NodeInput = Node.Interface.Inputs.Last();
+				Node.InputLiterals = { FMetasoundFrontendVertexLiteral { NodeInput.VertexID, MoveTemp(DefaultLiteral) } };
 
-					RemoveEdgeToNodeInput(GraphOutput.NodeID, NodeInput.VertexID);
-					GraphOutput.TypeName = DataType;
+				RemoveEdgeToNodeInput(GraphOutput.NodeID, NodeInput.VertexID);
+				GraphOutput.TypeName = DataType;
 
-				}
+			}
 
-				const bool bNodeConformed = ConformGraphOutputNodeToClass(GraphOutput);
-				if (bNodeConformed)
-				{
+			const bool bNodeConformed = ConformGraphOutputNodeToClass(GraphOutput);
+			if (bNodeConformed)
+			{
 
 #if WITH_EDITORONLY_DATA
-					ClearMemberMetadata(GraphOutput.NodeID);
-					Document.Metadata.ModifyContext.AddMemberIDModified(GraphOutput.NodeID);
+				ClearMemberMetadata(GraphOutput.NodeID);
+				Document.Metadata.ModifyContext.AddMemberIDModified(GraphOutput.NodeID);
 #endif // WITH_EDITORONLY_DATA
 
-					return true;
-				}
+				return true;
 			}
 		}
 	}
