First pass at member display name code clean-up, misc bug fixes
#rnx

diff --git a/Engine/Plugins/Runtime/Metasound/Source/MetasoundEditor/Private/MetasoundEditor.cpp b/Engine/Plugins/Runtime/Metasound/Source/MetasoundEditor/Private/MetasoundEditor.cpp
--- a/Engine/Plugins/Runtime/Metasound/Source/MetasoundEditor/Private/MetasoundEditor.cpp
+++ b/Engine/Plugins/Runtime/Metasound/Source/MetasoundEditor/Private/MetasoundEditor.cpp
@@ -3169,6 +3169,8 @@ namespace Metasound
 
 		void FEditor::CollectAllActions(FGraphActionListBuilderBase& OutAllActions)
 		{
+			using namespace Frontend;
+
 			// Uses the builder rather than the local edit object as it may not be set
 			// initially when loading the editor prior to init call on the underlying AssetToolKit.
 			if (!Builder.IsValid())
@@ -3203,6 +3205,7 @@ namespace Metasound
 				const FGuid MemberID;
 			};
 
+			constexpr bool bDisplayNamespace = false;
 			const FMetasoundAssetBase& AssetBase = DocBuilder.GetMetasoundAsset();
 			UEdGraph& EdGraph = AssetBase.GetGraphChecked();
 			auto AddMemberAction = [&](const FAddActionParams& Params)
@@ -3220,28 +3223,23 @@ namespace Metasound
 				OutAllActions.AddAction(NewFuncAction);
 			};
 
-			auto GetVertexDisplayName = [&](const FMetasoundFrontendNode& Node)
-			{
-				if (const FMetasoundFrontendClass* Class = DocBuilder.FindDependency(Node.ClassID))
-				{
-					constexpr bool bDisplayNamespace = false;
-					return FGraphBuilder::GetDisplayName(Class->Metadata, Node.Name, bDisplayNamespace);
-				}
-
-				return FText();
-			};
-
 			for (const FMetasoundFrontendClassInput& Input : DocBuilder.GetConstDocumentChecked().RootGraph.Interface.Inputs)
 			{
 				if (const FMetasoundFrontendNode* Node = DocBuilder.FindGraphInputNode(Input.Name))
 				{
+					FText DisplayName;
+					if (const FMetasoundFrontendClassInput* ClassInput = DocBuilder.FindGraphInput(Node->Name))
+					{
+						DisplayName = ClassInput->Metadata.GetDisplayName();
+					}
+
 					AddMemberAction(FAddActionParams
 					{
-						Input.Name,							// FullName
-						Input.Metadata.GetDescription(),	// Tooltip
-						GetVertexDisplayName(*Node),		// MenuDesc
-						ENodeSection::Inputs,				// Section
-						Node->GetID()						// MemberID
+						Input.Name, // FullName
+						Input.Metadata.GetDescription(), // Tooltip
+						INodeTemplate::ResolveMemberDisplayName(Node->Name, DisplayName, bDisplayNamespace), // MenuDesc
+						ENodeSection::Inputs, // Section
+						Node->GetID() // MemberID
 					});
 				}
 			}
@@ -3251,13 +3249,19 @@ namespace Metasound
 			{
 				if (const FMetasoundFrontendNode* Node = DocBuilder.FindGraphOutputNode(Output.Name))
 				{
+					FText DisplayName;
+					if (const FMetasoundFrontendClassOutput* ClassOutput = DocBuilder.FindGraphOutput(Node->Name))
+					{
+						DisplayName = ClassOutput->Metadata.GetDisplayName();
+					}
+
 					AddMemberAction(FAddActionParams
 					{
-						Output.Name,						// FullName
-						Output.Metadata.GetDescription(),	// Tooltip
-						GetVertexDisplayName(*Node),		// MenuDesc
-						ENodeSection::Outputs,				// Section
-						Node->GetID()						// MemberID
+						Output.Name, // FullName
+						Output.Metadata.GetDescription(), // Tooltip
+						INodeTemplate::ResolveMemberDisplayName(Node->Name, DisplayName, bDisplayNamespace), // MenuDesc
+						ENodeSection::Outputs, // Section
+						Node->GetID() // MemberID
 					});
 				}
 			}
@@ -3265,22 +3269,13 @@ namespace Metasound
 			const FMetasoundFrontendGraph& Graph = DocBuilder.FindConstBuildGraphChecked();
 			for (const FMetasoundFrontendVariable& Variable : Graph.Variables)
 			{
-				FText MenuDesc = Variable.DisplayName;
-				if (MenuDesc.IsEmptyOrWhitespace())
-				{
-					FName Namespace;
-					FName ParameterName;
-					Audio::FParameterPath::SplitName(Variable.Name, Namespace, ParameterName);
-					MenuDesc = FText::FromName(ParameterName);
-				}
-
 				AddMemberAction(FAddActionParams
 				{
-					Variable.Name,				// FullName
-					Variable.Description,		// Tooltip
-					MenuDesc,					// MenuDesc
-					ENodeSection::Variables,	// Section
-					Variable.ID					 // MemberID
+					Variable.Name, // FullName
+					Variable.Description, // Tooltip
+					INodeTemplate::ResolveMemberDisplayName(Variable.Name, Variable.DisplayName, bDisplayNamespace), // MenuDesc
+					ENodeSection::Variables, // Section
+					Variable.ID // MemberID
 				});
 			}
 		}
diff --git a/Engine/Plugins/Runtime/Metasound/Source/MetasoundEditor/Private/MetasoundEditorGraphInputNode.cpp b/Engine/Plugins/Runtime/Metasound/Source/MetasoundEditor/Private/MetasoundEditorGraphInputNode.cpp
--- a/Engine/Plugins/Runtime/Metasound/Source/MetasoundEditor/Private/MetasoundEditorGraphInputNode.cpp
+++ b/Engine/Plugins/Runtime/Metasound/Source/MetasoundEditor/Private/MetasoundEditorGraphInputNode.cpp
@@ -30,17 +30,8 @@
 
 void UMetasoundEditorGraphInputNode::CacheTitle()
 {
-	using namespace Metasound::Frontend;
-
-	if (Input)
-	{
-		FConstNodeHandle NodeHandle = Input->GetNodeHandle();
-		CachedTitle = NodeHandle->GetDisplayTitle();
-	}
-	else
-	{
-		CachedTitle = FText::FromName(Breadcrumb.MemberName);
-	}
+	static FText InputDisplayTitle = LOCTEXT("InputNode_Title", "Input");
+	CachedTitle = InputDisplayTitle;
 }
 
 const FMetasoundEditorGraphVertexNodeBreadcrumb& UMetasoundEditorGraphInputNode::GetBreadcrumb() const
@@ -86,8 +77,14 @@ FMetasoundFrontendClassName UMetasoundEditorGraphInputNode::GetClassName() const
 
 	if (Input)
 	{
-		FConstNodeHandle NodeHandle = Input->GetConstNodeHandle();
-		return NodeHandle->GetClassMetadata().GetClassName();
+		const FMetaSoundFrontendDocumentBuilder& Builder = Input->GetFrontendBuilderChecked();
+		if (const FMetasoundFrontendNode* Node = Builder.FindNode(Input->NodeID))
+		{
+			if (const FMetasoundFrontendClass* Class = Builder.FindDependency(Node->ClassID))
+			{
+				return Class->Metadata.GetClassName();
+			}
+		}
 	}
 
 	return Breadcrumb.ClassName;
diff --git a/Engine/Plugins/Runtime/Metasound/Source/MetasoundEditor/Private/MetasoundEditorGraphSchema.cpp b/Engine/Plugins/Runtime/Metasound/Source/MetasoundEditor/Private/MetasoundEditorGraphSchema.cpp
--- a/Engine/Plugins/Runtime/Metasound/Source/MetasoundEditor/Private/MetasoundEditorGraphSchema.cpp
+++ b/Engine/Plugins/Runtime/Metasound/Source/MetasoundEditor/Private/MetasoundEditorGraphSchema.cpp
@@ -1716,11 +1716,22 @@ FText UMetasoundEditorGraphSchema::GetPinDisplayName(const UEdGraphPin* Pin) con
 
 	check(Pin);
 
-	UMetasoundEditorGraphNode* Node = CastChecked<UMetasoundEditorGraphNode>(Pin->GetOwningNode());
-	FConstNodeHandle NodeHandle = Node->GetConstNodeHandle();
-	const FMetasoundFrontendClassMetadata& ClassMetadata = NodeHandle->GetClassMetadata();
-	const EMetasoundFrontendClassType ClassType = ClassMetadata.GetType();
+	UMetasoundEditorGraphNode* EdNode = CastChecked<UMetasoundEditorGraphNode>(Pin->GetOwningNode());
+	const FMetaSoundFrontendDocumentBuilder& Builder = EdNode->GetBuilderChecked().GetBuilder();
 
+	const FMetasoundFrontendNode* Node = Builder.FindNode(EdNode->GetNodeID());
+	if (!Node)
+	{
+		return Super::GetPinDisplayName(Pin);
+	}
+
+	const FMetasoundFrontendClass* Class = Builder.FindDependency(Node->ClassID);
+	if (!Class)
+	{
+		return Super::GetPinDisplayName(Pin);
+	}
+
+	const EMetasoundFrontendClassType ClassType = Class->Metadata.GetType();
 	switch (ClassType)
 	{
 		case EMetasoundFrontendClassType::Input:
@@ -1730,7 +1741,7 @@ FText UMetasoundEditorGraphSchema::GetPinDisplayName(const UEdGraphPin* Pin) con
 		case EMetasoundFrontendClassType::VariableDeferredAccessor:
 		case EMetasoundFrontendClassType::VariableMutator:
 		{
-			UMetasoundEditorGraphMemberNode* MemberNode = Cast<UMetasoundEditorGraphMemberNode>(Node);
+			UMetasoundEditorGraphMemberNode* MemberNode = Cast<UMetasoundEditorGraphMemberNode>(EdNode);
 			if (ensure(MemberNode))
 			{
 				UMetasoundEditorGraphMember* Member = MemberNode->GetMember();
@@ -1739,56 +1750,55 @@ FText UMetasoundEditorGraphSchema::GetPinDisplayName(const UEdGraphPin* Pin) con
 					return MemberNode->GetMember()->GetDisplayName();
 				}
 			}
-			return Super::GetPinDisplayName(Pin);
 		}
+		break;
 
 		case EMetasoundFrontendClassType::Literal:
 		case EMetasoundFrontendClassType::External:
 		{
+			auto PinMatchesClassVertex = [&Pin](const FMetasoundFrontendClassVertex& OtherVertex) { return OtherVertex.Name == Pin->GetFName(); };
+			const FMetasoundFrontendVertex* Vertex = nullptr;
+			const FMetasoundFrontendClassVertex* ClassVertex = nullptr;
 			if (Pin->Direction == EGPD_Input)
 			{
-				FConstInputHandle InputHandle = NodeHandle->GetConstInputWithVertexName(Pin->GetFName());
-				if (InputHandle->IsValid())
-				{
-					return FGraphBuilder::GetDisplayName(*InputHandle);
-				}
+				Vertex = Builder.FindNodeInput(EdNode->GetNodeID(), Pin->GetFName());
+				ClassVertex = Class->Interface.Inputs.FindByPredicate(PinMatchesClassVertex);
 			}
 			else
 			{
-				FConstOutputHandle OutputHandle = NodeHandle->GetConstOutputWithVertexName(Pin->GetFName());
-				if (OutputHandle->IsValid())
+				Vertex = Builder.FindNodeOutput(EdNode->GetNodeID(), Pin->GetFName());
+				ClassVertex = Class->Interface.Outputs.FindByPredicate(PinMatchesClassVertex);
+			}
+
+			if (Vertex && ClassVertex)
+			{
+				FName Namespace, ParamName;
+				ClassVertex->SplitName(Namespace, ParamName);
+				FText DisplayName = ClassVertex->Metadata.GetDisplayName();
+				if (DisplayName.IsEmptyOrWhitespace())
 				{
-					return FGraphBuilder::GetDisplayName(*OutputHandle);
+					DisplayName = FText::FromName(ParamName);
 				}
-			}
 
-			return Super::GetPinDisplayName(Pin);
+				return DisplayName;
+			}
 		}
+		break;
 
 		case EMetasoundFrontendClassType::Template:
 		{
-			const INodeTemplate* Template = INodeTemplateRegistry::Get().FindTemplate(ClassMetadata.GetClassName());
+			const INodeTemplate* Template = INodeTemplateRegistry::Get().FindTemplate(Class->Metadata.GetClassName());
 			if (ensure(Template))
 			{
 				if (Pin->Direction == EGPD_Input)
 				{
-					FConstInputHandle InputHandle = NodeHandle->GetConstInputWithVertexName(Pin->GetFName());
-					if (InputHandle->IsValid())
-					{
-						return Template->GetInputPinDisplayName(*InputHandle);
-					}
+					return Template->GetInputVertexDisplayName(Builder, Node->GetID(), Pin->GetFName());
 				}
 				else
 				{
-					FConstOutputHandle OutputHandle = NodeHandle->GetConstOutputWithVertexName(Pin->GetFName());
-					if (OutputHandle->IsValid())
-					{
-						return Template->GetOutputPinDisplayName(*OutputHandle);
-					}
+					return Template->GetOutputVertexDisplayName(Builder, Node->GetID(), Pin->GetFName());
 				}
 			}
-
-			return Super::GetPinDisplayName(Pin);
 		}
 		break;
 
@@ -1797,9 +1807,11 @@ FText UMetasoundEditorGraphSchema::GetPinDisplayName(const UEdGraphPin* Pin) con
 		default:
 		{
 			static_assert(static_cast<int32>(EMetasoundFrontendClassType::Invalid) == 10, "Possible missing EMetasoundFrontendClassType case coverage");
-			return Super::GetPinDisplayName(Pin);
 		}
+		break;
 	}
+
+	return Super::GetPinDisplayName(Pin);
 }
 
 FLinearColor UMetasoundEditorGraphSchema::GetPinTypeColor(const FEdGraphPinType& PinType) const
diff --git a/Engine/Plugins/Runtime/Metasound/Source/MetasoundEditor/Private/MetasoundNodeDetailCustomization.cpp b/Engine/Plugins/Runtime/Metasound/Source/MetasoundEditor/Private/MetasoundNodeDetailCustomization.cpp
--- a/Engine/Plugins/Runtime/Metasound/Source/MetasoundEditor/Private/MetasoundNodeDetailCustomization.cpp
+++ b/Engine/Plugins/Runtime/Metasound/Source/MetasoundEditor/Private/MetasoundNodeDetailCustomization.cpp
@@ -1365,6 +1365,11 @@ PRAGMA_ENABLE_DEPRECATION_WARNINGS
 
 			if (!bIsNameInvalid && GraphMember.IsValid())
 			{
+				if (GraphMember->GetMemberName() == InNewName.ToString())
+				{
+					return;
+				}
+
 				const FText TransactionLabel = FText::Format(LOCTEXT("RenameGraphMember_Format", "Set MetaSound {0}'s Name"), GraphMember->GetGraphMemberLabel());
 				const FScopedTransaction Transaction(TransactionLabel);
 
diff --git a/Engine/Plugins/Runtime/Metasound/Source/MetasoundFrontend/Private/MetasoundFrontendNodeTemplateRegistry.cpp b/Engine/Plugins/Runtime/Metasound/Source/MetasoundFrontend/Private/MetasoundFrontendNodeTemplateRegistry.cpp
--- a/Engine/Plugins/Runtime/Metasound/Source/MetasoundFrontend/Private/MetasoundFrontendNodeTemplateRegistry.cpp
+++ b/Engine/Plugins/Runtime/Metasound/Source/MetasoundFrontend/Private/MetasoundFrontendNodeTemplateRegistry.cpp
@@ -12,6 +12,26 @@
 
 namespace Metasound::Frontend
 {
+#if WITH_EDITOR
+	FText INodeTemplate::ResolveMemberDisplayName(FName FullName, FText DisplayName, bool bIncludeNamespace)
+	{
+		FName Namespace;
+		FName ShortParamName;
+		Audio::FParameterPath::SplitName(FullName, Namespace, ShortParamName);
+		if (DisplayName.IsEmpty())
+		{
+			DisplayName = FText::FromName(ShortParamName);
+		}
+
+		if (bIncludeNamespace && !Namespace.IsNone())
+		{
+			return FText::Format(NSLOCTEXT("MetasoundFrontend", "DisplayNameWithNamespaceFormat", "{0} ({1})"), DisplayName, FText::FromName(Namespace));
+		}
+
+		return DisplayName;
+	}
+#endif // WITH_EDITOR
+
 	class FNodeTemplateRegistry : public INodeTemplateRegistry
 	{
 	public:
@@ -93,14 +113,14 @@ namespace Metasound::Frontend
 		return { };
 	}
 
-	FText FNodeTemplateBase::GetInputPinDisplayName(const Frontend::IInputController& InInput) const
+	FText FNodeTemplateBase::GetInputVertexDisplayName(const FMetaSoundFrontendDocumentBuilder& InBuilder, const FGuid& InNodeID, FName InputName) const
 	{
-		return InInput.GetDisplayName();
+		return FText::FromName(InputName);
 	}
 
-	FText FNodeTemplateBase::GetOutputPinDisplayName(const Frontend::IOutputController& InOutput) const
+	FText FNodeTemplateBase::GetOutputVertexDisplayName(const FMetaSoundFrontendDocumentBuilder& InBuilder, const FGuid& InNodeID, FName OutputName) const
 	{
-		return InOutput.GetDisplayName();
+		return FText::FromName(OutputName);
 	}
 
 	bool FNodeTemplateBase::HasRequiredConnections(const FMetaSoundFrontendDocumentBuilder& InBuilder, const FGuid& InNodeID, FString* OutMessage) const
diff --git a/Engine/Plugins/Runtime/Metasound/Source/MetasoundFrontend/Private/NodeTemplates/MetasoundFrontendNodeTemplateInput.cpp b/Engine/Plugins/Runtime/Metasound/Source/MetasoundFrontend/Private/NodeTemplates/MetasoundFrontendNodeTemplateInput.cpp
--- a/Engine/Plugins/Runtime/Metasound/Source/MetasoundFrontend/Private/NodeTemplates/MetasoundFrontendNodeTemplateInput.cpp
+++ b/Engine/Plugins/Runtime/Metasound/Source/MetasoundFrontend/Private/NodeTemplates/MetasoundFrontendNodeTemplateInput.cpp
@@ -149,34 +149,30 @@ namespace Metasound::Frontend
 	}
 
 #if WITH_EDITOR
-	FText FInputNodeTemplate::GetOutputPinDisplayName(const IOutputController& InOutput) const
+	FText FInputNodeTemplate::GetOutputVertexDisplayName(const FMetaSoundFrontendDocumentBuilder& InBuilder, const FGuid& InNodeID, FName OutputName) const
 	{
-		FConstNodeHandle OwningNode = InOutput.GetOwningNode();
-		FConstOutputHandle ConnectedOutput = OwningNode->GetConstInputs().Last()->GetConnectedOutput();
-		if (ensureMsgf(ConnectedOutput->IsValid(), TEXT("Input template node should always be connected to associated input node's only output")))
+		const FMetasoundFrontendNode* OwningNode = InBuilder.FindNode(InNodeID);
+		if (!OwningNode)
 		{
-			FConstNodeHandle ConnectedInputNode = ConnectedOutput->GetOwningNode();
-
-			FName NodeName = ConnectedInputNode->GetNodeName();
-			FName Namespace;
-			FName ParameterName;
-			Audio::FParameterPath::SplitName(NodeName, Namespace, ParameterName);
-
-			FText DisplayName = ConnectedInputNode->GetDisplayName();
-			if (DisplayName.IsEmpty())
-			{
-				DisplayName = FText::FromName(NodeName);
-			}
+			return FText::FromName(OutputName);
+		}
 
-			if (!Namespace.IsNone())
+		const FMetasoundFrontendNode* ConnectedInputNode = nullptr;
+		const FMetasoundFrontendVertex* ConnectedOutput = InBuilder.FindNodeOutputConnectedToNodeInput(InNodeID, OwningNode->Interface.Inputs.Last().VertexID, &ConnectedInputNode);
+		if (ensureMsgf(ConnectedInputNode, TEXT("Input template node should always be connected to associated input node's only output")))
+		{
+			FName NodeName = ConnectedInputNode->Name;
+			FText DisplayName;
+			if (const FMetasoundFrontendClassInput* Input = InBuilder.FindGraphInput(NodeName))
 			{
-				return FText::Format(NSLOCTEXT("MetasoundFrontend", "InputNodeTemplate_DisplayNameWithNamespaceFormat", "{0} ({1})"), DisplayName, FText::FromName(Namespace));
+				DisplayName = Input->Metadata.GetDisplayName();
 			}
 
-			return DisplayName;
+			constexpr bool bIncludeNamespace = true;
+			return INodeTemplate::ResolveMemberDisplayName(NodeName, DisplayName, bIncludeNamespace);
 		}
 
-		return FRerouteNodeTemplate::GetOutputPinDisplayName(InOutput);
+		return FRerouteNodeTemplate::GetOutputVertexDisplayName(InBuilder, InNodeID, OutputName);
 	}
 
 	bool FInputNodeTemplate::HasRequiredConnections(const FMetaSoundFrontendDocumentBuilder& InBuilder, const FGuid& InNodeID, FString* OutMessage) const
diff --git a/Engine/Plugins/Runtime/Metasound/Source/MetasoundFrontend/Public/MetasoundFrontendDocumentBuilder.h b/Engine/Plugins/Runtime/Metasound/Source/MetasoundFrontend/Public/MetasoundFrontendDocumentBuilder.h
--- a/Engine/Plugins/Runtime/Metasound/Source/MetasoundFrontend/Public/MetasoundFrontendDocumentBuilder.h
+++ b/Engine/Plugins/Runtime/Metasound/Source/MetasoundFrontend/Public/MetasoundFrontendDocumentBuilder.h
@@ -347,11 +347,26 @@ public:
 	// induce cash thrashing.
 	bool SetBuildPageID(const FGuid& InBuildPageID);
 
+	// Sets the given graph input's access type. If connected to other nodes and access type is not compatible,
+	// associated edges/connections are removed.  Returns true if either DataType was successfully set to new
+	// value or if AccessType is already the given AccessType.
 	bool SetGraphInputAccessType(FName InputName, EMetasoundFrontendVertexAccessType AccessType);
+
+	// Sets the given graph input's data type. If connected to other nodes, associated edges/connections
+	// are removed.  Returns true if either DataType was successfully set to new value or if DataType is
+	// already the given DataType.
 	bool SetGraphInputDataType(FName InputName, FName DataType);
+
 	bool SetGraphInputDefault(FName InputName, const FMetasoundFrontendLiteral& InDefaultLiteral);
 
+	// Sets the given graph output's access type. If connected to other nodes and access type is not compatible,
+	// associated edges/connections are removed.  Returns true if either DataType was successfully set to new
+	// value or if AccessType is already the given AccessType.
 	bool SetGraphOutputAccessType(FName OutputName, EMetasoundFrontendVertexAccessType AccessType);
+
+	// Sets the given graph output's data type. If connected to other nodes, associated edges/connections
+	// are removed.  Returns true if either DataType was successfully set to new value or if DataType is
+	// already the given DataType.
 	bool SetGraphOutputDataType(FName OutputName, FName DataType);
 
 #if WITH_EDITOR
diff --git a/Engine/Plugins/Runtime/Metasound/Source/MetasoundFrontend/Public/MetasoundFrontendNodeTemplateRegistry.h b/Engine/Plugins/Runtime/Metasound/Source/MetasoundFrontend/Public/MetasoundFrontendNodeTemplateRegistry.h
--- a/Engine/Plugins/Runtime/Metasound/Source/MetasoundFrontend/Public/MetasoundFrontendNodeTemplateRegistry.h
+++ b/Engine/Plugins/Runtime/Metasound/Source/MetasoundFrontend/Public/MetasoundFrontendNodeTemplateRegistry.h
@@ -65,8 +65,18 @@ namespace Metasound::Frontend
 
 #if WITH_EDITOR
 		virtual FText GetNodeDisplayName(const IMetaSoundDocumentInterface& Interface, const FGuid& InNodeID) const = 0;
-		virtual FText GetInputPinDisplayName(const Frontend::IInputController& InInput) const = 0;
-		virtual FText GetOutputPinDisplayName(const Frontend::IOutputController& InOutput) const = 0;
+
+		UE_DEPRECATED(5.5, "Use overload GetOutputVertexDisplayName with supplied builder instead")
+		virtual FText GetInputPinDisplayName(const Frontend::IInputController& InInput) const { return FText(); }
+
+		virtual FText GetInputVertexDisplayName(const FMetaSoundFrontendDocumentBuilder& InBuilder, const FGuid& InNodeID, FName InputName) const = 0;
+
+		UE_DEPRECATED(5.5, "Use overload GetOutputVertexDisplayName with supplied builder instead")
+		virtual FText GetOutputPinDisplayName(const Frontend::IOutputController& InOutput) const { return FText(); };
+
+		virtual FText GetOutputVertexDisplayName(const FMetaSoundFrontendDocumentBuilder& InBuilder, const FGuid& InNodeID, FName OutputName) const = 0;
+
+		static FText ResolveMemberDisplayName(FName VertexName, FText DisplayName, bool bIncludeNamespace);
 #endif // WITH_EDITOR
 
 		// Generates node transform that is used to preprocess nodes.
@@ -133,8 +143,8 @@ namespace Metasound::Frontend
 	public:
 #if WITH_EDITOR
 		virtual FText GetNodeDisplayName(const IMetaSoundDocumentInterface& Interface, const FGuid& InNodeID) const override;
-		virtual FText GetInputPinDisplayName(const IInputController& InInput) const override;
-		virtual FText GetOutputPinDisplayName(const IOutputController& InOutput) const override;
+		virtual FText GetInputVertexDisplayName(const FMetaSoundFrontendDocumentBuilder& InBuilder, const FGuid& InNodeID, FName InputName) const override;
+		virtual FText GetOutputVertexDisplayName(const FMetaSoundFrontendDocumentBuilder& InBuilder, const FGuid& InNodeID, FName OutputName) const override;
 		virtual bool HasRequiredConnections(const FMetaSoundFrontendDocumentBuilder& InBuilder, const FGuid& InNodeID, FString* OutMessage = nullptr) const override;
 #endif // WITH_EDITOR
 	};
diff --git a/Engine/Plugins/Runtime/Metasound/Source/MetasoundFrontend/Public/NodeTemplates/MetasoundFrontendNodeTemplateInput.h b/Engine/Plugins/Runtime/Metasound/Source/MetasoundFrontend/Public/NodeTemplates/MetasoundFrontendNodeTemplateInput.h
--- a/Engine/Plugins/Runtime/Metasound/Source/MetasoundFrontend/Public/NodeTemplates/MetasoundFrontendNodeTemplateInput.h
+++ b/Engine/Plugins/Runtime/Metasound/Source/MetasoundFrontend/Public/NodeTemplates/MetasoundFrontendNodeTemplateInput.h
@@ -42,7 +42,7 @@ namespace Metasound::Frontend
 		virtual EMetasoundFrontendVertexAccessType GetNodeOutputAccessType(const FMetaSoundFrontendDocumentBuilder& InBuilder, const FGuid& InNodeID, const FGuid& InVertexID) const override;
 
 #if WITH_EDITOR
-		virtual FText GetOutputPinDisplayName(const IOutputController& InOutput) const override;
+		virtual FText GetOutputVertexDisplayName(const FMetaSoundFrontendDocumentBuilder& InBuilder, const FGuid& InNodeID, FName OutputName) const override;
 		virtual bool HasRequiredConnections(const FMetaSoundFrontendDocumentBuilder& InBuilder, const FGuid& InNodeID, FString* OutMessage = nullptr) const override;
 
 		// Injects template nodes between builder's document inputs not connected
