Fix for crash when double clicking on a reroute node
#rnx

diff --git a/Engine/Plugins/Runtime/Metasound/Source/MetasoundEditor/Private/MetasoundEditor.cpp b/Engine/Plugins/Runtime/Metasound/Source/MetasoundEditor/Private/MetasoundEditor.cpp
--- a/Engine/Plugins/Runtime/Metasound/Source/MetasoundEditor/Private/MetasoundEditor.cpp
+++ b/Engine/Plugins/Runtime/Metasound/Source/MetasoundEditor/Private/MetasoundEditor.cpp
@@ -2055,14 +2055,24 @@ namespace Metasound
 				return;
 			}
 
-			if (UMetasoundEditorGraphExternalNode* ExternalNode = Cast<UMetasoundEditorGraphExternalNode>(InNode))
+			if (UAssetEditorSubsystem* AssetSubsystem = GEditor->GetEditorSubsystem<UAssetEditorSubsystem>())
 			{
-				FConstNodeHandle NodeHandle = ExternalNode->GetConstNodeHandle();
-				const FNodeRegistryKey Key(NodeHandle->GetClassMetadata());
-
-				if (const FMetasoundAssetBase* Asset = IMetaSoundAssetManager::GetChecked().FindAsset(Key))
+				if (UMetasoundEditorGraphExternalNode* ExternalNode = Cast<UMetasoundEditorGraphExternalNode>(InNode))
 				{
-					GEditor->GetEditorSubsystem<UAssetEditorSubsystem>()->OpenEditorForAsset(Asset->GetOwningAsset());
+					if (const FMetasoundFrontendClass* Class = ExternalNode->GetFrontendClass())
+					{
+						// Editor external nodes can represent frontend template nodes, so check
+						// to make sure underlying frontend node is of type 'External' to avoid
+						// ensure when generating asset key.
+						if (Class->Metadata.GetType() == EMetasoundFrontendClassType::External)
+						{
+							const FAssetKey AssetKey(Class->Metadata);
+							if (const FMetasoundAssetBase* Asset = IMetaSoundAssetManager::GetChecked().FindAsset(AssetKey))
+							{
+								AssetSubsystem->OpenEditorForAsset(Asset->GetOwningAsset());
+							}
+						}
+					}
 				}
 			}
 		}
diff --git a/Engine/Plugins/Runtime/Metasound/Source/MetasoundEditor/Private/MetasoundEditorGraphNode.cpp b/Engine/Plugins/Runtime/Metasound/Source/MetasoundEditor/Private/MetasoundEditorGraphNode.cpp
--- a/Engine/Plugins/Runtime/Metasound/Source/MetasoundEditor/Private/MetasoundEditorGraphNode.cpp
+++ b/Engine/Plugins/Runtime/Metasound/Source/MetasoundEditor/Private/MetasoundEditorGraphNode.cpp
@@ -206,6 +206,23 @@ UMetaSoundBuilderBase& UMetasoundEditorGraphNode::GetBuilderChecked() const
 	return FDocumentBuilderRegistry::GetChecked().FindOrBeginBuilding(EdGraph->GetMetasoundChecked());
 }
 
+const FMetasoundFrontendClass* UMetasoundEditorGraphNode::GetFrontendClass() const
+{
+	using namespace Metasound::Engine;
+
+	if (UObject* Outermost = GetOutermostObject())
+	{
+		const FGuid NodeID = GetNodeID();
+		const FMetaSoundFrontendDocumentBuilder& Builder = FDocumentBuilderRegistry::GetChecked().FindOrBeginBuilding(Outermost);
+		if (const FMetasoundFrontendNode* Node = Builder.FindNode(NodeID))
+		{
+			return Builder.FindDependency(Node->ClassID);
+		}
+	}
+
+	return nullptr;
+}
+
 const FMetasoundFrontendNode* UMetasoundEditorGraphNode::GetFrontendNode() const
 {
 	using namespace Metasound::Engine;
diff --git a/Engine/Plugins/Runtime/Metasound/Source/MetasoundEditor/Public/MetasoundEditorGraphNode.h b/Engine/Plugins/Runtime/Metasound/Source/MetasoundEditor/Public/MetasoundEditorGraphNode.h
--- a/Engine/Plugins/Runtime/Metasound/Source/MetasoundEditor/Public/MetasoundEditorGraphNode.h
+++ b/Engine/Plugins/Runtime/Metasound/Source/MetasoundEditor/Public/MetasoundEditorGraphNode.h
@@ -120,6 +120,14 @@ public:
 	virtual bool CanAddInputPin() const { return false; }
 
 	UMetaSoundBuilderBase& GetBuilderChecked() const;
+
+	// Returns document's cached Frontend Node Class (as defined by the
+	// document's dependency array). If node or class is not found on
+	// document, returns null.
+	const FMetasoundFrontendClass* GetFrontendClass() const;
+
+	// Returns document's cached Frontend Node.
+	// If node is not found on document, returns null.
 	const FMetasoundFrontendNode* GetFrontendNode() const;
 	const FMetasoundFrontendNode& GetFrontendNodeChecked() const;
 
