When an audio bus writer node is assigned a bus with more channels than it expects, ensure the extra channels are filled with silence.
#tests audio bus engine tests pass

diff --git a/Engine/Plugins/Runtime/Metasound/Source/MetasoundEngine/Private/MetasoundAudioBusWriterNode.cpp b/Engine/Plugins/Runtime/Metasound/Source/MetasoundEngine/Private/MetasoundAudioBusWriterNode.cpp
--- a/Engine/Plugins/Runtime/Metasound/Source/MetasoundEngine/Private/MetasoundAudioBusWriterNode.cpp
+++ b/Engine/Plugins/Runtime/Metasound/Source/MetasoundEngine/Private/MetasoundAudioBusWriterNode.cpp
@@ -155,7 +155,7 @@ namespace Metasound
 					return;
 				}
 
-				AudioBusChannels = uint32(FMath::Min(AudioBusProxy->NumChannels, int32(EAudioBusChannels::MaxChannelCount)));
+				AudioBusChannels = FMath::Min(uint32(AudioBusProxy->NumChannels), uint32(EAudioBusChannels::MaxChannelCount));
 				AudioBusId = AudioBusProxy->AudioBusId;
 
 				Audio::FAudioBusKey AudioBusKey(AudioBusId);
@@ -168,8 +168,10 @@ namespace Metasound
 					AudioBusPatchInput.PushAudio(nullptr, NumBlocksToNumSamples(NumBlocksToPush));
 				}
 
+				// Allocate and fill the interleaved buffer with silence,
+				// in case it contains more channels than the node supports.
 				InterleavedBuffer.Reset();
-				InterleavedBuffer.AddUninitialized(NumBlocksToNumSamples(1));
+				InterleavedBuffer.AddZeroed(NumBlocksToNumSamples(1));
 			}
 		}
 		
@@ -283,10 +285,14 @@ namespace Metasound
 				uint32 MinChannels = FMath::Min(AudioBusChannels, NumChannels);
 				for (int32 FrameIndex = 0; FrameIndex < BlockSizeFrames; ++FrameIndex)
 				{
+					// Fill as many channels in the interleaved buffer as possible,
+					// given the number of available audio buffers.
 					for (uint32 ChannelIndex = 0; ChannelIndex < MinChannels; ++ChannelIndex)
 					{
 						InterleavedBufferPtr[ChannelIndex] = *AudioInputBufferPtrs[ChannelIndex]++;
 					}
+
+					// The interleaved buffer has as many channels as the assigned audio bus.
 					InterleavedBufferPtr += AudioBusChannels;
 				}
 			}
