[MetaSounds] - Delete Unused Members no longer working
#rnx

diff --git a/Engine/Plugins/Runtime/Metasound/Source/MetasoundEditor/Private/MetasoundEditor.cpp b/Engine/Plugins/Runtime/Metasound/Source/MetasoundEditor/Private/MetasoundEditor.cpp
--- a/Engine/Plugins/Runtime/Metasound/Source/MetasoundEditor/Private/MetasoundEditor.cpp
+++ b/Engine/Plugins/Runtime/Metasound/Source/MetasoundEditor/Private/MetasoundEditor.cpp
@@ -3461,25 +3461,23 @@ namespace Metasound
 				return false;
 			}
 
+			//Check if there is any Actions to remove in the section
+			TArray<TSharedPtr<FEdGraphSchemaAction>> Actions;
+			GraphMembersMenu->GetSelectedCategorySubActions(Actions);
+			if (Actions.IsEmpty())
+			{
+				return false;
+			}
+
+			//Check if selected is not a Member
 			TArray<TSharedPtr<FEdGraphSchemaAction>> SelectedActions;
 			GraphMembersMenu->GetSelectedActions(SelectedActions);
-
 			if (SelectedActions.IsEmpty())
 			{
 				return true;
 			}
-
-			//Option to delete unused member becomes pressable only if right-click happened on non valid Action.
-			for (const TSharedPtr<FEdGraphSchemaAction>& SelectedAction : SelectedActions)
-			{
-				TSharedPtr<FMetasoundGraphMemberSchemaAction> MetasoundAction = StaticCastSharedPtr<FMetasoundGraphMemberSchemaAction>(SelectedAction);
-				if (MetasoundAction.IsValid())
-				{					
-					return false;				
-				}
-			}
 			
-			return true;
+			return false;
 		}
 
 		FActionMenuContent FEditor::OnCreateGraphActionMenu(UEdGraph* InGraph, const FVector2D& InNodePosition, const TArray<UEdGraphPin*>& InDraggedPins, bool bAutoExpand, SGraphEditor::FActionMenuClosed InOnMenuClosed)
@@ -3541,35 +3539,37 @@ namespace Metasound
 				return nullptr;
 			}
 
-			// Context menu should only open when graph members are selected
+			FMenuBuilder MenuBuilder(true, ToolkitCommands);
 			TArray<TSharedPtr<FEdGraphSchemaAction>> Actions;
 			GraphMembersMenu->GetSelectedActions(Actions);
-			if (Actions.IsEmpty())
-			{
-				return nullptr;
-			}
-
-			FMenuBuilder MenuBuilder(true, ToolkitCommands);
 
-			MenuBuilder.AddMenuEntry(FGenericCommands::Get().Delete);
-			MenuBuilder.AddMenuEntry(FGenericCommands::Get().Rename);
-			MenuBuilder.AddMenuEntry(FGenericCommands::Get().Duplicate);
-
-			MenuBuilder.AddMenuEntry(
-				LOCTEXT("JumpToNodesMenuEntry", "Jump to Node(s) in Graph"),
-				LOCTEXT("JumpToNodesMenuEntryTooltip", "Jump to the corresponding node(s) in the MetaSound graph"),
-				FSlateIcon(),
-				FUIAction(
-					FExecuteAction::CreateSP(this, &FEditor::JumpToNodesForSelectedInterfaceItem), 
-					FCanExecuteAction::CreateSP(this, &FEditor::CanJumpToNodesForSelectedInterfaceItem)));
-					
-			MenuBuilder.AddMenuEntry(
-				LOCTEXT("DeleteAllUnusedInSection", "Delete Unused Members"),
-				LOCTEXT("DeleteAllUnusedInSectionTooltip", "Delete all Unused Members under this Header"),
-				FSlateIcon(),
-				FUIAction(
-					FExecuteAction::CreateSP(this, &FEditor::DeleteAllUnusedInSection),
-					FCanExecuteAction::CreateSP(this, &FEditor::CanDeleteUnusedMembers)));
+			if (Actions.IsEmpty())//Section is selected
+			{
+				MenuBuilder.BeginSection("GraphActionMenuSectionActions", LOCTEXT("SectionActionsMenuHeader", "Section Actions"));
+				MenuBuilder.AddMenuEntry(
+					LOCTEXT("DeleteAllUnusedInSection", "Delete Unused Members"),
+					LOCTEXT("DeleteAllUnusedInSectionTooltip", "Delete all Unused Members under this Section"),
+					FSlateIcon(),
+					FUIAction(
+						FExecuteAction::CreateSP(this, &FEditor::DeleteAllUnusedInSection),
+						FCanExecuteAction::CreateSP(this, &FEditor::CanDeleteUnusedMembers)));
+				MenuBuilder.EndSection();
+			}
+			else //Member is selected
+			{
+				MenuBuilder.BeginSection("GraphActionMenuMemberActions", LOCTEXT("MemberActionsMenuHeader", "Member Actions"));
+				MenuBuilder.AddMenuEntry(FGenericCommands::Get().Delete);
+				MenuBuilder.AddMenuEntry(FGenericCommands::Get().Rename);
+				MenuBuilder.AddMenuEntry(FGenericCommands::Get().Duplicate);
+				MenuBuilder.AddMenuEntry(
+					LOCTEXT("JumpToNodesMenuEntry", "Jump to Node(s) in Graph"),
+					LOCTEXT("JumpToNodesMenuEntryTooltip", "Jump to the corresponding node(s) in the MetaSound graph"),
+					FSlateIcon(),
+					FUIAction(
+						FExecuteAction::CreateSP(this, &FEditor::JumpToNodesForSelectedInterfaceItem), 
+						FCanExecuteAction::CreateSP(this, &FEditor::CanJumpToNodesForSelectedInterfaceItem)));
+				MenuBuilder.EndSection();
+			}
 
 			return MenuBuilder.MakeWidget();
 		}
