Versioning guid collision fixes
#rnx
#tests local edits, audioQA automated tests, project fast cook

diff --git a/Engine/Plugins/Runtime/Metasound/Source/MetasoundEditor/Private/MetasoundEditorModule.cpp b/Engine/Plugins/Runtime/Metasound/Source/MetasoundEditor/Private/MetasoundEditorModule.cpp
--- a/Engine/Plugins/Runtime/Metasound/Source/MetasoundEditor/Private/MetasoundEditorModule.cpp
+++ b/Engine/Plugins/Runtime/Metasound/Source/MetasoundEditor/Private/MetasoundEditorModule.cpp
@@ -27,6 +27,7 @@
 #include "MetasoundEditorSettings.h"
 #include "MetasoundFrontendDataTypeRegistry.h"
 #include "MetasoundFrontendDocument.h"
+#include "MetasoundFrontendGraph.h"
 #include "MetasoundFrontendRegistries.h"
 #include "MetasoundFrontendTransform.h"
 #include "MetasoundNodeDetailCustomization.h"
@@ -240,9 +241,11 @@ namespace Metasound
 		{
 			void LoadAndRegisterAsset(const FAssetData& InAssetData)
 			{
+				using namespace Metasound;
+			
 				// Ignore requests if running cook, as presave will register and cook using the
 				// proper node register call which avoids async registration/igraph generation.
-				if (IsRunningCookCommandlet())
+				if (!FFrontendGraphBuilder::CanEverExecute())
 				{
 					return;
 				}
diff --git a/Engine/Plugins/Runtime/Metasound/Source/MetasoundEngine/Private/MetasoundDocumentBuilderRegistry.cpp b/Engine/Plugins/Runtime/Metasound/Source/MetasoundEngine/Private/MetasoundDocumentBuilderRegistry.cpp
--- a/Engine/Plugins/Runtime/Metasound/Source/MetasoundEngine/Private/MetasoundDocumentBuilderRegistry.cpp
+++ b/Engine/Plugins/Runtime/Metasound/Source/MetasoundEngine/Private/MetasoundDocumentBuilderRegistry.cpp
@@ -2,6 +2,8 @@
 
 #include "MetasoundDocumentBuilderRegistry.h"
 
+#include "MetasoundAssetManager.h"
+#include "MetasoundFrontendGraph.h"
 #include "MetasoundTrace.h"
 #include "MetasoundUObjectRegistry.h"
 
@@ -272,6 +274,9 @@ namespace Metasound::Engine
 
 	void FDocumentBuilderRegistry::FinishBuildingInternal(UMetaSoundBuilderBase& Builder, bool bForceUnregisterNodeClass) const
 	{
+		using namespace Metasound;
+		using namespace Metasound::Frontend;
+
 		// If the builder has applied transactions to its document object that are not mirrored in the frontend registry,
 		// unregister version in registry. This will ensure that future requests for the builder's associated asset will
 		// register a fresh version from the object as the transaction history is intrinsically lost once this builder
@@ -280,7 +285,7 @@ namespace Metasound::Engine
 		FMetaSoundFrontendDocumentBuilder& DocBuilder = Builder.GetBuilder();
 		if (DocBuilder.IsValid())
 		{
-			if (!IsRunningCookCommandlet())
+			if (FFrontendGraphBuilder::CanEverExecute())
 			{
 				const int32 TransactionCount = DocBuilder.GetTransactionCount();
 				const int32 LastTransactionRegistered = Builder.GetLastTransactionRegistered();
diff --git a/Engine/Plugins/Runtime/Metasound/Source/MetasoundEngine/Private/MetasoundEngineAsset.h b/Engine/Plugins/Runtime/Metasound/Source/MetasoundEngine/Private/MetasoundEngineAsset.h
--- a/Engine/Plugins/Runtime/Metasound/Source/MetasoundEngine/Private/MetasoundEngineAsset.h
+++ b/Engine/Plugins/Runtime/Metasound/Source/MetasoundEngine/Private/MetasoundEngineAsset.h
@@ -6,6 +6,7 @@
 #include "MetasoundAssetManager.h"
 #include "MetasoundDocumentBuilderRegistry.h"
 #include "MetasoundFrontendDocumentIdGenerator.h"
+#include "MetasoundFrontendGraph.h"
 #include "MetasoundFrontendRegistryKey.h"
 #include "MetasoundUObjectRegistry.h"
 #include "Misc/App.h"
@@ -14,9 +15,9 @@
 #if WITH_EDITORONLY_DATA
 #include "Algo/Transform.h"
 #include "MetasoundFrontendRegistryContainer.h"
+#include "UObject/GarbageCollection.h"
 #include "UObject/StrongObjectPtrTemplates.h"
 #include "UObject/ObjectMacros.h"
-#include "UObject/GarbageCollection.h"
 #endif // WITH_EDITORONLY_DATA
 
 
@@ -28,6 +29,11 @@ namespace Metasound::Engine
 	 */
 	struct FAssetHelper
 	{
+		static bool IsDeterministic(bool bIsContextCooking)
+		{
+			return bIsContextCooking || IsRunningCookCommandlet();
+		}
+
 #if WITH_EDITOR
 		static void PreDuplicate(TScriptInterface<IMetaSoundDocumentInterface> MetaSound, FObjectDuplicationParameters& DupParams)
 		{
@@ -138,11 +144,13 @@ namespace Metasound::Engine
 				AssetManager->WaitUntilAsyncLoadReferencedAssetsComplete(InMetaSound);
 			}
 
-			if (InSaveContext.IsCooking() || IsRunningCommandlet())
+			const bool bIsCooking = InSaveContext.IsCooking();
+			const bool bCanEverExecute = FFrontendGraphBuilder::CanEverExecute(bIsCooking);
+			if (!bCanEverExecute)
 			{
-				constexpr bool bIsDeterministic = true;
-				FDocumentIDGenerator::FScopeDeterminism DeterminismScope = FDocumentIDGenerator::FScopeDeterminism(bIsDeterministic);
-				InMetaSound.CookMetaSound();
+				const bool bIsDeterministic = IsDeterministic(bIsCooking);
+				FDocumentIDGenerator::FScopeDeterminism DeterminismScope(bIsDeterministic);
+				InMetaSound.PreSaveDocument();
 			}
  			else if (FApp::CanEverRenderAudio())
 			{
@@ -165,6 +173,8 @@ namespace Metasound::Engine
 		static void SerializeToArchive(TMetaSoundObject& InMetaSound, FArchive& InArchive)
 		{
 #if WITH_EDITORONLY_DATA
+			using namespace Frontend;
+
 			bool bVersionedAsset = false;
 
 			if (InArchive.IsLoading())
@@ -175,8 +185,14 @@ namespace Metasound::Engine
 					Builder.Reset(&FDocumentBuilderRegistry::GetChecked().FindOrBeginBuilding(InMetaSound));
 				}
 
-				check(Builder.IsValid());
-				bVersionedAsset = InMetaSound.VersionAsset(Builder->GetBuilder());
+				{
+					const bool bIsCooking = InArchive.IsCooking();
+					const bool bIsDeterministic = IsDeterministic(bIsCooking);
+					FDocumentIDGenerator::FScopeDeterminism DeterminismScope(bIsDeterministic);
+					check(Builder.IsValid());
+					bVersionedAsset = InMetaSound.VersionAsset(Builder->GetBuilder());
+				}
+
 				Builder->ClearInternalFlags(EInternalObjectFlags::Async);
 			}
 
diff --git a/Engine/Plugins/Runtime/Metasound/Source/MetasoundEngine/Private/MetasoundSource.cpp b/Engine/Plugins/Runtime/Metasound/Source/MetasoundEngine/Private/MetasoundSource.cpp
--- a/Engine/Plugins/Runtime/Metasound/Source/MetasoundEngine/Private/MetasoundSource.cpp
+++ b/Engine/Plugins/Runtime/Metasound/Source/MetasoundEngine/Private/MetasoundSource.cpp
@@ -26,6 +26,7 @@
 #include "MetasoundFrontendDataTypeRegistry.h"
 #include "MetasoundFrontendDocumentBuilder.h"
 #include "MetasoundFrontendDocumentIdGenerator.h"
+#include "MetasoundFrontendGraph.h"
 #include "MetasoundFrontendQuery.h"
 #include "MetasoundFrontendQuerySteps.h"
 #include "MetasoundFrontendTransform.h"
@@ -804,12 +805,14 @@ void UMetaSoundSource::InitParameters(TArray<FAudioParameter>& ParametersToInit,
 
 void UMetaSoundSource::InitResources()
 {
+	using namespace Metasound;
+	using namespace Metasound::Engine;
 	using namespace Metasound::Frontend;
 	using namespace Metasound::SourcePrivate;
 
 	METASOUND_LLM_SCOPE;
 	METASOUND_TRACE_CPUPROFILER_EVENT_SCOPE(UMetaSoundSource::InitResources); 
-	ensureMsgf(!IsRunningCookCommandlet(), TEXT("UMetaSoundSource::InitResources should not be called during cook."));
+	ensureMsgf(!FFrontendGraphBuilder::CanEverExecute(), TEXT("UMetaSoundSource::InitResources should not be called in builds that cannot execute MetaSounds."));
 
 	if (IsInGameThread())
 	{
@@ -838,12 +841,15 @@ void UMetaSoundSource::InitResources()
 
 void UMetaSoundSource::RegisterGraphWithFrontend(Metasound::Frontend::FMetaSoundAssetRegistrationOptions InRegistrationOptions)
 {
+	using namespace Metasound;
+	using namespace Metasound::Frontend;
+
 	check(IsInGameThread());
 
 	FMetasoundAssetBase::RegisterGraphWithFrontend(InRegistrationOptions);
 	const bool bIsRuntimeInputDataValid = RuntimeInputData.bIsValid.load();
-	// Runtime data does not need to and should not be created at cook
-	if (!bIsRuntimeInputDataValid && !IsRunningCookCommandlet())
+	// Runtime data does not need to and should not be created in builds that don't execute graphs
+	if (!bIsRuntimeInputDataValid && FFrontendGraphBuilder::CanEverExecute())
 	{
 		CacheRuntimeInputData();
 	}
diff --git a/Engine/Plugins/Runtime/Metasound/Source/MetasoundFrontend/Private/MetasoundAssetBase.cpp b/Engine/Plugins/Runtime/Metasound/Source/MetasoundFrontend/Private/MetasoundAssetBase.cpp
--- a/Engine/Plugins/Runtime/Metasound/Source/MetasoundFrontend/Private/MetasoundAssetBase.cpp
+++ b/Engine/Plugins/Runtime/Metasound/Source/MetasoundFrontend/Private/MetasoundAssetBase.cpp
@@ -268,7 +268,7 @@ void FMetasoundAssetBase::RegisterGraphWithFrontend(Metasound::Frontend::FMetaSo
 
 	// Graph registration must only happen on one thread to avoid race conditions on graph registration.
 	checkf(IsInGameThread(), TEXT("MetaSound %s graph can only be registered on the GameThread"), *GetOwningAssetName());
-	checkf(!IsRunningCookCommandlet(), TEXT("Cook of asset must call RegisterNode directly providing FDocumentNodeRegistryEntryForCook to avoid proxy/runtime graph generation."));
+	checkf(FFrontendGraphBuilder::CanEverExecute(), TEXT("Cannot generate proxies/runtime graph when graph execution is not enabled."));
 
 	METASOUND_TRACE_CPUPROFILER_EVENT_SCOPE(MetaSoundAssetBase::RegisterGraphWithFrontend);
 	METASOUND_TRACE_CPUPROFILER_EVENT_SCOPE_TEXT(*FString::Printf(TEXT("MetaSoundAssetBase::RegisterGraphWithFrontend asset %s"), *this->GetOwningAssetName()));
@@ -348,17 +348,27 @@ void FMetasoundAssetBase::RegisterGraphWithFrontend(Metasound::Frontend::FMetaSo
 }
 
 void FMetasoundAssetBase::CookMetaSound()
+{
+#if WITH_EDITORONLY_DATA
+	PreSaveDocument();
+#endif // WITH_EDITORONLY_DATA
+}
+
+#if WITH_EDITORONLY_DATA
+void FMetasoundAssetBase::PreSaveDocument()
 {
 	using namespace Metasound;
 	using namespace Metasound::Frontend;
 
-	METASOUND_TRACE_CPUPROFILER_EVENT_SCOPE(MetaSoundAssetBase::CookMetaSound);
+	METASOUND_TRACE_CPUPROFILER_EVENT_SCOPE(MetaSoundAssetBase::PreSaveDocument);
+
+	// If already registered, nothing to condition for presaving
 	if (IsRegistered())
 	{
 		return;
 	}
 
-	CookReferencedMetaSounds();
+	PreSaveReferencedDocuments();
 	IMetaSoundAssetManager::GetChecked().AddOrUpdateAsset(*GetOwningAsset());
 
 	// Auto update must be done after all referenced asset classes are registered
@@ -380,7 +390,6 @@ void FMetasoundAssetBase::CookMetaSound()
 	check(Owner);
 
 	{
-#if WITH_EDITORONLY_DATA
 		// Performs document transforms on local copy, which reduces document footprint & renders transforming unnecessary at runtime
 		FMetaSoundFrontendDocumentBuilder& DocBuilder = IDocumentBuilderRegistry::GetChecked().FindOrBeginBuilding(Owner);
 		const bool bContainsTemplateDependency = DocBuilder.ContainsDependencyOfType(EMetasoundFrontendClassType::Template);
@@ -388,7 +397,6 @@ void FMetasoundAssetBase::CookMetaSound()
 		{
 			DocBuilder.TransformTemplateNodes();
 		}
-#endif // WITH_EDITORONLY_DATA
 
 		if (GraphRegistryKey.IsValid())
 		{
@@ -396,9 +404,9 @@ void FMetasoundAssetBase::CookMetaSound()
 			GraphRegistryKey = { };
 		}
 
-		// During cook, we need to register the node so that it is available for other graphs, but we need to avoid
-		// creating proxies. To do so, we use a special node registration object which reflects the necessary information
-		// for the node registry, but does not create INodes.
+		// Need to register the node so that it is available for other graphs, but avoids creating proxies.
+		// This is accomplished by using a special node registration object which reflects the necessary
+		// information for the node registry, but does not create the runtime graph model (i.e. INodes).
 		TScriptInterface<IMetaSoundDocumentInterface> DocInterface(Owner);
 		const FMetasoundFrontendDocument& Document = DocInterface->GetConstDocument();
 		const FTopLevelAssetPath AssetPath = DocInterface->GetAssetPathChecked();
@@ -410,29 +418,33 @@ void FMetasoundAssetBase::CookMetaSound()
 
 	if (GraphRegistryKey.IsValid())
 	{
-#if WITH_EDITORONLY_DATA
 		UpdateAssetRegistry();
-#endif // WITH_EDITORONLY_DATA
 	}
 	else
 	{
 		const UClass* Class = Owner->GetClass();
 		check(Class);
 		const FString ClassName = Class->GetName();
-		UE_LOG(LogMetaSound, Error, TEXT("Registration failed during cook for MetaSound node class '%s' of UObject class '%s'"), *GetOwningAssetName(), *ClassName);
+		UE_LOG(LogMetaSound, Error, TEXT("Presave failed for MetaSound node class '%s' of UObject class '%s'"), *GetOwningAssetName(), *ClassName);
 	}
 }
+#endif // WITH_EDITORONLY_DATA
 
 void FMetasoundAssetBase::OnNotifyBeginDestroy()
 {
+	using namespace Metasound;
 	using namespace Metasound::Frontend;
 
 	UObject* OwningAsset = GetOwningAsset();
 	check(OwningAsset);
 
-	// Unregistration of graph is not necessary when cooking as deserialized objects are not mutable and, should they be reloaded,
-	// omitting unregistration avoids potentially kicking off an invalid asynchronous task to unregister a non-existent runtime graph.
-	if (IsRunningCookCommandlet())
+	// Unregistration of graph using local call is not necessary when cooking as deserialized objects are not mutable and, should they be
+	// reloaded, omitting unregistration avoids potentially kicking off an invalid asynchronous task to unregister a non-existent runtime graph.
+	if (FFrontendGraphBuilder::CanEverExecute())
+	{
+		UnregisterGraphWithFrontend();
+	}
+	else
 	{
 		if (GraphRegistryKey.IsValid())
 		{
@@ -440,10 +452,6 @@ void FMetasoundAssetBase::OnNotifyBeginDestroy()
 			GraphRegistryKey = { };
 		}
 	}
-	else
-	{
-		UnregisterGraphWithFrontend();
-	}
 
 	if (IMetaSoundAssetManager* AssetManager = IMetaSoundAssetManager::Get())
 	{
@@ -453,11 +461,12 @@ void FMetasoundAssetBase::OnNotifyBeginDestroy()
 
 void FMetasoundAssetBase::UnregisterGraphWithFrontend()
 {
+	using namespace Metasound;
 	using namespace Metasound::Frontend;
 	METASOUND_TRACE_CPUPROFILER_EVENT_SCOPE(MetaSoundAssetBase::UnregisterGraphWithFrontend);
 
 	check(IsInGameThread());
-	checkf(!IsRunningCookCommandlet(), TEXT("Cook of asset must call UnregisterNode directly providing FDocumentNodeRegistryEntryForCook to avoid proxy/runtime graph generation."));
+	checkf(FFrontendGraphBuilder::CanEverExecute(), TEXT("If execution is not supported, UnregisterNode must be called directly to avoid async attempt at destroying runtime graph that does not exist."));
 
 	if (GraphRegistryKey.IsValid())
 	{
@@ -507,8 +516,6 @@ bool FMetasoundAssetBase::VersionAsset(FMetaSoundFrontendDocumentBuilder& Builde
 	using namespace Metasound::Frontend;
 
 	METASOUND_TRACE_CPUPROFILER_EVENT_SCOPE(MetaSoundAssetBase::VersionAsset);
-	constexpr bool bIsDeterministic = true;
-	FDocumentIDGenerator::FScopeDeterminism DeterminismScope(bIsDeterministic);
 
 	bool bDidEdit = false;
 
@@ -907,6 +914,14 @@ void FMetasoundAssetBase::RegisterAssetDependencies(const Metasound::Frontend::F
 }
 
 void FMetasoundAssetBase::CookReferencedMetaSounds()
+{
+#if WITH_EDITORONLY_DATA
+	PreSaveReferencedDocuments();
+#endif // WITH_EDITORONLY_DATA
+}
+
+#if WITH_EDITORONLY_DATA
+void FMetasoundAssetBase::PreSaveReferencedDocuments()
 {
 	using namespace Metasound::Frontend;
 
@@ -918,10 +933,11 @@ void FMetasoundAssetBase::CookReferencedMetaSounds()
 		{
 			// TODO: Check for infinite recursion and error if so
 			AssetManager.AddOrUpdateAsset(*(Reference->GetOwningAsset()));
-			Reference->CookMetaSound();
+			Reference->PreSaveDocument();
 		}
 	}
 }
+#endif // WITH_EDITORONLY_DATA
 
 bool FMetasoundAssetBase::AutoUpdate(bool bInLogWarningsOnDroppedConnection)
 {
diff --git a/Engine/Plugins/Runtime/Metasound/Source/MetasoundFrontend/Private/MetasoundFrontendDocumentIdGenerator.cpp b/Engine/Plugins/Runtime/Metasound/Source/MetasoundFrontend/Private/MetasoundFrontendDocumentIdGenerator.cpp
--- a/Engine/Plugins/Runtime/Metasound/Source/MetasoundFrontend/Private/MetasoundFrontendDocumentIdGenerator.cpp
+++ b/Engine/Plugins/Runtime/Metasound/Source/MetasoundFrontend/Private/MetasoundFrontendDocumentIdGenerator.cpp
@@ -1,12 +1,15 @@
 // Copyright Epic Games, Inc. All Rights Reserved.
-
 #include "MetasoundFrontendDocumentIdGenerator.h"
 
+#include "MetasoundAssetManager.h"
+#include "MetasoundFrontendDocument.h"
+#include "MetasoundFrontendDocumentVersioning.h"
+
 namespace Metasound
 {
 	namespace Frontend
 	{
-		namespace
+		namespace DocumentIDGeneratorPrivate
 		{
 			std::atomic<uint64> GlobalAtomicMetasoundIdCounter = 1; // First ID will be 1. This is because we will make this into an FGuid which cannot be zero.
 		}
@@ -63,9 +66,19 @@ namespace Metasound
 		{
 			if (bIsDeterministic)
 			{
+				// A bug caused collisions between serialized content and newly generated values.
+				// The use of this base guid ensures no such collision generation continues.
+				constexpr FGuid BaseGuid(0x8BC4C7C3, 0x591449C4, 0xA35830F8, 0xE7F9052E);
+
 				const uint32 Value = InDocument.GetNextIdCounter();
 				const FGuid CounterGuid = FGuid(Value << 6, Value << 4, Value << 2, Value);
-				return FGuid::Combine(InDocument.RootGraph.ID, CounterGuid);
+				const FGuid UpdatedGuid = FGuid::Combine(CounterGuid, BaseGuid);
+
+				FGuid ClassID;
+				const FMetasoundFrontendClassName& ClassName = InDocument.RootGraph.Metadata.GetClassName();
+				ensureAlwaysMsgf(IMetaSoundAssetManager::GetChecked().TryGetAssetIDFromClassName(ClassName, ClassID), TEXT("Failed to retrieve AssetID from MetaSoundClassName"));
+
+				return FGuid::Combine(ClassID, UpdatedGuid);
 			}
 			else
 			{
@@ -81,7 +94,7 @@ namespace Metasound
 
 		FGuid FClassIDGenerator::CreateInputID(const FMetasoundFrontendClassInput& Input) const
 		{
-			const FGuid ClassInputNamespaceGuid = FGuid("{149FEB6E-B9F9-47A6-AD4F-B78655F6EBE8}");
+			constexpr FGuid ClassInputNamespaceGuid = FGuid(0x149FEB6E, 0xB9F947A6, 0xAD4FB786, 0x55F6EBE8);
 			FString NameToHash = FString::Printf(TEXT("ClassInput.%s.%s.%s"), *Input.Name.ToString(), *Input.TypeName.ToString(), LexToString(Input.AccessType));
 
 			return CreateNamespacedIDFromString(ClassInputNamespaceGuid, NameToHash);
@@ -89,7 +102,7 @@ namespace Metasound
 
 		FGuid FClassIDGenerator::CreateInputID(const Audio::FParameterInterface::FInput& Input) const
 		{
-			const FGuid ParameterInterfaceInputNamespaceGuid = FGuid("{D9E893C0-92B3-4CB4-8306-4525ABEACADD}");
+			constexpr FGuid ParameterInterfaceInputNamespaceGuid(0xD9E893C0, 0x92B34CB4, 0x83064525, 0xABEACADD);
 			FString NameToHash = FString::Printf(TEXT("ParameterInterfaceInput.%s.%s"), *Input.InitValue.ParamName.ToString(), *Input.DataType.ToString());
 
 			return CreateNamespacedIDFromString(ParameterInterfaceInputNamespaceGuid, NameToHash);
@@ -97,7 +110,7 @@ namespace Metasound
 
 		FGuid FClassIDGenerator::CreateOutputID(const FMetasoundFrontendClassOutput& Output) const
 		{
-			const FGuid ClassOutputNamespaceGuid = FGuid("{C7B3ED2C-4407-4B2A-9144-7F1108387EBB}");
+			constexpr FGuid ClassOutputNamespaceGuid(0xC7B3ED2C, 0x44074B2A, 0x91447F11, 0x08387EBB);
 			FString NameToHash = FString::Printf(TEXT("ClassOutput.%s.%s.%s"), *Output.Name.ToString(), *Output.TypeName.ToString(), LexToString(Output.AccessType));
 
 			return CreateNamespacedIDFromString(ClassOutputNamespaceGuid, NameToHash);
@@ -105,13 +118,13 @@ namespace Metasound
 
 		FGuid FClassIDGenerator::CreateOutputID(const Audio::FParameterInterface::FOutput& Output) const
 		{
-			const FGuid ParameterInterfaceOutputNamespaceGuid = FGuid("{6F41342A-2436-4462-81A0-8517887BB729}");
+			constexpr FGuid ParameterInterfaceOutputNamespaceGuid(0x6F41342A, 0x24364462, 0x81A08517, 0x887BB729);
 			FString NameToHash = FString::Printf(TEXT("ParameterInterfaceOutput.%s.%s"), *Output.ParamName.ToString(), *Output.DataType.ToString());
 
 			return CreateNamespacedIDFromString(ParameterInterfaceOutputNamespaceGuid, NameToHash);
 		}
 
-		FGuid FClassIDGenerator::CreateNamespacedIDFromString(const FGuid NamespaceGuid, const FString& StringToHash) const
+		FGuid FClassIDGenerator::CreateNamespacedIDFromString(const FGuid& NamespaceGuid, const FString& StringToHash) const
 		{
 			FSHA1 Hasher;
 			Hasher.Update(reinterpret_cast<const uint8*>(&NamespaceGuid), sizeof(FGuid));
@@ -124,6 +137,7 @@ namespace Metasound
 
 		FGuid CreateLocallyUniqueId()
 		{
+			using namespace DocumentIDGeneratorPrivate;
 			uint64 NextId = GlobalAtomicMetasoundIdCounter.fetch_add(1, std::memory_order_relaxed);
 			return FGuid(0, 0, NextId >> 32, NextId & 0xFFFFFFFF);
 		}
diff --git a/Engine/Plugins/Runtime/Metasound/Source/MetasoundFrontend/Private/MetasoundFrontendGraph.cpp b/Engine/Plugins/Runtime/Metasound/Source/MetasoundFrontend/Private/MetasoundFrontendGraph.cpp
--- a/Engine/Plugins/Runtime/Metasound/Source/MetasoundFrontend/Private/MetasoundFrontendGraph.cpp
+++ b/Engine/Plugins/Runtime/Metasound/Source/MetasoundFrontend/Private/MetasoundFrontendGraph.cpp
@@ -7,6 +7,7 @@
 #include "Algo/TopologicalSort.h"
 #include "Algo/Transform.h"
 #include "CoreMinimal.h"
+#include "CoreGlobals.h"
 #include "MetasoundFrontend.h"
 #include "MetasoundFrontendDataTypeRegistry.h"
 #include "MetasoundFrontendNodeTemplateRegistry.h"
@@ -17,10 +18,19 @@
 #include "MetasoundLog.h"
 #include "MetasoundNodeInterface.h"
 
+
 namespace Metasound
 {
 	namespace FrontendGraphPrivate
 	{
+		int32 ForceCommandletExecution = 0;
+			FAutoConsoleVariableRef CVarEnableCommandletExecution(
+				TEXT("au.MetaSound.EnableCommandletExecution"),
+				ForceCommandletExecution,
+				TEXT("If application is a commandlet, enables execution of MetaSounds irrespective of whether sound is rendered to hardware or not. (Ignored if cooking)")
+				TEXT("Default: 0"),
+				ECVF_Default);
+
 		FNodeInitData CreateNodeInitData(const FMetasoundFrontendNode& InNode)
 		{
 			FNodeInitData InitData;
@@ -646,6 +656,21 @@ namespace Metasound
 		return DefaultLiteralData;
 	}
 
+	bool FFrontendGraphBuilder::CanEverExecute(bool bIsContextCooking /* = false */)
+	{
+		if (bIsContextCooking || IsRunningCookCommandlet())
+		{
+			return false;
+		}
+
+		if (IsRunningCommandlet())
+		{
+			return FrontendGraphPrivate::ForceCommandletExecution != 0;
+		}
+
+		return true;
+	}
+
 	/** Check that all dependencies are C++ class dependencies. */
 	bool FFrontendGraphBuilder::IsFlat(const FMetasoundFrontendDocument& InDocument)
 	{
diff --git a/Engine/Plugins/Runtime/Metasound/Source/MetasoundFrontend/Public/MetasoundAssetBase.h b/Engine/Plugins/Runtime/Metasound/Source/MetasoundFrontend/Public/MetasoundAssetBase.h
--- a/Engine/Plugins/Runtime/Metasound/Source/MetasoundFrontend/Public/MetasoundAssetBase.h
+++ b/Engine/Plugins/Runtime/Metasound/Source/MetasoundFrontend/Public/MetasoundAssetBase.h
@@ -80,15 +80,24 @@ public:
 	UE_DEPRECATED(5.5, "Moved to IMetaSoundDocumentInterface::ConformObjectToDocument")
 	virtual bool ConformObjectDataToInterfaces();
 
-	// Registers the root graph of the given asset with the MetaSound Frontend.
+	// Registers the root graph of the given asset with the MetaSound Frontend. Unlike 'PreSaveDocument", this call
+	// generates all necessary runtime data to execute the given graph (i.e. INodes).
 	virtual void RegisterGraphWithFrontend(Metasound::Frontend::FMetaSoundAssetRegistrationOptions InRegistrationOptions = Metasound::Frontend::FMetaSoundAssetRegistrationOptions());
 
 	// Unregisters the root graph of the given asset with the MetaSound Frontend.
 	void UnregisterGraphWithFrontend();
 
-	// Cooks this MetaSound and recursively checks and cooks referenced MetaSounds if necessary. Cook includes autoupdating and resolving the document, which is then registered with the MetaSound Frontend.
+	UE_DEPRECATED(5.5, "Moved to PreSaveDocument instead, which is only in builds set to load editor-only data.")
 	void CookMetaSound();
 
+#if WITH_EDITORONLY_DATA
+	// Executes on this and referenced MetaSound document objects. Presave includes autoupdating and
+	// optimizing the document for runtime use, which is then registered with the MetaSound Frontend.
+	// Unlike 'RegisterGraphWithFrontend', this call does not generate required runtime data for graph
+	// execution.
+	void PreSaveDocument();
+#endif // WITH_EDITORONLY_DATA
+
 #if WITH_EDITOR
 	// Rebuild dependent asset classes
 	void RebuildReferencedAssetClasses();
@@ -255,6 +264,8 @@ protected:
 	PRAGMA_ENABLE_DEPRECATION_WARNINGS
 
 	bool AutoUpdate(bool bInLogWarningsOnDroppedConnection);
+
+	UE_DEPRECATED(5.5, "Moved to private, PreSaveReferencedDocuments implementation")
 	void CookReferencedMetaSounds();
 
 	// Ensures all referenced graph classes are registered (or re-registers depending on options).
@@ -268,6 +279,8 @@ private:
 	void UpdateAssetRegistry();
 
 	bool bVersionedOnLoad = false;
+
+	void PreSaveReferencedDocuments();
 #endif
 
 	// Checks if version is up-to-date. If so, returns true. If false, updates the interfaces within the given asset's document to the most recent version.
diff --git a/Engine/Plugins/Runtime/Metasound/Source/MetasoundFrontend/Public/MetasoundFrontendDocument.h b/Engine/Plugins/Runtime/Metasound/Source/MetasoundFrontend/Public/MetasoundFrontendDocument.h
--- a/Engine/Plugins/Runtime/Metasound/Source/MetasoundFrontend/Public/MetasoundFrontendDocument.h
+++ b/Engine/Plugins/Runtime/Metasound/Source/MetasoundFrontend/Public/MetasoundFrontendDocument.h
@@ -1761,7 +1761,10 @@ private:
 	TArray<FMetasoundFrontendVersion> InterfaceVersions;
 #endif // WITH_EDITORONLY_DATA
 
-	// Used for generating deterministic IDs per document
+	// Used for generating deterministic IDs per document. Serialized to avoid id collisions if deterministic IDs
+	// are ever serialized (not idea, but can occur in less common commandlet use cases such as resaving serialized
+	// assets procedurally).
+	UPROPERTY()
 	mutable uint32 IdCounter = 1;
 
 public:
diff --git a/Engine/Plugins/Runtime/Metasound/Source/MetasoundFrontend/Public/MetasoundFrontendDocumentIdGenerator.h b/Engine/Plugins/Runtime/Metasound/Source/MetasoundFrontend/Public/MetasoundFrontendDocumentIdGenerator.h
--- a/Engine/Plugins/Runtime/Metasound/Source/MetasoundFrontend/Public/MetasoundFrontendDocumentIdGenerator.h
+++ b/Engine/Plugins/Runtime/Metasound/Source/MetasoundFrontend/Public/MetasoundFrontendDocumentIdGenerator.h
@@ -86,7 +86,7 @@ namespace Metasound
 			FGuid CreateOutputID(const FMetasoundFrontendClassOutput& Output) const;
 			FGuid CreateOutputID(const Audio::FParameterInterface::FOutput& Output) const;
 
-			FGuid CreateNamespacedIDFromString(const FGuid NamespaceGuid, const FString& StringToHash) const;
+			FGuid CreateNamespacedIDFromString(const FGuid& NamespaceGuid, const FString& StringToHash) const;
 		};
 
 		METASOUNDFRONTEND_API FGuid CreateLocallyUniqueId();
diff --git a/Engine/Plugins/Runtime/Metasound/Source/MetasoundFrontend/Public/MetasoundFrontendGraph.h b/Engine/Plugins/Runtime/Metasound/Source/MetasoundFrontend/Public/MetasoundFrontendGraph.h
--- a/Engine/Plugins/Runtime/Metasound/Source/MetasoundFrontend/Public/MetasoundFrontendGraph.h
+++ b/Engine/Plugins/Runtime/Metasound/Source/MetasoundFrontend/Public/MetasoundFrontendGraph.h
@@ -83,6 +83,8 @@ namespace Metasound
 	class METASOUNDFRONTEND_API FFrontendGraphBuilder
 	{
 	public:
+		// Returns true if a MetaSound is every expected to execute in the given application, false if not.
+		static bool CanEverExecute(bool bIsContextCooking = false);
 
 		/** Check that all dependencies are C++ class dependencies. 
 		 * 
