Add ability to cache multiple document edge/node subcaches by PageID
#rnx

diff --git a/Engine/Plugins/Runtime/Metasound/Source/MetasoundFrontend/Private/MetasoundFrontendDocumentBuilder.cpp b/Engine/Plugins/Runtime/Metasound/Source/MetasoundFrontend/Private/MetasoundFrontendDocumentBuilder.cpp
--- a/Engine/Plugins/Runtime/Metasound/Source/MetasoundFrontend/Private/MetasoundFrontendDocumentBuilder.cpp
+++ b/Engine/Plugins/Runtime/Metasound/Source/MetasoundFrontend/Private/MetasoundFrontendDocumentBuilder.cpp
@@ -3123,7 +3123,7 @@ bool FMetaSoundFrontendDocumentBuilder::SetBuildPageID(const FGuid& InBuildPageI
 			BuildPageID = BuildGraph->PageID;
 
 			constexpr bool bPrimeCache = false;
-			DocumentCache = FDocumentCache::Create(Document, DocumentDelegates.ToSharedRef(), BuildPageID, bPrimeCache);
+			DocumentCache->SetBuildPageID(BuildPageID);
 		}
 		return true;
 	}
diff --git a/Engine/Plugins/Runtime/Metasound/Source/MetasoundFrontend/Private/MetasoundFrontendDocumentCache.cpp b/Engine/Plugins/Runtime/Metasound/Source/MetasoundFrontend/Private/MetasoundFrontendDocumentCache.cpp
--- a/Engine/Plugins/Runtime/Metasound/Source/MetasoundFrontend/Private/MetasoundFrontendDocumentCache.cpp
+++ b/Engine/Plugins/Runtime/Metasound/Source/MetasoundFrontend/Private/MetasoundFrontendDocumentCache.cpp
@@ -97,12 +97,12 @@ namespace Metasound::Frontend
 	{
 	}
 
-	void FDocumentCache::Init(const FGuid& InPageID, bool bPrimeCache)
+	void FDocumentCache::Init(const FGuid& InBuildPageID, bool bPrimeCache)
 	{
 		METASOUND_TRACE_CPUPROFILER_EVENT_SCOPE(Metasound::Frontend::FDocumentCache::Init);
 		check(Document);
 
-		PageID = InPageID;
+		BuildPageID = InBuildPageID;
 		if (bPrimeCache)
 		{
 			METASOUND_TRACE_CPUPROFILER_EVENT_SCOPE(Metasound::Frontend::FDocumentCache::Init_Prime);
@@ -110,10 +110,10 @@ namespace Metasound::Frontend
 			TSharedRef<const FDocumentCache> ThisShared = StaticCastSharedRef<const FDocumentCache>(AsShared());
 			DependencyCache = MakeShared<FDocumentDependencyCache>(GetDocument());
 
-			// Currently just hardcoded to use cache for default graph (i.e. PageID of zero guid).  Needs to be
-			// transitioned to support multiple caches for each page, where cache is generated on demand per page.
-			EdgeCache = FDocumentGraphEdgeCache::Create(ThisShared, PageID, ModifyDelegates->EdgeDelegates);
-			NodeCache = FDocumentGraphNodeCache::Create(ThisShared, PageID, ModifyDelegates->NodeDelegates);
+			TSharedRef<FDocumentGraphEdgeCache> EdgeCache = FDocumentGraphEdgeCache::Create(ThisShared, InBuildPageID, ModifyDelegates.Get());
+			EdgeCacheMap.Add(InBuildPageID, EdgeCache);
+			TSharedRef<FDocumentGraphNodeCache> NodeCache = FDocumentGraphNodeCache::Create(ThisShared, InBuildPageID, ModifyDelegates.Get());
+			NodeCacheMap.Add(InBuildPageID, NodeCache);
 			InterfaceCache = FDocumentGraphInterfaceCache::Create(ThisShared, ModifyDelegates->InterfaceDelegates);
 		}
 
@@ -128,7 +128,7 @@ namespace Metasound::Frontend
 		return *Document;
 	}
 
-	TSharedRef<FDocumentCache> FDocumentCache::Create(const FMetasoundFrontendDocument& InDocument, TSharedRef<FDocumentModifyDelegates> Delegates, const FGuid& InPageID, bool bPrimeCache)
+	TSharedRef<FDocumentCache> FDocumentCache::Create(const FMetasoundFrontendDocument& InDocument, TSharedRef<FDocumentModifyDelegates> Delegates, const FGuid& InBuildPageID, bool bPrimeCache)
 	{
 		METASOUND_TRACE_CPUPROFILER_EVENT_SCOPE(Metasound::Frontend::FDocumentCache::Create);
 
@@ -137,7 +137,7 @@ namespace Metasound::Frontend
 
 		// Right now the assumption is the whole doc cache is invalidated when swapping pages.  Need to migrate to supporting associated caches per page
 		// (i.e. a node and edge cache per page)
-		Cache->Init(InPageID, bPrimeCache);
+		Cache->Init(InBuildPageID, bPrimeCache);
 		return Cache;
 	}
 
@@ -208,22 +208,38 @@ namespace Metasound::Frontend
 
 	const IDocumentGraphEdgeCache& FDocumentCache::GetEdgeCache() const
 	{
+		return GetEdgeCache(BuildPageID);
+	}
+
+	const IDocumentGraphEdgeCache& FDocumentCache::GetEdgeCache(const FGuid& InPageID) const
+	{
+		TSharedPtr<FDocumentGraphEdgeCache> EdgeCache = EdgeCacheMap.FindRef(InPageID);
 		if (!EdgeCache.IsValid())
 		{
 			TSharedRef<const FDocumentCache> ThisShared = StaticCastSharedRef<const FDocumentCache>(AsShared());
-			EdgeCache = FDocumentGraphEdgeCache::Create(ThisShared, PageID, ModifyDelegates->EdgeDelegates);
+			EdgeCache = FDocumentGraphEdgeCache::Create(ThisShared, InPageID, ModifyDelegates.Get());
+			EdgeCacheMap.Add(InPageID, EdgeCache);
 		}
-		return *EdgeCache;
+
+		return *EdgeCache.Get();
 	}
 
 	const IDocumentGraphNodeCache& FDocumentCache::GetNodeCache() const
 	{
+		return GetNodeCache(BuildPageID);
+	}
+
+	const IDocumentGraphNodeCache& FDocumentCache::GetNodeCache(const FGuid& InPageID) const
+	{
+		TSharedPtr<FDocumentGraphNodeCache> NodeCache = NodeCacheMap.FindRef(InPageID);
 		if (!NodeCache.IsValid())
 		{
 			TSharedRef<const FDocumentCache> ThisShared = StaticCastSharedRef<const FDocumentCache>(AsShared());
-			NodeCache = FDocumentGraphNodeCache::Create(ThisShared, PageID, ModifyDelegates->NodeDelegates);
+			NodeCache = FDocumentGraphNodeCache::Create(ThisShared, InPageID, ModifyDelegates.Get());
+			NodeCacheMap.Add(InPageID, NodeCache);
 		}
-		return *NodeCache;
+
+		return *NodeCache.Get();
 	}
 
 	const IDocumentGraphInterfaceCache& FDocumentCache::GetInterfaceCache() const
@@ -237,13 +253,20 @@ namespace Metasound::Frontend
 	}
 
 	int32 FDocumentCache::GetTransactionCount() const
+	{
+		return GetTransactionCount(BuildPageID);
+	}
+
+	int32 FDocumentCache::GetTransactionCount(const FGuid& InPageID) const
 	{
 		int32 TotalCount = TransactionCount;
+		TSharedPtr<FDocumentGraphNodeCache> NodeCache = NodeCacheMap.FindRef(InPageID);
 		if (NodeCache.IsValid())
 		{
 			TotalCount += NodeCache->GetTransactionCount();
 		}
 
+		TSharedPtr<FDocumentGraphEdgeCache> EdgeCache = EdgeCacheMap.FindRef(InPageID);
 		if (EdgeCache.IsValid())
 		{
 			TotalCount += EdgeCache->GetTransactionCount();
@@ -307,6 +330,15 @@ namespace Metasound::Frontend
 		++TransactionCount;
 	}
 
+	void FDocumentCache::SetBuildPageID(const FGuid& InPageID)
+	{
+		if (InPageID != BuildPageID)
+		{
+			check(Document->RootGraph.ContainsGraphPage(InPageID));
+			BuildPageID = InPageID;
+		}
+	}
+
 	FDocumentGraphInterfaceCache::FDocumentGraphInterfaceCache(TSharedRef<const FDocumentCache> ParentCache)
 		: Parent(ParentCache)
 	{
@@ -453,13 +485,15 @@ namespace Metasound::Frontend
 		return IDToIndex.Contains(InNodeID);
 	}
 
-	TSharedRef<FDocumentGraphNodeCache> FDocumentGraphNodeCache::Create(TSharedRef<const FDocumentCache> ParentCache, const FGuid& InPageID, FNodeModifyDelegates& OutDelegates)
+	TSharedRef<FDocumentGraphNodeCache> FDocumentGraphNodeCache::Create(TSharedRef<const FDocumentCache> ParentCache, const FGuid& InPageID, FDocumentModifyDelegates& OutDelegates)
 	{
 		METASOUND_TRACE_CPUPROFILER_EVENT_SCOPE(Metasound::Frontend::FDocumentGraphNodeCache::Create);
 
 		// Factory style constructor as restrictions on construction of shared pointers disallows passing of this document cache's pointer to sub-caches
 		TSharedRef<FDocumentGraphNodeCache> Cache = MakeShared<FDocumentGraphNodeCache>(ParentCache, InPageID);
-		Cache->Init(OutDelegates);
+		Cache->Init(InPageID.IsValid()
+			? OutDelegates.PageNodeDelegates.FindOrAdd(InPageID)
+			: OutDelegates.NodeDelegates);
 		return Cache;
 	}
 
@@ -853,13 +887,16 @@ namespace Metasound::Frontend
 		return false;
 	}
 
-	TSharedRef<FDocumentGraphEdgeCache> FDocumentGraphEdgeCache::Create(TSharedRef<const FDocumentCache> ParentCache, const FGuid& InPageID, FEdgeModifyDelegates& OutDelegates)
+	TSharedRef<FDocumentGraphEdgeCache> FDocumentGraphEdgeCache::Create(TSharedRef<const FDocumentCache> ParentCache, const FGuid& InPageID, FDocumentModifyDelegates& OutDelegates)
 	{
 		METASOUND_TRACE_CPUPROFILER_EVENT_SCOPE(Metasound::Frontend::FDocumentGraphEdgeCache::Create);
 
 		// Factory style constructor as restrictions on construction of shared pointers disallows passing of this document cache's pointer to sub-caches
 		TSharedRef<FDocumentGraphEdgeCache> Cache = MakeShared<FDocumentGraphEdgeCache>(ParentCache, InPageID);
-		Cache->Init(OutDelegates);
+		Cache->Init(InPageID.IsValid()
+			? OutDelegates.PageEdgeDelegates.FindOrAdd(InPageID)
+			: OutDelegates.EdgeDelegates);
+
 		return Cache;
 	}
 
diff --git a/Engine/Plugins/Runtime/Metasound/Source/MetasoundFrontend/Private/MetasoundFrontendDocumentCache.h b/Engine/Plugins/Runtime/Metasound/Source/MetasoundFrontend/Private/MetasoundFrontendDocumentCache.h
--- a/Engine/Plugins/Runtime/Metasound/Source/MetasoundFrontend/Private/MetasoundFrontendDocumentCache.h
+++ b/Engine/Plugins/Runtime/Metasound/Source/MetasoundFrontend/Private/MetasoundFrontendDocumentCache.h
@@ -26,7 +26,7 @@ namespace Metasound::Frontend
 		FDocumentGraphEdgeCache(TSharedRef<const FDocumentCache> ParentCache, const FGuid& InPageID);
 
 	public:
-		static TSharedRef<FDocumentGraphEdgeCache> Create(TSharedRef<const FDocumentCache> ParentCache, const FGuid& InPageID, FEdgeModifyDelegates& OutDelegates);
+		static TSharedRef<FDocumentGraphEdgeCache> Create(TSharedRef<const FDocumentCache> ParentCache, const FGuid& InPageID, FDocumentModifyDelegates& OutDelegates);
 		virtual ~FDocumentGraphEdgeCache() = default;
 
 		// IDocumentGraphEdgeCache implementation
@@ -68,7 +68,7 @@ namespace Metasound::Frontend
 		FDocumentGraphNodeCache(TSharedRef<const FDocumentCache> ParentCache, const FGuid& InPageID);
 
 	public:
-		static TSharedRef<FDocumentGraphNodeCache> Create(TSharedRef<const FDocumentCache> ParentCache, const FGuid& InPageID, FNodeModifyDelegates& OutDelegates);
+		static TSharedRef<FDocumentGraphNodeCache> Create(TSharedRef<const FDocumentCache> ParentCache, const FGuid& InPageID, FDocumentModifyDelegates& OutDelegates);
 		virtual ~FDocumentGraphNodeCache() = default;
 
 		// IDocumentGraphNodeCache implementation
@@ -162,7 +162,7 @@ namespace Metasound::Frontend
 		FDocumentCache(const FMetasoundFrontendDocument& InDocument, TSharedRef<FDocumentModifyDelegates> Delegates);
 
 	public:
-		static TSharedRef<FDocumentCache> Create(const FMetasoundFrontendDocument& InDocument, TSharedRef<FDocumentModifyDelegates> Delegates, const FGuid& InPageID, bool bPrimeCache);
+		static TSharedRef<FDocumentCache> Create(const FMetasoundFrontendDocument& InDocument, TSharedRef<FDocumentModifyDelegates> Delegates, const FGuid& InBuildPageID, bool bPrimeCache);
 		virtual ~FDocumentCache() = default;
 
 		virtual bool ContainsDependency(const FNodeRegistryKey& InClassKey) const override;
@@ -175,13 +175,18 @@ namespace Metasound::Frontend
 
 		virtual const FMetasoundFrontendDocument& GetDocument() const override;
 		virtual const IDocumentGraphEdgeCache& GetEdgeCache() const override;
+		virtual const IDocumentGraphEdgeCache& GetEdgeCache(const FGuid& InPageID) const override;
 		virtual const IDocumentGraphNodeCache& GetNodeCache() const override;
+		virtual const IDocumentGraphNodeCache& GetNodeCache(const FGuid& InPageID) const override;
 		virtual const IDocumentGraphInterfaceCache& GetInterfaceCache() const override;
 
+		virtual void SetBuildPageID(const FGuid& InPageID) override;
+
 		int32 GetTransactionCount() const;
+		int32 GetTransactionCount(const FGuid& InPageID) const;
 
 	private:
-		void Init(const FGuid& InPageID, bool bPrimeCache);
+		void Init(const FGuid& InBuildPageID, bool bPrimeCache);
 
 		void OnDependencyAdded(int32 InNewIndex);
 		void OnRemoveSwappingDependency(int32 SwapIndex, int32 LastIndex);
@@ -205,10 +210,15 @@ namespace Metasound::Frontend
 		FDocumentDependencyCache& GetDependencyCache();
 		const FDocumentDependencyCache& GetDependencyCache() const;
 
-		// Mutable to be able to be loaded on demand if cache does not exist at the time of an initial query
+		// Caches are mutable to be able to be loaded on demand if cache does not exist at the time of an initial query
 		mutable TSharedPtr<FDocumentDependencyCache> DependencyCache;
-		mutable TSharedPtr<FDocumentGraphEdgeCache> EdgeCache;
-		mutable TSharedPtr<FDocumentGraphNodeCache> NodeCache;
+
+		// Sorted map of PageID to EdgeCache
+		mutable TSortedMap<FGuid, TSharedPtr<FDocumentGraphEdgeCache>> EdgeCacheMap;
+
+		// Sorted map of PageID to NodeCache
+		mutable TSortedMap<FGuid, TSharedPtr<FDocumentGraphNodeCache>> NodeCacheMap;
+
 		mutable TSharedPtr<FDocumentGraphInterfaceCache> InterfaceCache;
 
 		const FMetasoundFrontendDocument* Document = nullptr;
@@ -216,6 +226,8 @@ namespace Metasound::Frontend
 
 		// Number of transactions processed since builder was instantiated
 		int32 TransactionCount = 0;
-		FGuid PageID;
+
+		// Currently set build PageID, which is the default PageID for API calls that do not specify a targeted PageID.
+		FGuid BuildPageID;
 	};
 } // namespace Metasound::Frontend
diff --git a/Engine/Plugins/Runtime/Metasound/Source/MetasoundFrontend/Public/MetasoundFrontendDocumentCacheInterface.h b/Engine/Plugins/Runtime/Metasound/Source/MetasoundFrontend/Public/MetasoundFrontendDocumentCacheInterface.h
--- a/Engine/Plugins/Runtime/Metasound/Source/MetasoundFrontend/Public/MetasoundFrontendDocumentCacheInterface.h
+++ b/Engine/Plugins/Runtime/Metasound/Source/MetasoundFrontend/Public/MetasoundFrontendDocumentCacheInterface.h
@@ -108,9 +108,23 @@ namespace Metasound::Frontend
 
 		virtual const FMetasoundFrontendDocument& GetDocument() const = 0;
 
+		// Returns NodeCache associated with the currently set build graph PageID
 		virtual const IDocumentGraphNodeCache& GetNodeCache() const = 0;
+
+		// Returns NodeCache associated with the given PageID
+		virtual const IDocumentGraphNodeCache& GetNodeCache(const FGuid& InPageID) const = 0;
+
+		// Returns EdgeCache associated with the currently set build graph PageID
 		virtual const IDocumentGraphEdgeCache& GetEdgeCache() const = 0;
+
+		// Returns EdgeCache associated with the given PageID
+		virtual const IDocumentGraphEdgeCache& GetEdgeCache(const FGuid& InPageID) const = 0;
+
+		// Returns InterfaceCache for the given document
 		virtual const IDocumentGraphInterfaceCache& GetInterfaceCache() const = 0;
+
+		// Sets the default BuildPageID for the given document
+		virtual void SetBuildPageID(const FGuid& InPageID) = 0;
 	};
 
 } // namespace Metasound::Frontend
diff --git a/Engine/Plugins/Runtime/Metasound/Source/MetasoundFrontend/Public/MetasoundFrontendDocumentModifyDelegates.h b/Engine/Plugins/Runtime/Metasound/Source/MetasoundFrontend/Public/MetasoundFrontendDocumentModifyDelegates.h
--- a/Engine/Plugins/Runtime/Metasound/Source/MetasoundFrontend/Public/MetasoundFrontendDocumentModifyDelegates.h
+++ b/Engine/Plugins/Runtime/Metasound/Source/MetasoundFrontend/Public/MetasoundFrontendDocumentModifyDelegates.h
@@ -52,5 +52,8 @@ namespace Metasound::Frontend
 		FInterfaceModifyDelegates InterfaceDelegates;
 		FNodeModifyDelegates NodeDelegates;
 		FEdgeModifyDelegates EdgeDelegates;
+
+		TSortedMap<FGuid, FNodeModifyDelegates> PageNodeDelegates;
+		TSortedMap<FGuid, FEdgeModifyDelegates> PageEdgeDelegates;
 	};
 } // namespace Metasound::Frontend
