Pesky, sneaky, snacky exclamat!on
#rnx

diff --git a/Engine/Plugins/Runtime/Metasound/Source/MetasoundEngine/Private/MetasoundSource.cpp b/Engine/Plugins/Runtime/Metasound/Source/MetasoundEngine/Private/MetasoundSource.cpp
--- a/Engine/Plugins/Runtime/Metasound/Source/MetasoundEngine/Private/MetasoundSource.cpp
+++ b/Engine/Plugins/Runtime/Metasound/Source/MetasoundEngine/Private/MetasoundSource.cpp
@@ -812,7 +812,7 @@ void UMetaSoundSource::InitResources()
 
 	METASOUND_LLM_SCOPE;
 	METASOUND_TRACE_CPUPROFILER_EVENT_SCOPE(UMetaSoundSource::InitResources); 
-	ensureMsgf(!FFrontendGraphBuilder::CanEverExecute(), TEXT("UMetaSoundSource::InitResources should not be called in builds that cannot execute MetaSounds."));
+	ensureMsgf(FFrontendGraphBuilder::CanEverExecute(), TEXT("UMetaSoundSource::InitResources should not be called in builds that cannot execute MetaSounds."));
 
 	if (IsInGameThread())
 	{
