[MetaSound] - Creating a new input/output doesn't prompt for rename
#rnx

diff --git a/Engine/Plugins/Runtime/Metasound/Source/MetasoundEditor/Private/MetasoundEditor.cpp b/Engine/Plugins/Runtime/Metasound/Source/MetasoundEditor/Private/MetasoundEditor.cpp
--- a/Engine/Plugins/Runtime/Metasound/Source/MetasoundEditor/Private/MetasoundEditor.cpp
+++ b/Engine/Plugins/Runtime/Metasound/Source/MetasoundEditor/Private/MetasoundEditor.cpp
@@ -2851,11 +2851,11 @@ namespace Metasound
 				{
 					if (const UMetasoundEditorGraphMember* Member = MemberNode->GetMember()) 
 					{
-						GraphMembersMenu->SelectItemByName(Member->GetMemberName(), ESelectInfo::Direct, static_cast<int32>(Member->GetSectionID()));
-
-						if (Member->OnRenameRequested.IsBound())
+						if (Member->CanRename())
 						{
-							Member->OnRenameRequested.Broadcast();
+							GraphMembersMenu->SelectItemByName(Member->GetMemberName(), ESelectInfo::Direct, static_cast<int32>(Member->GetSectionID()));
+							GraphMembersMenu->RefreshAllActions(/*bPreserveExpansion=*/ true, /*bHandleOnSelectionEvent=*/ true);
+							GraphMembersMenu->OnRequestRenameOnActionNode();
 						}
 					}
 				}
@@ -2880,11 +2880,8 @@ namespace Metasound
 							{
 								if (GraphMember->CanRename())
 								{
-									if (GraphMember->OnRenameRequested.IsBound())
-									{
-										GraphMember->OnRenameRequested.Broadcast();
-										return;
-									}
+									GraphMembersMenu->RefreshAllActions(/*bPreserveExpansion=*/ true, /*bHandleOnSelectionEvent=*/ true);
+									GraphMembersMenu->OnRequestRenameOnActionNode();
 								}
 							}
 						}
@@ -3038,6 +3035,7 @@ namespace Metasound
 				{
 					GraphMembersMenu->SelectItemByName(NameToSelect);
 					SetSelection(SelectedObjects);
+					SetDelayedRename();
 				}
 			}
 		}
@@ -3658,8 +3656,6 @@ namespace Metasound
 			FMetasoundAssetBase* MetaSoundAsset = IMetasoundUObjectRegistry::Get().GetObjectAsAssetBase(&MetaSound);
 			check(MetaSoundAsset);
 
-			bool bNodesSelected = false;
-
 			// Capture after synchronizing as the modification state may be modified therein
 			const FMetasoundFrontendDocumentModifyContext& ModifyContext = MetaSoundAsset->GetConstModifyContext();
 			const bool bForceRefreshViews = ModifyContext.GetForceRefreshViews();
@@ -3755,7 +3751,6 @@ namespace Metasound
 				if (!Selection.IsEmpty())
 				{
 					SetSelection(Selection);
-					bNodesSelected = true;
 				}
 
 				// Avoids details panel displaying
@@ -3763,8 +3758,8 @@ namespace Metasound
 				RemoveInvalidSelection();
 			}
 
-			// Wait for GraphMembersMenu to get updated with selected nodes
-			if (bMemberRenameRequested && bNodesSelected)
+			// Prompt to Rename if requested on Member Creation.
+			if (bMemberRenameRequested)
 			{
 				GraphMembersMenu->RefreshAllActions(/*bPreserveExpansion=*/ true, /*bHandleOnSelectionEvent=*/ true);
 				GraphMembersMenu->OnRequestRenameOnActionNode();
@@ -3976,6 +3971,7 @@ namespace Metasound
 				{
 					GraphMembersMenu->SelectItemByName(NameToSelect);
 					SetSelection(SelectedObjects);
+					SetDelayedRename();
 				}
 			}
 			return FReply::Handled();
