UE-216288 - [MetaSounds] - There are no icons when dragging an input or variable to a pin to show if the action is premitted

diff --git a/Engine/Plugins/Runtime/Metasound/Source/MetasoundEditor/Private/MetasoundEditor.cpp b/Engine/Plugins/Runtime/Metasound/Source/MetasoundEditor/Private/MetasoundEditor.cpp
--- a/Engine/Plugins/Runtime/Metasound/Source/MetasoundEditor/Private/MetasoundEditor.cpp
+++ b/Engine/Plugins/Runtime/Metasound/Source/MetasoundEditor/Private/MetasoundEditor.cpp
@@ -431,6 +431,37 @@ namespace Metasound
 				return false;
 			}
 
+			Frontend::FConnectability CanBeConnected(const FName& DataType0, const FName& DataType1)
+			{
+				using namespace Frontend;
+
+				FConnectability OutConnectability;
+				OutConnectability.Connectable = FConnectability::EConnectable::No;
+				OutConnectability.Reason = FConnectability::EReason::None;
+
+				if (DataType0 == FName())
+				{
+					OutConnectability.Connectable = FConnectability::EConnectable::No;
+					OutConnectability.Reason = FConnectability::EReason::IncompatibleDataTypes;
+				}
+				else if (DataType0 == DataType1)
+				{
+					OutConnectability.Connectable = FConnectability::EConnectable::Yes;
+					OutConnectability.Reason = FConnectability::EReason::None;
+				}
+				else
+				{
+					OutConnectability.PossibleConverterNodeClasses = FMetasoundFrontendRegistryContainer::Get()->GetPossibleConverterNodes(DataType0, DataType1);
+
+					if (OutConnectability.PossibleConverterNodeClasses.Num() > 0)
+					{
+						OutConnectability.Connectable = FConnectability::EConnectable::YesWithConverterNode;
+					}
+				}
+
+				return OutConnectability;
+			}
+
 			virtual void HoverTargetChanged() override
 			{
 				using namespace Frontend;
@@ -467,10 +498,30 @@ namespace Metasound
 								Style::GetSlateBrushSafe("MetasoundEditor.Graph.Node.Class.Input");
 								SecondarySymbol = nullptr;
 
-								if (const UMetasoundEditorSettings* EditorSettings = GetDefault<UMetasoundEditorSettings>())
+								UEdGraphPin* PinUnderCursor = GetHoveredPin();
+								
+								if (PinUnderCursor != nullptr && PinUnderCursor->Direction == EGPD_Input)
 								{
-									PrimaryColor = EditorSettings->InputNodeTitleColor;
-									SecondaryColor = EditorSettings->InputNodeTitleColor;
+									Frontend::FConstInputHandle InputHandle = Editor::FGraphBuilder::GetConstInputHandleFromPin(PinUnderCursor);
+									const FName& DataType = InputHandle->GetDataType();
+									const FName& OtherDataType = GraphMember->GetDataType();
+
+									FConnectability Connectability = CanBeConnected(OtherDataType, DataType);
+
+									PrimarySymbol = FAppStyle::GetBrush(TEXT("Graph.ConnectorFeedback.OK"));
+									Message = FText();
+									if (Connectability.Connectable == Frontend::FConnectability::EConnectable::No)
+									{
+										PrimarySymbol = FAppStyle::GetBrush(TEXT("Graph.ConnectorFeedback.Error"));
+										Message = FText::Format(LOCTEXT("NotCompatible_Error", "'{0}' is not compatible with '{1}'"),
+											FText::FromName(DataType), FText::FromName(OtherDataType));
+									}
+									else if (Connectability.Connectable == Frontend::FConnectability::EConnectable::YesWithConverterNode)
+									{
+										PrimarySymbol = FAppStyle::GetBrush(TEXT("Graph.ConnectorFeedback.ViaCast"));
+										Message = FText::Format(LOCTEXT("NotCompatible_Error", "Convert {0} to {1}."),
+											FText::FromName(DataType), FText::FromName(OtherDataType));
+									}
 								}
 							}
 							else if (UMetasoundEditorGraphOutput* Output = Cast<UMetasoundEditorGraphOutput>(GraphMember.Get()))
@@ -491,10 +542,30 @@ namespace Metasound
 										SecondarySymbol = nullptr;
 									}
 
-									if (const UMetasoundEditorSettings* EditorSettings = GetDefault<UMetasoundEditorSettings>())
+									UEdGraphPin* PinUnderCursor = GetHoveredPin();
+
+									if (PinUnderCursor != nullptr && PinUnderCursor->Direction == EGPD_Output)
 									{
-										PrimaryColor = EditorSettings->OutputNodeTitleColor;
-										SecondaryColor = EditorSettings->OutputNodeTitleColor;
+										Frontend::FConstOutputHandle OutputHandle = Editor::FGraphBuilder::GetConstOutputHandleFromPin(PinUnderCursor);
+										const FName& DataType = OutputHandle->GetDataType();
+										const FName& OtherDataType = GraphMember->GetDataType();
+
+										FConnectability Connectability = CanBeConnected(OtherDataType, DataType);
+
+										PrimarySymbol = FAppStyle::GetBrush(TEXT("Graph.ConnectorFeedback.OK"));
+										Message = FText();
+										if (Connectability.Connectable == Frontend::FConnectability::EConnectable::No)
+										{
+											PrimarySymbol = FAppStyle::GetBrush(TEXT("Graph.ConnectorFeedback.Error"));
+											Message = FText::Format(LOCTEXT("NotCompatible_Error", "'{0}' is not compatible with '{1}'"),
+												FText::FromName(DataType), FText::FromName(OtherDataType));
+										}
+										else if (Connectability.Connectable == Frontend::FConnectability::EConnectable::YesWithConverterNode)
+										{
+											PrimarySymbol = FAppStyle::GetBrush(TEXT("Graph.ConnectorFeedback.ViaCast"));
+											Message = FText::Format(LOCTEXT("NotCompatible_Error", "Convert {0} to {1}."),
+												FText::FromName(DataType), FText::FromName(OtherDataType));
+										}
 									}
 								}
 							}
@@ -510,12 +581,6 @@ namespace Metasound
 									SecondarySymbol = nullptr;
 								}
 
-								if (const UMetasoundEditorSettings* EditorSettings = GetDefault<UMetasoundEditorSettings>())
-								{
-									PrimaryColor = EditorSettings->VariableNodeTitleColor;
-									SecondaryColor = EditorSettings->VariableNodeTitleColor;
-								}
-
 								const FText DisplayName = GraphMember->GetDisplayName();
 								const FText GetterToolTip = FText::Format(LOCTEXT("DropTargetGetterVariableToolTipFormat", "{0}\nAdd:\n* Get (Drop)\n* Get Delayed (Alt+Drop)\n"), DisplayName);
 								static const FText GetJumpToToolTip = LOCTEXT("JumpToGettersToolTip", "Get (Ctrl+Drop)");
@@ -538,6 +603,32 @@ namespace Metasound
 										Message = FText::Format(LOCTEXT("DropTargetVariableAddSetJumpToGetFormat", "{0}* Set (Shift+Drop)\n\nJump To:\n* {1}"), GetterToolTip, GetJumpToToolTip);
 									}
 								}
+
+								UEdGraphPin* PinUnderCursor = GetHoveredPin();
+
+								if (PinUnderCursor != nullptr && PinUnderCursor->Direction == EGPD_Input)
+								{
+									Frontend::FConstInputHandle InputHandle = Editor::FGraphBuilder::GetConstInputHandleFromPin(PinUnderCursor);
+									const FName& DataType = InputHandle->GetDataType();
+									const FName& OtherDataType = GraphMember->GetDataType();
+
+									FConnectability Connectability = CanBeConnected(OtherDataType, DataType);
+
+									PrimarySymbol = FAppStyle::GetBrush(TEXT("Graph.ConnectorFeedback.OK"));
+									Message = FText();
+									if (Connectability.Connectable == Frontend::FConnectability::EConnectable::No)
+									{
+										PrimarySymbol = FAppStyle::GetBrush(TEXT("Graph.ConnectorFeedback.Error"));
+										Message = FText::Format(LOCTEXT("NotCompatible_Error", "'{0}' is not compatible with '{1}'"),
+											FText::FromName(DataType), FText::FromName(OtherDataType));
+									}
+									else if (Connectability.Connectable == Frontend::FConnectability::EConnectable::YesWithConverterNode)
+									{
+										PrimarySymbol = FAppStyle::GetBrush(TEXT("Graph.ConnectorFeedback.ViaCast"));
+										Message = FText::Format(LOCTEXT("NotCompatible_Error", "Convert {0} to {1}."),
+											FText::FromName(DataType), FText::FromName(OtherDataType));
+									}
+								}
 							}
 						}
 						else
