Temp fix until we set correct cvar for applicable builds allowing MetaSound execution in test commandlets
#rnx

diff --git a/Engine/Plugins/Runtime/Metasound/Source/MetasoundFrontend/Private/MetasoundFrontendGraph.cpp b/Engine/Plugins/Runtime/Metasound/Source/MetasoundFrontend/Private/MetasoundFrontendGraph.cpp
--- a/Engine/Plugins/Runtime/Metasound/Source/MetasoundFrontend/Private/MetasoundFrontendGraph.cpp
+++ b/Engine/Plugins/Runtime/Metasound/Source/MetasoundFrontend/Private/MetasoundFrontendGraph.cpp
@@ -663,10 +663,11 @@ namespace Metasound
 			return false;
 		}
 
-		if (IsRunningCommandlet())
-		{
-			return FrontendGraphPrivate::ForceCommandletExecution != 0;
-		}
+		// TODO: Test builds need ForceCommandletExecution set to true, otherwise they fail.
+		// if (IsRunningCommandlet())
+		// {
+		// 	return FrontendGraphPrivate::ForceCommandletExecution != 0;
+		// }
 
 		return true;
 	}
