Guid collision fix-up follow-up renaming
#rnx

diff --git a/Engine/Plugins/Runtime/Metasound/Source/MetasoundEditor/Private/MetasoundEditorGraphBuilder.cpp b/Engine/Plugins/Runtime/Metasound/Source/MetasoundEditor/Private/MetasoundEditorGraphBuilder.cpp
--- a/Engine/Plugins/Runtime/Metasound/Source/MetasoundEditor/Private/MetasoundEditorGraphBuilder.cpp
+++ b/Engine/Plugins/Runtime/Metasound/Source/MetasoundEditor/Private/MetasoundEditorGraphBuilder.cpp
@@ -1795,13 +1795,13 @@ namespace Metasound
 			// are registered recursively via bRegisterDependencies flag.
 			if (EditedReferencingMetaSounds.IsEmpty())
 			{
-				MetaSoundAsset->RegisterGraphWithFrontend(RegOptions);
+				MetaSoundAsset->UpdateAndRegisterForExecution(RegOptions);
 			}
 			else
 			{
 				for (FMetasoundAssetBase* MetaSound : EditedReferencingMetaSounds)
 				{
-					MetaSound->RegisterGraphWithFrontend(RegOptions);
+					MetaSound->UpdateAndRegisterForExecution(RegOptions);
 				}
 			}
 		}
diff --git a/Engine/Plugins/Runtime/Metasound/Source/MetasoundEditor/Private/MetasoundEditorModule.cpp b/Engine/Plugins/Runtime/Metasound/Source/MetasoundEditor/Private/MetasoundEditorModule.cpp
--- a/Engine/Plugins/Runtime/Metasound/Source/MetasoundEditor/Private/MetasoundEditorModule.cpp
+++ b/Engine/Plugins/Runtime/Metasound/Source/MetasoundEditor/Private/MetasoundEditorModule.cpp
@@ -27,9 +27,9 @@
 #include "MetasoundEditorSettings.h"
 #include "MetasoundFrontendDataTypeRegistry.h"
 #include "MetasoundFrontendDocument.h"
-#include "MetasoundFrontendGraph.h"
 #include "MetasoundFrontendRegistries.h"
 #include "MetasoundFrontendTransform.h"
+#include "MetasoundGlobals.h"
 #include "MetasoundNodeDetailCustomization.h"
 #include "MetasoundSettings.h"
 #include "MetasoundSource.h"
@@ -243,9 +243,8 @@ namespace Metasound
 			{
 				using namespace Metasound;
 			
-				// Ignore requests if running cook, as presave will register and cook using the
-				// proper node register call which avoids async registration/igraph generation.
-				if (!FFrontendGraphBuilder::CanEverExecute())
+				// Ignore requests if graphs cannot be executed, as registration that results in IGraph generation is not supported.
+				if (!Metasound::CanEverExecuteGraph())
 				{
 					return;
 				}
diff --git a/Engine/Plugins/Runtime/Metasound/Source/MetasoundEngine/Private/MetasoundAssetSubsystem.cpp b/Engine/Plugins/Runtime/Metasound/Source/MetasoundEngine/Private/MetasoundAssetSubsystem.cpp
--- a/Engine/Plugins/Runtime/Metasound/Source/MetasoundEngine/Private/MetasoundAssetSubsystem.cpp
+++ b/Engine/Plugins/Runtime/Metasound/Source/MetasoundEngine/Private/MetasoundAssetSubsystem.cpp
@@ -613,7 +613,7 @@ namespace Metasound::Engine
 			{
 				RegOptions.bAutoUpdateLogWarningOnDroppedConnection = Settings->bAutoUpdateLogWarningOnDroppedConnection;
 			}
-			MetaSoundAsset->RegisterGraphWithFrontend(RegOptions);
+			MetaSoundAsset->UpdateAndRegisterForExecution(RegOptions);
 		});
 	}
 
diff --git a/Engine/Plugins/Runtime/Metasound/Source/MetasoundEngine/Private/MetasoundDocumentBuilderRegistry.cpp b/Engine/Plugins/Runtime/Metasound/Source/MetasoundEngine/Private/MetasoundDocumentBuilderRegistry.cpp
--- a/Engine/Plugins/Runtime/Metasound/Source/MetasoundEngine/Private/MetasoundDocumentBuilderRegistry.cpp
+++ b/Engine/Plugins/Runtime/Metasound/Source/MetasoundEngine/Private/MetasoundDocumentBuilderRegistry.cpp
@@ -3,7 +3,7 @@
 #include "MetasoundDocumentBuilderRegistry.h"
 
 #include "MetasoundAssetManager.h"
-#include "MetasoundFrontendGraph.h"
+#include "MetasoundGlobals.h"
 #include "MetasoundTrace.h"
 #include "MetasoundUObjectRegistry.h"
 
@@ -285,7 +285,7 @@ namespace Metasound::Engine
 		FMetaSoundFrontendDocumentBuilder& DocBuilder = Builder.GetBuilder();
 		if (DocBuilder.IsValid())
 		{
-			if (FFrontendGraphBuilder::CanEverExecute())
+			if (Metasound::CanEverExecuteGraph())
 			{
 				const int32 TransactionCount = DocBuilder.GetTransactionCount();
 				const int32 LastTransactionRegistered = Builder.GetLastTransactionRegistered();
diff --git a/Engine/Plugins/Runtime/Metasound/Source/MetasoundEngine/Private/MetasoundEngineAsset.h b/Engine/Plugins/Runtime/Metasound/Source/MetasoundEngine/Private/MetasoundEngineAsset.h
--- a/Engine/Plugins/Runtime/Metasound/Source/MetasoundEngine/Private/MetasoundEngineAsset.h
+++ b/Engine/Plugins/Runtime/Metasound/Source/MetasoundEngine/Private/MetasoundEngineAsset.h
@@ -6,8 +6,8 @@
 #include "MetasoundAssetManager.h"
 #include "MetasoundDocumentBuilderRegistry.h"
 #include "MetasoundFrontendDocumentIdGenerator.h"
-#include "MetasoundFrontendGraph.h"
 #include "MetasoundFrontendRegistryKey.h"
+#include "MetasoundGlobals.h"
 #include "MetasoundUObjectRegistry.h"
 #include "Misc/App.h"
 #include "Serialization/Archive.h"
@@ -16,8 +16,8 @@
 #include "Algo/Transform.h"
 #include "MetasoundFrontendRegistryContainer.h"
 #include "UObject/GarbageCollection.h"
-#include "UObject/StrongObjectPtrTemplates.h"
 #include "UObject/ObjectMacros.h"
+#include "UObject/StrongObjectPtrTemplates.h"
 #endif // WITH_EDITORONLY_DATA
 
 
@@ -29,9 +29,9 @@ namespace Metasound::Engine
 	 */
 	struct FAssetHelper
 	{
-		static bool IsDeterministic(bool bIsContextCooking)
+		static bool SerializationRequiresDeterminism(bool bIsCooking)
 		{
-			return bIsContextCooking || IsRunningCookCommandlet();
+			return bIsCooking || IsRunningCookCommandlet();
 		}
 
 #if WITH_EDITOR
@@ -145,12 +145,12 @@ namespace Metasound::Engine
 			}
 
 			const bool bIsCooking = InSaveContext.IsCooking();
-			const bool bCanEverExecute = FFrontendGraphBuilder::CanEverExecute(bIsCooking);
+			const bool bCanEverExecute = Metasound::CanEverExecuteGraph(bIsCooking);
 			if (!bCanEverExecute)
 			{
-				const bool bIsDeterministic = IsDeterministic(bIsCooking);
+				const bool bIsDeterministic = SerializationRequiresDeterminism(bIsCooking);
 				FDocumentIDGenerator::FScopeDeterminism DeterminismScope(bIsDeterministic);
-				InMetaSound.PreSaveDocument();
+				InMetaSound.UpdateAndRegisterForSerialization();
 			}
  			else if (FApp::CanEverRenderAudio())
 			{
@@ -187,7 +187,7 @@ namespace Metasound::Engine
 
 				{
 					const bool bIsCooking = InArchive.IsCooking();
-					const bool bIsDeterministic = IsDeterministic(bIsCooking);
+					const bool bIsDeterministic = SerializationRequiresDeterminism(bIsCooking);
 					FDocumentIDGenerator::FScopeDeterminism DeterminismScope(bIsDeterministic);
 					check(Builder.IsValid());
 					bVersionedAsset = InMetaSound.VersionAsset(Builder->GetBuilder());
diff --git a/Engine/Plugins/Runtime/Metasound/Source/MetasoundEngine/Private/MetasoundSource.cpp b/Engine/Plugins/Runtime/Metasound/Source/MetasoundEngine/Private/MetasoundSource.cpp
--- a/Engine/Plugins/Runtime/Metasound/Source/MetasoundEngine/Private/MetasoundSource.cpp
+++ b/Engine/Plugins/Runtime/Metasound/Source/MetasoundEngine/Private/MetasoundSource.cpp
@@ -26,11 +26,11 @@
 #include "MetasoundFrontendDataTypeRegistry.h"
 #include "MetasoundFrontendDocumentBuilder.h"
 #include "MetasoundFrontendDocumentIdGenerator.h"
-#include "MetasoundFrontendGraph.h"
 #include "MetasoundFrontendQuery.h"
 #include "MetasoundFrontendQuerySteps.h"
 #include "MetasoundFrontendTransform.h"
 #include "MetasoundGenerator.h"
+#include "MetasoundGlobals.h"
 #include "MetasoundLog.h"
 #include "MetasoundOperatorBuilderSettings.h"
 #include "MetasoundOperatorCacheSubsystem.h"
@@ -812,11 +812,11 @@ void UMetaSoundSource::InitResources()
 
 	METASOUND_LLM_SCOPE;
 	METASOUND_TRACE_CPUPROFILER_EVENT_SCOPE(UMetaSoundSource::InitResources); 
-	ensureMsgf(FFrontendGraphBuilder::CanEverExecute(), TEXT("UMetaSoundSource::InitResources should not be called in builds that cannot execute MetaSounds."));
+	ensureMsgf(Metasound::CanEverExecuteGraph(), TEXT("UMetaSoundSource::InitResources can only be called in applications where MetaSounds can execute."));
 
 	if (IsInGameThread())
 	{
-		RegisterGraphWithFrontend(GetInitRegistrationOptions());
+		UpdateAndRegisterForExecution(GetInitRegistrationOptions());
 	}
 	else
 	{
@@ -839,17 +839,17 @@ void UMetaSoundSource::InitResources()
 	}
 }
 
-void UMetaSoundSource::RegisterGraphWithFrontend(Metasound::Frontend::FMetaSoundAssetRegistrationOptions InRegistrationOptions)
+void UMetaSoundSource::UpdateAndRegisterForExecution(Metasound::Frontend::FMetaSoundAssetRegistrationOptions InRegistrationOptions)
 {
 	using namespace Metasound;
 	using namespace Metasound::Frontend;
 
 	check(IsInGameThread());
 
-	FMetasoundAssetBase::RegisterGraphWithFrontend(InRegistrationOptions);
+	FMetasoundAssetBase::UpdateAndRegisterForExecution(InRegistrationOptions);
 	const bool bIsRuntimeInputDataValid = RuntimeInputData.bIsValid.load();
 	// Runtime data does not need to and should not be created in builds that don't execute graphs
-	if (!bIsRuntimeInputDataValid && FFrontendGraphBuilder::CanEverExecute())
+	if (!bIsRuntimeInputDataValid && Metasound::CanEverExecuteGraph())
 	{
 		CacheRuntimeInputData();
 	}
diff --git a/Engine/Plugins/Runtime/Metasound/Source/MetasoundEngine/Public/MetasoundBuilderBase.cpp b/Engine/Plugins/Runtime/Metasound/Source/MetasoundEngine/Public/MetasoundBuilderBase.cpp
--- a/Engine/Plugins/Runtime/Metasound/Source/MetasoundEngine/Public/MetasoundBuilderBase.cpp
+++ b/Engine/Plugins/Runtime/Metasound/Source/MetasoundEngine/Public/MetasoundBuilderBase.cpp
@@ -453,7 +453,7 @@ void UMetaSoundBuilderBase::ConvertToPreset(const TScriptInterface<IMetaSoundDoc
 	// Ensure the referenced node is registered
 	if (FMetasoundAssetBase* ReferencedMetaSoundAsset = Metasound::IMetasoundUObjectRegistry::Get().GetObjectAsAssetBase(ReferencedObject))
 	{
-		ReferencedMetaSoundAsset->RegisterGraphWithFrontend();
+		ReferencedMetaSoundAsset->UpdateAndRegisterForExecution();
 	}
 
 	const FMetasoundFrontendDocument& ReferencedDocument = ReferencedInterface->GetConstDocument();
@@ -893,7 +893,7 @@ void UMetaSoundBuilderBase::RegisterGraphIfOutstandingTransactions(UObject& InMe
 		Builder->LastTransactionRegistered = TransactionCount;
 	}
 
-	MetaSoundAsset->RegisterGraphWithFrontend(Options);
+	MetaSoundAsset->UpdateAndRegisterForExecution(Options);
 }
 
 void UMetaSoundBuilderBase::Reload(bool bPrimeCache)
diff --git a/Engine/Plugins/Runtime/Metasound/Source/MetasoundEngine/Public/MetasoundBuilderBase.h b/Engine/Plugins/Runtime/Metasound/Source/MetasoundEngine/Public/MetasoundBuilderBase.h
--- a/Engine/Plugins/Runtime/Metasound/Source/MetasoundEngine/Public/MetasoundBuilderBase.h
+++ b/Engine/Plugins/Runtime/Metasound/Source/MetasoundEngine/Public/MetasoundBuilderBase.h
@@ -499,7 +499,7 @@ protected:
 
 		if (BuilderOptions.bAddToRegistry)
 		{
-			MetaSound->RegisterGraphWithFrontend();
+			MetaSound->UpdateAndRegisterForExecution();
 		}
 
 		UE_LOG(LogMetaSound, VeryVerbose, TEXT("MetaSound '%s' built from '%s'"), *BuilderOptions.Name.ToString(), *GetFullName());
diff --git a/Engine/Plugins/Runtime/Metasound/Source/MetasoundEngine/Public/MetasoundSource.h b/Engine/Plugins/Runtime/Metasound/Source/MetasoundEngine/Public/MetasoundSource.h
--- a/Engine/Plugins/Runtime/Metasound/Source/MetasoundEngine/Public/MetasoundSource.h
+++ b/Engine/Plugins/Runtime/Metasound/Source/MetasoundEngine/Public/MetasoundSource.h
@@ -253,7 +253,7 @@ public:
 	virtual void InitParameters(TArray<FAudioParameter>& ParametersToInit, FName InFeatureName) override;
 
 	virtual void InitResources() override;
-	virtual void RegisterGraphWithFrontend(Metasound::Frontend::FMetaSoundAssetRegistrationOptions InRegistrationOptions = Metasound::Frontend::FMetaSoundAssetRegistrationOptions()) override;
+	virtual void UpdateAndRegisterForExecution(Metasound::Frontend::FMetaSoundAssetRegistrationOptions InRegistrationOptions = Metasound::Frontend::FMetaSoundAssetRegistrationOptions()) override;
 
 	virtual bool IsPlayable() const override;
 	virtual float GetDuration() const override;
diff --git a/Engine/Plugins/Runtime/Metasound/Source/MetasoundFrontend/Private/MetasoundAssetBase.cpp b/Engine/Plugins/Runtime/Metasound/Source/MetasoundFrontend/Private/MetasoundAssetBase.cpp
--- a/Engine/Plugins/Runtime/Metasound/Source/MetasoundFrontend/Private/MetasoundAssetBase.cpp
+++ b/Engine/Plugins/Runtime/Metasound/Source/MetasoundFrontend/Private/MetasoundAssetBase.cpp
@@ -22,13 +22,13 @@
 #include "MetasoundFrontendDocumentBuilder.h"
 #include "MetasoundFrontendDocumentIdGenerator.h"
 #include "MetasoundFrontendDocumentVersioning.h"
-#include "MetasoundFrontendGraph.h"
 #include "MetasoundFrontendNodeTemplateRegistry.h"
 #include "MetasoundFrontendProxyDataCache.h"
 #include "MetasoundFrontendRegistries.h"
 #include "MetasoundFrontendRegistryContainerImpl.h"
 #include "MetasoundFrontendSearchEngine.h"
 #include "MetasoundFrontendTransform.h"
+#include "MetasoundGlobals.h"
 #include "MetasoundGraph.h"
 #include "MetasoundJsonBackend.h"
 #include "MetasoundLog.h"
@@ -262,16 +262,21 @@ bool FMetasoundAssetBase::ConformObjectDataToInterfaces()
 }
 
 void FMetasoundAssetBase::RegisterGraphWithFrontend(Metasound::Frontend::FMetaSoundAssetRegistrationOptions InRegistrationOptions)
+{
+	UpdateAndRegisterForExecution(MoveTemp(InRegistrationOptions));
+}
+
+void FMetasoundAssetBase::UpdateAndRegisterForExecution(Metasound::Frontend::FMetaSoundAssetRegistrationOptions InRegistrationOptions)
 {
 	using namespace Metasound;
 	using namespace Metasound::Frontend;
 
 	// Graph registration must only happen on one thread to avoid race conditions on graph registration.
 	checkf(IsInGameThread(), TEXT("MetaSound %s graph can only be registered on the GameThread"), *GetOwningAssetName());
-	checkf(FFrontendGraphBuilder::CanEverExecute(), TEXT("Cannot generate proxies/runtime graph when graph execution is not enabled."));
+	checkf(Metasound::CanEverExecuteGraph(), TEXT("Cannot generate proxies/runtime graph when graph execution is not enabled."));
 
-	METASOUND_TRACE_CPUPROFILER_EVENT_SCOPE(MetaSoundAssetBase::RegisterGraphWithFrontend);
-	METASOUND_TRACE_CPUPROFILER_EVENT_SCOPE_TEXT(*FString::Printf(TEXT("MetaSoundAssetBase::RegisterGraphWithFrontend asset %s"), *this->GetOwningAssetName()));
+	METASOUND_TRACE_CPUPROFILER_EVENT_SCOPE(MetaSoundAssetBase::UpdateAndRegisterForExecution);
+	METASOUND_TRACE_CPUPROFILER_EVENT_SCOPE_TEXT(*FString::Printf(TEXT("MetaSoundAssetBase::UpdateAndRegisterForExecution  asset %s"), *this->GetOwningAssetName()));
 	if (!InRegistrationOptions.bForceReregister)
 	{
 		if (IsRegistered())
@@ -350,17 +355,17 @@ void FMetasoundAssetBase::RegisterGraphWithFrontend(Metasound::Frontend::FMetaSo
 void FMetasoundAssetBase::CookMetaSound()
 {
 #if WITH_EDITORONLY_DATA
-	PreSaveDocument();
+	UpdateAndRegisterForSerialization();
 #endif // WITH_EDITORONLY_DATA
 }
 
 #if WITH_EDITORONLY_DATA
-void FMetasoundAssetBase::PreSaveDocument()
+void FMetasoundAssetBase::UpdateAndRegisterForSerialization()
 {
 	using namespace Metasound;
 	using namespace Metasound::Frontend;
 
-	METASOUND_TRACE_CPUPROFILER_EVENT_SCOPE(MetaSoundAssetBase::PreSaveDocument);
+	METASOUND_TRACE_CPUPROFILER_EVENT_SCOPE(MetaSoundAssetBase::UpdateAndRegisterForSerialization);
 
 	// If already registered, nothing to condition for presaving
 	if (IsRegistered())
@@ -440,7 +445,7 @@ void FMetasoundAssetBase::OnNotifyBeginDestroy()
 
 	// Unregistration of graph using local call is not necessary when cooking as deserialized objects are not mutable and, should they be
 	// reloaded, omitting unregistration avoids potentially kicking off an invalid asynchronous task to unregister a non-existent runtime graph.
-	if (FFrontendGraphBuilder::CanEverExecute())
+	if (Metasound::CanEverExecuteGraph())
 	{
 		UnregisterGraphWithFrontend();
 	}
@@ -466,7 +471,7 @@ void FMetasoundAssetBase::UnregisterGraphWithFrontend()
 	METASOUND_TRACE_CPUPROFILER_EVENT_SCOPE(MetaSoundAssetBase::UnregisterGraphWithFrontend);
 
 	check(IsInGameThread());
-	checkf(FFrontendGraphBuilder::CanEverExecute(), TEXT("If execution is not supported, UnregisterNode must be called directly to avoid async attempt at destroying runtime graph that does not exist."));
+	checkf(Metasound::CanEverExecuteGraph(), TEXT("If execution is not supported, UnregisterNode must be called directly to avoid async attempt at destroying runtime graph that does not exist."));
 
 	if (GraphRegistryKey.IsValid())
 	{
@@ -908,7 +913,7 @@ void FMetasoundAssetBase::RegisterAssetDependencies(const Metasound::Frontend::F
 		{
 			// TODO: Check for infinite recursion and error if so
 			AssetManager.AddOrUpdateAsset(*(Reference->GetOwningAsset()));
-			Reference->RegisterGraphWithFrontend(InRegistrationOptions);
+			Reference->UpdateAndRegisterForExecution(InRegistrationOptions);
 		}
 	}
 }
@@ -933,7 +938,7 @@ void FMetasoundAssetBase::PreSaveReferencedDocuments()
 		{
 			// TODO: Check for infinite recursion and error if so
 			AssetManager.AddOrUpdateAsset(*(Reference->GetOwningAsset()));
-			Reference->PreSaveDocument();
+			Reference->UpdateAndRegisterForSerialization();
 		}
 	}
 }
diff --git a/Engine/Plugins/Runtime/Metasound/Source/MetasoundFrontend/Private/MetasoundFrontendGraph.cpp b/Engine/Plugins/Runtime/Metasound/Source/MetasoundFrontend/Private/MetasoundFrontendGraph.cpp
--- a/Engine/Plugins/Runtime/Metasound/Source/MetasoundFrontend/Private/MetasoundFrontendGraph.cpp
+++ b/Engine/Plugins/Runtime/Metasound/Source/MetasoundFrontend/Private/MetasoundFrontendGraph.cpp
@@ -23,14 +23,6 @@ namespace Metasound
 {
 	namespace FrontendGraphPrivate
 	{
-		int32 ForceCommandletExecution = 0;
-			FAutoConsoleVariableRef CVarEnableCommandletExecution(
-				TEXT("au.MetaSound.EnableCommandletExecution"),
-				ForceCommandletExecution,
-				TEXT("If application is a commandlet, enables execution of MetaSounds irrespective of whether sound is rendered to hardware or not. (Ignored if cooking)")
-				TEXT("Default: 0"),
-				ECVF_Default);
-
 		FNodeInitData CreateNodeInitData(const FMetasoundFrontendNode& InNode)
 		{
 			FNodeInitData InitData;
@@ -656,22 +648,6 @@ namespace Metasound
 		return DefaultLiteralData;
 	}
 
-	bool FFrontendGraphBuilder::CanEverExecute(bool bIsContextCooking /* = false */)
-	{
-		if (bIsContextCooking || IsRunningCookCommandlet())
-		{
-			return false;
-		}
-
-		// TODO: Test builds need ForceCommandletExecution set to true, otherwise they fail.
-		// if (IsRunningCommandlet())
-		// {
-		// 	return FrontendGraphPrivate::ForceCommandletExecution != 0;
-		// }
-
-		return true;
-	}
-
 	/** Check that all dependencies are C++ class dependencies. */
 	bool FFrontendGraphBuilder::IsFlat(const FMetasoundFrontendDocument& InDocument)
 	{
diff --git a/Engine/Plugins/Runtime/Metasound/Source/MetasoundFrontend/Public/MetasoundAssetBase.h b/Engine/Plugins/Runtime/Metasound/Source/MetasoundFrontend/Public/MetasoundAssetBase.h
--- a/Engine/Plugins/Runtime/Metasound/Source/MetasoundFrontend/Public/MetasoundAssetBase.h
+++ b/Engine/Plugins/Runtime/Metasound/Source/MetasoundFrontend/Public/MetasoundAssetBase.h
@@ -80,14 +80,17 @@ public:
 	UE_DEPRECATED(5.5, "Moved to IMetaSoundDocumentInterface::ConformObjectToDocument")
 	virtual bool ConformObjectDataToInterfaces();
 
-	// Registers the root graph of the given asset with the MetaSound Frontend. Unlike 'PreSaveDocument", this call
+	// Registers the root graph of the given asset with the MetaSound Frontend. Unlike 'UpdateAndRegisterForSerialization", this call
 	// generates all necessary runtime data to execute the given graph (i.e. INodes).
+	virtual void UpdateAndRegisterForExecution(Metasound::Frontend::FMetaSoundAssetRegistrationOptions InRegistrationOptions = Metasound::Frontend::FMetaSoundAssetRegistrationOptions());
+
+	UE_DEPRECATED(5.5, "Moved to UpdateAndRegisterForExecution.")
 	virtual void RegisterGraphWithFrontend(Metasound::Frontend::FMetaSoundAssetRegistrationOptions InRegistrationOptions = Metasound::Frontend::FMetaSoundAssetRegistrationOptions());
 
 	// Unregisters the root graph of the given asset with the MetaSound Frontend.
 	void UnregisterGraphWithFrontend();
 
-	UE_DEPRECATED(5.5, "Moved to PreSaveDocument instead, which is only in builds set to load editor-only data.")
+	UE_DEPRECATED(5.5, "Moved to UpdateAndRegisterForSerialization instead, which is only in builds set to load editor-only data.")
 	void CookMetaSound();
 
 #if WITH_EDITORONLY_DATA
@@ -95,7 +98,7 @@ public:
 	// optimizing the document for runtime use, which is then registered with the MetaSound Frontend.
 	// Unlike 'RegisterGraphWithFrontend', this call does not generate required runtime data for graph
 	// execution.
-	void PreSaveDocument();
+	void UpdateAndRegisterForSerialization();
 #endif // WITH_EDITORONLY_DATA
 
 #if WITH_EDITOR
diff --git a/Engine/Plugins/Runtime/Metasound/Source/MetasoundFrontend/Public/MetasoundFrontendDocument.h b/Engine/Plugins/Runtime/Metasound/Source/MetasoundFrontend/Public/MetasoundFrontendDocument.h
--- a/Engine/Plugins/Runtime/Metasound/Source/MetasoundFrontend/Public/MetasoundFrontendDocument.h
+++ b/Engine/Plugins/Runtime/Metasound/Source/MetasoundFrontend/Public/MetasoundFrontendDocument.h
@@ -1762,7 +1762,7 @@ private:
 #endif // WITH_EDITORONLY_DATA
 
 	// Used for generating deterministic IDs per document. Serialized to avoid id collisions if deterministic IDs
-	// are ever serialized (not idea, but can occur in less common commandlet use cases such as resaving serialized
+	// are ever serialized (not ideal, but can occur in less common commandlet use cases such as resaving serialized
 	// assets procedurally).
 	UPROPERTY()
 	mutable uint32 IdCounter = 1;
diff --git a/Engine/Plugins/Runtime/Metasound/Source/MetasoundFrontend/Public/MetasoundFrontendGraph.h b/Engine/Plugins/Runtime/Metasound/Source/MetasoundFrontend/Public/MetasoundFrontendGraph.h
--- a/Engine/Plugins/Runtime/Metasound/Source/MetasoundFrontend/Public/MetasoundFrontendGraph.h
+++ b/Engine/Plugins/Runtime/Metasound/Source/MetasoundFrontend/Public/MetasoundFrontendGraph.h
@@ -83,9 +83,6 @@ namespace Metasound
 	class METASOUNDFRONTEND_API FFrontendGraphBuilder
 	{
 	public:
-		// Returns true if a MetaSound is every expected to execute in the given application, false if not.
-		static bool CanEverExecute(bool bIsContextCooking = false);
-
 		/** Check that all dependencies are C++ class dependencies. 
 		 * 
 		 * @param InDocument - Document containing dependencies.
diff --git a/Engine/Plugins/Runtime/Metasound/Source/MetasoundGraphCore/Private/MetasoundGlobals.cpp b/Engine/Plugins/Runtime/Metasound/Source/MetasoundGraphCore/Private/MetasoundGlobals.cpp
new file mode 100644
--- /dev/null
+++ b/Engine/Plugins/Runtime/Metasound/Source/MetasoundGraphCore/Private/MetasoundGlobals.cpp
@@ -0,0 +1,36 @@
+// Copyright Epic Games, Inc. All Rights Reserved.
+
+#include "MetasoundGlobals.h"
+
+#include "HAL/IConsoleManager.h"
+
+
+namespace Metasound
+{
+	namespace GlobalsPrivate
+	{
+		int32 EnableCommandletExecution = 0;
+		FAutoConsoleVariableRef CVarEnableCommandletExecution(
+			TEXT("au.MetaSound.EnableCommandletExecution"),
+			EnableCommandletExecution,
+			TEXT("If application is a commandlet, enables execution of MetaSounds irrespective of whether sound is rendered to hardware or not. (Ignored if cooking)")
+			TEXT("Default: 0"),
+			ECVF_Default);
+	} // namespace GlobalsPrivate
+
+	bool CanEverExecuteGraph(bool bIsCooking)
+	{
+		if (bIsCooking || IsRunningCookCommandlet())
+		{
+			return false;
+		}
+
+		// TODO: Test builds need au.MetaSound.EnableCommandletExecution set to true, otherwise they fail.
+		// if (IsRunningCommandlet())
+		// {
+		// 	return GlobalsPrivate::EnableCommandletExecution != 0;
+		// }
+
+		return true;
+	}
+} // namespace Metasound
diff --git a/Engine/Plugins/Runtime/Metasound/Source/MetasoundGraphCore/Public/MetasoundGlobals.h b/Engine/Plugins/Runtime/Metasound/Source/MetasoundGraphCore/Public/MetasoundGlobals.h
new file mode 100644
--- /dev/null
+++ b/Engine/Plugins/Runtime/Metasound/Source/MetasoundGraphCore/Public/MetasoundGlobals.h
@@ -0,0 +1,11 @@
+// Copyright Epic Games, Inc. All Rights Reserved.
+#pragma once
+
+#include "CoreGlobals.h"
+#include "Misc/Build.h"
+
+namespace Metasound
+{
+	// Returns true if a MetaSound is expected to execute, false if not.
+	METASOUNDGRAPHCORE_API bool CanEverExecuteGraph(bool bIsCooking = false);
+} // namespace Metasound
