Fix for corrupt builder access when caching ticked attribute getter
#rnx

diff --git a/Engine/Plugins/Runtime/Metasound/Source/MetasoundEditor/Private/MetasoundDetailCustomization.cpp b/Engine/Plugins/Runtime/Metasound/Source/MetasoundEditor/Private/MetasoundDetailCustomization.cpp
--- a/Engine/Plugins/Runtime/Metasound/Source/MetasoundEditor/Private/MetasoundDetailCustomization.cpp
+++ b/Engine/Plugins/Runtime/Metasound/Source/MetasoundEditor/Private/MetasoundDetailCustomization.cpp
@@ -75,22 +75,22 @@ namespace Metasound::Editor
 	{
 		using namespace Engine;
 		Builder.Reset(&FDocumentBuilderRegistry::GetChecked().FindOrBeginBuilding(MetaSound));
+	}
 
-		IsGraphEditableAttribute = TAttribute<bool>::Create([this]()
-		{
-			using namespace Engine;
+	bool FMetaSoundDetailCustomizationBase::IsGraphEditable() const
+	{
+		using namespace Engine;
 
-			if (Builder.IsValid())
+		if (Builder.IsValid())
+		{
+			const FMetaSoundFrontendDocumentBuilder& DocBuilder = Builder->GetConstBuilder();
+			if (DocBuilder.IsValid())
 			{
-				const FMetaSoundFrontendDocumentBuilder& DocBuilder = Builder->GetConstBuilder();
-				if (DocBuilder.IsValid())
-				{
-					return DocBuilder.FindConstBuildGraphChecked().Style.bIsGraphEditable;
-				}
+				return DocBuilder.FindConstBuildGraphChecked().Style.bIsGraphEditable;
 			}
+		}
 
-			return false;
-		});
+		return false;
 	}
 
 	FMetasoundDetailCustomization::FMetasoundDetailCustomization(FName InDocumentPropertyName)
@@ -208,7 +208,7 @@ namespace Metasound::Editor
 					}
 
 					TSharedRef<SWidget> OutputFormatValueWidget = OutputFormat->CreatePropertyValueWidget();
-					OutputFormatValueWidget->SetEnabled(IsGraphEditableAttribute);
+					OutputFormatValueWidget->SetEnabled(TAttribute<bool>::Create(TAttribute<bool>::FGetter::CreateSP(this, &FMetaSoundDetailCustomizationBase::IsGraphEditable)));
 
 					static const FText OutputFormatName = LOCTEXT("MetasoundOutputFormatPropertyName", "Output Format");
 					GeneralCategoryBuilder.AddCustomRow(OutputFormatName)
diff --git a/Engine/Plugins/Runtime/Metasound/Source/MetasoundEditor/Private/MetasoundDetailCustomization.h b/Engine/Plugins/Runtime/Metasound/Source/MetasoundEditor/Private/MetasoundDetailCustomization.h
--- a/Engine/Plugins/Runtime/Metasound/Source/MetasoundEditor/Private/MetasoundDetailCustomization.h
+++ b/Engine/Plugins/Runtime/Metasound/Source/MetasoundEditor/Private/MetasoundDetailCustomization.h
@@ -24,11 +24,12 @@ namespace Metasound
 		public:
 			virtual ~FMetaSoundDetailCustomizationBase() = default;
 
+			bool IsGraphEditable() const;
+
 		protected:
 			UObject* GetMetaSound() const;
 			void InitBuilder(UObject& MetaSound);
 
-			TAttribute<bool> IsGraphEditableAttribute;
 			TStrongObjectPtr<UMetaSoundBuilderBase> Builder;
 		};
 
