Recurse MetaSound Preset Graph Inflation as much as possible.
#rnx

diff --git a/Engine/Plugins/Runtime/Metasound/Source/MetasoundEngine/Private/MetasoundOperatorCacheSubsystem.cpp b/Engine/Plugins/Runtime/Metasound/Source/MetasoundEngine/Private/MetasoundOperatorCacheSubsystem.cpp
--- a/Engine/Plugins/Runtime/Metasound/Source/MetasoundEngine/Private/MetasoundOperatorCacheSubsystem.cpp
+++ b/Engine/Plugins/Runtime/Metasound/Source/MetasoundEngine/Private/MetasoundOperatorCacheSubsystem.cpp
@@ -116,9 +116,13 @@ void UMetaSoundCacheSubsystem::PrecacheMetaSoundInternal(UMetaSoundSource* InMet
 		return;
 	}
 
+	// Graph inflation may interact with cache. Need to find the same graph registry key
+	// that is found when a MetaSound generator is created. 
+	const UMetaSoundSource&	NoninflatableSource = InMetaSound->FindFirstNoninflatableSource(InitParams->Environment, [](const UMetaSoundSource&){});
+
 	TUniquePtr<FOperatorBuildData> Data = MakeUnique<FOperatorBuildData>(
 		  MoveTemp(InitParams.GetValue())
-		, InMetaSound->GetGraphRegistryKey()
+		, NoninflatableSource.GetGraphRegistryKey()
 		, InMetaSound->AssetClassID
 		, InNumInstances
 		, bTouchExisting
@@ -157,6 +161,6 @@ void UMetaSoundCacheSubsystem::RemoveCachedOperatorsForMetaSound(UMetaSoundSourc
 
 	if (TSharedPtr<FOperatorPool> OperatorPool = Module->GetOperatorPool())
 	{
-		OperatorPool->RemoveOperatorsWithID(InMetaSound->AssetClassID);
+		OperatorPool->RemoveOperatorsWithAssetClassID(InMetaSound->AssetClassID);
 	}
-}
\ No newline at end of file
+}
diff --git a/Engine/Plugins/Runtime/Metasound/Source/MetasoundEngine/Private/MetasoundSource.cpp b/Engine/Plugins/Runtime/Metasound/Source/MetasoundEngine/Private/MetasoundSource.cpp
--- a/Engine/Plugins/Runtime/Metasound/Source/MetasoundEngine/Private/MetasoundSource.cpp
+++ b/Engine/Plugins/Runtime/Metasound/Source/MetasoundEngine/Private/MetasoundSource.cpp
@@ -1,4 +1,5 @@
 // Copyright Epic Games, Inc. All Rights Reserved.
+
 #include "MetasoundSource.h"
 
 #include "Algo/AnyOf.h"
@@ -330,6 +331,74 @@ FAutoConsoleVariableRef CVarMetaSoundEnableExperimentalRUntimePresetGraphInflati
 	TEXT("Default: false"),
 	ECVF_Default);
 
+// FAudioParameterCollector is used during preset graph inflation to collect audio parameters
+// overriden by preset graphs. This collection process is needed because preset graph inflation
+// circumvents the normal storage location of default inputs (They are stored on the registered
+// IGraph). 
+class UMetaSoundSource::FAudioParameterCollector
+{
+	// Utilitiy for sorting parameters by ParamName
+	struct FSortPredicate
+	{
+		FORCEINLINE bool operator()(const FAudioParameter& InLHS, const FAudioParameter& InRHS) const
+		{
+			return InLHS.ParamName.FastLess(InRHS.ParamName);
+		}
+	};
+
+	// Utility for accessing the ParamName from a audio parameter
+	struct FProjection
+	{
+		FORCEINLINE FName operator()(const FAudioParameter& InParam) const
+		{
+			return InParam.ParamName;
+		}
+	};
+
+public:
+	
+	// FAudioParameterCollection merges new parameter into the provided array. The array is manipulated
+	// in-place and so care must be taken that the `FAudioParameterCollector` does not attempt to access
+	// invalid parameter arrays. 
+	FAudioParameterCollector(TArray<FAudioParameter>& InOutParameters)
+	: Parameters(InOutParameters)
+	{
+		Algo::Sort(Parameters, FSortPredicate());
+	}
+	
+	// Do not allow the FAudioParameterCollector to escape scope in order to prevent accidental
+	// invalid access to Parameter array reference.
+	FAudioParameterCollector(const FAudioParameterCollector&) = delete;
+	FAudioParameterCollector(FAudioParameterCollector&&) = delete;
+	FAudioParameterCollector& operator=(const FAudioParameterCollector&) = delete;
+	FAudioParameterCollector& operator=(FAudioParameterCollector&&) = delete;
+
+	// Merge in parameters from this preset.
+	void CollectPresetOverrides(const TSet<Metasound::FVertexName>& InInputsInheritingDefault, const Metasound::TSortedVertexNameMap<FRuntimeInput>& InInputMap)
+	{
+		using namespace Metasound;
+
+		for (const TPair<FVertexName, FRuntimeInput>& Pair : InInputMap)
+		{
+			// Avoid adding parameters which already exist in the ParameterArray
+			int32 InsertPos = Algo::LowerBoundBy(Parameters, Pair.Key, FProjection(), FSortPredicate());
+			const bool bDoesNotAlreadyExist = (InsertPos >= Parameters.Num()) || (Pair.Key != Parameters[InsertPos].ParamName);
+
+			if (bDoesNotAlreadyExist)
+			{
+				// Avoid adding parameters which are NOT overriding a default. 
+				if (!InInputsInheritingDefault.Contains(Pair.Key))
+				{
+					Parameters.Insert(Pair.Value.DefaultParameter, InsertPos);
+				}
+			}
+		}
+	}
+
+private:
+
+	TArray<FAudioParameter>& Parameters;
+};
 
 UMetaSoundSource::UMetaSoundSource(const FObjectInitializer& ObjectInitializer)
 	: Super(ObjectInitializer)
@@ -997,46 +1066,12 @@ ISoundGeneratorPtr UMetaSoundSource::CreateSoundGenerator(const FSoundGeneratorI
 		// us to retrieve that specific graph. We also supply the parameters that were overridden
 		// in the preset to the FMetaSoundGenerator, because they are not backed into
 		// the base IGraph. 
-		if (ConsoleVariables::bEnableExperimentalRuntimePresetGraphInflation && bIsPresetGraphInflationSupported)
-		{
-			// Get the graph associated with base graph which this preset wraps.
-			TSharedPtr<const IGraph> MetasoundGraph = TryGetMetaSoundPresetBaseGraph();
-
-			if (MetasoundGraph.IsValid())
-			{
-				// Combine the supplied parameters with the overridden parameters set on the preset.
-				TArray<FAudioParameter> MergedParameters;
-				MergePresetOverridesAndSuppliedDefaults(InDefaultParameters, MergedParameters);
-
-				// Update Graph Hierarchy with the asset unique ID (the hierarchy was collapsed by only using the base graph)
-				Environment.SetValue<TArray<FGuid>>(CoreInterface::Environment::GraphHierarchy, TArray<FGuid>({ AssetClassID }));
-
-				// Create generator.
-				FMetasoundGeneratorInitParams InitParams
-				{
-					InSettings,
-					MoveTemp(BuilderSettings),
-					MetasoundGraph,
-					Environment,
-					GetName(),
-					GetOutputAudioChannelOrder(),
-					MoveTemp(MergedParameters),
-					bBuildSynchronous,
-					DataChannel
-				};
-
-				Generator = MakeShared<FMetasoundConstGraphGenerator>(MoveTemp(InitParams));
-			}
-		}
+		FAudioParameterCollector AudioParameterCollector(InDefaultParameters);
+		TSharedPtr<const IGraph> MetasoundGraph = FindFirstNoninflatableGraph(AudioParameterCollector, Environment);
 
-		if (!Generator.IsValid())
+		if (MetasoundGraph.IsValid())
 		{
-			TSharedPtr<const FGraph> MetasoundGraph = FMetasoundFrontendRegistryContainer::Get()->GetGraph(GetGraphRegistryKey());
-			if (!MetasoundGraph.IsValid())
-			{
-				return ISoundGeneratorPtr(nullptr);
-			}
-
+			// Create generator.
 			FMetasoundGeneratorInitParams InitParams
 			{
 				InSettings,
@@ -1777,75 +1812,68 @@ void UMetaSoundSource::InvalidateCachedRuntimeInputData()
 	RuntimeInputData.bIsValid.store(false);
 }
 
-TSharedPtr<const Metasound::IGraph> UMetaSoundSource::TryGetMetaSoundPresetBaseGraph() const
+const UMetaSoundSource& UMetaSoundSource::FindFirstNoninflatableSource(Metasound::FMetasoundEnvironment& InOutEnvironment, TFunctionRef<void(const UMetaSoundSource&)> OnTraversal) const
 {
-	using namespace Metasound;
+	TArray<FGuid> GraphHierarchy;
 
-	TSharedPtr<const IGraph> MetasoundGraph;
-	if (ReferencedAssetClassObjects.Num() == 1)
-	{
-		// Get first element from TSet<>
-		TObjectPtr<const UObject> BaseGraph = *ReferencedAssetClassObjects.CreateConstIterator();
+	const UMetaSoundSource& FoundSource = FindFirstNoninflatableSourceInternal(GraphHierarchy, OnTraversal);
 
-		// Get the reference graph as a UMetaSoundSource
-		TObjectPtr<const UMetaSoundSource> BaseMetaSoundSource = Cast<const UMetaSoundSource>(BaseGraph);
-		if (BaseMetaSoundSource)
-		{
-			// Get registered graph of base metasound source.
-			MetasoundGraph = FMetasoundFrontendRegistryContainer::Get()->GetGraph(BaseMetaSoundSource->GetGraphRegistryKey());
-		}
-	}
-	else
+	if (GraphHierarchy.Num())
 	{
-		UE_LOG(LogMetaSound, Warning, TEXT("Attempt to reference parent of metasound preset failed due to unexpected number of reference asses (%d) from MetaSound Preset %s"), ReferencedAssetClassObjects.Num(), *GetOwningAssetName());
+		// Preset graph inflation needs to emulate the graph hierarchy for nodes which depend upon accurate graph hierarchies. 
+		InOutEnvironment.SetValue<TArray<FGuid>>(Metasound::CoreInterface::Environment::GraphHierarchy, MoveTemp(GraphHierarchy));
 	}
 
-	return MetasoundGraph;
+	return FoundSource;
 }
 
-void UMetaSoundSource::MergePresetOverridesAndSuppliedDefaults(const TArray<FAudioParameter>& InSuppliedDefaults, TArray<FAudioParameter>& OutMerged)
+const UMetaSoundSource& UMetaSoundSource::FindFirstNoninflatableSourceInternal(TArray<FGuid>& OutHierarchy, TFunctionRef<void(const UMetaSoundSource&)> OnTraversal) const
 {
 	using namespace Metasound;
 
-	if(!RuntimeInputData.bIsValid.load())
-	{
-		UE_CLOG(SourcePrivate::HasNotBeenLoggedForThisObject(*this, __LINE__), LogMetaSound, Warning, TEXT("Failed to use experimental preset inflation due to invalid runtime data on MetaSound %s. Ensure that UMetaSoundSource::InitResources() is executed on the game thread before CreateSoundGenerator"), *GetOwningAssetName());
-	}
-	else
+	const bool bIsDynamic = DynamicTransactor.IsValid();
+	if (!bIsDynamic && ConsoleVariables::bEnableExperimentalRuntimePresetGraphInflation && bIsPresetGraphInflationSupported)
 	{
-		// Get parameters that are overridden in the Preset.
-		for (const TPair<FVertexName, FRuntimeInput>& Pair : RuntimeInputData.InputMap)
+		if (ReferencedAssetClassObjects.Num() == 1)
 		{
-			if (!RootMetasoundDocument.RootGraph.PresetOptions.InputsInheritingDefault.Contains(Pair.Key))
-			{
-				OutMerged.Add(Pair.Value.DefaultParameter);
-			}
-		}
+			// Get first element from TSet<>
+			TObjectPtr<const UObject> BaseGraph = *ReferencedAssetClassObjects.CreateConstIterator();
 
-		// Instead of using `FAudioParameter::Merge(...)` we roll custom merge logic because
-		// FAudioParameter::Merge(...) requires us to move the InSuppliedDefaults into the
-		// OutMerged array. This unwanted in this case because we want to maintain InSuppliedDefaults
-		// for the scenario that the experimental preset override fails to create a generator
-		// in which case we are forced to use the prior codepath which does not merge these
-		// two parameter arrays. 
-		for (const FAudioParameter& SuppliedParameter : InSuppliedDefaults)
-		{
-			auto HasSameName = [&Name=SuppliedParameter.ParamName](const FAudioParameter& InOtherParam) -> bool
+			// Get the reference graph as a UMetaSoundSource
+			TObjectPtr<const UMetaSoundSource> BaseMetaSoundSource = Cast<const UMetaSoundSource>(BaseGraph);
+			if (ensure(BaseMetaSoundSource)) // SourcePresets assume they are referencing a UMetaSoundSource
 			{
-				return InOtherParam.ParamName == Name;
-			};
+				// Preset graph inflation needs to emulate the graph hierarchy for nodes which depend upon accurate graph hierarchies. 
+				OutHierarchy.Add(AssetClassID);
 
-			// The supplied defaults generally come from BP and should override any other parameters.
-			if (FAudioParameter* Param = OutMerged.FindByPredicate(HasSameName))
-			{
-				*Param = SuppliedParameter;
-			}
-			else
-			{
-				OutMerged.Add(SuppliedParameter);
+				OnTraversal(*this);
+
+				// If the base metasound is also a preset that can be inflate, recurse into it. 
+				return BaseMetaSoundSource->FindFirstNoninflatableSourceInternal(OutHierarchy, OnTraversal);
 			}
 		}
+		else
+		{
+			UE_LOG(LogMetaSound, Warning, TEXT("Attempt to reference parent of metasound preset failed due to unexpected number of reference asses (%d) from MetaSound Preset %s"), ReferencedAssetClassObjects.Num(), *GetOwningAssetName());
+		}
 	}
+
+	return *this;
+}
+
+TSharedPtr<const Metasound::IGraph> UMetaSoundSource::FindFirstNoninflatableGraph(UMetaSoundSource::FAudioParameterCollector& InOutParameterCollector, Metasound::FMetasoundEnvironment& InOutEnvironment) const
+{
+	using namespace Metasound;
+
+	auto OnGraphInflation = [&InOutParameterCollector](const UMetaSoundSource& InInflatedSource)
+	{
+		// Any preset overrides on this object need to be baked in to the parameters
+		InOutParameterCollector.CollectPresetOverrides(InInflatedSource.RootMetasoundDocument.RootGraph.PresetOptions.InputsInheritingDefault, InInflatedSource.RuntimeInputData.InputMap);
+	};
+
+	const UMetaSoundSource& NoninflatableSource = FindFirstNoninflatableSource(InOutEnvironment, OnGraphInflation);
+
+	return FMetasoundFrontendRegistryContainer::Get()->GetGraph(NoninflatableSource.GetGraphRegistryKey());
 }
 
 #undef LOCTEXT_NAMESPACE // MetaSound
diff --git a/Engine/Plugins/Runtime/Metasound/Source/MetasoundEngine/Public/MetasoundSource.h b/Engine/Plugins/Runtime/Metasound/Source/MetasoundEngine/Public/MetasoundSource.h
--- a/Engine/Plugins/Runtime/Metasound/Source/MetasoundEngine/Public/MetasoundSource.h
+++ b/Engine/Plugins/Runtime/Metasound/Source/MetasoundEngine/Public/MetasoundSource.h
@@ -75,6 +75,9 @@ class METASOUNDENGINE_API UMetaSoundSource : public USoundWaveProcedural, public
 
 	friend struct Metasound::Engine::FAssetHelper;
 	friend class UMetaSoundSourceBuilder;
+	
+	//Forward declare
+	class FAudioParameterCollector;
 
 	// FRuntimeInput represents an input to a MetaSound which can be manipulated.
 	struct FRuntimeInput
@@ -308,9 +311,20 @@ public:
 	Metasound::FMetasoundEnvironment CreateEnvironment(const FSoundGeneratorInitParams& InParams) const;
 	const TArray<Metasound::FVertexName>& GetOutputAudioChannelOrder() const;
 
+	/** Find the Source related to this Preset.
+	 * 
+	 * If this MetaSound is a preset and preset graph inflation is enabled, this
+	 * will traverse the MetaSound Preset hierarchy until a UMetaSoundSource is 
+	 * found which is either 
+	 *  	- Not a preset
+	 * 		AND/OR
+	 *  	- Has modified constructor pin overrides.
+	 */
+	const UMetaSoundSource& FindFirstNoninflatableSource(Metasound::FMetasoundEnvironment& InOutEnvironment, TFunctionRef<void(const UMetaSoundSource&)> OnTraversal) const;
+
 private:
-	TSharedPtr<const Metasound::IGraph> TryGetMetaSoundPresetBaseGraph() const;
-	void MergePresetOverridesAndSuppliedDefaults(const TArray<FAudioParameter>& InSuppliedDefaults, TArray<FAudioParameter>& OutMerged);
+	const UMetaSoundSource& FindFirstNoninflatableSourceInternal(TArray<FGuid>& OutHierarchy, TFunctionRef<void(const UMetaSoundSource&)> OnTraversal) const;
+	TSharedPtr<const Metasound::IGraph> FindFirstNoninflatableGraph(UMetaSoundSource::FAudioParameterCollector& InOutParameterCollector, Metasound::FMetasoundEnvironment& InOutEnvironment) const;
 	
 	Metasound::FMetasoundEnvironment CreateEnvironment() const;
 	Metasound::FMetasoundEnvironment CreateEnvironment(const Audio::FParameterTransmitterInitParams& InParams) const;
diff --git a/Engine/Plugins/Runtime/Metasound/Source/MetasoundGenerator/Private/MetasoundGenerator.cpp b/Engine/Plugins/Runtime/Metasound/Source/MetasoundGenerator/Private/MetasoundGenerator.cpp
--- a/Engine/Plugins/Runtime/Metasound/Source/MetasoundGenerator/Private/MetasoundGenerator.cpp
+++ b/Engine/Plugins/Runtime/Metasound/Source/MetasoundGenerator/Private/MetasoundGenerator.cpp
@@ -491,6 +491,7 @@ PRAGMA_ENABLE_DEPRECATION_WARNINGS
 		GraphOutputAudio.Reset();
 		OnFinishedTriggerRef = TDataWriteReference<FTrigger>::CreateNew(OperatorSettings);
 		GraphAnalyzer.Reset();
+		OutputAnalyzers.Reset();
 		ParameterSetters.Reset();
 		ParameterPackSetters.Reset();
 
@@ -924,6 +925,7 @@ PRAGMA_ENABLE_DEPRECATION_WARNINGS
 
 		if (bUseOperatorPool)
 		{
+			UpdateGraphIfPending();
 			ReleaseOperatorToCache();
 		}
 	}
@@ -938,11 +940,11 @@ PRAGMA_ENABLE_DEPRECATION_WARNINGS
 		if (OperatorPool.IsValid())
 		{
 			// See if the cache has an operator with matching OperatorID
-			OperatorID = InInitParams.Graph->GetInstanceID();
-			FOperatorAndInputs GraphOperatorAndInputs = OperatorPool->ClaimOperator(OperatorID);
+			OperatorPoolID = FOperatorPoolEntryID{InInitParams.Graph->GetInstanceID(), InInitParams.OperatorSettings};
+			FOperatorAndInputs GraphOperatorAndInputs = OperatorPool->ClaimOperator(*OperatorPoolID);
 			if (GraphOperatorAndInputs.Operator.IsValid())
 			{
-				UE_LOG(LogMetasoundGenerator, VeryVerbose, TEXT("Using cached operator %s for MetaSound %s"), *LexToString(OperatorID), *InInitParams.MetaSoundName);
+				UE_LOG(LogMetasoundGenerator, VeryVerbose, TEXT("Using cached operator %s for MetaSound %s"), *OperatorPoolID->ToString(), *InInitParams.MetaSoundName);
 				bUseOperatorPool = true; // raise this flag to make sure we put the operator back in the pool (regardless of CVAR state)
 
 				// Apply and default inputs to the operator.
@@ -973,16 +975,17 @@ PRAGMA_ENABLE_DEPRECATION_WARNINGS
 		using namespace MetasoundGeneratorPrivate;
 		METASOUND_TRACE_CPUPROFILER_EVENT_SCOPE(FMetasoundConstGraphGenerator::ReleaseOperatorToCache);
 
-		if (OperatorID.IsValid())
+		if (OperatorPoolID.IsSet())
 		{
 			FMetasoundGeneratorModule& Module = FModuleManager::GetModuleChecked<FMetasoundGeneratorModule>("MetasoundGenerator");
 			TSharedPtr<FOperatorPool> OperatorPool = Module.GetOperatorPool();
 			TUniquePtr<IOperator> GraphOperator = ReleaseGraphOperator();
+			FInputVertexInterfaceData InputData = ReleaseInputVertexData();
 
 			if (OperatorPool.IsValid() && GraphOperator.IsValid())
 			{
 				// Release graph operator and input data to the cache
-				UE_LOG(LogMetasoundGenerator, VeryVerbose, TEXT("Caching operator %s"), *LexToString(OperatorID));
+				UE_LOG(LogMetasoundGenerator, VeryVerbose, TEXT("Caching operator %s"), *OperatorPoolID->ToString());
 
 				// give the operator a chance to reduce its memory footprint before being cached
 				// in the future this should be a conanical phase of an operator's lifecycle (i.e. Reset, Execute, Hybernate)
@@ -995,13 +998,18 @@ PRAGMA_ENABLE_DEPRECATION_WARNINGS
 					}
 				}
 
-				OperatorPool->AddOperator(OperatorID, MoveTemp(GraphOperator), ReleaseInputVertexData());
+				// Clear out any internal references to graph data before adding to
+				// operator pool. MetaSound data references are not thread safe which
+				// requires that data references be removed before transferring 
+				// data references to the operator pool. The operator pool is accessed
+				// from multiple threads and so failure to remove internal data
+				// references may result in a race condition on deallocating references.
+				ClearGraph();
+
+				OperatorPool->AddOperator(*OperatorPoolID, MoveTemp(GraphOperator), MoveTemp(InputData));
 
 			}
 		}
-
-		// Clear out any internal references to graph data
-		ClearGraph();
 	}
 
 	void FMetasoundConstGraphGenerator::BuildGraph(FMetasoundGeneratorInitParams&& InInitParams)
diff --git a/Engine/Plugins/Runtime/Metasound/Source/MetasoundGenerator/Private/MetasoundOperatorCache.cpp b/Engine/Plugins/Runtime/Metasound/Source/MetasoundGenerator/Private/MetasoundOperatorCache.cpp
--- a/Engine/Plugins/Runtime/Metasound/Source/MetasoundGenerator/Private/MetasoundOperatorCache.cpp
+++ b/Engine/Plugins/Runtime/Metasound/Source/MetasoundGenerator/Private/MetasoundOperatorCache.cpp
@@ -27,6 +27,9 @@ namespace Metasound
 
 	namespace OperatorPoolPrivate
 	{
+		constexpr int32 DefaultSampleRateForDeprecatedAPI = 48000;
+		constexpr float BlockRate = 100.f;
+
 		static bool bMetasoundPoolSyncGraphRetrieval = true;
 		FAutoConsoleVariableRef CVarMetasoundPoolSyncGraphRetrieval(
 			TEXT("au.MetaSound.OperatorPoolSyncGraphRetrieval"),
@@ -175,6 +178,38 @@ namespace Metasound
 	{
 	}
 
+	FOperatorPoolEntryID::FOperatorPoolEntryID(FGuid InOperatorID, FOperatorSettings InSettings)
+	: OperatorID(MoveTemp(InOperatorID))
+	, OperatorSettings(MoveTemp(InSettings))
+	{
+	}
+
+	FString FOperatorPoolEntryID::ToString() const
+	{
+		return FString::Printf(TEXT("%s %s"), *OperatorID.ToString(), *OperatorSettings.ToString());
+	}
+
+	bool operator<(const FOperatorPoolEntryID& InLHS, const FOperatorPoolEntryID& InRHS)
+	{
+		if (InLHS.OperatorID < InRHS.OperatorID)
+		{
+			return true;
+		}
+		else if (InRHS.OperatorID < InLHS.OperatorID)
+		{
+			return false;
+		}
+		else 
+		{
+			return InLHS.OperatorSettings < InRHS.OperatorSettings;
+		}
+	}
+
+	bool operator==(const FOperatorPoolEntryID& InLHS, const FOperatorPoolEntryID& InRHS)
+	{
+		return (InLHS.OperatorID == InRHS.OperatorID) && (InLHS.OperatorSettings == InRHS.OperatorSettings);
+	}
+
 
 	FOperatorPool::FOperatorPool(const FOperatorPoolSettings& InSettings)
 	: Settings(InSettings)
@@ -188,6 +223,11 @@ namespace Metasound
 	}
 
 	FOperatorAndInputs FOperatorPool::ClaimOperator(const FGuid& InOperatorID)
+	{
+		return ClaimOperator(FOperatorPoolEntryID{InOperatorID, FOperatorSettings{OperatorPoolPrivate::DefaultSampleRateForDeprecatedAPI, OperatorPoolPrivate::BlockRate}});
+	}
+
+	FOperatorAndInputs FOperatorPool::ClaimOperator(const FOperatorPoolEntryID& InOperatorID)
 	{
 		FOperatorAndInputs OpAndInputs;
 
@@ -215,12 +255,19 @@ namespace Metasound
 			OperatorPoolPrivate::CacheAttemptCount++;
 			TRACE_COUNTER_SET(MetaSound_OperatorPool_HitRatio, OperatorPoolPrivate::GetHitRatio());
 #endif
+			UE_LOG(LogMetasoundGenerator, VeryVerbose, TEXT("Attempt to claim operator with ID %s from operator pool resulted in %s"), *InOperatorID.ToString(), *::LexToString(bCacheHit));
 		}
 
 		return OpAndInputs;
 	}
 
-	void FOperatorPool::AddOperator(const FGuid& InOperatorID, TUniquePtr<IOperator>&& InOperator, FInputVertexInterfaceData&& InInputData)
+	void FOperatorPool::AddOperator(const FGuid& InOperatorID, TUniquePtr<IOperator>&& InOperator, FInputVertexInterfaceData&& InputData)
+	{
+		const FOperatorPoolEntryID EntryID{InOperatorID, FOperatorSettings{OperatorPoolPrivate::DefaultSampleRateForDeprecatedAPI, OperatorPoolPrivate::BlockRate}};
+		AddOperator(EntryID, MoveTemp(InOperator), MoveTemp(InputData));
+	}
+
+	void FOperatorPool::AddOperator(const FOperatorPoolEntryID& InOperatorID, TUniquePtr<IOperator>&& InOperator, FInputVertexInterfaceData&& InInputData)
 	{
 		AddOperator(InOperatorID, { MoveTemp(InOperator), MoveTemp(InInputData) });
 	}
@@ -256,6 +303,12 @@ namespace Metasound
 	}
 
 	void FOperatorPool::AddOperator(const FGuid& InOperatorID, FOperatorAndInputs&& OperatorAndInputs)
+	{
+		const FOperatorPoolEntryID EntryID{InOperatorID, FOperatorSettings{OperatorPoolPrivate::DefaultSampleRateForDeprecatedAPI, OperatorPoolPrivate::BlockRate}};
+		AddOperator(EntryID, MoveTemp(OperatorAndInputs));
+	}
+
+	void FOperatorPool::AddOperator(const FOperatorPoolEntryID& InOperatorID, FOperatorAndInputs&& OperatorAndInputs)
 	{
 		using namespace UE::Tasks;
 
@@ -270,7 +323,7 @@ namespace Metasound
 		});
 	}	
 
-	void FOperatorPool::AddOperatorInternal(const FGuid& InOperatorID, FOperatorAndInputs && OperatorAndInputs)
+	void FOperatorPool::AddOperatorInternal(const FOperatorPoolEntryID& InOperatorID, FOperatorAndInputs && OperatorAndInputs)
 	{
 		if (!OperatorAndInputs.Operator.IsValid())
 		{
@@ -282,6 +335,8 @@ namespace Metasound
 #if METASOUND_OPERATORCACHEPROFILER_ENABLED
 		TRACE_COUNTER_INCREMENT(MetaSound_OperatorPool_NumOperators);
 #endif // #if METASOUND_OPERATORCACHEPROFILER_ENABLED
+	   
+		UE_LOG(LogMetasoundGenerator, VeryVerbose, TEXT("Adding operator with ID %s to operator pool"), *InOperatorID.ToString());
 
 		if (TArray<FOperatorAndInputs>* OperatorArray = Operators.Find(InOperatorID))
 		{
@@ -408,8 +463,8 @@ namespace Metasound
 				if (OperatorPool.IsValid())
 				{
 					// Get the number of instances already in the cache & move pre-existing to the top of the cache
-					const int32 NumInCache = OperatorPool->GetNumCachedOperatorsWithID(PreCacheData->AssetClassID);
-					OperatorPool->TouchOperators(PreCacheData->AssetClassID, FMath::Min(NumInCache, NumToBuild));
+					const int32 NumInCache = OperatorPool->GetNumCachedOperatorsWithAssetClassID(PreCacheData->AssetClassID);
+					OperatorPool->TouchOperatorsViaAssetClassID(PreCacheData->AssetClassID, FMath::Min(NumInCache, NumToBuild));
 					NumToBuild -= NumInCache;
 				}
 			}
@@ -430,19 +485,28 @@ namespace Metasound
 				FOperatorAndInputs OperatorAndInputs = GeneratorBuilder::BuildGraphOperator(PreCacheData->InitParams.OperatorSettings, PreCacheData->InitParams, BuildResults);
 				GeneratorBuilder::LogBuildErrors(PreCacheData->InitParams.MetaSoundName, BuildResults);
 
-				const FGuid& GraphID = PreCacheData->InitParams.Graph->GetInstanceID();
-				OperatorPool->AddOperatorInternal(GraphID, MoveTemp(OperatorAndInputs));
+				FOperatorPoolEntryID EntryID{PreCacheData->InitParams.Graph->GetInstanceID(), PreCacheData->InitParams.OperatorSettings};
+				OperatorPool->AddOperatorInternal(EntryID, MoveTemp(OperatorAndInputs));
+				OperatorPool->AddAssetIdToGraphIdLookUpInternal(PreCacheData->AssetClassID, EntryID);
 			}
 		});
 	}
 
 	void FOperatorPool::TouchOperators(const FGuid& InOperatorID, int32 NumToTouch)
+	{
+		const FOperatorPoolEntryID EntryID{InOperatorID, FOperatorSettings{OperatorPoolPrivate::DefaultSampleRateForDeprecatedAPI, OperatorPoolPrivate::BlockRate}};
+		TouchOperators(EntryID, NumToTouch);
+	}
+
+	void FOperatorPool::TouchOperators(const FOperatorPoolEntryID& InOperatorID, int32 NumToTouch)
 	{
 		if (!IsStopping())
 		{
 			FScopeLock Lock(&CriticalSection);
 
+
 			const int32 NumToMove = FMath::Min(NumToTouch, GetNumCachedOperatorsWithID(InOperatorID));
+			UE_LOG(LogMetasoundGenerator, VeryVerbose, TEXT("Touching operator %d operators with ID %s in operator pool"), NumToMove, *InOperatorID.ToString());
 			if (!NumToMove)
 			{
 				return;
@@ -464,10 +528,24 @@ namespace Metasound
 
 	void FOperatorPool::TouchOperatorsViaAssetClassID(const FGuid& InAssetClassID, int32 NumToTouch)
 	{
-		TouchOperators(InAssetClassID, NumToTouch);
+		if (!IsStopping())
+		{
+			FScopeLock Lock(&CriticalSection);
+			FOperatorPoolEntryID* GraphIdPtr = AssetIdToGraphIdLookUp.Find(InAssetClassID);
+			if (GraphIdPtr)
+			{
+				TouchOperators(*GraphIdPtr, NumToTouch);
+			}
+		}
 	}
 
 	void FOperatorPool::RemoveOperatorsWithID(const FGuid& InOperatorID)
+	{
+		const FOperatorPoolEntryID EntryID{InOperatorID, FOperatorSettings{OperatorPoolPrivate::DefaultSampleRateForDeprecatedAPI, OperatorPoolPrivate::BlockRate}};
+		RemoveOperatorsWithID(EntryID);
+	}
+
+	void FOperatorPool::RemoveOperatorsWithID(const FOperatorPoolEntryID& InOperatorID)
 	{
 		using namespace UE::Tasks;
 
@@ -480,6 +558,7 @@ namespace Metasound
 		{
 			if (TSharedPtr<FOperatorPool> Pool = WeakPoolPtr.Pin())
 			{
+				UE_LOG(LogMetasoundGenerator, VeryVerbose, TEXT("Removing operators with ID %s from operator pool"), *OperatorID.ToString());
 				FScopeLock Lock(&Pool->CriticalSection);
 				Pool->Operators.Remove(OperatorID);
 				const int32 NumRemoved = Pool->Stack.Remove(OperatorID);
@@ -493,7 +572,19 @@ namespace Metasound
 
 	void FOperatorPool::RemoveOperatorsWithAssetClassID(const FGuid& InAssetClassID)
 	{
-		return RemoveOperatorsWithID(InAssetClassID);
+		FScopeLock Lock(&CriticalSection);
+		const FOperatorPoolEntryID* GraphIdPtr = AssetIdToGraphIdLookUp.Find(InAssetClassID);
+		if (GraphIdPtr)
+		{
+			if (GraphIdToAssetIdLookUp.Num(*GraphIdPtr) <= 1)
+			{
+				// Only remove operators if there are no other assets pointing to the same operator
+				RemoveOperatorsWithID(*GraphIdPtr);
+			}
+
+			GraphIdToAssetIdLookUp.Remove(*GraphIdPtr, InAssetClassID);
+			AssetIdToGraphIdLookUp.Remove(InAssetClassID);
+		}
 	}
 		
 	void FOperatorPool::SetMaxNumOperators(uint32 InMaxNumOperators)
@@ -506,7 +597,7 @@ namespace Metasound
 		}
 	}
 
-	int32 FOperatorPool::GetNumCachedOperatorsWithID(const FGuid& InOperatorID) const
+	int32 FOperatorPool::GetNumCachedOperatorsWithID(const FOperatorPoolEntryID& InOperatorID) const
 	{
 		FScopeLock Lock(&CriticalSection);
 		if (TArray<FOperatorAndInputs> const* OperatorsWithID = Operators.Find(InOperatorID))
@@ -519,7 +610,20 @@ namespace Metasound
 
 	int32 FOperatorPool::GetNumCachedOperatorsWithAssetClassID(const FGuid& InAssetClassID) const
 	{
-		return GetNumCachedOperatorsWithID(InAssetClassID);
+		FScopeLock Lock(&CriticalSection);
+		if (const FOperatorPoolEntryID* GraphIdPtr = AssetIdToGraphIdLookUp.Find(InAssetClassID))
+		{
+			return GetNumCachedOperatorsWithID(*GraphIdPtr);
+		}
+
+		return 0;
+	}
+
+	void FOperatorPool::AddAssetIdToGraphIdLookUpInternal(const FGuid& InAssetClassID, const FOperatorPoolEntryID& InOperatorID)
+	{
+		FScopeLock Lock(&CriticalSection);
+		AssetIdToGraphIdLookUp.Add(InAssetClassID, InOperatorID);
+		GraphIdToAssetIdLookUp.AddUnique(InOperatorID, InAssetClassID);
 	}
 
 #if METASOUND_OPERATORCACHEPROFILER_ENABLED
@@ -538,7 +642,7 @@ namespace Metasound
 			UE_LOG(LogMetasoundGenerator, Verbose, TEXT("Trimming %d operators"), NumToTrim);
 			for (int32 i = 0; i < NumToTrim; i++)
 			{
-				UE_LOG(LogMetasoundGenerator, VeryVerbose, TEXT("Trimming operator with ID %s"), *LexToString(Stack[i]));
+				UE_LOG(LogMetasoundGenerator, VeryVerbose, TEXT("Trimming operator with ID %s"), *Stack[i].ToString());
 				TArray<FOperatorAndInputs>* OperatorArray = Operators.Find(Stack[i]);
 				if (ensure(OperatorArray))
 				{
@@ -548,6 +652,15 @@ namespace Metasound
 						if (OperatorArray->Num() == 0)
 						{
 							Operators.Remove(Stack[i]);
+							// Remove related asset IDs
+							TArray<FGuid> RelatedAssetIDs;
+							GraphIdToAssetIdLookUp.MultiFind(Stack[i], RelatedAssetIDs);
+
+							for (const FGuid& AssetID : RelatedAssetIDs)
+							{
+								AssetIdToGraphIdLookUp.Remove(AssetID);
+								GraphIdToAssetIdLookUp.Remove(Stack[i], AssetID);
+							}
 						}
 					}
 				}
diff --git a/Engine/Plugins/Runtime/Metasound/Source/MetasoundGenerator/Public/MetasoundGenerator.h b/Engine/Plugins/Runtime/Metasound/Source/MetasoundGenerator/Public/MetasoundGenerator.h
--- a/Engine/Plugins/Runtime/Metasound/Source/MetasoundGenerator/Public/MetasoundGenerator.h
+++ b/Engine/Plugins/Runtime/Metasound/Source/MetasoundGenerator/Public/MetasoundGenerator.h
@@ -80,6 +80,30 @@ namespace Metasound
 		};
 	}
 
+	/** ID for looking up a operator in the operator pool */
+	struct METASOUNDGENERATOR_API FOperatorPoolEntryID final
+	{
+		/** Construct an ID
+		 *
+		 * InOperatorID - ID of the operator.
+		 * InSettings - Operator settings used to create the operator.
+		 */
+		FOperatorPoolEntryID(FGuid InOperatorID, FOperatorSettings InSettings);
+
+		FString ToString() const;
+
+		METASOUNDGENERATOR_API friend bool operator<(const FOperatorPoolEntryID& InLHS, const FOperatorPoolEntryID& InRHS);
+		METASOUNDGENERATOR_API friend bool operator==(const FOperatorPoolEntryID& InLHS, const FOperatorPoolEntryID& InRHS);
+		friend FORCEINLINE uint32 GetTypeHash(const FOperatorPoolEntryID& InID)
+		{
+			return HashCombineFast(GetTypeHash(InID.OperatorID), GetTypeHash(InID.OperatorSettings));
+		}
+
+	private:
+		FGuid OperatorID;
+		FOperatorSettings OperatorSettings;
+	};
+
 	// Struct needed for building the metasound graph
 	struct METASOUNDGENERATOR_API FMetasoundGeneratorInitParams
 	{
@@ -402,7 +426,7 @@ namespace Metasound
 
 		TUniquePtr<FMetasoundEnvironment> EnvironmentPtr;
 		TUniquePtr<FAsyncTaskBase> BuilderTask;
-		FGuid OperatorID;
+		TOptional<FOperatorPoolEntryID> OperatorPoolID;
 		bool bUseOperatorPool = false;
 	};
 
diff --git a/Engine/Plugins/Runtime/Metasound/Source/MetasoundGenerator/Public/MetasoundOperatorCache.h b/Engine/Plugins/Runtime/Metasound/Source/MetasoundGenerator/Public/MetasoundOperatorCache.h
--- a/Engine/Plugins/Runtime/Metasound/Source/MetasoundGenerator/Public/MetasoundOperatorCache.h
+++ b/Engine/Plugins/Runtime/Metasound/Source/MetasoundGenerator/Public/MetasoundOperatorCache.h
@@ -63,6 +63,7 @@ namespace Metasound
 		uint32 MaxNumOperators = 64;
 	};
 
+
 	// Data required to build an operator without immediately playing it
 	struct METASOUNDGENERATOR_API FOperatorBuildData
 	{
@@ -90,32 +91,44 @@ namespace Metasound
 	class METASOUNDGENERATOR_API FOperatorPool : public TSharedFromThis<FOperatorPool>
 	{
 	public:
+
 		FOperatorPool(const FOperatorPoolSettings& InSettings);
 		~FOperatorPool();
 
+
+		UE_DEPRECATED(5.5, "Use ClaimOperator(const FOperatorPoolEntryID&) instead")
 		FOperatorAndInputs ClaimOperator(const FGuid& InOperatorID);
+		FOperatorAndInputs ClaimOperator(const FOperatorPoolEntryID& InOperatorID);
 
+		UE_DEPRECATED(5.5, "Use AddOperator(const FOperatorPoolEntryID&, ...) instead")
 		void AddOperator(const FGuid& InOperatorID, TUniquePtr<IOperator>&& InOperator, FInputVertexInterfaceData&& InputData);
+		void AddOperator(const FOperatorPoolEntryID& InOperatorID, TUniquePtr<IOperator>&& InOperator, FInputVertexInterfaceData&& InputData);
+
+		UE_DEPRECATED(5.5, "Use AddOperator(const FOperatorPoolEntryID&, ...) instead")
 		void AddOperator(const FGuid& InOperatorID, FOperatorAndInputs&& OperatorAndInputs);
+		void AddOperator(const FOperatorPoolEntryID& InOperatorID, FOperatorAndInputs&& OperatorAndInputs);
 
 		void BuildAndAddOperator(TUniquePtr<FOperatorBuildData> InBuildData);
 
+		UE_DEPRECATED(5.5, "Use TouchOperators(const FOperatorPoolEntryID&, ...) instead")
 		void TouchOperators(const FGuid& InOperatorID, int32 NumToTouch = 1);
-		UE_DEPRECATED(5.5, "Operators can now be directly identified by asset class id.")
+		void TouchOperators(const FOperatorPoolEntryID& InOperatorID, int32 NumToTouch = 1);
 		void TouchOperatorsViaAssetClassID(const FGuid& InAssetClassID, int32 NumToTouch = 1);
 
 		bool IsStopping() const;
 
+		UE_DEPRECATED(5.5, "Use RemoveOperatorsWithID(const FOperatorPoolEntryID&) instead")
 		void RemoveOperatorsWithID(const FGuid& InOperatorID);
-		UE_DEPRECATED(5.5, "Operators can now be directly identified by asset class id.")
+		void RemoveOperatorsWithID(const FOperatorPoolEntryID& InOperatorID);
 		void RemoveOperatorsWithAssetClassID(const FGuid& InAssetClassID);
 
+		UE_DEPRECATED(5.5, "Use GetNumCachedOperatorsWithID(const FOperatorPoolEntryID&) instead")
 		int32 GetNumCachedOperatorsWithID(const FGuid& InOperatorID) const;
-		UE_DEPRECATED(5.5, "Operators can now be directly identified by asset class id.")
+		int32 GetNumCachedOperatorsWithID(const FOperatorPoolEntryID& InOperatorID) const;
 		int32 GetNumCachedOperatorsWithAssetClassID(const FGuid& InAssetClassID) const;
 
-		UE_DEPRECATED(5.5, "Operators can now be directly identified by asset class id, so lookup is not needed anymore.")
-		void AddAssetIdToGraphIdLookUp(const FGuid& InAssetClassID, const FGuid& InOperatorID) {}
+		UE_DEPRECATED(5.5, "Adding id to look-up is now private implementation")
+		void AddAssetIdToGraphIdLookUp(const FGuid& InAssetClassID, const FOperatorPoolEntryID& InOperatorID) { }
 
 		void SetMaxNumOperators(uint32 InMaxNumOperators);
 #if METASOUND_OPERATORCACHEPROFILER_ENABLED
@@ -133,7 +146,8 @@ namespace Metasound
 	private:
 		FTaskId LastTaskId = 0;
 
-		void AddOperatorInternal(const FGuid& InOperatorID, FOperatorAndInputs&& OperatorAndInputs);
+		void AddAssetIdToGraphIdLookUpInternal(const FGuid& InAssetClassID, const FOperatorPoolEntryID& InOperatorID);
+		void AddOperatorInternal(const FOperatorPoolEntryID& InOperatorID, FOperatorAndInputs&& OperatorAndInputs);
 		bool ExecuteTaskAsync(FTaskFunction&& InFunction);
 		void Trim();
 
@@ -151,8 +165,10 @@ namespace Metasound
 		TMap<FTaskId, UE::Tasks::FTask> ActiveBuildTasks;
 		UE::Tasks::FPipe AsyncBuildPipe;
 
-		TMap<FGuid, TArray<FOperatorAndInputs>> Operators;
-		TArray<FGuid> Stack;
+		TMap<FOperatorPoolEntryID, TArray<FOperatorAndInputs>> Operators;
+		TMap<FGuid, FOperatorPoolEntryID> AssetIdToGraphIdLookUp;
+		TMultiMap<FOperatorPoolEntryID, FGuid> GraphIdToAssetIdLookUp;
+		TArray<FOperatorPoolEntryID> Stack;
 	};
 } // namespace Metasound
 
diff --git a/Engine/Plugins/Runtime/Metasound/Source/MetasoundGraphCore/Private/MetasoundOperatorSettings.cpp b/Engine/Plugins/Runtime/Metasound/Source/MetasoundGraphCore/Private/MetasoundOperatorSettings.cpp
--- a/Engine/Plugins/Runtime/Metasound/Source/MetasoundGraphCore/Private/MetasoundOperatorSettings.cpp
+++ b/Engine/Plugins/Runtime/Metasound/Source/MetasoundGraphCore/Private/MetasoundOperatorSettings.cpp
@@ -104,4 +104,34 @@ namespace Metasound
 
 		return AlignedNum;
 	}
+
+	bool operator<(const FOperatorSettings& InLHS, const FOperatorSettings& InRHS)
+	{
+		if (InLHS.SampleRate < InRHS.SampleRate)
+		{
+			return true;
+		}
+		else if (InRHS.SampleRate < InLHS.SampleRate)
+		{
+			return false;
+		}
+
+		// Ignore TargetBlockRate in comparison because it is not used in any external calculations
+		// Ignore ActualBlockRate because it is derived from NumFramesPerBlock
+
+		return InLHS.NumFramesPerBlock < InRHS.NumFramesPerBlock;
+	}
+
+	FString FOperatorSettings::ToString() const
+	{
+		return FString::Printf(TEXT("SampleRate:%d BlockSize:%d"), SampleRate, NumFramesPerBlock);
+	}
+
+	bool operator==(const FOperatorSettings& InLHS, const FOperatorSettings& InRHS)
+	{
+		// Ignore TargetBlockRate in comparison because it is not used in any external calculations
+		// Ignore ActualBlockRate because it is derived from NumFramesPerBlock
+		return (InLHS.SampleRate == InRHS.SampleRate) && (InLHS.NumFramesPerBlock == InRHS.NumFramesPerBlock);
+	}
 }
+
diff --git a/Engine/Plugins/Runtime/Metasound/Source/MetasoundGraphCore/Public/MetasoundOperatorSettings.h b/Engine/Plugins/Runtime/Metasound/Source/MetasoundGraphCore/Public/MetasoundOperatorSettings.h
--- a/Engine/Plugins/Runtime/Metasound/Source/MetasoundGraphCore/Public/MetasoundOperatorSettings.h
+++ b/Engine/Plugins/Runtime/Metasound/Source/MetasoundGraphCore/Public/MetasoundOperatorSettings.h
@@ -84,6 +84,19 @@ namespace Metasound
 				return NumFramesPerBlock;
 			}
 
+			/** Convert to human readable string. */
+			FString ToString() const;
+
+			METASOUNDGRAPHCORE_API friend bool operator<(const FOperatorSettings& InLHS, const FOperatorSettings& InRHS);
+			METASOUNDGRAPHCORE_API friend bool operator==(const FOperatorSettings& InLHS, const FOperatorSettings& InRHS);
+
+			friend FORCEINLINE uint32 GetTypeHash(const FOperatorSettings& InSettings)
+			{
+				// Ignore TargetBlockRate in comparison because it is not used in any external calculations
+				// Ignore ActualBlockRate because it is derived from NumFramesPerBlock
+				return HashCombineFast(GetTypeHash(InSettings.SampleRate), GetTypeHash(InSettings.NumFramesPerBlock));
+			}
+
 		private:
 			// Update actual block rate and frames per block.
 			void Update();
