Const handle fix-ups
#rnx

diff --git a/Engine/Plugins/Runtime/Metasound/Source/MetasoundEditor/Private/MetasoundEditorGraphNode.cpp b/Engine/Plugins/Runtime/Metasound/Source/MetasoundEditor/Private/MetasoundEditorGraphNode.cpp
--- a/Engine/Plugins/Runtime/Metasound/Source/MetasoundEditor/Private/MetasoundEditorGraphNode.cpp
+++ b/Engine/Plugins/Runtime/Metasound/Source/MetasoundEditor/Private/MetasoundEditorGraphNode.cpp
@@ -100,7 +100,7 @@ void UMetasoundEditorGraphNode::SyncCommentFromFrontendNode()
 {
 	using namespace Metasound::Frontend;
 
-	FConstNodeHandle NodeHandle = GetNodeHandle();
+	FConstNodeHandle NodeHandle = GetConstNodeHandle();
 	const FMetasoundFrontendNodeStyle& Style = NodeHandle->GetNodeStyle();
 	NodeComment = Style.Display.Comment;
 	bCommentBubbleMakeVisible = Style.Display.bCommentVisible;
@@ -335,7 +335,7 @@ void UMetasoundEditorGraphNode::CacheTitle()
 	using namespace Metasound;
 	using namespace Metasound::Frontend;
 
-	FConstNodeHandle NodeHandle = GetNodeHandle();
+	FConstNodeHandle NodeHandle = GetConstNodeHandle();
 	CachedTitle = NodeHandle->GetDisplayTitle();
 }
 
@@ -394,7 +394,7 @@ void UMetasoundEditorGraphNode::ReconstructNode()
 	// and subsequent steps clean-up unused pins.  This can be called mid-copy, which means the node
 	// handle may be invalid.  Setting to remove unused causes premature removal and then default values
 	// are lost.
-	FConstNodeHandle NodeHandle = GetNodeHandle();
+	FConstNodeHandle NodeHandle = GetConstNodeHandle();
 	if (NodeHandle->IsValid())
 	{
 		FGraphBuilder::SynchronizeNodePins(*this, NodeHandle, false /* bRemoveUnusedPins */, false /* bLogChanges */);
@@ -678,7 +678,7 @@ void UMetasoundEditorGraphNode::PostEditUndo()
 	// This can trigger and the handle is no longer valid if transaction
 	// is being undone on a graph node that is orphaned.  If orphaned,
 	// bail early.
-	FNodeHandle NodeHandle = GetNodeHandle();
+	FConstNodeHandle NodeHandle = GetConstNodeHandle();
 	if (!NodeHandle->IsValid())
 	{
 		return;
@@ -1552,7 +1552,7 @@ void UMetasoundEditorGraphVariableNode::PinDefaultValueChanged(UEdGraphPin* Pin)
 			if (FGraphBuilder::GetPinLiteral(*Pin, LiteralValue))
 			{
 				// If this is the mutator node, synchronize the variable default literal with this default.
-				FNodeHandle MutatorNode = Variable->GetVariableHandle()->FindMutatorNode();
+				FConstNodeHandle MutatorNode = Variable->GetConstVariableHandle()->FindMutatorNode();
 				if (MutatorNode->IsValid())
 				{
 					if (MutatorNode->GetID() == NodeID)
