#1 Pages Checkpoint: Update BuilderAPI, associated delegates to support providing Pages
#rnx

diff --git a/Engine/Plugins/Runtime/Metasound/Source/MetasoundEngine/Private/MetasoundBuilderSubsystem.cpp b/Engine/Plugins/Runtime/Metasound/Source/MetasoundEngine/Private/MetasoundBuilderSubsystem.cpp
--- a/Engine/Plugins/Runtime/Metasound/Source/MetasoundEngine/Private/MetasoundBuilderSubsystem.cpp
+++ b/Engine/Plugins/Runtime/Metasound/Source/MetasoundEngine/Private/MetasoundBuilderSubsystem.cpp
@@ -258,18 +258,64 @@ void UMetaSoundSourceBuilder::InitDelegates(Metasound::Frontend::FDocumentModify
 {
 	Super::InitDelegates(OutDocumentDelegates);
 
-	OutDocumentDelegates.EdgeDelegates.OnEdgeAdded.AddUObject(this, &UMetaSoundSourceBuilder::OnEdgeAdded);
-	OutDocumentDelegates.EdgeDelegates.OnRemoveSwappingEdge.AddUObject(this, &UMetaSoundSourceBuilder::OnRemoveSwappingEdge);
+	OutDocumentDelegates.PageDelegates.OnPageAdded.AddUObject(this, &UMetaSoundSourceBuilder::OnPageAdded);
+	OutDocumentDelegates.PageDelegates.OnRemovingPage.AddUObject(this, &UMetaSoundSourceBuilder::OnRemovingPage);
 
 	OutDocumentDelegates.InterfaceDelegates.OnInputAdded.AddUObject(this, &UMetaSoundSourceBuilder::OnInputAdded);
 	OutDocumentDelegates.InterfaceDelegates.OnOutputAdded.AddUObject(this, &UMetaSoundSourceBuilder::OnOutputAdded);
 	OutDocumentDelegates.InterfaceDelegates.OnRemovingInput.AddUObject(this, &UMetaSoundSourceBuilder::OnRemovingInput);
 	OutDocumentDelegates.InterfaceDelegates.OnRemovingOutput.AddUObject(this, &UMetaSoundSourceBuilder::OnRemovingOutput);
 
-	OutDocumentDelegates.NodeDelegates.OnNodeAdded.AddUObject(this, &UMetaSoundSourceBuilder::OnNodeAdded);
-	OutDocumentDelegates.NodeDelegates.OnNodeInputLiteralSet.AddUObject(this, &UMetaSoundSourceBuilder::OnNodeInputLiteralSet);
-	OutDocumentDelegates.NodeDelegates.OnRemoveSwappingNode.AddUObject(this, &UMetaSoundSourceBuilder::OnRemoveSwappingNode);
-	OutDocumentDelegates.NodeDelegates.OnRemovingNodeInputLiteral.AddUObject(this, &UMetaSoundSourceBuilder::OnRemovingNodeInputLiteral);
+	InitExecutablePageDelegates(OutDocumentDelegates);
+}
+
+void UMetaSoundSourceBuilder::InitExecutablePageDelegates(Metasound::Frontend::FDocumentModifyDelegates& OutDocumentDelegates)
+{
+	using namespace Metasound::DynamicGraph;
+	using namespace Metasound::Engine;
+	using namespace Metasound::Frontend;
+
+	// If currently executing live audition, must call stop as provided transactions may
+	// get corrupted by the fact that the executable page ID may now resolve to a different value.
+	ExecuteAuditionableTransaction([this](FDynamicOperatorTransactor& Transactor)
+	{
+		bool bComponentsStopped = false;
+		for (uint64 AudioComponentID : LiveComponentIDs)
+		{
+			if (UAudioComponent* AudioComponent = UAudioComponent::GetAudioComponentFromID(AudioComponentID))
+			{
+				AudioComponent->Stop();
+				bComponentsStopped = true;
+			}
+		}
+		return bComponentsStopped;
+	});
+
+	OutDocumentDelegates.IterateGraphNodeDelegates([this](FNodeModifyDelegates& NodeDelegates)
+	{
+		NodeDelegates.OnNodeAdded.RemoveAll(this);
+		NodeDelegates.OnNodeInputLiteralSet.RemoveAll(this);
+		NodeDelegates.OnRemoveSwappingNode.RemoveAll(this);
+		NodeDelegates.OnRemovingNodeInputLiteral.RemoveAll(this);
+	});
+
+	OutDocumentDelegates.IterateGraphEdgeDelegates([this](FEdgeModifyDelegates& EdgeDelegates)
+	{
+		EdgeDelegates.OnEdgeAdded.RemoveAll(this);
+		EdgeDelegates.OnRemoveSwappingEdge.RemoveAll(this);
+	});
+
+	const FGuid ExecutablePageID = FDocumentBuilderRegistry::GetChecked().ResolveExecutablePageID(&GetMetaSoundSource());
+
+	FEdgeModifyDelegates& EdgeDelegates = OutDocumentDelegates.FindEdgeDelegatesChecked(ExecutablePageID);
+	EdgeDelegates.OnEdgeAdded.AddUObject(this, &UMetaSoundSourceBuilder::OnEdgeAdded);
+	EdgeDelegates.OnRemoveSwappingEdge.AddUObject(this, &UMetaSoundSourceBuilder::OnRemoveSwappingEdge);
+
+	FNodeModifyDelegates& NodeDelegates = OutDocumentDelegates.FindNodeDelegatesChecked(ExecutablePageID);
+	NodeDelegates.OnNodeAdded.AddUObject(this, &UMetaSoundSourceBuilder::OnNodeAdded);
+	NodeDelegates.OnNodeInputLiteralSet.AddUObject(this, &UMetaSoundSourceBuilder::OnNodeInputLiteralSet);
+	NodeDelegates.OnRemoveSwappingNode.AddUObject(this, &UMetaSoundSourceBuilder::OnRemoveSwappingNode);
+	NodeDelegates.OnRemovingNodeInputLiteral.AddUObject(this, &UMetaSoundSourceBuilder::OnRemovingNodeInputLiteral);
 }
 
 void UMetaSoundSourceBuilder::OnAssetReferenceAdded(TScriptInterface<IMetaSoundDocumentInterface> DocInterface)
@@ -506,6 +552,23 @@ void UMetaSoundSourceBuilder::OnOutputAdded(int32 OutputIndex) const
 	});
 }
 
+void UMetaSoundSourceBuilder::OnPageAdded(const Metasound::Frontend::FDocumentMutatePageArgs& Args)
+{
+	using namespace Metasound::Frontend;
+
+	FDocumentModifyDelegates& DocDelegates = Builder.GetDocumentDelegates();
+	InitExecutablePageDelegates(DocDelegates);
+}
+
+void UMetaSoundSourceBuilder::OnRemovingPage(const Metasound::Frontend::FDocumentMutatePageArgs& Args)
+{
+	using namespace Metasound::Frontend;
+
+	FDocumentModifyDelegates& DocDelegates = Builder.GetDocumentDelegates();
+	InitExecutablePageDelegates(DocDelegates);
+
+}
+
 void UMetaSoundSourceBuilder::OnRemoveSwappingEdge(int32 SwapIndex, int32 LastIndex) const
 {
 	using namespace Metasound::DynamicGraph;
diff --git a/Engine/Plugins/Runtime/Metasound/Source/MetasoundEngine/Private/MetasoundDocumentBuilderRegistry.cpp b/Engine/Plugins/Runtime/Metasound/Source/MetasoundEngine/Private/MetasoundDocumentBuilderRegistry.cpp
--- a/Engine/Plugins/Runtime/Metasound/Source/MetasoundEngine/Private/MetasoundDocumentBuilderRegistry.cpp
+++ b/Engine/Plugins/Runtime/Metasound/Source/MetasoundEngine/Private/MetasoundDocumentBuilderRegistry.cpp
@@ -4,6 +4,7 @@
 
 #include "MetasoundAssetManager.h"
 #include "MetasoundGlobals.h"
+#include "MetasoundSettings.h"
 #include "MetasoundTrace.h"
 #include "MetasoundUObjectRegistry.h"
 
@@ -232,6 +233,12 @@ namespace Metasound::Engine
 		return FoundBuilders;
 	}
 
+	FGuid FDocumentBuilderRegistry::ResolveExecutablePageID(const TScriptInterface<IMetaSoundDocumentInterface> DocumentInterface) const
+	{
+		// TODO: implement selection logic.  For now default to 0 guid, which is always valid on an asset.
+		return FGuid();
+	}
+
 	FMetaSoundFrontendDocumentBuilder* FDocumentBuilderRegistry::FindOutermostBuilder(const UObject& InSubObject) const
 	{
 		using namespace Metasound::Frontend;
diff --git a/Engine/Plugins/Runtime/Metasound/Source/MetasoundEngine/Private/MetasoundEngineAsset.h b/Engine/Plugins/Runtime/Metasound/Source/MetasoundEngine/Private/MetasoundEngineAsset.h
--- a/Engine/Plugins/Runtime/Metasound/Source/MetasoundEngine/Private/MetasoundEngineAsset.h
+++ b/Engine/Plugins/Runtime/Metasound/Source/MetasoundEngine/Private/MetasoundEngineAsset.h
@@ -69,10 +69,8 @@ namespace Metasound::Engine
 		{
 			InMetaSound.GetModifyContext().SetForceRefreshViews();
 
-			if (UMetaSoundBuilderBase* Builder = UMetaSoundBuilderSubsystem::GetChecked().FindBuilderOfDocument(&InMetaSound))
-			{
-				Builder->Reload();
-			}
+			const FMetasoundFrontendClassName& ClassName = InMetaSound.GetConstDocument().RootGraph.Metadata.GetClassName();
+			Frontend::IDocumentBuilderRegistry::GetChecked().ReloadBuilder(ClassName);
 
 			if (UMetasoundEditorGraphBase* Graph = Cast<UMetasoundEditorGraphBase>(InMetaSound.GetGraph()))
 			{
diff --git a/Engine/Plugins/Runtime/Metasound/Source/MetasoundEngine/Public/MetasoundBuilderBase.cpp b/Engine/Plugins/Runtime/Metasound/Source/MetasoundEngine/Public/MetasoundBuilderBase.cpp
--- a/Engine/Plugins/Runtime/Metasound/Source/MetasoundEngine/Public/MetasoundBuilderBase.cpp
+++ b/Engine/Plugins/Runtime/Metasound/Source/MetasoundEngine/Public/MetasoundBuilderBase.cpp
@@ -458,7 +458,7 @@ void UMetaSoundBuilderBase::ConvertToPreset(const TScriptInterface<IMetaSoundDoc
 
 	const FMetasoundFrontendDocument& ReferencedDocument = ReferencedInterface->GetConstDocument();
 
-	TSharedRef<FDocumentModifyDelegates> DocumentDelegates = MakeShared<FDocumentModifyDelegates>();
+	TSharedRef<FDocumentModifyDelegates> DocumentDelegates = MakeShared<FDocumentModifyDelegates>(ReferencedDocument);
 	InitDelegates(*DocumentDelegates);
 	const bool bConvertToPreset = Builder.ConvertToPreset(ReferencedDocument, DocumentDelegates);
 	ConformObjectToDocument();
@@ -772,9 +772,9 @@ void UMetaSoundBuilderBase::Initialize()
 
 	const EObjectFlags NewObjectFlags = RF_Public | RF_Transient;
 	TScriptInterface<IMetaSoundDocumentInterface> DocObject = NewObject<UObject>(GetTransientPackage(), &GetBaseMetaSoundUClass(), { }, NewObjectFlags);
-	TSharedRef<FDocumentModifyDelegates> DocumentDelegates = MakeShared<FDocumentModifyDelegates>();
-	InitDelegates(*DocumentDelegates);
+	TSharedRef<FDocumentModifyDelegates> DocumentDelegates = MakeShared<FDocumentModifyDelegates>(DocObject->GetConstDocument());
 	Builder = FMetaSoundFrontendDocumentBuilder(DocObject, DocumentDelegates);
+	InitDelegates(*DocumentDelegates);
 	Builder.InitDocument();
 }
 
@@ -900,7 +900,7 @@ void UMetaSoundBuilderBase::Reload(bool bPrimeCache)
 {
 	using namespace Metasound::Frontend;
 
-	TSharedRef<FDocumentModifyDelegates> DocumentDelegates = MakeShared<FDocumentModifyDelegates>();
+	TSharedRef<FDocumentModifyDelegates> DocumentDelegates = MakeShared<FDocumentModifyDelegates>(GetConstBuilder().GetConstDocumentChecked());
 	InitDelegates(*DocumentDelegates);
 	Builder.Reload(DocumentDelegates, bPrimeCache);
 }
diff --git a/Engine/Plugins/Runtime/Metasound/Source/MetasoundEngine/Public/MetasoundBuilderSubsystem.h b/Engine/Plugins/Runtime/Metasound/Source/MetasoundEngine/Public/MetasoundBuilderSubsystem.h
--- a/Engine/Plugins/Runtime/Metasound/Source/MetasoundEngine/Public/MetasoundBuilderSubsystem.h
+++ b/Engine/Plugins/Runtime/Metasound/Source/MetasoundEngine/Public/MetasoundBuilderSubsystem.h
@@ -121,6 +121,7 @@ private:
 	const UMetaSoundSource& GetMetaSoundSource() const;
 	UMetaSoundSource& GetMetaSoundSource();
 
+	void InitExecutablePageDelegates(Metasound::Frontend::FDocumentModifyDelegates& OutDocumentDelegates);
 
 	void OnEdgeAdded(int32 EdgeIndex) const;
 	void OnInputAdded(int32 InputIndex);
@@ -128,11 +129,13 @@ private:
 	void OnNodeAdded(int32 NodeIndex) const;
 	void OnNodeInputLiteralSet(int32 NodeIndex, int32 VertexIndex, int32 LiteralIndex) const;
 	void OnOutputAdded(int32 OutputIndex) const;
+	void OnPageAdded(const Metasound::Frontend::FDocumentMutatePageArgs& InArgs);
 	void OnRemoveSwappingEdge(int32 SwapIndex, int32 LastIndex) const;
 	void OnRemovingInput(int32 InputIndex);
 	void OnRemoveSwappingNode(int32 SwapIndex, int32 LastIndex) const;
 	void OnRemovingNodeInputLiteral(int32 NodeIndex, int32 VertexIndex, int32 LiteralIndex) const;
 	void OnRemovingOutput(int32 OutputIndex) const;
+	void OnRemovingPage(const Metasound::Frontend::FDocumentMutatePageArgs& InArgs);
 
 	using FAuditionableTransaction = TFunctionRef<bool(Metasound::DynamicGraph::FDynamicOperatorTransactor&)>;
 	bool ExecuteAuditionableTransaction(FAuditionableTransaction Transaction) const;
diff --git a/Engine/Plugins/Runtime/Metasound/Source/MetasoundEngine/Public/MetasoundDocumentBuilderRegistry.h b/Engine/Plugins/Runtime/Metasound/Source/MetasoundEngine/Public/MetasoundDocumentBuilderRegistry.h
--- a/Engine/Plugins/Runtime/Metasound/Source/MetasoundEngine/Public/MetasoundDocumentBuilderRegistry.h
+++ b/Engine/Plugins/Runtime/Metasound/Source/MetasoundEngine/Public/MetasoundDocumentBuilderRegistry.h
@@ -119,6 +119,9 @@ namespace Metasound::Engine
 		// Returns all builder objects registered and active associated with the given ClassName.
 		TArray<UMetaSoundBuilderBase*> FindBuilderObjects(const FMetasoundFrontendClassName& InClassName) const;
 
+		// Given the provided document, returns the valid page ID to be executed.
+		virtual FGuid ResolveExecutablePageID(const TScriptInterface<IMetaSoundDocumentInterface> DocumentInterface) const;
+
 		bool ReloadBuilder(const FMetasoundFrontendClassName& InClassName) const override;
 
 		void SetEventLogVerbosity(ELogEvent Event, ELogVerbosity::Type Verbosity);
diff --git a/Engine/Plugins/Runtime/Metasound/Source/MetasoundFrontend/Private/MetasoundFrontendDocument.cpp b/Engine/Plugins/Runtime/Metasound/Source/MetasoundFrontend/Private/MetasoundFrontendDocument.cpp
--- a/Engine/Plugins/Runtime/Metasound/Source/MetasoundFrontend/Private/MetasoundFrontendDocument.cpp
+++ b/Engine/Plugins/Runtime/Metasound/Source/MetasoundFrontend/Private/MetasoundFrontendDocument.cpp
@@ -877,58 +877,86 @@ FMetasoundFrontendGraphClass::FMetasoundFrontendGraphClass()
 	Metadata.SetType(EMetasoundFrontendClassType::Graph);
 }
 
-void FMetasoundFrontendGraphClass::AddNewGraphPage(const FGuid& InPageID)
+const FMetasoundFrontendGraph& FMetasoundFrontendGraphClass::AddGraphPage(const FGuid& InPageID, bool bDuplicateLastGraph, bool bSetAsBuildGraph)
 {
-	checkf(ContainsGraphPage(InPageID), TEXT("Cannot add new graph page with existing PageID"));
 	checkf(InPageID.IsValid(), TEXT("Cannot add graph with invalid PageID"))
+	checkf(!ContainsGraphPage(InPageID), TEXT("Cannot add new graph page with existing PageID"));
 
-	FMetasoundFrontendGraph& NewGraph = PagedGraphs.AddDefaulted_GetRef();
-	NewGraph.PageID = InPageID;
+	FMetasoundFrontendGraph* NewGraph = nullptr;
+	if (bDuplicateLastGraph)
+	{
+PRAGMA_DISABLE_DEPRECATION_WARNINGS
+		FMetasoundFrontendGraph ToDuplicate = PagedGraphs.IsEmpty() ? Graph : PagedGraphs.Last();
+PRAGMA_ENABLE_DEPRECATION_WARNINGS
+		NewGraph = &PagedGraphs.Add_GetRef(MoveTemp(ToDuplicate));
+	}
+	else
+	{
+		NewGraph = &PagedGraphs.AddDefaulted_GetRef();
+	}
+
+	NewGraph->PageID = InPageID;
+	return *NewGraph;
 }
 
 bool FMetasoundFrontendGraphClass::ContainsGraphPage(const FGuid& InPageID) const
 {
+PRAGMA_DISABLE_DEPRECATION_WARNINGS
 	if (InPageID == Graph.PageID)
 	{
 		return true;
 	}
+PRAGMA_ENABLE_DEPRECATION_WARNINGS
 
 	auto MatchesPageID = [&InPageID](const FMetasoundFrontendGraph& Iter) { return Iter.PageID == InPageID; };
 	return PagedGraphs.ContainsByPredicate(MatchesPageID);
 }
 
-void FMetasoundFrontendGraphClass::DuplicateLastGraphPage(const FGuid& InPageID)
-{
-	checkf(ContainsGraphPage(InPageID), TEXT("Cannot add new graph page with existing PageID"));
-	checkf(InPageID.IsValid(), TEXT("Cannot add graph with invalid PageID"))
-
-	FMetasoundFrontendGraph* ToDuplicate = PagedGraphs.IsEmpty() ? &Graph : &PagedGraphs.Last();
-	FMetasoundFrontendGraph& NewGraph = PagedGraphs.Add_GetRef(*ToDuplicate);
-	NewGraph.PageID = InPageID;
-}
-
 void FMetasoundFrontendGraphClass::RemoveAllGraphPages()
 {
 	PagedGraphs.Empty();
+
+PRAGMA_DISABLE_DEPRECATION_WARNINGS
 	Graph.PageID = FGuid();
+PRAGMA_ENABLE_DEPRECATION_WARNINGS
 }
 
-bool FMetasoundFrontendGraphClass::RemoveGraphPage(const FGuid& InPageID)
+bool FMetasoundFrontendGraphClass::RemoveGraphPage(const FGuid& InPageID, FGuid* OutAdjacentPageID)
 {
-	const bool bRemoved = PagedGraphs.RemoveAllSwap([&InPageID](const FMetasoundFrontendGraph& Iter)
+	if (InPageID != FGuid())
 	{
-		return Iter.PageID == InPageID;
-	}, EAllowShrinking::Yes) > 0;
+		for (int32 Index = 0; Index < PagedGraphs.Num(); ++Index)
+		{
+			if (PagedGraphs[Index].PageID == InPageID)
+			{
+				PagedGraphs.RemoveAtSwap(Index, EAllowShrinking::Yes);
 
-	return bRemoved;
+				if (OutAdjacentPageID)
+				{
+					*OutAdjacentPageID = (Index > 0) ? PagedGraphs[Index - 1].PageID : FGuid();
+				}
+
+				return true;
+			}
+		}
+	}
+
+	if (OutAdjacentPageID)
+	{
+		*OutAdjacentPageID = FGuid();
+	}
+
+	return false;
 }
 
 FMetasoundFrontendGraph* FMetasoundFrontendGraphClass::FindGraph(const FGuid& InPageID)
 {
+PRAGMA_DISABLE_DEPRECATION_WARNINGS
 	if (InPageID == Graph.PageID)
 	{
 		return &Graph;
 	}
+PRAGMA_ENABLE_DEPRECATION_WARNINGS
 
 	auto MatchesPageID = [this, &InPageID](const FMetasoundFrontendGraph& Iter) { return Iter.PageID == InPageID; };
 	FMetasoundFrontendGraph* PageGraph = PagedGraphs.FindByPredicate(MatchesPageID);
@@ -944,10 +972,12 @@ FMetasoundFrontendGraph& FMetasoundFrontendGraphClass::FindGraphChecked(const FG
 
 const FMetasoundFrontendGraph* FMetasoundFrontendGraphClass::FindConstGraph(const FGuid& InPageID) const
 {
+PRAGMA_DISABLE_DEPRECATION_WARNINGS
 	if (InPageID == Graph.PageID)
 	{
 		return &Graph;
 	}
+PRAGMA_ENABLE_DEPRECATION_WARNINGS
 
 	auto MatchesPageID = [this, &InPageID](const FMetasoundFrontendGraph& Iter) { return Iter.PageID == InPageID; };
 	const FMetasoundFrontendGraph* PageGraph = PagedGraphs.FindByPredicate(MatchesPageID);
@@ -973,7 +1003,9 @@ const FMetasoundFrontendGraph& FMetasoundFrontendGraphClass::GetConstDefaultGrap
 
 void FMetasoundFrontendGraphClass::IterateGraphPages(TFunctionRef<void(FMetasoundFrontendGraph&)> IterFunc)
 {
+PRAGMA_DISABLE_DEPRECATION_WARNINGS
 	IterFunc(Graph);
+PRAGMA_ENABLE_DEPRECATION_WARNINGS
 
 	for (FMetasoundFrontendGraph& Iter : PagedGraphs)
 	{
@@ -983,7 +1015,9 @@ void FMetasoundFrontendGraphClass::IterateGraphPages(TFunctionRef<void(FMetasoun
 
 void FMetasoundFrontendGraphClass::IterateGraphPages(TFunctionRef<void(const FMetasoundFrontendGraph&)> IterFunc) const
 {
+PRAGMA_DISABLE_DEPRECATION_WARNINGS
 	IterFunc(Graph);
+PRAGMA_ENABLE_DEPRECATION_WARNINGS
 
 	for (const FMetasoundFrontendGraph& Iter : PagedGraphs)
 	{
@@ -994,8 +1028,11 @@ void FMetasoundFrontendGraphClass::IterateGraphPages(TFunctionRef<void(const FMe
 void FMetasoundFrontendGraphClass::ResetGraphs()
 {
 	PagedGraphs.Empty();
+
+PRAGMA_DISABLE_DEPRECATION_WARNINGS
 	Graph.Nodes.Reset();
 	Graph.Edges.Reset();
+PRAGMA_ENABLE_DEPRECATION_WARNINGS
 }
 
 FMetasoundFrontendVersionNumber FMetasoundFrontendDocument::GetMaxVersion()
diff --git a/Engine/Plugins/Runtime/Metasound/Source/MetasoundFrontend/Private/MetasoundFrontendDocumentBuilder.cpp b/Engine/Plugins/Runtime/Metasound/Source/MetasoundFrontend/Private/MetasoundFrontendDocumentBuilder.cpp
--- a/Engine/Plugins/Runtime/Metasound/Source/MetasoundFrontend/Private/MetasoundFrontendDocumentBuilder.cpp
+++ b/Engine/Plugins/Runtime/Metasound/Source/MetasoundFrontend/Private/MetasoundFrontendDocumentBuilder.cpp
@@ -722,34 +722,38 @@ const FMetasoundFrontendClass* FMetaSoundFrontendDocumentBuilder::AddDependency(
 	return Dependency;
 }
 
-void FMetaSoundFrontendDocumentBuilder::AddEdge(FMetasoundFrontendEdge&& InNewEdge)
+void FMetaSoundFrontendDocumentBuilder::AddEdge(FMetasoundFrontendEdge&& InNewEdge, const FGuid* InPageID)
 {
 	using namespace Metasound::Frontend;
 
+	const FGuid& PageID = InPageID ? *InPageID : BuildPageID;
+
 #if DO_CHECK
 	{
-		const IDocumentGraphEdgeCache& EdgeCache = DocumentCache->GetEdgeCache();
+		const IDocumentGraphEdgeCache& EdgeCache = DocumentCache->GetEdgeCache(PageID);
 		checkf(!EdgeCache.IsNodeInputConnected(InNewEdge.ToNodeID, InNewEdge.ToVertexID), TEXT("Failed to add edge in MetaSound Builder: Destination input already connected"));
 
-		const EInvalidEdgeReason Reason = IsValidEdge(InNewEdge);
+		const EInvalidEdgeReason Reason = IsValidEdge(InNewEdge, &PageID);
 		checkf(Reason == Metasound::Frontend::EInvalidEdgeReason::None, TEXT("Attempted call to AddEdge in MetaSound Builder where edge is invalid: %s."), *LexToString(Reason));
 	}
 #endif // DO_CHECK
 
 	FMetasoundFrontendDocument& Document = GetDocumentChecked();
-	FMetasoundFrontendGraph& Graph = FindBuildGraphChecked();
+	FMetasoundFrontendGraph& Graph = Document.RootGraph.FindGraphChecked(PageID);
 	Graph.Edges.Add(MoveTemp(InNewEdge));
 	const int32 NewIndex = Graph.Edges.Num() - 1;
-	DocumentDelegates->EdgeDelegates.OnEdgeAdded.Broadcast(NewIndex);
+	DocumentDelegates->FindEdgeDelegatesChecked(PageID).OnEdgeAdded.Broadcast(NewIndex);
 }
 
-bool FMetaSoundFrontendDocumentBuilder::AddNamedEdges(const TSet<Metasound::Frontend::FNamedEdge>& EdgesToMake, TArray<const FMetasoundFrontendEdge*>* OutNewEdges, bool bReplaceExistingConnections)
+bool FMetaSoundFrontendDocumentBuilder::AddNamedEdges(const TSet<Metasound::Frontend::FNamedEdge>& EdgesToMake, TArray<const FMetasoundFrontendEdge*>* OutNewEdges, bool bReplaceExistingConnections, const FGuid* InPageID)
 {
 	using namespace Metasound::Frontend;
 
+	const FGuid& PageID = InPageID ? *InPageID : BuildPageID;
+
 	FMetasoundFrontendDocument& Document = GetDocumentChecked();
-	const IDocumentGraphNodeCache& NodeCache = DocumentCache->GetNodeCache();
-	FMetasoundFrontendGraph& Graph = FindBuildGraphChecked();
+	const IDocumentGraphNodeCache& NodeCache = DocumentCache->GetNodeCache(PageID);
+	FMetasoundFrontendGraph& Graph = Document.RootGraph.FindGraphChecked(PageID);
 
 	if (OutNewEdges)
 	{
@@ -787,7 +791,7 @@ bool FMetaSoundFrontendDocumentBuilder::AddNamedEdges(const TSet<Metasound::Fron
 		}
 	}
 
-	const TArray<FMetasoundFrontendEdge>& Edges = FindConstBuildGraphChecked().Edges;
+	const TArray<FMetasoundFrontendEdge>& Edges = Graph.Edges;
 	const int32 LastIndex = Edges.Num() - 1;
 	for (FNewEdgeData& EdgeToAdd : EdgesToAdd)
 	{
@@ -795,10 +799,10 @@ bool FMetaSoundFrontendDocumentBuilder::AddNamedEdges(const TSet<Metasound::Fron
 		{
 #if !NO_LOGGING
 			const FMetasoundFrontendNode* OldOutputNode = nullptr;
-			const FMetasoundFrontendVertex* OldOutputVertex = FindNodeOutputConnectedToNodeInput(EdgeToAdd.NewEdge.ToNodeID, EdgeToAdd.NewEdge.ToVertexID, &OldOutputNode);
+			const FMetasoundFrontendVertex* OldOutputVertex = FindNodeOutputConnectedToNodeInput(EdgeToAdd.NewEdge.ToNodeID, EdgeToAdd.NewEdge.ToVertexID, &OldOutputNode, &PageID);
 #endif // !NO_LOGGING
 
-			const bool bRemovedEdge = RemoveEdgeToNodeInput(EdgeToAdd.NewEdge.ToNodeID, EdgeToAdd.NewEdge.ToVertexID);
+			const bool bRemovedEdge = RemoveEdgeToNodeInput(EdgeToAdd.NewEdge.ToNodeID, EdgeToAdd.NewEdge.ToVertexID, &PageID);
 
 #if !NO_LOGGING
 			if (bRemovedEdge)
@@ -822,11 +826,11 @@ bool FMetaSoundFrontendDocumentBuilder::AddNamedEdges(const TSet<Metasound::Fron
 			}
 #endif // !NO_LOGGING
 
-			AddEdge(MoveTemp(EdgeToAdd.NewEdge));
+			AddEdge(MoveTemp(EdgeToAdd.NewEdge), &PageID);
 		}
-		else if (!IsNodeInputConnected(EdgeToAdd.NewEdge.ToNodeID, EdgeToAdd.NewEdge.ToVertexID))
+		else if (!IsNodeInputConnected(EdgeToAdd.NewEdge.ToNodeID, EdgeToAdd.NewEdge.ToVertexID, &PageID))
 		{
-			AddEdge(MoveTemp(EdgeToAdd.NewEdge));
+			AddEdge(MoveTemp(EdgeToAdd.NewEdge), &PageID);
 		}
 		else
 		{
@@ -834,15 +838,15 @@ bool FMetaSoundFrontendDocumentBuilder::AddNamedEdges(const TSet<Metasound::Fron
 
 #if !NO_LOGGING
 			FMetasoundFrontendEdge EdgeToRemove;
-			if (const int32* EdgeIndex = DocumentCache->GetEdgeCache().FindEdgeIndexToNodeInput(EdgeToAdd.NewEdge.ToNodeID, EdgeToAdd.NewEdge.ToVertexID))
+			if (const int32* EdgeIndex = DocumentCache->GetEdgeCache(PageID).FindEdgeIndexToNodeInput(EdgeToAdd.NewEdge.ToNodeID, EdgeToAdd.NewEdge.ToVertexID))
 			{
 				EdgeToRemove = Graph.Edges[*EdgeIndex];
 			}
 
-			const FMetasoundFrontendVertex* Input = FindNodeInput(EdgeToAdd.NewEdge.ToNodeID, EdgeToAdd.NewEdge.ToVertexID);
+			const FMetasoundFrontendVertex* Input = FindNodeInput(EdgeToAdd.NewEdge.ToNodeID, EdgeToAdd.NewEdge.ToVertexID, &PageID);
 			checkf(Input, TEXT("Prior loop to check edge validity should protect against missing input vertex"));
 
-			const FMetasoundFrontendVertex* Output = FindNodeOutput(EdgeToAdd.NewEdge.FromNodeID, EdgeToAdd.NewEdge.FromVertexID);
+			const FMetasoundFrontendVertex* Output = FindNodeOutput(EdgeToAdd.NewEdge.FromNodeID, EdgeToAdd.NewEdge.FromVertexID, &PageID);
 			checkf(Input, TEXT("Prior loop to check edge validity should protect against missing output vertex"));
 
 			UE_LOG(LogMetaSound, Warning, TEXT("Connection between MetaSound output '%s' and input '%s' not added: Input already connected to '%s'."), *Output->Name.ToString(), *Input->Name.ToString(), *Output->Name.ToString());
@@ -861,19 +865,21 @@ bool FMetaSoundFrontendDocumentBuilder::AddNamedEdges(const TSet<Metasound::Fron
 	return bSuccess;
 }
 
-bool FMetaSoundFrontendDocumentBuilder::AddEdgesByNodeClassInterfaceBindings(const FGuid& InFromNodeID, const FGuid& InToNodeID, bool bReplaceExistingConnections)
+bool FMetaSoundFrontendDocumentBuilder::AddEdgesByNodeClassInterfaceBindings(const FGuid& InFromNodeID, const FGuid& InToNodeID, bool bReplaceExistingConnections, const FGuid* InPageID)
 {
 	using namespace Metasound;
 	using namespace Metasound::Frontend;
 
+	const FGuid& PageID = InPageID ? *InPageID : BuildPageID;
+
 	TSet<FMetasoundFrontendVersion> FromInterfaceVersions;
 	TSet<FMetasoundFrontendVersion> ToInterfaceVersions;
-	if (FindNodeClassInterfaces(InFromNodeID, FromInterfaceVersions) && FindNodeClassInterfaces(InToNodeID, ToInterfaceVersions))
+	if (FindNodeClassInterfaces(InFromNodeID, FromInterfaceVersions, PageID) && FindNodeClassInterfaces(InToNodeID, ToInterfaceVersions, PageID))
 	{
 		TSet<FNamedEdge> NamedEdges;
 		if (DocumentBuilderPrivate::TryGetInterfaceBoundEdges(InFromNodeID, FromInterfaceVersions, InToNodeID, ToInterfaceVersions, NamedEdges))
 		{
-			return AddNamedEdges(NamedEdges, nullptr, bReplaceExistingConnections);
+			return AddNamedEdges(NamedEdges, nullptr, bReplaceExistingConnections, &PageID);
 		}
 	}
 
@@ -881,22 +887,24 @@ bool FMetaSoundFrontendDocumentBuilder::AddEdgesByNodeClassInterfaceBindings(con
 
 }
 
-bool FMetaSoundFrontendDocumentBuilder::AddEdgesFromMatchingInterfaceNodeOutputsToGraphOutputs(const FGuid& InNodeID, TArray<const FMetasoundFrontendEdge*>& OutEdgesCreated, bool bReplaceExistingConnections)
+bool FMetaSoundFrontendDocumentBuilder::AddEdgesFromMatchingInterfaceNodeOutputsToGraphOutputs(const FGuid& InNodeID, TArray<const FMetasoundFrontendEdge*>& OutEdgesCreated, bool bReplaceExistingConnections, const FGuid* InPageID)
 {
 	METASOUND_TRACE_CPUPROFILER_EVENT_SCOPE(FMetaSoundFrontendDocumentBuilder::AddEdgesFromMatchingInterfaceNodeOutputsToGraphOutputs);
 
 	using namespace Metasound::Frontend;
 
+	const FGuid& PageID = InPageID ? *InPageID : BuildPageID;
+
 	OutEdgesCreated.Reset();
 
 	TSet<FMetasoundFrontendVersion> NodeInterfaces;
-	if (!FindNodeClassInterfaces(InNodeID, NodeInterfaces))
+	if (!FindNodeClassInterfaces(InNodeID, NodeInterfaces, PageID))
 	{
 		// Did not find any node interfaces
 		return false;
 	}
 
-	const IDocumentGraphNodeCache& NodeCache = DocumentCache->GetNodeCache();
+	const IDocumentGraphNodeCache& NodeCache = DocumentCache->GetNodeCache(PageID);
 	const IDocumentGraphInterfaceCache& InterfaceCache = DocumentCache->GetInterfaceCache();
 	const TSet<FMetasoundFrontendVersion> CommonInterfaces = NodeInterfaces.Intersect(GetDocumentChecked().Interfaces);
 
@@ -922,37 +930,39 @@ bool FMetaSoundFrontendDocumentBuilder::AddEdgesFromMatchingInterfaceNodeOutputs
 		}
 	}
 
-	return AddNamedEdges(EdgesToMake, &OutEdgesCreated, bReplaceExistingConnections);
+	return AddNamedEdges(EdgesToMake, &OutEdgesCreated, bReplaceExistingConnections, &PageID);
 }
 
-bool FMetaSoundFrontendDocumentBuilder::AddEdgesFromMatchingInterfaceNodeInputsToGraphInputs(const FGuid& InNodeID, TArray<const FMetasoundFrontendEdge*>& OutEdgesCreated, bool bReplaceExistingConnections)
+bool FMetaSoundFrontendDocumentBuilder::AddEdgesFromMatchingInterfaceNodeInputsToGraphInputs(const FGuid& InNodeID, TArray<const FMetasoundFrontendEdge*>& OutEdgesCreated, bool bReplaceExistingConnections, const FGuid* InPageID)
 {
 	METASOUND_TRACE_CPUPROFILER_EVENT_SCOPE(FMetaSoundFrontendDocumentBuilder::AddEdgesFromMatchingInterfaceNodeInputsToGraphInputs);
 
 	using namespace Metasound::Frontend;
 
+	const FGuid& PageID = InPageID ? *InPageID : BuildPageID;
+
 	OutEdgesCreated.Reset();
 
 	TSet<FMetasoundFrontendVersion> NodeInterfaces;
-	if (!FindNodeClassInterfaces(InNodeID, NodeInterfaces))
+	if (!FindNodeClassInterfaces(InNodeID, NodeInterfaces, PageID))
 	{
 		// Did not find any node interfaces
 		return false;
 	}
 
-	const IDocumentGraphNodeCache& NodeCache = DocumentCache->GetNodeCache();
+	const IDocumentGraphNodeCache& NodeCache = DocumentCache->GetNodeCache(PageID);
 	const IDocumentGraphInterfaceCache& InterfaceCache = DocumentCache->GetInterfaceCache();
 	const TSet<FMetasoundFrontendVersion> CommonInterfaces = NodeInterfaces.Intersect(GetDocumentChecked().Interfaces);
 
 	TSet<FNamedEdge> EdgesToMake;
+	const FMetasoundFrontendGraph& Graph = GetDocumentChecked().RootGraph.FindConstGraphChecked(PageID);
 	for (const FMetasoundFrontendVersion& Version : CommonInterfaces)
 	{
 		const FInterfaceRegistryKey InterfaceKey = GetInterfaceRegistryKey(Version);
 		if (const IInterfaceRegistryEntry* RegistryEntry = IInterfaceRegistry::Get().FindInterfaceRegistryEntry(InterfaceKey))
 		{
-			Algo::Transform(RegistryEntry->GetInterface().Inputs, EdgesToMake, [this, &NodeCache, &InterfaceCache, InNodeID](const FMetasoundFrontendClassInput& Input)
+			Algo::Transform(RegistryEntry->GetInterface().Inputs, EdgesToMake, [this, &Graph, &NodeCache, &InterfaceCache, InNodeID](const FMetasoundFrontendClassInput& Input)
 			{
-				const FMetasoundFrontendGraph& Graph = FindConstBuildGraphChecked();
 				const FMetasoundFrontendVertex* NodeVertex = NodeCache.FindInputVertex(InNodeID, Input.Name);
 				check(NodeVertex);
 				const FMetasoundFrontendClassInput* InputClass = InterfaceCache.FindInput(Input.Name);
@@ -966,19 +976,17 @@ bool FMetaSoundFrontendDocumentBuilder::AddEdgesFromMatchingInterfaceNodeInputsT
 		}
 	}
 
-	return AddNamedEdges(EdgesToMake, &OutEdgesCreated, bReplaceExistingConnections);
+	return AddNamedEdges(EdgesToMake, &OutEdgesCreated, bReplaceExistingConnections, &PageID);
 }
 
-const FMetasoundFrontendNode* FMetaSoundFrontendDocumentBuilder::AddGraphInput(const FMetasoundFrontendClassInput& InClassInput)
+const FMetasoundFrontendNode* FMetaSoundFrontendDocumentBuilder::AddGraphInput(const FMetasoundFrontendClassInput& InClassInput, const FGuid* InPageID)
 {
 	using namespace Metasound::Frontend;
 
-	FMetasoundFrontendDocument& Document = GetDocumentChecked();
-	FMetasoundFrontendGraphClass& RootGraph = Document.RootGraph;
-
 	checkf(InClassInput.NodeID.IsValid(), TEXT("Unassigned NodeID when adding graph input"));
 	checkf(InClassInput.VertexID.IsValid(), TEXT("Unassigned VertexID when adding graph input"));
 
+	const FGuid& PageID = InPageID ? *InPageID : BuildPageID;
 	if (InClassInput.TypeName.IsNone())
 	{
 		UE_LOG(LogMetaSound, Error, TEXT("TypeName unset when attempting to add class input '%s'"), *InClassInput.Name.ToString());
@@ -987,7 +995,7 @@ const FMetasoundFrontendNode* FMetaSoundFrontendDocumentBuilder::AddGraphInput(c
 	else if (const FMetasoundFrontendClassInput* Input = DocumentCache->GetInterfaceCache().FindInput(InClassInput.Name))
 	{
 		UE_LOG(LogMetaSound, Error, TEXT("Attempting to add MetaSound graph input '%s' when input with name already exists"), *InClassInput.Name.ToString());
-		const FMetasoundFrontendNode* OutputNode = DocumentCache->GetNodeCache().FindNode(Input->NodeID);
+		const FMetasoundFrontendNode* OutputNode = DocumentCache->GetNodeCache(PageID).FindNode(Input->NodeID);
 		check(OutputNode);
 		return OutputNode;
 	}
@@ -1005,6 +1013,9 @@ const FMetasoundFrontendNode* FMetaSoundFrontendDocumentBuilder::AddGraphInput(c
 			AddDependency(Class);
 		}
 
+		FMetasoundFrontendDocument& Document = GetDocumentChecked();
+		FMetasoundFrontendGraphClass& RootGraph = Document.RootGraph;
+
 		auto FinalizeNode = [this, &InClassInput](FMetasoundFrontendNode& InOutNode, const Metasound::Frontend::FNodeRegistryKey&)
 		{
 			// Sets the name of the node an vertices on the node to match the class vertex name
@@ -1017,11 +1028,11 @@ const FMetasoundFrontendNode* FMetaSoundFrontendDocumentBuilder::AddGraphInput(c
 
 		// Must add input node to all paged graphs to maintain API parity for all page implementations
 		FMetasoundFrontendNode* NewNode = nullptr;
-		Document.RootGraph.IterateGraphPages([this, &NewNode, &Class, &FinalizeNode, &InClassInput](const FMetasoundFrontendGraph& Graph)
+		RootGraph.IterateGraphPages([&](const FMetasoundFrontendGraph& Graph)
 		{
 			constexpr int32* NewNodeIndex = nullptr;
-			FMetasoundFrontendNode* NewPageNode = AddNodeInternal(Class.Metadata, FinalizeNode, InClassInput.NodeID, NewNodeIndex, &Graph.PageID);
-			if (Graph.PageID == BuildPageID)
+			FMetasoundFrontendNode* NewPageNode = AddNodeInternal(Class.Metadata, FinalizeNode, Graph.PageID, InClassInput.NodeID, NewNodeIndex);
+			if (Graph.PageID == PageID)
 			{
 				NewNode = NewPageNode;
 			}
@@ -1052,35 +1063,14 @@ const FMetasoundFrontendNode* FMetaSoundFrontendDocumentBuilder::AddGraphInput(c
 	return nullptr;
 }
 
-const FMetasoundFrontendNode* FMetaSoundFrontendDocumentBuilder::DuplicateGraphInput(const FMetasoundFrontendClassInput& InClassInput, FMetasoundFrontendLiteral DefaultValue, const FName InName)
-{
-	using namespace Metasound::Frontend;
-
-	FDocumentIDGenerator& IDGenerator = FDocumentIDGenerator::Get();
-	const FMetasoundFrontendDocument& Doc = GetConstDocumentChecked();
-		
-	FMetasoundFrontendClassInput ClassInput = InClassInput;
-	ClassInput.NodeID = IDGenerator.CreateNodeID(Doc);
-	ClassInput.VertexID = IDGenerator.CreateVertexID(Doc);
-#if WITH_EDITORONLY_DATA
-	ClassInput.Metadata.SetDisplayName(FText::GetEmpty());
-#endif // WITH_EDITORONLY_DATA
-	ClassInput.Name = InName;
-	ClassInput.DefaultLiteral = static_cast<FMetasoundFrontendLiteral>(DefaultValue);
-
-	return AddGraphInput(ClassInput);
-}
-
-const FMetasoundFrontendNode* FMetaSoundFrontendDocumentBuilder::AddGraphOutput(const FMetasoundFrontendClassOutput& InClassOutput)
+const FMetasoundFrontendNode* FMetaSoundFrontendDocumentBuilder::AddGraphOutput(const FMetasoundFrontendClassOutput& InClassOutput, const FGuid* InPageID)
 {
 	using namespace Metasound::Frontend;
 
 	checkf(InClassOutput.NodeID.IsValid(), TEXT("Unassigned NodeID when adding graph output"));
 	checkf(InClassOutput.VertexID.IsValid(), TEXT("Unassigned VertexID when adding graph output"));
 
-	FMetasoundFrontendDocument& Document = GetDocumentChecked();
-	FMetasoundFrontendGraphClass& RootGraph = Document.RootGraph;
-
+	const FGuid& PageID = InPageID ? *InPageID : BuildPageID;
 	if (InClassOutput.TypeName.IsNone())
 	{
 		UE_LOG(LogMetaSound, Error, TEXT("TypeName unset when attempting to add class output '%s'"), *InClassOutput.Name.ToString());
@@ -1089,7 +1079,7 @@ const FMetasoundFrontendNode* FMetaSoundFrontendDocumentBuilder::AddGraphOutput(
 	else if (const FMetasoundFrontendClassOutput* Output = DocumentCache->GetInterfaceCache().FindOutput(InClassOutput.Name))
 	{
 		UE_LOG(LogMetaSound, Error, TEXT("Attempting to add MetaSound graph output '%s' when output with name already exists"), *InClassOutput.Name.ToString());
-		return DocumentCache->GetNodeCache().FindNode(Output->NodeID);
+		return DocumentCache->GetNodeCache(PageID).FindNode(Output->NodeID);
 	}
 	else if (!IDataTypeRegistry::Get().IsRegistered(InClassOutput.TypeName))
 	{
@@ -1109,7 +1099,22 @@ const FMetasoundFrontendNode* FMetaSoundFrontendDocumentBuilder::AddGraphOutput(
 		{
 			DocumentBuilderPrivate::SetNodeAndVertexNames(InOutNode, InClassOutput);
 		};
-		if (FMetasoundFrontendNode* NewNode = AddNodeInternal(Class.Metadata, FinalizeNode, InClassOutput.NodeID))
+
+		bool bAddedNodes = true;
+		FMetasoundFrontendNode* NewNodeToReturn = nullptr;
+		FMetasoundFrontendDocument& Document = GetDocumentChecked();
+		FMetasoundFrontendGraphClass& RootGraph = Document.RootGraph;
+		Document.RootGraph.IterateGraphPages([&](FMetasoundFrontendGraph& Graph)
+		{
+			FMetasoundFrontendNode* NewNode = AddNodeInternal(Class.Metadata, FinalizeNode, Graph.PageID, InClassOutput.NodeID);
+			if (Graph.PageID == PageID)
+			{
+				NewNodeToReturn = NewNode;
+			}
+			bAddedNodes &= NewNode != nullptr;
+		});
+
+		if (bAddedNodes)
 		{
 			const int32 NewIndex = RootGraph.Interface.Outputs.Num();
 			FMetasoundFrontendClassOutput& NewOutput = RootGraph.Interface.Outputs.Add_GetRef(InClassOutput);
@@ -1122,32 +1127,15 @@ const FMetasoundFrontendNode* FMetaSoundFrontendDocumentBuilder::AddGraphOutput(
 #if WITH_EDITORONLY_DATA
 			Document.Metadata.ModifyContext.AddMemberIDModified(InClassOutput.NodeID);
 #endif // WITH_EDITORONLY_DATA
-
-			return NewNode;
 		}
+
+		check(NewNodeToReturn);
+		return NewNodeToReturn;
 	}
 
 	return nullptr;
 }
 
-const FMetasoundFrontendNode* FMetaSoundFrontendDocumentBuilder::DuplicateGraphOutput(const FMetasoundFrontendClassOutput& InClassOutput, const FName InName)
-{
-	using namespace Metasound::Frontend;
-
-	FDocumentIDGenerator& IDGenerator = FDocumentIDGenerator::Get();
-	const FMetasoundFrontendDocument& Doc = GetConstDocumentChecked();
-
-	FMetasoundFrontendClassOutput ClassOutput = InClassOutput;
-	ClassOutput.NodeID = IDGenerator.CreateNodeID(Doc);
-	ClassOutput.VertexID = IDGenerator.CreateVertexID(Doc);
-#if WITH_EDITORONLY_DATA
-	ClassOutput.Metadata.SetDisplayName(FText::GetEmpty());
-#endif // WITH_EDITORONLY_DATA
-	ClassOutput.Name = InName;
-	
-	return AddGraphOutput(ClassOutput);
-}
-
 bool FMetaSoundFrontendDocumentBuilder::AddInterface(FName InterfaceName)
 {
 	using namespace Metasound::Frontend;
@@ -1182,7 +1170,7 @@ bool FMetaSoundFrontendDocumentBuilder::AddInterface(FName InterfaceName)
 	return false;
 }
 
-const FMetasoundFrontendNode* FMetaSoundFrontendDocumentBuilder::AddGraphNode(const FMetasoundFrontendGraphClass& InGraphClass, FGuid InNodeID)
+const FMetasoundFrontendNode* FMetaSoundFrontendDocumentBuilder::AddGraphNode(const FMetasoundFrontendGraphClass& InGraphClass, FGuid InNodeID, const FGuid* InPageID)
 {
 	auto FinalizeNode = [](FMetasoundFrontendNode& InOutNode, const Metasound::Frontend::FNodeRegistryKey& ClassKey)
 	{
@@ -1209,10 +1197,12 @@ const FMetasoundFrontendNode* FMetaSoundFrontendDocumentBuilder::AddGraphNode(co
 		AddDependency(InGraphClass);
 	}
 
-	return AddNodeInternal(NewClassMetadata, FinalizeNode, InNodeID);
+	constexpr int32* NewNodeIndex = nullptr;
+	const FGuid& PageID = InPageID ? *InPageID : BuildPageID;
+	return AddNodeInternal(NewClassMetadata, FinalizeNode, PageID, InNodeID, NewNodeIndex);
 }
 
-const FMetasoundFrontendNode* FMetaSoundFrontendDocumentBuilder::AddNodeByClassName(const FMetasoundFrontendClassName& InClassName, int32 InMajorVersion, FGuid InNodeID)
+const FMetasoundFrontendNode* FMetaSoundFrontendDocumentBuilder::AddNodeByClassName(const FMetasoundFrontendClassName& InClassName, int32 InMajorVersion, FGuid InNodeID, const FGuid* InPageID)
 {
 	using namespace Metasound;
 	using namespace Metasound::Frontend;
@@ -1240,7 +1230,10 @@ const FMetasoundFrontendNode* FMetaSoundFrontendDocumentBuilder::AddNodeByClassN
 
 		if (Dependency)
 		{
-			return AddNodeInternal(Dependency->Metadata, [](const FMetasoundFrontendNode& Node, const Metasound::Frontend::FNodeRegistryKey& ClassKey) { return Node.Name; }, InNodeID);
+			auto FinalizeNode = [](const FMetasoundFrontendNode& Node, const Metasound::Frontend::FNodeRegistryKey& ClassKey) { return Node.Name; };
+			constexpr int32* NewNodeIndex = nullptr;
+			const FGuid& PageID = InPageID ? *InPageID : BuildPageID;
+			return AddNodeInternal(Dependency->Metadata, FinalizeNode, PageID, InNodeID, NewNodeIndex);
 		}
 	}
 
@@ -1248,7 +1241,7 @@ const FMetasoundFrontendNode* FMetaSoundFrontendDocumentBuilder::AddNodeByClassN
 	return nullptr;
 }
 
-const FMetasoundFrontendNode* FMetaSoundFrontendDocumentBuilder::AddNodeByTemplate(const Metasound::Frontend::INodeTemplate& InTemplate, Metasound::Frontend::FNodeTemplateGenerateInterfaceParams Params, FGuid InNodeID)
+const FMetasoundFrontendNode* FMetaSoundFrontendDocumentBuilder::AddNodeByTemplate(const Metasound::Frontend::INodeTemplate& InTemplate, Metasound::Frontend::FNodeTemplateGenerateInterfaceParams Params, FGuid InNodeID, const FGuid* InPageID)
 {
 	using namespace Metasound;
 	using namespace Metasound::Frontend;
@@ -1267,14 +1260,16 @@ const FMetasoundFrontendNode* FMetaSoundFrontendDocumentBuilder::AddNodeByTempla
 		return Node.Name;
 	};
 
-	FMetasoundFrontendNode* NewNode = AddNodeInternal(Dependency->Metadata, FinalizeNodeFunction, InNodeID);
+	constexpr int32* NewNodeIndex = nullptr;
+	const FGuid & PageID = InPageID ? *InPageID : BuildPageID;
+	FMetasoundFrontendNode* NewNode = AddNodeInternal(Dependency->Metadata, FinalizeNodeFunction, PageID, InNodeID, NewNodeIndex);
 	check(NewNode);
 	NewNode->Interface = InTemplate.GenerateNodeInterface(MoveTemp(Params));
 
 	return NewNode;
 }
 
-FMetasoundFrontendNode* FMetaSoundFrontendDocumentBuilder::AddNodeInternal(const FMetasoundFrontendClassMetadata& InClassMetadata, FFinalizeNodeFunctionRef FinalizeNode, FGuid InNodeID, int32* NewNodeIndex, const FGuid* InGraphPageID)
+FMetasoundFrontendNode* FMetaSoundFrontendDocumentBuilder::AddNodeInternal(const FMetasoundFrontendClassMetadata& InClassMetadata, FFinalizeNodeFunctionRef FinalizeNode, const FGuid& InPageID, FGuid InNodeID, int32* NewNodeIndex)
 {
 	METASOUND_TRACE_CPUPROFILER_EVENT_SCOPE(FMetaSoundFrontendDocumentBuilder::AddNodeInternal);
 
@@ -1284,43 +1279,55 @@ FMetasoundFrontendNode* FMetaSoundFrontendDocumentBuilder::AddNodeInternal(const
 	if (const FMetasoundFrontendClass* Dependency = DocumentCache->FindDependency(ClassKey))
 	{
 		FMetasoundFrontendDocument& Document = GetDocumentChecked();
-		FMetasoundFrontendGraph& Graph = Document.RootGraph.FindGraphChecked(InGraphPageID ? *InGraphPageID : BuildPageID);
+		FMetasoundFrontendGraph& Graph = Document.RootGraph.FindGraphChecked(InPageID);
 		TArray<FMetasoundFrontendNode>& Nodes = Graph.Nodes;
 		FMetasoundFrontendNode& Node = Nodes.Emplace_GetRef(*Dependency);
 		Node.UpdateID(InNodeID);
 		FinalizeNode(Node, ClassKey);
 
-		if (!InGraphPageID || *InGraphPageID == BuildPageID)
-		{
-			const int32 NewIndex = Nodes.Num() - 1;
-			const IDocumentGraphNodeCache& NodeCache = DocumentCache->GetNodeCache();
-			DocumentDelegates->NodeDelegates.OnNodeAdded.Broadcast(NewIndex);
+		const int32 NewIndex = Nodes.Num() - 1;
+		const IDocumentGraphNodeCache& NodeCache = DocumentCache->GetNodeCache(InPageID);
+		DocumentDelegates->FindNodeDelegatesChecked(InPageID).OnNodeAdded.Broadcast(NewIndex);
 
-			if (NewNodeIndex)
-			{
-				*NewNodeIndex = NewIndex;
-			}
+		if (NewNodeIndex)
+		{
+			*NewNodeIndex = NewIndex;
+		}
 
 #if WITH_EDITORONLY_DATA
-			Document.Metadata.ModifyContext.AddNodeIDModified(InNodeID);
+		Document.Metadata.ModifyContext.AddNodeIDModified(InNodeID);
 #endif // WITH_EDITORONLY_DATA
-		}
+
 		return &Node;
 	}
 
 	return nullptr;
 }
 
-bool FMetaSoundFrontendDocumentBuilder::CanAddEdge(const FMetasoundFrontendEdge& InEdge) const
+const FMetasoundFrontendGraph& FMetaSoundFrontendDocumentBuilder::AddGraphPage(const FGuid& InPageID, bool bDuplicateLastGraph, bool bSetAsBuildGraph)
+{
+	using namespace Metasound::Frontend;
+
+	const FMetasoundFrontendGraph& ToReturn = GetDocumentChecked().RootGraph.AddGraphPage(InPageID, bDuplicateLastGraph);
+	DocumentDelegates->AddPageDelegates(InPageID);
+	if (bSetAsBuildGraph)
+	{
+		SetBuildPageID(InPageID);
+	}
+	return ToReturn;
+}
+
+bool FMetaSoundFrontendDocumentBuilder::CanAddEdge(const FMetasoundFrontendEdge& InEdge, const FGuid* InPageID) const
 {
 	using namespace Metasound::Frontend;
 
+	const FGuid& PageID = InPageID ? *InPageID : BuildPageID;
 	const FMetasoundFrontendDocument& Document = GetConstDocumentChecked();
-	const IDocumentGraphEdgeCache& EdgeCache = DocumentCache->GetEdgeCache();
+	const IDocumentGraphEdgeCache& EdgeCache = DocumentCache->GetEdgeCache(PageID);
 
 	if (!EdgeCache.IsNodeInputConnected(InEdge.ToNodeID, InEdge.ToVertexID))
 	{
-		return IsValidEdge(InEdge) == EInvalidEdgeReason::None;
+		return IsValidEdge(InEdge, InPageID) == EInvalidEdgeReason::None;
 	}
 
 	return false;
@@ -1344,9 +1351,12 @@ void FMetaSoundFrontendDocumentBuilder::ClearDocument(TSharedRef<Metasound::Fron
 	GraphClass.PresetOptions.bIsPreset = false;
 
 	GraphClass.RemoveAllGraphPages();
+
+PRAGMA_DISABLE_DEPRECATION_WARNINGS
 	FMetasoundFrontendGraph& Graph = GraphClass.GetDefaultGraph();
 	Graph.Nodes.Reset();
 	Graph.Edges.Reset();
+PRAGMA_ENABLE_DEPRECATION_WARNINGS
 
 	Reload(ModifyDelegates);
 }
@@ -1362,50 +1372,54 @@ bool FMetaSoundFrontendDocumentBuilder::ConformGraphInputNodeToClass(const FMeta
 {
 	using namespace Metasound::Frontend;
 
-	FMetasoundFrontendDocument& Document = GetDocumentChecked();
-	const IDocumentGraphNodeCache& NodeCache = DocumentCache->GetNodeCache();
-	if (const int32* NodeIndexPtr = NodeCache.FindNodeIndex(GraphInput.NodeID))
+	FMetasoundFrontendClass Class;
+	const bool bClassFound = DocumentBuilderPrivate::FindInputRegistryClass(GraphInput.TypeName, GraphInput.AccessType, Class);
+	if (ensureAlways(bClassFound))
 	{
-		TArray<FMetasoundFrontendNode>& Nodes = FindBuildGraphChecked().Nodes;
-		FMetasoundFrontendNode& Node = Nodes[*NodeIndexPtr];
-		FMetasoundFrontendClass Class;
-		const bool bClassFound = DocumentBuilderPrivate::FindInputRegistryClass(GraphInput.TypeName, GraphInput.AccessType, Class);
-		if (ensureAlways(bClassFound))
+		FMetasoundFrontendDocument& Document = GetDocumentChecked();
+		const FMetasoundFrontendClass* Dependency = FindDependency(Class.Metadata);
+		if (!Dependency)
 		{
-			const FMetasoundFrontendClass* Dependency = FindDependency(Class.Metadata);
-			if (!Dependency)
-			{
-				Dependency = AddDependency(Class);
-			}
+			Dependency = AddDependency(Class);
+		}
 
-			if (ensureAlways(Dependency))
+		if (ensureAlways(Dependency))
+		{
+			Document.RootGraph.IterateGraphPages([this, &Document, &Dependency, &GraphInput](FMetasoundFrontendGraph& Graph)
 			{
-				const int32 RemovalIndex = *NodeIndexPtr; // Have to cache as next delegate broadcast invalidates index pointer
-				DocumentDelegates->NodeDelegates.OnRemoveSwappingNode.Broadcast(RemovalIndex, Nodes.Num() - 1);
-				FMetasoundFrontendNode NewNode = MoveTemp(Node);
-				Nodes.RemoveAtSwap(RemovalIndex, EAllowShrinking::No);
-				NewNode.ClassID = Dependency->ID;
-				NewNode.Interface.Inputs.Last().TypeName = GraphInput.TypeName;
-				NewNode.Interface.Outputs.Last().TypeName = GraphInput.TypeName;
+				const IDocumentGraphNodeCache& NodeCache = DocumentCache->GetNodeCache(Graph.PageID);
+				if (const int32* NodeIndexPtr = NodeCache.FindNodeIndex(GraphInput.NodeID))
+				{
+					TArray<FMetasoundFrontendNode>& Nodes = Graph.Nodes;
+					FMetasoundFrontendNode& Node = Nodes[*NodeIndexPtr];
+					FNodeModifyDelegates& NodeDelegates = DocumentDelegates->FindNodeDelegatesChecked(Graph.PageID);
+					const int32 RemovalIndex = *NodeIndexPtr; // Have to cache as next delegate broadcast invalidates index pointer
+					NodeDelegates.OnRemoveSwappingNode.Broadcast(RemovalIndex, Nodes.Num() - 1);
+					FMetasoundFrontendNode NewNode = MoveTemp(Node);
+					Nodes.RemoveAtSwap(RemovalIndex, EAllowShrinking::No);
+					NewNode.ClassID = Dependency->ID;
+					NewNode.Interface.Inputs.Last().TypeName = GraphInput.TypeName;
+					NewNode.Interface.Outputs.Last().TypeName = GraphInput.TypeName;
 
 #if WITH_EDITORONLY_DATA
-				Document.Metadata.ModifyContext.AddNodeIDModified(NewNode.GetID());
+					Document.Metadata.ModifyContext.AddNodeIDModified(NewNode.GetID());
 #endif // WITH_EDITORONLY_DATA
 
-				// Set the default literal on the nodes inputs so that it gets passed to the instantiated TInputNode on a live
-				// auditioned MetaSound.
-				DocumentBuilderPrivate::SetDefaultLiteralOnInputNode(NewNode, GraphInput);
+					// Set the default literal on the nodes inputs so that it gets passed to the instantiated TInputNode on a live
+					// auditioned MetaSound.
+					DocumentBuilderPrivate::SetDefaultLiteralOnInputNode(NewNode, GraphInput);
 
-				FMetasoundFrontendNode& NewNodeRef = Nodes.Add_GetRef(MoveTemp(NewNode));
-				DocumentDelegates->NodeDelegates.OnNodeAdded.Broadcast(Nodes.Num() - 1);
+					FMetasoundFrontendNode& NewNodeRef = Nodes.Add_GetRef(MoveTemp(NewNode));
+					NodeDelegates.OnNodeAdded.Broadcast(Nodes.Num() - 1);
 
-				// Remove the default literal on the node added during the "FinalizeNode" call. This matches how 
-				// nodes are serialized in editor. The default literals are only stored on the FMetasoundFrontendClassInputs.
-				NewNodeRef.InputLiterals.Reset();
+					// Remove the default literal on the node added during the "FinalizeNode" call. This matches how 
+					// nodes are serialized in editor. The default literals are only stored on the FMetasoundFrontendClassInputs.
+					NewNodeRef.InputLiterals.Reset();
+				}
+			});
 
-				RemoveUnusedDependencies();
-				return true;
-			}
+			RemoveUnusedDependencies();
+			return true;
 		}
 	}
 
@@ -1416,41 +1430,45 @@ bool FMetaSoundFrontendDocumentBuilder::ConformGraphOutputNodeToClass(const FMet
 {
 	using namespace Metasound::Frontend;
 
-	FMetasoundFrontendDocument& Document = GetDocumentChecked();
-	const IDocumentGraphNodeCache& NodeCache = DocumentCache->GetNodeCache();
-	if (const int32* NodeIndexPtr = NodeCache.FindNodeIndex(GraphOutput.NodeID))
+	FMetasoundFrontendClass Class;
+	const bool bClassFound = DocumentBuilderPrivate::FindOutputRegistryClass(GraphOutput.TypeName, GraphOutput.AccessType, Class);
+	if (ensureAlways(bClassFound))
 	{
-		TArray<FMetasoundFrontendNode>& Nodes = FindBuildGraphChecked().Nodes;
-		FMetasoundFrontendNode& Node = Nodes[*NodeIndexPtr];
-		FMetasoundFrontendClass Class;
-		const bool bClassFound = DocumentBuilderPrivate::FindOutputRegistryClass(GraphOutput.TypeName, GraphOutput.AccessType, Class);
-		if (ensureAlways(bClassFound))
+		FMetasoundFrontendDocument& Document = GetDocumentChecked();
+		const FMetasoundFrontendClass* Dependency = FindDependency(Class.Metadata);
+		if (!Dependency)
 		{
-			const FMetasoundFrontendClass* Dependency = FindDependency(Class.Metadata);
-			if (!Dependency)
-			{
-				Dependency = AddDependency(Class);
-			}
+			Dependency = AddDependency(Class);
+		}
 
-			if (ensureAlways(Dependency))
+		if (ensureAlways(Dependency))
+		{
+			Document.RootGraph.IterateGraphPages([this, &Document, &Dependency, &GraphOutput](FMetasoundFrontendGraph& Graph)
 			{
-				const int32 RemovalIndex = *NodeIndexPtr; // Have to cache as next delegate broadcast invalidates index pointer
-				DocumentDelegates->NodeDelegates.OnRemoveSwappingNode.Broadcast(RemovalIndex, Nodes.Num() - 1);
-				FMetasoundFrontendNode NewNode = MoveTemp(Node);
-				Nodes.RemoveAtSwap(RemovalIndex, EAllowShrinking::No);
-				NewNode.ClassID = Dependency->ID;
-				NewNode.Interface.Inputs.Last().TypeName = GraphOutput.TypeName;
-				NewNode.Interface.Outputs.Last().TypeName = GraphOutput.TypeName;
+				const IDocumentGraphNodeCache& NodeCache = DocumentCache->GetNodeCache(Graph.PageID);
+				if (const int32* NodeIndexPtr = NodeCache.FindNodeIndex(GraphOutput.NodeID))
+				{
+					TArray<FMetasoundFrontendNode>& Nodes = Graph.Nodes;
+					FMetasoundFrontendNode& Node = Nodes[*NodeIndexPtr];
+					FNodeModifyDelegates& NodeDelegates = DocumentDelegates->FindNodeDelegatesChecked(Graph.PageID);
+					const int32 RemovalIndex = *NodeIndexPtr; // Have to cache as next delegate broadcast invalidates index pointer
+					NodeDelegates.OnRemoveSwappingNode.Broadcast(RemovalIndex, Nodes.Num() - 1);
+					FMetasoundFrontendNode NewNode = MoveTemp(Node);
+					Nodes.RemoveAtSwap(RemovalIndex, EAllowShrinking::No);
+					NewNode.ClassID = Dependency->ID;
+					NewNode.Interface.Inputs.Last().TypeName = GraphOutput.TypeName;
+					NewNode.Interface.Outputs.Last().TypeName = GraphOutput.TypeName;
 
 #if WITH_EDITORONLY_DATA
-				Document.Metadata.ModifyContext.AddNodeIDModified(NewNode.GetID());
+					Document.Metadata.ModifyContext.AddNodeIDModified(NewNode.GetID());
 #endif // WITH_EDITORONLY_DATA
-				Nodes.Add(MoveTemp(NewNode));
-				DocumentDelegates->NodeDelegates.OnNodeAdded.Broadcast(Nodes.Num() - 1);
+					Nodes.Add(MoveTemp(NewNode));
+					NodeDelegates.OnNodeAdded.Broadcast(Nodes.Num() - 1);
+				}	
+			});	
 
-				RemoveUnusedDependencies();
-				return true;
-			}
+			RemoveUnusedDependencies();
+			return true;
 		}
 	}
 
@@ -1462,17 +1480,17 @@ bool FMetaSoundFrontendDocumentBuilder::ContainsDependencyOfType(EMetasoundFront
 	return DocumentCache->ContainsDependencyOfType(ClassType);
 }
 
-bool FMetaSoundFrontendDocumentBuilder::ContainsEdge(const FMetasoundFrontendEdge& InEdge) const
+bool FMetaSoundFrontendDocumentBuilder::ContainsEdge(const FMetasoundFrontendEdge& InEdge, const FGuid* InPageID) const
 {
 	using namespace Metasound::Frontend;
-	const IDocumentGraphEdgeCache& EdgeCache = DocumentCache->GetEdgeCache();
+	const IDocumentGraphEdgeCache& EdgeCache = DocumentCache->GetEdgeCache(InPageID ? *InPageID : BuildPageID);
 	return EdgeCache.ContainsEdge(InEdge);
 }
 
-bool FMetaSoundFrontendDocumentBuilder::ContainsNode(const FGuid& InNodeID) const
+bool FMetaSoundFrontendDocumentBuilder::ContainsNode(const FGuid& InNodeID, const FGuid* InPageID) const
 {
 	using namespace Metasound::Frontend;
-	const IDocumentGraphNodeCache& NodeCache = DocumentCache->GetNodeCache();
+	const IDocumentGraphNodeCache& NodeCache = DocumentCache->GetNodeCache(InPageID ? * InPageID : BuildPageID);
 	return NodeCache.ContainsNode(InNodeID);
 }
 
@@ -1540,6 +1558,47 @@ bool FMetaSoundFrontendDocumentBuilder::ConvertToPreset(const FMetasoundFrontend
 	return false;
 }
 
+const FMetasoundFrontendNode* FMetaSoundFrontendDocumentBuilder::DuplicateGraphInput(const FMetasoundFrontendClassInput& InClassInput, FMetasoundFrontendLiteral DefaultValue, const FName InName, const FGuid* InPageID)
+{
+	using namespace Metasound::Frontend;
+
+	FDocumentIDGenerator& IDGenerator = FDocumentIDGenerator::Get();
+	const FMetasoundFrontendDocument& Doc = GetConstDocumentChecked();
+
+	const FGuid& PageID = InPageID ? *InPageID : BuildPageID;
+
+	FMetasoundFrontendClassInput ClassInput = InClassInput;
+	ClassInput.NodeID = IDGenerator.CreateNodeID(Doc);
+	ClassInput.VertexID = IDGenerator.CreateVertexID(Doc);
+#if WITH_EDITORONLY_DATA
+	ClassInput.Metadata.SetDisplayName(FText::GetEmpty());
+#endif // WITH_EDITORONLY_DATA
+	ClassInput.Name = InName;
+	ClassInput.DefaultLiteral = static_cast<FMetasoundFrontendLiteral>(DefaultValue);
+
+	return AddGraphInput(ClassInput, &PageID);
+}
+
+const FMetasoundFrontendNode* FMetaSoundFrontendDocumentBuilder::DuplicateGraphOutput(const FMetasoundFrontendClassOutput& InClassOutput, const FName InName, const FGuid* InPageID)
+{
+	using namespace Metasound::Frontend;
+
+	FDocumentIDGenerator& IDGenerator = FDocumentIDGenerator::Get();
+	const FMetasoundFrontendDocument& Doc = GetConstDocumentChecked();
+
+	const FGuid& PageID = InPageID ? *InPageID : BuildPageID;
+
+	FMetasoundFrontendClassOutput ClassOutput = InClassOutput;
+	ClassOutput.NodeID = IDGenerator.CreateNodeID(Doc);
+	ClassOutput.VertexID = IDGenerator.CreateVertexID(Doc);
+#if WITH_EDITORONLY_DATA
+	ClassOutput.Metadata.SetDisplayName(FText::GetEmpty());
+#endif // WITH_EDITORONLY_DATA
+	ClassOutput.Name = InName;
+
+	return AddGraphOutput(ClassOutput, &PageID);
+}
+
 FMetasoundFrontendGraph& FMetaSoundFrontendDocumentBuilder::FindBuildGraphChecked()
 {
 	return GetDocumentChecked().RootGraph.FindGraphChecked(BuildPageID);
@@ -1593,11 +1652,11 @@ const FMetasoundFrontendClass* FMetaSoundFrontendDocumentBuilder::FindDependency
 	return DocumentCache->FindDependency(RegistryKey);
 }
 
-TArray<const FMetasoundFrontendEdge*> FMetaSoundFrontendDocumentBuilder::FindEdges(const FGuid& InNodeID, const FGuid& InVertexID) const
+TArray<const FMetasoundFrontendEdge*> FMetaSoundFrontendDocumentBuilder::FindEdges(const FGuid& InNodeID, const FGuid& InVertexID, const FGuid* InPageID) const
 {
 	using namespace Metasound::Frontend;
 
-	const IDocumentGraphEdgeCache& EdgeCache = DocumentCache->GetEdgeCache();
+	const IDocumentGraphEdgeCache& EdgeCache = DocumentCache->GetEdgeCache(InPageID ? *InPageID : BuildPageID);
 	return EdgeCache.FindEdges(InNodeID, InVertexID);
 }
 
@@ -1615,18 +1674,20 @@ FMetaSoundFrontendGraphComment* FMetaSoundFrontendDocumentBuilder::FindGraphComm
 }
 #endif // WITH_EDITOR
 
-bool FMetaSoundFrontendDocumentBuilder::FindInterfaceInputNodes(FName InterfaceName, TArray<const FMetasoundFrontendNode*>& OutInputs) const
+bool FMetaSoundFrontendDocumentBuilder::FindInterfaceInputNodes(FName InterfaceName, TArray<const FMetasoundFrontendNode*>& OutInputs, const FGuid* InPageID) const
 {
 	using namespace Metasound::Frontend;
 
 	OutInputs.Reset();
 
+	const FGuid& PageID = InPageID ? *InPageID : BuildPageID;
 	FMetasoundFrontendInterface Interface;
+	const TSet<FMetasoundFrontendVersion>& Interfaces = GetConstDocumentChecked().Interfaces;
 	if (ISearchEngine::Get().FindInterfaceWithHighestVersion(InterfaceName, Interface))
 	{
-		if (GetConstDocumentChecked().Interfaces.Contains(Interface.Version))
+		if (Interfaces.Contains(Interface.Version))
 		{
-			const IDocumentGraphNodeCache& NodeCache = DocumentCache->GetNodeCache();
+			const IDocumentGraphNodeCache& NodeCache = DocumentCache->GetNodeCache(PageID);
 			const IDocumentGraphInterfaceCache& InterfaceCache = DocumentCache->GetInterfaceCache();
 
 			TArray<const FMetasoundFrontendNode*> InterfaceInputs;
@@ -1656,18 +1717,21 @@ bool FMetaSoundFrontendDocumentBuilder::FindInterfaceInputNodes(FName InterfaceN
 	return false;
 }
 
-bool FMetaSoundFrontendDocumentBuilder::FindInterfaceOutputNodes(FName InterfaceName, TArray<const FMetasoundFrontendNode*>& OutOutputs) const
+bool FMetaSoundFrontendDocumentBuilder::FindInterfaceOutputNodes(FName InterfaceName, TArray<const FMetasoundFrontendNode*>& OutOutputs, const FGuid* InPageID) const
 {
 	using namespace Metasound::Frontend;
 
 	OutOutputs.Reset();
 
+	const FGuid& PageID = InPageID ? *InPageID : BuildPageID;
+
 	FMetasoundFrontendInterface Interface;
+	const TSet<FMetasoundFrontendVersion>& Interfaces = GetConstDocumentChecked().Interfaces;
 	if (ISearchEngine::Get().FindInterfaceWithHighestVersion(InterfaceName, Interface))
 	{
-		if (GetConstDocumentChecked().Interfaces.Contains(Interface.Version))
+		if (Interfaces.Contains(Interface.Version))
 		{
-			const IDocumentGraphNodeCache& NodeCache = DocumentCache->GetNodeCache();
+			const IDocumentGraphNodeCache& NodeCache = DocumentCache->GetNodeCache(PageID);
 			const IDocumentGraphInterfaceCache& InterfaceCache = DocumentCache->GetInterfaceCache();
 
 			TArray<const FMetasoundFrontendNode*> InterfaceOutputs;
@@ -1702,13 +1766,14 @@ const FMetasoundFrontendClassInput* FMetaSoundFrontendDocumentBuilder::FindGraph
 	return DocumentCache->GetInterfaceCache().FindInput(InputName);
 }
 
-const FMetasoundFrontendNode* FMetaSoundFrontendDocumentBuilder::FindGraphInputNode(FName InputName) const
+const FMetasoundFrontendNode* FMetaSoundFrontendDocumentBuilder::FindGraphInputNode(FName InputName, const FGuid* InPageID) const
 {
 	using namespace Metasound::Frontend;
 
 	if (const FMetasoundFrontendClassInput* InputClass = FindGraphInput(InputName))
 	{
-		const IDocumentGraphNodeCache& NodeCache = DocumentCache->GetNodeCache();
+		const FGuid& PageID = InPageID ? *InPageID : BuildPageID;
+		const IDocumentGraphNodeCache& NodeCache = DocumentCache->GetNodeCache(PageID);
 		return NodeCache.FindNode(InputClass->NodeID);
 	}
 
@@ -1720,13 +1785,14 @@ const FMetasoundFrontendClassOutput* FMetaSoundFrontendDocumentBuilder::FindGrap
 	return DocumentCache->GetInterfaceCache().FindOutput(OutputName);
 }
 
-const FMetasoundFrontendNode* FMetaSoundFrontendDocumentBuilder::FindGraphOutputNode(FName OutputName) const
+const FMetasoundFrontendNode* FMetaSoundFrontendDocumentBuilder::FindGraphOutputNode(FName OutputName, const FGuid* InPageID) const
 {
 	using namespace Metasound::Frontend;
 
 	if (const FMetasoundFrontendClassOutput* OutputClass = FindGraphOutput(OutputName))
 	{
-		const IDocumentGraphNodeCache& NodeCache = DocumentCache->GetNodeCache();
+		const FGuid& PageID = InPageID ? *InPageID : BuildPageID;
+		const IDocumentGraphNodeCache& NodeCache = DocumentCache->GetNodeCache(PageID);
 		return NodeCache.FindNode(OutputClass->NodeID);
 	}
 
@@ -1743,20 +1809,21 @@ UMetaSoundFrontendMemberMetadata* FMetaSoundFrontendDocumentBuilder::FindMemberM
 }
 #endif // WITH_EDITOR
 
-const FMetasoundFrontendNode* FMetaSoundFrontendDocumentBuilder::FindNode(const FGuid& InNodeID) const
+const FMetasoundFrontendNode* FMetaSoundFrontendDocumentBuilder::FindNode(const FGuid& InNodeID, const FGuid* InPageID) const
 {
 	using namespace Metasound::Frontend;
-	const IDocumentGraphNodeCache& NodeCache = DocumentCache->GetNodeCache();
+	const FGuid& PageID = InPageID ? *InPageID : BuildPageID;
+	const IDocumentGraphNodeCache& NodeCache = DocumentCache->GetNodeCache(PageID);
 	return NodeCache.FindNode(InNodeID);
 }
 
-bool FMetaSoundFrontendDocumentBuilder::FindNodeClassInterfaces(const FGuid& InNodeID, TSet<FMetasoundFrontendVersion>& OutInterfaces) const
+bool FMetaSoundFrontendDocumentBuilder::FindNodeClassInterfaces(const FGuid& InNodeID, TSet<FMetasoundFrontendVersion>& OutInterfaces, const FGuid& InPageID) const
 {
 	using namespace Metasound;
 	using namespace Metasound::Frontend;
 
 	const FMetasoundFrontendDocument& Document = GetConstDocumentChecked();
-	const IDocumentGraphNodeCache& NodeCache = DocumentCache->GetNodeCache();
+	const IDocumentGraphNodeCache& NodeCache = DocumentCache->GetNodeCache(InPageID);
 	if (const FMetasoundFrontendNode* Node = NodeCache.FindNode(InNodeID))
 	{
 		if (const FMetasoundFrontendClass* NodeClass = DocumentCache->FindDependency(Node->ClassID))
@@ -1769,31 +1836,35 @@ bool FMetaSoundFrontendDocumentBuilder::FindNodeClassInterfaces(const FGuid& InN
 	return false;
 }
 
-const FMetasoundFrontendVertex* FMetaSoundFrontendDocumentBuilder::FindNodeInput(const FGuid& InNodeID, const FGuid& InVertexID) const
+const FMetasoundFrontendVertex* FMetaSoundFrontendDocumentBuilder::FindNodeInput(const FGuid& InNodeID, const FGuid& InVertexID, const FGuid* InPageID) const
 {
 	using namespace Metasound::Frontend;
+	const FGuid& PageID = InPageID ? *InPageID : BuildPageID;
 	const IDocumentGraphNodeCache& NodeCache = DocumentCache->GetNodeCache();
 	return NodeCache.FindInputVertex(InNodeID, InVertexID);
 }
 
-const FMetasoundFrontendVertex* FMetaSoundFrontendDocumentBuilder::FindNodeInput(const FGuid& InNodeID, FName InVertexName) const
+const FMetasoundFrontendVertex* FMetaSoundFrontendDocumentBuilder::FindNodeInput(const FGuid& InNodeID, FName InVertexName, const FGuid* InPageID) const
 {
 	using namespace Metasound::Frontend;
-	const IDocumentGraphNodeCache& NodeCache = DocumentCache->GetNodeCache();
+	const FGuid& PageID = InPageID ? *InPageID : BuildPageID;
+	const IDocumentGraphNodeCache& NodeCache = DocumentCache->GetNodeCache(PageID);
 	return NodeCache.FindInputVertex(InNodeID, InVertexName);
 }
 
-TArray<const FMetasoundFrontendVertex*> FMetaSoundFrontendDocumentBuilder::FindNodeInputs(const FGuid& InNodeID, FName TypeName) const
+TArray<const FMetasoundFrontendVertex*> FMetaSoundFrontendDocumentBuilder::FindNodeInputs(const FGuid& InNodeID, FName TypeName, const FGuid* InPageID) const
 {
-	return DocumentCache->GetNodeCache().FindNodeInputs(InNodeID, TypeName);
+	const FGuid& PageID = InPageID ? *InPageID : BuildPageID;
+	return DocumentCache->GetNodeCache(PageID).FindNodeInputs(InNodeID, TypeName);
 }
 
-TArray<const FMetasoundFrontendVertex*> FMetaSoundFrontendDocumentBuilder::FindNodeInputsConnectedToNodeOutput(const FGuid& InOutputNodeID, const FGuid& InOutputVertexID, TArray<const FMetasoundFrontendNode*>* ConnectedInputNodes) const
+TArray<const FMetasoundFrontendVertex*> FMetaSoundFrontendDocumentBuilder::FindNodeInputsConnectedToNodeOutput(const FGuid& InOutputNodeID, const FGuid& InOutputVertexID, TArray<const FMetasoundFrontendNode*>* ConnectedInputNodes, const FGuid* InPageID) const
 {
 	using namespace Metasound::Frontend;
 
-	const IDocumentGraphEdgeCache& EdgeCache = DocumentCache->GetEdgeCache();
-	const IDocumentGraphNodeCache& NodeCache = DocumentCache->GetNodeCache();
+	const FGuid& PageID = InPageID ? *InPageID : BuildPageID;
+	const IDocumentGraphEdgeCache& EdgeCache = DocumentCache->GetEdgeCache(PageID);
+	const IDocumentGraphNodeCache& NodeCache = DocumentCache->GetNodeCache(PageID);
 
 	const FMetasoundFrontendDocument& Document = GetConstDocumentChecked();
 
@@ -1803,7 +1874,7 @@ TArray<const FMetasoundFrontendVertex*> FMetaSoundFrontendDocumentBuilder::FindN
 	}
 
 	TArray<const FMetasoundFrontendVertex*> Inputs;
-	const FMetasoundFrontendGraph& Graph = FindConstBuildGraphChecked();
+	const FMetasoundFrontendGraph& Graph = Document.RootGraph.FindConstGraphChecked(PageID);
 	const TArrayView<const int32> Indices = EdgeCache.FindEdgeIndicesFromNodeOutput(InOutputNodeID, InOutputVertexID);
 	Algo::Transform(Indices, Inputs, [&Graph, &Document, &NodeCache, &ConnectedInputNodes](const int32& Index)
 	{
@@ -1817,51 +1888,57 @@ TArray<const FMetasoundFrontendVertex*> FMetaSoundFrontendDocumentBuilder::FindN
 	return Inputs;
 }
 
-FMetasoundFrontendNode* FMetaSoundFrontendDocumentBuilder::FindNodeInternal(const FGuid& InNodeID)
+FMetasoundFrontendNode* FMetaSoundFrontendDocumentBuilder::FindNodeInternal(const FGuid& InNodeID, const FGuid* InPageID)
 {
 	using namespace Metasound::Frontend;
 
 	using namespace Metasound;
 	using namespace Metasound::Frontend;
 
-	const IDocumentGraphNodeCache& NodeCache = DocumentCache->GetNodeCache();
+	const FGuid& PageID = InPageID ? *InPageID : BuildPageID;
+	const IDocumentGraphNodeCache& NodeCache = DocumentCache->GetNodeCache(PageID);
 	if (const int32* NodeIndex = NodeCache.FindNodeIndex(InNodeID))
 	{
-		return &FindBuildGraphChecked().Nodes[*NodeIndex];
+		FMetasoundFrontendGraph& Graph = GetDocumentChecked().RootGraph.FindGraphChecked(PageID);
+		return &Graph.Nodes[*NodeIndex];
 	}
 
 	return nullptr;
 }
 
-const FMetasoundFrontendVertex* FMetaSoundFrontendDocumentBuilder::FindNodeOutput(const FGuid& InNodeID, const FGuid& InVertexID) const
+const FMetasoundFrontendVertex* FMetaSoundFrontendDocumentBuilder::FindNodeOutput(const FGuid& InNodeID, const FGuid& InVertexID, const FGuid* InPageID) const
 {
 	using namespace Metasound::Frontend;
-	const IDocumentGraphNodeCache& NodeCache = DocumentCache->GetNodeCache();
+	const FGuid& PageID = InPageID ? *InPageID : BuildPageID;
+	const IDocumentGraphNodeCache& NodeCache = DocumentCache->GetNodeCache(PageID);
 	return NodeCache.FindOutputVertex(InNodeID, InVertexID);
 }
 
-const FMetasoundFrontendVertex* FMetaSoundFrontendDocumentBuilder::FindNodeOutput(const FGuid& InNodeID, FName InVertexName) const
+const FMetasoundFrontendVertex* FMetaSoundFrontendDocumentBuilder::FindNodeOutput(const FGuid& InNodeID, FName InVertexName, const FGuid* InPageID) const
 {
 	using namespace Metasound::Frontend;
-	const IDocumentGraphNodeCache& NodeCache = DocumentCache->GetNodeCache();
+	const FGuid& PageID = InPageID ? *InPageID : BuildPageID;
+	const IDocumentGraphNodeCache& NodeCache = DocumentCache->GetNodeCache(PageID);
 	return NodeCache.FindOutputVertex(InNodeID, InVertexName);
 }
 
-TArray<const FMetasoundFrontendVertex*> FMetaSoundFrontendDocumentBuilder::FindNodeOutputs(const FGuid& InNodeID, FName TypeName) const
+TArray<const FMetasoundFrontendVertex*> FMetaSoundFrontendDocumentBuilder::FindNodeOutputs(const FGuid& InNodeID, FName TypeName, const FGuid* InPageID) const
 {
-	return DocumentCache->GetNodeCache().FindNodeOutputs(InNodeID, TypeName);
+	const FGuid& PageID = InPageID ? *InPageID : BuildPageID;
+	return DocumentCache->GetNodeCache(PageID).FindNodeOutputs(InNodeID, TypeName);
 }
 
-const FMetasoundFrontendVertex* FMetaSoundFrontendDocumentBuilder::FindNodeOutputConnectedToNodeInput(const FGuid& InInputNodeID, const FGuid& InInputVertexID, const FMetasoundFrontendNode** ConnectedOutputNode) const
+const FMetasoundFrontendVertex* FMetaSoundFrontendDocumentBuilder::FindNodeOutputConnectedToNodeInput(const FGuid& InInputNodeID, const FGuid& InInputVertexID, const FMetasoundFrontendNode** ConnectedOutputNode, const FGuid* InPageID) const
 {
 	using namespace Metasound::Frontend;
 
-	const IDocumentGraphEdgeCache& EdgeCache = DocumentCache->GetEdgeCache();
+	const FGuid& PageID = InPageID ? *InPageID : BuildPageID;
+	const IDocumentGraphEdgeCache& EdgeCache = DocumentCache->GetEdgeCache(PageID);
 	if (const int32* Index = EdgeCache.FindEdgeIndexToNodeInput(InInputNodeID, InInputVertexID))
 	{
 		const FMetasoundFrontendDocument& Document = GetConstDocumentChecked();
-		const FMetasoundFrontendEdge& Edge = FindConstBuildGraphChecked().Edges[*Index];
-		const IDocumentGraphNodeCache& NodeCache = DocumentCache->GetNodeCache();
+		const FMetasoundFrontendEdge& Edge = Document.RootGraph.FindConstGraphChecked(PageID).Edges[*Index];
+		const IDocumentGraphNodeCache& NodeCache = DocumentCache->GetNodeCache(PageID);
 		if (ConnectedOutputNode)
 		{
 			(*ConnectedOutputNode) = NodeCache.FindNode(Edge.FromNodeID);
@@ -1934,7 +2011,7 @@ FMetasoundFrontendDocument& FMetaSoundFrontendDocumentBuilder::GetDocumentChecke
 	return GetDocumentInterfaceChecked().GetDocument();
 }
 
-const Metasound::Frontend::FDocumentModifyDelegates& FMetaSoundFrontendDocumentBuilder::GetDocumentDelegates() const
+Metasound::Frontend::FDocumentModifyDelegates& FMetaSoundFrontendDocumentBuilder::GetDocumentDelegates()
 {
 	return *DocumentDelegates;
 }
@@ -1953,6 +2030,46 @@ IMetaSoundDocumentInterface& FMetaSoundFrontendDocumentBuilder::GetDocumentInter
 	return *Interface;
 }
 
+TArray<const FMetasoundFrontendNode*> FMetaSoundFrontendDocumentBuilder::GetGraphInputTemplateNodes(FName InInputName, const FGuid* InPageID)
+{
+	using namespace Metasound::Frontend;
+
+	TArray<const FMetasoundFrontendNode*> TemplateNodes;
+
+	const FGuid PageID = InPageID ? *InPageID : BuildPageID;
+	FMetasoundFrontendGraphClass& RootGraph = GetDocumentChecked().RootGraph;
+	if (const int32* Index = DocumentCache->GetInterfaceCache().FindInputIndex(InInputName))
+	{
+		const FMetasoundFrontendClassInput& InputClass = RootGraph.Interface.Inputs[*Index];
+		const FMetasoundFrontendGraph& Graph = RootGraph.FindConstGraphChecked(PageID);
+		const IDocumentGraphNodeCache& NodeCache = DocumentCache->GetNodeCache(PageID);
+		const IDocumentGraphEdgeCache& EdgeCache = DocumentCache->GetEdgeCache(PageID);
+
+		if (const FMetasoundFrontendNode* InputNode = NodeCache.FindNode(InputClass.NodeID))
+		{
+			const FGuid OutputVertexID = InputNode->Interface.Outputs.Last().VertexID;
+			const TArray<const FMetasoundFrontendEdge*> ConnectedEdges = EdgeCache.FindEdges(InputClass.NodeID, OutputVertexID);
+			for (const FMetasoundFrontendEdge* Edge : ConnectedEdges)
+			{
+				check(Edge);
+				if (const int32* ConnectedNodeIndex = NodeCache.FindNodeIndex(Edge->ToNodeID))
+				{
+					const FMetasoundFrontendNode& ConnectedNode = Graph.Nodes[*ConnectedNodeIndex];
+					if (const FMetasoundFrontendClass* ConnectedNodeClass = FindDependency(ConnectedNode.ClassID))
+					{
+						if (ConnectedNodeClass->Metadata.GetClassName() == FInputNodeTemplate::ClassName)
+						{
+							TemplateNodes.Add(&ConnectedNode);
+						}
+					}
+				}
+			}
+		}
+	}
+
+	return TemplateNodes;
+}
+
 FMetasoundAssetBase& FMetaSoundFrontendDocumentBuilder::GetMetasoundAsset() const
 {
 	using namespace Metasound::Frontend;
@@ -1969,14 +2086,15 @@ const FGuid& FMetaSoundFrontendDocumentBuilder::GetBuildPageID() const
 	return BuildPageID;
 }
 
-EMetasoundFrontendVertexAccessType FMetaSoundFrontendDocumentBuilder::GetNodeInputAccessType(const FGuid& InNodeID, const FGuid& InVertexID) const
+EMetasoundFrontendVertexAccessType FMetaSoundFrontendDocumentBuilder::GetNodeInputAccessType(const FGuid& InNodeID, const FGuid& InVertexID, const FGuid* InPageID) const
 {
 	using namespace Metasound::Frontend;
 
-	const IDocumentGraphNodeCache& NodeCache = DocumentCache->GetNodeCache();
+	const FGuid& PageID = InPageID ? *InPageID : BuildPageID;
+	const IDocumentGraphNodeCache& NodeCache = DocumentCache->GetNodeCache(PageID);
 	if (const int32* NodeIndex = NodeCache.FindNodeIndex(InNodeID))
 	{
-		const FMetasoundFrontendGraph& Graph = FindConstBuildGraphChecked();
+		const FMetasoundFrontendGraph& Graph = GetConstDocumentChecked().RootGraph.FindConstGraphChecked(PageID);
 		const FMetasoundFrontendNode& Node = Graph.Nodes[*NodeIndex];
 		auto IsVertexID = [&InVertexID](const FMetasoundFrontendVertex& Vertex) { return Vertex.VertexID == InVertexID; };
 		if (const FMetasoundFrontendClass* Class = DocumentCache->FindDependency(Node.ClassID))
@@ -2023,14 +2141,16 @@ EMetasoundFrontendVertexAccessType FMetaSoundFrontendDocumentBuilder::GetNodeInp
 	return EMetasoundFrontendVertexAccessType::Unset;
 }
 
-const FMetasoundFrontendLiteral* FMetaSoundFrontendDocumentBuilder::GetNodeInputClassDefault(const FGuid& InNodeID, const FGuid& InVertexID) const
+const FMetasoundFrontendLiteral* FMetaSoundFrontendDocumentBuilder::GetNodeInputClassDefault(const FGuid& InNodeID, const FGuid& InVertexID, const FGuid* InPageID) const
 {
 	using namespace Metasound::Frontend;
 
-	const IDocumentGraphNodeCache& NodeCache = DocumentCache->GetNodeCache();
+	const FGuid& PageID = InPageID ? *InPageID : BuildPageID;
+	const IDocumentGraphNodeCache& NodeCache = DocumentCache->GetNodeCache(PageID);
 	if (const int32* NodeIndex = NodeCache.FindNodeIndex(InNodeID))
 	{
-		const FMetasoundFrontendNode& Node = FindConstBuildGraphChecked().Nodes[*NodeIndex];
+		const FMetasoundFrontendDocument& Document = GetConstDocumentChecked();
+		const FMetasoundFrontendNode& Node = Document.RootGraph.FindConstGraphChecked(PageID).Nodes[*NodeIndex];
 		auto IsVertexID = [&InVertexID](const FMetasoundFrontendVertex& Vertex) { return Vertex.VertexID == InVertexID; };
 		if (const FMetasoundFrontendVertex* Vertex = Node.Interface.Inputs.FindByPredicate(IsVertexID))
 		{
@@ -2066,14 +2186,16 @@ const FMetasoundFrontendLiteral* FMetaSoundFrontendDocumentBuilder::GetNodeInput
 	return nullptr;
 }
 
-const FMetasoundFrontendLiteral* FMetaSoundFrontendDocumentBuilder::GetNodeInputDefault(const FGuid& InNodeID, const FGuid& InVertexID) const
+const FMetasoundFrontendLiteral* FMetaSoundFrontendDocumentBuilder::GetNodeInputDefault(const FGuid& InNodeID, const FGuid& InVertexID, const FGuid* InPageID) const
 {
 	using namespace Metasound::Frontend;
 
-	const IDocumentGraphNodeCache& NodeCache = DocumentCache->GetNodeCache();
+	const FGuid& PageID = InPageID ? *InPageID : BuildPageID;
+	const IDocumentGraphNodeCache& NodeCache = DocumentCache->GetNodeCache(PageID);
 	if (const int32* NodeIndex = NodeCache.FindNodeIndex(InNodeID))
 	{
-		const FMetasoundFrontendNode& Node = FindConstBuildGraphChecked().Nodes[*NodeIndex];
+		const FMetasoundFrontendGraph& Graph = GetConstDocumentChecked().RootGraph.FindConstGraphChecked(PageID);
+		const FMetasoundFrontendNode& Node = Graph.Nodes[*NodeIndex];
 
 		auto IsVertex = [&InVertexID](const FMetasoundFrontendVertex& Vertex) { return Vertex.VertexID == InVertexID; };
 		const int32 VertexIndex = Node.Interface.Inputs.IndexOfByPredicate(IsVertex);
@@ -2094,14 +2216,16 @@ const FMetasoundFrontendLiteral* FMetaSoundFrontendDocumentBuilder::GetNodeInput
 	return nullptr;
 }
 
-EMetasoundFrontendVertexAccessType FMetaSoundFrontendDocumentBuilder::GetNodeOutputAccessType(const FGuid& InNodeID, const FGuid& InVertexID) const
+EMetasoundFrontendVertexAccessType FMetaSoundFrontendDocumentBuilder::GetNodeOutputAccessType(const FGuid& InNodeID, const FGuid& InVertexID, const FGuid* InPageID) const
 {
 	using namespace Metasound::Frontend;
 
-	const IDocumentGraphNodeCache& NodeCache = DocumentCache->GetNodeCache();
+	const FGuid& PageID = InPageID ? *InPageID : BuildPageID;
+	const IDocumentGraphNodeCache& NodeCache = DocumentCache->GetNodeCache(PageID);
 	if (const int32* NodeIndex = NodeCache.FindNodeIndex(InNodeID))
 	{
-		const FMetasoundFrontendNode& Node = FindConstBuildGraphChecked().Nodes[*NodeIndex];
+		const FMetasoundFrontendGraph& Graph = GetConstDocumentChecked().RootGraph.FindConstGraphChecked(PageID);
+		const FMetasoundFrontendNode& Node = Graph.Nodes[*NodeIndex];
 		if (const FMetasoundFrontendClass* Class = DocumentCache->FindDependency(Node.ClassID))
 		{
 			const EMetasoundFrontendClassType ClassType = Class->Metadata.GetType();
@@ -2285,14 +2409,16 @@ bool FMetaSoundFrontendDocumentBuilder::IsDependencyReferenced(const FGuid& InCl
 	return DocumentCache->GetNodeCache().ContainsNodesOfClassID(InClassID);
 }
 
-bool FMetaSoundFrontendDocumentBuilder::IsNodeInputConnected(const FGuid& InNodeID, const FGuid& InVertexID) const
+bool FMetaSoundFrontendDocumentBuilder::IsNodeInputConnected(const FGuid& InNodeID, const FGuid& InVertexID, const FGuid* InPageID) const
 {
-	return DocumentCache->GetEdgeCache().IsNodeInputConnected(InNodeID, InVertexID);
+	const FGuid& PageID = InPageID ? *InPageID : BuildPageID;
+	return DocumentCache->GetEdgeCache(PageID).IsNodeInputConnected(InNodeID, InVertexID);
 }
 
-bool FMetaSoundFrontendDocumentBuilder::IsNodeOutputConnected(const FGuid& InNodeID, const FGuid& InVertexID) const
+bool FMetaSoundFrontendDocumentBuilder::IsNodeOutputConnected(const FGuid& InNodeID, const FGuid& InVertexID, const FGuid* InPageID) const
 {
-	return DocumentCache->GetEdgeCache().IsNodeOutputConnected(InNodeID, InVertexID);
+	const FGuid& PageID = InPageID ? *InPageID : BuildPageID;
+	return DocumentCache->GetEdgeCache(PageID).IsNodeOutputConnected(InNodeID, InVertexID);
 }
 
 bool FMetaSoundFrontendDocumentBuilder::IsInterfaceDeclared(FName InInterfaceName) const
@@ -2318,11 +2444,12 @@ bool FMetaSoundFrontendDocumentBuilder::IsPreset() const
 	return GetConstDocumentChecked().RootGraph.PresetOptions.bIsPreset;
 }
 
-Metasound::Frontend::EInvalidEdgeReason FMetaSoundFrontendDocumentBuilder::IsValidEdge(const FMetasoundFrontendEdge& InEdge) const
+Metasound::Frontend::EInvalidEdgeReason FMetaSoundFrontendDocumentBuilder::IsValidEdge(const FMetasoundFrontendEdge& InEdge, const FGuid* InPageID) const
 {
 	using namespace Metasound::Frontend;
 
-	const IDocumentGraphNodeCache& NodeCache = DocumentCache->GetNodeCache();
+	const FGuid& PageID = InPageID ? *InPageID : BuildPageID;
+	const IDocumentGraphNodeCache& NodeCache = DocumentCache->GetNodeCache(PageID);
 
 	const FMetasoundFrontendVertex* OutputVertex = NodeCache.FindOutputVertex(InEdge.FromNodeID, InEdge.FromVertexID);
 	if (!OutputVertex)
@@ -2343,8 +2470,8 @@ Metasound::Frontend::EInvalidEdgeReason FMetaSoundFrontendDocumentBuilder::IsVal
 
 	// TODO: Add cycle detection here
 
-	const EMetasoundFrontendVertexAccessType OutputAccessType = GetNodeOutputAccessType(InEdge.FromNodeID, InEdge.FromVertexID);
-	const EMetasoundFrontendVertexAccessType InputAccessType = GetNodeInputAccessType(InEdge.ToNodeID, InEdge.ToVertexID);
+	const EMetasoundFrontendVertexAccessType OutputAccessType = GetNodeOutputAccessType(InEdge.FromNodeID, InEdge.FromVertexID, InPageID);
+	const EMetasoundFrontendVertexAccessType InputAccessType = GetNodeInputAccessType(InEdge.ToNodeID, InEdge.ToVertexID, InPageID);
 	if (!FMetasoundFrontendClassVertex::CanConnectVertexAccessTypes(OutputAccessType, InputAccessType))
 	{
 		return EInvalidEdgeReason::MismatchedAccessType;
@@ -2353,14 +2480,14 @@ Metasound::Frontend::EInvalidEdgeReason FMetaSoundFrontendDocumentBuilder::IsVal
 	return EInvalidEdgeReason::None;
 }
 
-void FMetaSoundFrontendDocumentBuilder::IterateNodesConnectedWithVertex(const FMetasoundFrontendVertexHandle& Vertex, TFunctionRef<void(const FMetasoundFrontendEdge&, FMetasoundFrontendNode&)> NodeIndexIterFunc)
+void FMetaSoundFrontendDocumentBuilder::IterateNodesConnectedWithVertex(const FMetasoundFrontendVertexHandle& Vertex, TFunctionRef<void(const FMetasoundFrontendEdge&, FMetasoundFrontendNode&)> NodeIndexIterFunc, const FGuid& InPageID)
 {
 	using namespace Metasound::Frontend;
 
-	FMetasoundFrontendGraph& Graph = FindBuildGraphChecked();
+	FMetasoundFrontendGraph& Graph = GetDocumentChecked().RootGraph.FindGraphChecked(InPageID);
 	TArray<FMetasoundFrontendEdge> EdgesToConnectedNodes; // Have to cache to avoid pointers becoming garbage in subsequent removal loop
-	const IDocumentGraphEdgeCache& EdgeCache = DocumentCache->GetEdgeCache();
-	const IDocumentGraphNodeCache& NodeCache = DocumentCache->GetNodeCache();
+	const IDocumentGraphEdgeCache& EdgeCache = DocumentCache->GetEdgeCache(InPageID);
+	const IDocumentGraphNodeCache& NodeCache = DocumentCache->GetNodeCache(InPageID);
 	const TArray<const FMetasoundFrontendEdge*> Edges = EdgeCache.FindEdges(Vertex.NodeID, Vertex.VertexID);
 	Algo::Transform(Edges, EdgesToConnectedNodes, [](const FMetasoundFrontendEdge* Edge) { check(Edge); return *Edge; });
 	for (const FMetasoundFrontendEdge& Edge : EdgesToConnectedNodes)
@@ -2468,12 +2595,26 @@ void FMetaSoundFrontendDocumentBuilder::BeginBuilding(TSharedPtr<Metasound::Fron
 {
 	using namespace Metasound::Frontend;
 
-	DocumentDelegates = Delegates.IsValid() ? Delegates : MakeShared<FDocumentModifyDelegates>();
+	if (Delegates.IsValid())
+	{
+		DocumentDelegates = Delegates;
+	}
+	else if (DocumentInterface)
+	{
+		const FMetasoundFrontendDocument& Document = GetConstDocumentChecked();
+		DocumentDelegates = Delegates.IsValid() ? Delegates : MakeShared<FDocumentModifyDelegates>(Document);
+	}
+	else
+	{
+		DocumentDelegates = MakeShared<FDocumentModifyDelegates>();
+	}
 
 	if (DocumentInterface)
 	{
 		DocumentInterface->OnBeginActiveBuilder();
-		DocumentCache = FDocumentCache::Create(GetDocumentChecked(), DocumentDelegates.ToSharedRef(), BuildPageID, bPrimeCache);
+
+		const FMetasoundFrontendDocument& Document = GetConstDocumentChecked();
+		DocumentCache = FDocumentCache::Create(Document, DocumentDelegates.ToSharedRef(), BuildPageID, bPrimeCache);
 	}
 }
 
@@ -2491,6 +2632,25 @@ void FMetaSoundFrontendDocumentBuilder::FinishBuilding()
 	DocumentCache.Reset();
 }
 
+void FMetaSoundFrontendDocumentBuilder::RemoveAllGraphPages()
+{
+	FMetasoundFrontendGraphClass& RootGraph = GetDocumentChecked().RootGraph;
+	RootGraph.IterateGraphPages([this](FMetasoundFrontendGraph& Graph)
+	{
+		using namespace Metasound::Frontend;
+
+		if (Graph.PageID.IsValid())
+		{
+			DocumentDelegates->PageDelegates.OnRemovingPage.Broadcast(FDocumentMutatePageArgs { Graph.PageID });
+		}
+	});
+	RootGraph.RemoveAllGraphPages();
+	if (BuildPageID.IsValid())
+	{
+		SetBuildPageID(FGuid());
+	}
+}
+
 bool FMetaSoundFrontendDocumentBuilder::RemoveDependency(const FGuid& InClassID)
 {
 	using namespace Metasound::Frontend;
@@ -2550,14 +2710,14 @@ bool FMetaSoundFrontendDocumentBuilder::RemoveDependency(EMetasoundFrontendClass
 	return true;
 }
 
-bool FMetaSoundFrontendDocumentBuilder::RemoveEdge(const FMetasoundFrontendEdge& EdgeToRemove)
+bool FMetaSoundFrontendDocumentBuilder::RemoveEdge(const FMetasoundFrontendEdge& EdgeToRemove, const FGuid* InPageID)
 {
 	using namespace Metasound::Frontend;
 
-	FMetasoundFrontendDocument& Document = GetDocumentChecked();
-	FMetasoundFrontendGraph& Graph = FindBuildGraphChecked();
+	const FGuid& PageID = InPageID ? *InPageID : BuildPageID;
+	FMetasoundFrontendGraph& Graph = GetDocumentChecked().RootGraph.FindGraphChecked(PageID);
 	TArray<FMetasoundFrontendEdge>& Edges = Graph.Edges;
-	const IDocumentGraphEdgeCache& EdgeCache = DocumentCache->GetEdgeCache();
+	const IDocumentGraphEdgeCache& EdgeCache = DocumentCache->GetEdgeCache(PageID);
 	if (const int32* IndexPtr = EdgeCache.FindEdgeIndexToNodeInput(EdgeToRemove.ToNodeID, EdgeToRemove.ToVertexID))
 	{
 		const int32 Index = *IndexPtr;
@@ -2565,7 +2725,7 @@ bool FMetaSoundFrontendDocumentBuilder::RemoveEdge(const FMetasoundFrontendEdge&
 		if (EdgeToRemove.FromNodeID == FoundEdge.FromNodeID && EdgeToRemove.FromVertexID == FoundEdge.FromVertexID)
 		{
 			const int32 LastIndex = Edges.Num() - 1;
-			DocumentDelegates->EdgeDelegates.OnRemoveSwappingEdge.Broadcast(Index, LastIndex);
+			DocumentDelegates->FindEdgeDelegatesChecked(PageID).OnRemoveSwappingEdge.Broadcast(Index, LastIndex);
 			Edges.RemoveAtSwap(Index, EAllowShrinking::No);
 			return true;
 		}
@@ -2574,12 +2734,13 @@ bool FMetaSoundFrontendDocumentBuilder::RemoveEdge(const FMetasoundFrontendEdge&
 	return false;
 }
 
-bool FMetaSoundFrontendDocumentBuilder::RemoveNamedEdges(const TSet<Metasound::Frontend::FNamedEdge>& InNamedEdgesToRemove, TArray<FMetasoundFrontendEdge>* OutRemovedEdges)
+bool FMetaSoundFrontendDocumentBuilder::RemoveNamedEdges(const TSet<Metasound::Frontend::FNamedEdge>& InNamedEdgesToRemove, TArray<FMetasoundFrontendEdge>* OutRemovedEdges, const FGuid* InPageID)
 {
 	using namespace Metasound::Frontend;
 
-	FMetasoundFrontendDocument& Document = GetDocumentChecked();
-	const IDocumentGraphNodeCache& NodeCache = DocumentCache->GetNodeCache();
+	const FGuid& PageID = InPageID ? *InPageID : BuildPageID;
+	const IDocumentGraphNodeCache& NodeCache = DocumentCache->GetNodeCache(PageID);
+	const IDocumentGraphEdgeCache& EdgeCache = DocumentCache->GetEdgeCache(PageID);
 
 	if (OutRemovedEdges)
 	{
@@ -2591,14 +2752,13 @@ bool FMetaSoundFrontendDocumentBuilder::RemoveNamedEdges(const TSet<Metasound::F
 	TArray<FMetasoundFrontendEdge> EdgesToRemove;
 	for (const FNamedEdge& NamedEdge : InNamedEdgesToRemove)
 	{
-		const FMetasoundFrontendGraph& Graph = FindBuildGraphChecked();
 		const FMetasoundFrontendVertex* OutputVertex = NodeCache.FindOutputVertex(NamedEdge.OutputNodeID, NamedEdge.OutputName);
 		const FMetasoundFrontendVertex* InputVertex = NodeCache.FindInputVertex(NamedEdge.InputNodeID, NamedEdge.InputName);
 
 		if (OutputVertex && InputVertex)
 		{
 			FMetasoundFrontendEdge NewEdge = { NamedEdge.OutputNodeID, OutputVertex->VertexID, NamedEdge.InputNodeID, InputVertex->VertexID };
-			if (ContainsEdge(NewEdge))
+			if (EdgeCache.ContainsEdge(NewEdge))
 			{
 				EdgesToRemove.Add(MoveTemp(NewEdge));
 			}
@@ -2612,7 +2772,7 @@ bool FMetaSoundFrontendDocumentBuilder::RemoveNamedEdges(const TSet<Metasound::F
 
 	for (const FMetasoundFrontendEdge& EdgeToRemove : EdgesToRemove)
 	{
-		const bool bRemovedEdge = RemoveEdgeToNodeInput(EdgeToRemove.ToNodeID, EdgeToRemove.ToVertexID);
+		const bool bRemovedEdge = RemoveEdgeToNodeInput(EdgeToRemove.ToNodeID, EdgeToRemove.ToVertexID, InPageID);
 		if (ensureAlwaysMsgf(bRemovedEdge, TEXT("Failed to remove MetaSound graph edge via DocumentBuilder when prior step validated edge remove was valid")))
 		{
 			if (OutRemovedEdges)
@@ -2638,23 +2798,25 @@ void FMetaSoundFrontendDocumentBuilder::Reload(TSharedPtr<Metasound::Frontend::F
 		DocumentInterface->OnFinishActiveBuilder();
 	}
 
-	DocumentDelegates = Delegates.IsValid() ? Delegates : MakeShared<FDocumentModifyDelegates>();
+	const FMetasoundFrontendDocument& Document = GetConstDocumentChecked();
+	DocumentDelegates = Delegates.IsValid() ? Delegates : MakeShared<FDocumentModifyDelegates>(Document);
 
 	if (DocumentInterface)
 	{
-		DocumentCache = FDocumentCache::Create(GetDocumentChecked(), DocumentDelegates.ToSharedRef(), BuildPageID, bPrimeCache);
+		DocumentCache = FDocumentCache::Create(Document, DocumentDelegates.ToSharedRef(), BuildPageID, bPrimeCache);
 		DocumentInterface->OnBeginActiveBuilder();
 	}
 }
 
-bool FMetaSoundFrontendDocumentBuilder::RemoveNodeInputDefault(const FGuid& InNodeID, const FGuid& InVertexID)
+bool FMetaSoundFrontendDocumentBuilder::RemoveNodeInputDefault(const FGuid& InNodeID, const FGuid& InVertexID, const FGuid* InPageID)
 {
 	using namespace Metasound::Frontend;
 
-	const IDocumentGraphNodeCache& NodeCache = DocumentCache->GetNodeCache();
+	const FGuid& PageID = InPageID ? *InPageID : BuildPageID;
+	const IDocumentGraphNodeCache& NodeCache = DocumentCache->GetNodeCache(PageID);
 	if (const int32* NodeIndex = NodeCache.FindNodeIndex(InNodeID))
 	{
-		FMetasoundFrontendGraph& Graph = FindBuildGraphChecked();
+		FMetasoundFrontendGraph& Graph = GetDocumentChecked().RootGraph.FindGraphChecked(PageID);
 		FMetasoundFrontendNode& Node = Graph.Nodes[*NodeIndex];
 
 		auto IsVertex = [&InVertexID](const FMetasoundFrontendVertex& Vertex) { return Vertex.VertexID == InVertexID; };
@@ -2665,7 +2827,8 @@ bool FMetaSoundFrontendDocumentBuilder::RemoveNodeInputDefault(const FGuid& InNo
 			const int32 LiteralIndex = Node.InputLiterals.IndexOfByPredicate(IsLiteral);
 			if (LiteralIndex != INDEX_NONE)
 			{
-				const FOnMetaSoundFrontendDocumentMutateNodeInputLiteralArray& OnRemovingNodeInputLiteral = DocumentDelegates->NodeDelegates.OnRemovingNodeInputLiteral;
+				FNodeModifyDelegates& NodeDelegates = DocumentDelegates->FindNodeDelegatesChecked(PageID);
+				const FOnMetaSoundFrontendDocumentMutateNodeInputLiteralArray& OnRemovingNodeInputLiteral = NodeDelegates.OnRemovingNodeInputLiteral;
 				const int32 LastIndex = Node.InputLiterals.Num() - 1;
 				OnRemovingNodeInputLiteral.Broadcast(*NodeIndex, VertexIndex, LastIndex);
 				if (LiteralIndex != LastIndex)
@@ -2676,7 +2839,7 @@ bool FMetaSoundFrontendDocumentBuilder::RemoveNodeInputDefault(const FGuid& InNo
 				Node.InputLiterals.RemoveAtSwap(LiteralIndex, EAllowShrinking::No);
 				if (LiteralIndex != LastIndex)
 				{
-					const FOnMetaSoundFrontendDocumentMutateNodeInputLiteralArray& OnNodeInputLiteralSet = DocumentDelegates->NodeDelegates.OnNodeInputLiteralSet;
+					const FOnMetaSoundFrontendDocumentMutateNodeInputLiteralArray& OnNodeInputLiteralSet = NodeDelegates.OnNodeInputLiteralSet;
 					OnNodeInputLiteralSet.Broadcast(*NodeIndex, VertexIndex, LiteralIndex);
 				}
 				return true;
@@ -2687,25 +2850,26 @@ bool FMetaSoundFrontendDocumentBuilder::RemoveNodeInputDefault(const FGuid& InNo
 	return false;
 }
 
-bool FMetaSoundFrontendDocumentBuilder::RemoveEdges(const FGuid& InNodeID)
+bool FMetaSoundFrontendDocumentBuilder::RemoveEdges(const FGuid& InNodeID, const FGuid* InPageID)
 {
 	using namespace Metasound::Frontend;
 	using namespace Metasound::VariableNames;
 
-	const IDocumentGraphNodeCache& NodeCache = DocumentCache->GetNodeCache();
+	const FGuid& PageID = InPageID ? *InPageID : BuildPageID;
+	const IDocumentGraphNodeCache& NodeCache = DocumentCache->GetNodeCache(PageID);
 	if (const FMetasoundFrontendNode* Node = NodeCache.FindNode(InNodeID))
 	{
-		const IDocumentGraphEdgeCache& EdgeCache = DocumentCache->GetEdgeCache();
+		const IDocumentGraphEdgeCache& EdgeCache = DocumentCache->GetEdgeCache(PageID);
 
 		for (const FMetasoundFrontendVertex& Vertex : Node->Interface.Inputs)
 		{
-			RemoveEdgeToNodeInput(InNodeID, Vertex.VertexID);
+			RemoveEdgeToNodeInput(InNodeID, Vertex.VertexID, InPageID);
 		}
 
 		TArray<FMetasoundFrontendVertexHandle> ToVertexHandles;
 		for (const FMetasoundFrontendVertex& Vertex : Node->Interface.Outputs)
 		{
-			RemoveEdgesFromNodeOutput(InNodeID, Vertex.VertexID);
+			RemoveEdgesFromNodeOutput(InNodeID, Vertex.VertexID, InPageID);
 		}
 
 		return true;
@@ -2714,7 +2878,7 @@ bool FMetaSoundFrontendDocumentBuilder::RemoveEdges(const FGuid& InNodeID)
 	return false;
 }
 
-bool FMetaSoundFrontendDocumentBuilder::RemoveEdgesByNodeClassInterfaceBindings(const FGuid& InFromNodeID, const FGuid& InToNodeID)
+bool FMetaSoundFrontendDocumentBuilder::RemoveEdgesByNodeClassInterfaceBindings(const FGuid& InFromNodeID, const FGuid& InToNodeID, const FGuid* InPageID)
 {
 	using namespace Metasound;
 	using namespace Metasound::Frontend;
@@ -2722,36 +2886,39 @@ bool FMetaSoundFrontendDocumentBuilder::RemoveEdgesByNodeClassInterfaceBindings(
 	TSet<FMetasoundFrontendVersion> FromInterfaceVersions;
 	TSet<FMetasoundFrontendVersion> ToInterfaceVersions;
 
-	if (FindNodeClassInterfaces(InFromNodeID, FromInterfaceVersions) && FindNodeClassInterfaces(InToNodeID, ToInterfaceVersions))
+	const FGuid& PageID = InPageID ? *InPageID : BuildPageID;
+	if (FindNodeClassInterfaces(InFromNodeID, FromInterfaceVersions, PageID) && FindNodeClassInterfaces(InToNodeID, ToInterfaceVersions, PageID))
 	{
 		TSet<FNamedEdge> NamedEdges;
 		if (DocumentBuilderPrivate::TryGetInterfaceBoundEdges(InFromNodeID, FromInterfaceVersions, InToNodeID, ToInterfaceVersions, NamedEdges))
 		{
-			return RemoveNamedEdges(NamedEdges);
+			constexpr TArray<FMetasoundFrontendEdge>* RemovedEdges = nullptr;
+			return RemoveNamedEdges(NamedEdges, RemovedEdges, InPageID);
 		}
 	}
 
 	return false;
 }
 
-bool FMetaSoundFrontendDocumentBuilder::RemoveEdgesFromNodeOutput(const FGuid& InNodeID, const FGuid& InVertexID)
+bool FMetaSoundFrontendDocumentBuilder::RemoveEdgesFromNodeOutput(const FGuid& InNodeID, const FGuid& InVertexID, const FGuid* InPageID)
 {
 	using namespace Metasound::Frontend;
 
-	const IDocumentGraphEdgeCache& EdgeCache = DocumentCache->GetEdgeCache();
+	const FGuid& PageID = InPageID ? *InPageID : BuildPageID;
+	const IDocumentGraphEdgeCache& EdgeCache = DocumentCache->GetEdgeCache(PageID);
 	const TArrayView<const int32> Indices = EdgeCache.FindEdgeIndicesFromNodeOutput(InNodeID, InVertexID);
 	if (!Indices.IsEmpty())
 	{
-		FMetasoundFrontendDocument& Document = GetDocumentChecked();
-		FMetasoundFrontendGraph& Graph = FindBuildGraphChecked();
+		FMetasoundFrontendGraph& Graph = GetDocumentChecked().RootGraph.FindGraphChecked(PageID);
 
 		// Copy off indices and sort descending as the edge array will be modified when notifying the cache in the loop below
 		TArray<int32> IndicesCopy(Indices.GetData(), Indices.Num());
 		Algo::Sort(IndicesCopy, [](const int32& L, const int32& R) { return L > R; });
+		FEdgeModifyDelegates& EdgeDelegates = DocumentDelegates->FindEdgeDelegatesChecked(PageID);
 		for (int32 Index : IndicesCopy)
 		{
 			const int32 LastIndex = Graph.Edges.Num() - 1;
-			DocumentDelegates->EdgeDelegates.OnRemoveSwappingEdge.Broadcast(Index, LastIndex);
+			EdgeDelegates.OnRemoveSwappingEdge.Broadcast(Index, LastIndex);
 			Graph.Edges.RemoveAtSwap(Index, EAllowShrinking::No);
 		}
 
@@ -2761,22 +2928,23 @@ bool FMetaSoundFrontendDocumentBuilder::RemoveEdgesFromNodeOutput(const FGuid& I
 	return false;
 }
 
-bool FMetaSoundFrontendDocumentBuilder::RemoveEdgeToNodeInput(const FGuid& InNodeID, const FGuid& InVertexID)
+bool FMetaSoundFrontendDocumentBuilder::RemoveEdgeToNodeInput(const FGuid& InNodeID, const FGuid& InVertexID, const FGuid* InPageID)
 {
 	using namespace Metasound::Frontend;
 
-	const IDocumentGraphEdgeCache& EdgeCache = DocumentCache->GetEdgeCache();
+	const FGuid& PageID = InPageID ? *InPageID : BuildPageID;
+	const IDocumentGraphEdgeCache& EdgeCache = DocumentCache->GetEdgeCache(PageID);
 	if (const int32* IndexPtr = EdgeCache.FindEdgeIndexToNodeInput(InNodeID, InVertexID))
 	{
-		FMetasoundFrontendDocument& Document = GetDocumentChecked();
-		FMetasoundFrontendGraph& Graph = FindBuildGraphChecked();
+		FMetasoundFrontendGraph& Graph = GetDocumentChecked().RootGraph.FindGraphChecked(PageID);
+		const FEdgeModifyDelegates& EdgeDelegates = DocumentDelegates->FindEdgeDelegatesChecked(PageID);
 		const int32 Index = *IndexPtr; // Copy off indices as the pointer may be modified when notifying the cache below
 		const int32 LastIndex = Graph.Edges.Num() - 1;
-		DocumentDelegates->EdgeDelegates.OnRemoveSwappingEdge.Broadcast(Index, LastIndex);
+		EdgeDelegates.OnRemoveSwappingEdge.Broadcast(Index, LastIndex);
 		Graph.Edges.RemoveAtSwap(Index, EAllowShrinking::No);
 
 #if WITH_EDITORONLY_DATA
-		Document.Metadata.ModifyContext.AddNodeIDModified(InNodeID);
+		GetDocumentChecked().Metadata.ModifyContext.AddNodeIDModified(InNodeID);
 #endif // WITH_EDITORONLY_DATA
 
 		return true;
@@ -2786,12 +2954,11 @@ bool FMetaSoundFrontendDocumentBuilder::RemoveEdgeToNodeInput(const FGuid& InNod
 }
 
 #if WITH_EDITOR
-bool FMetaSoundFrontendDocumentBuilder::RemoveGraphComment(const FGuid& InCommentID)
+bool FMetaSoundFrontendDocumentBuilder::RemoveGraphComment(const FGuid& InCommentID, const FGuid* InPageID)
 {
 	FMetasoundFrontendDocument& Document = GetDocumentChecked();
-	TMap<FGuid, FMetaSoundFrontendGraphComment>& Comments = FindBuildGraphChecked().Style.Comments;
-
-	if (Comments.Remove(InCommentID) > 0)
+	FMetasoundFrontendGraph& Graph = Document.RootGraph.FindGraphChecked(InPageID ? *InPageID : BuildPageID);
+	if (Graph.Style.Comments.Remove(InCommentID) > 0)
 	{
 #if WITH_EDITORONLY_DATA
 		Document.Metadata.ModifyContext.SetDocumentModified();
@@ -2804,44 +2971,6 @@ bool FMetaSoundFrontendDocumentBuilder::RemoveGraphComment(const FGuid& InCommen
 }
 #endif // WITH_EDITOR
 
-TArray<const FMetasoundFrontendNode*> FMetaSoundFrontendDocumentBuilder::GetGraphInputTemplateNodes(FName InInputName)
-{
-	using namespace Metasound::Frontend;
-
-	TArray<const FMetasoundFrontendNode*> TemplateNodes;
-
-	FMetasoundFrontendDocument& Document = GetDocumentChecked();
-	if (const int32* Index = DocumentCache->GetInterfaceCache().FindInputIndex(InInputName))
-	{
-		const FMetasoundFrontendClassInput& InputClass = Document.RootGraph.Interface.Inputs[*Index];
-		const IDocumentGraphNodeCache& NodeCache = DocumentCache->GetNodeCache();
-		const IDocumentGraphEdgeCache& EdgeCache = DocumentCache->GetEdgeCache();
-
-		if (const FMetasoundFrontendNode* InputNode = NodeCache.FindNode(InputClass.NodeID))
-		{
-			const FGuid OutputVertexID = InputNode->Interface.Outputs.Last().VertexID;
-			const TArray<const FMetasoundFrontendEdge*> ConnectedEdges = EdgeCache.FindEdges(InputClass.NodeID, OutputVertexID);
-			for (const FMetasoundFrontendEdge* Edge : ConnectedEdges)
-			{
-				check(Edge);
-				if (const int32* ConnectedNodeIndex = NodeCache.FindNodeIndex(Edge->ToNodeID))
-				{
-					const FMetasoundFrontendNode& ConnectedNode = FindConstBuildGraphChecked().Nodes[*ConnectedNodeIndex];
-					if (const FMetasoundFrontendClass* ConnectedNodeClass = FindDependency(ConnectedNode.ClassID))
-					{
-						if (ConnectedNodeClass->Metadata.GetClassName() == FInputNodeTemplate::ClassName)
-						{
-							TemplateNodes.Add(&ConnectedNode);
-						}
-					}
-				}
-			}
-		}
-	}
-
-	return TemplateNodes;
-}
-
 bool FMetaSoundFrontendDocumentBuilder::RemoveGraphInput(FName InInputName)
 {
 	using namespace Metasound::Frontend;
@@ -2852,34 +2981,38 @@ bool FMetaSoundFrontendDocumentBuilder::RemoveGraphInput(FName InInputName)
 		TArray<FMetasoundFrontendClassInput>& Inputs = Document.RootGraph.Interface.Inputs;
 		const FGuid NodeID = Inputs[*IndexPtr].NodeID;
 		FGuid ClassID;
-		TArray<FGuid> NodeIDsToRemove { NodeID };
-
-		if (const FMetasoundFrontendNode* Node = FindNode(NodeID))
-		{
-			ClassID = Node->ClassID;
-		}
-		else
-		{
-			return false;
-		}
-
-		const TArray<const FMetasoundFrontendNode*> TemplateNodes = GetGraphInputTemplateNodes(InInputName);
-		Algo::Transform(TemplateNodes, NodeIDsToRemove, [](const FMetasoundFrontendNode* Node) { return Node->GetID(); });
-
 		bool bNodesRemoved = true;
-		for (const FGuid& ToRemove : NodeIDsToRemove)
+		Document.RootGraph.IterateGraphPages([this, &Document, &ClassID, &NodeID, &InInputName, &bNodesRemoved](const FMetasoundFrontendGraph& Graph)
 		{
-			if (RemoveNode(ToRemove))
+			TArray<FGuid> NodeIDsToRemove { NodeID };
+
+			if (const FMetasoundFrontendNode* Node = FindNode(NodeID, &Graph.PageID))
 			{
-#if WITH_EDITORONLY_DATA
-				Document.Metadata.ModifyContext.AddNodeIDModified(ToRemove);
-#endif // WITH_EDITORONLY_DATA
+				ClassID = Node->ClassID;
 			}
 			else
 			{
 				bNodesRemoved = false;
+				return;
 			}
-		}
+
+			const TArray<const FMetasoundFrontendNode*> TemplateNodes = GetGraphInputTemplateNodes(InInputName, &Graph.PageID);
+			Algo::Transform(TemplateNodes, NodeIDsToRemove, [](const FMetasoundFrontendNode* Node) { return Node->GetID(); });
+
+			for (const FGuid& ToRemove : NodeIDsToRemove)
+			{
+				if (RemoveNode(ToRemove, &Graph.PageID))
+				{
+#if WITH_EDITORONLY_DATA
+					Document.Metadata.ModifyContext.AddNodeIDModified(ToRemove);
+#endif // WITH_EDITORONLY_DATA
+				}
+				else
+				{
+					bNodesRemoved = false;
+				}
+			}
+		});
 
 		if (bNodesRemoved)
 		{
@@ -2915,46 +3048,57 @@ bool FMetaSoundFrontendDocumentBuilder::RemoveGraphInput(FName InInputName)
 
 bool FMetaSoundFrontendDocumentBuilder::RemoveGraphOutput(FName InOutputName)
 {
-	if (const FMetasoundFrontendNode* Node = FindGraphOutputNode(InOutputName))
+	bool bNodesRemoved = true;
+	FGuid ClassID;
+	FGuid NodeID;
+	FMetasoundFrontendDocument& Document = GetDocumentChecked();
+	Document.RootGraph.IterateGraphPages([this, &InOutputName, &Document, &ClassID, &NodeID, &bNodesRemoved](const FMetasoundFrontendGraph& Graph)
 	{
-		const FGuid ClassID = Node->ClassID;
-		const FGuid NodeID = Node->GetID();
-		if (RemoveNode(NodeID))
+		if (const FMetasoundFrontendNode* Node = FindGraphOutputNode(InOutputName, &Graph.PageID))
 		{
-			FMetasoundFrontendDocument& Document = GetDocumentChecked();
+			ClassID = Node->ClassID;
+			NodeID = Node->GetID();
+			if (!RemoveNode(NodeID, &Graph.PageID))
+			{
+				bNodesRemoved = false;
+				return;
+			}
 
 #if WITH_EDITORONLY_DATA
 			Document.Metadata.ModifyContext.AddNodeIDModified(NodeID);
 #endif // WITH_EDITORONLY_DATA
+		}
+	});
 
-			TArray<FMetasoundFrontendClassOutput>& Outputs = Document.RootGraph.Interface.Outputs;
-			auto OutputNameMatches = [InOutputName](const FMetasoundFrontendClassOutput& Output) { return Output.Name == InOutputName; };
-			const int32 Index = Outputs.IndexOfByPredicate(OutputNameMatches);
-			if (Index != INDEX_NONE)
-			{
-				DocumentDelegates->InterfaceDelegates.OnRemovingOutput.Broadcast(Index);
+	if (bNodesRemoved)
+	{
+		TArray<FMetasoundFrontendClassOutput>& Outputs = Document.RootGraph.Interface.Outputs;
+		auto OutputNameMatches = [InOutputName](const FMetasoundFrontendClassOutput& Output) { return Output.Name == InOutputName; };
+		const int32 Index = Outputs.IndexOfByPredicate(OutputNameMatches);
+		if (Index != INDEX_NONE)
+		{
+			DocumentDelegates->InterfaceDelegates.OnRemovingOutput.Broadcast(Index);
 
-				const int32 LastIndex = Outputs.Num() - 1;
-				if (Index != LastIndex)
-				{
-					DocumentDelegates->InterfaceDelegates.OnRemovingOutput.Broadcast(LastIndex);
-				}
-				Outputs.RemoveAtSwap(Index, EAllowShrinking::No);
-				if (Index != LastIndex)
-				{
-					DocumentDelegates->InterfaceDelegates.OnOutputAdded.Broadcast(Index);
-				}
+			const int32 LastIndex = Outputs.Num() - 1;
+			if (Index != LastIndex)
+			{
+				DocumentDelegates->InterfaceDelegates.OnRemovingOutput.Broadcast(LastIndex);
+			}
+			Outputs.RemoveAtSwap(Index, EAllowShrinking::No);
+			if (Index != LastIndex)
+			{
+				DocumentDelegates->InterfaceDelegates.OnOutputAdded.Broadcast(Index);
+			}
 
 #if WITH_EDITORONLY_DATA
-				ClearMemberMetadata(NodeID);
-				Document.Metadata.ModifyContext.AddMemberIDModified(NodeID);
+			ClearMemberMetadata(NodeID);
+			Document.Metadata.ModifyContext.AddMemberIDModified(NodeID);
 #endif // WITH_EDITORONLY_DATA
 
-				const bool bDependencyReferenced = IsDependencyReferenced(ClassID);
-				if (bDependencyReferenced || RemoveDependency(ClassID))
-				{
-					return true;
-				}
+			const bool bDependencyReferenced = IsDependencyReferenced(ClassID);
+			if (bDependencyReferenced || RemoveDependency(ClassID))
+			{
+				return true;
 			}
 		}
 	}
@@ -2962,6 +3106,31 @@ bool FMetaSoundFrontendDocumentBuilder::RemoveGraphOutput(FName InOutputName)
 	return false;
 }
 
+bool FMetaSoundFrontendDocumentBuilder::RemoveGraphPage(const FGuid& InPageID)
+{
+	using namespace Metasound::Frontend;
+
+	FMetasoundFrontendDocument& Document = GetDocumentChecked();
+	FGuid AdjacentPageID;
+
+	if (Document.RootGraph.ContainsGraphPage(InPageID))
+	{
+		DocumentDelegates->RemovePageDelegates(InPageID);
+	}
+
+	const bool bPageRemoved = Document.RootGraph.RemoveGraphPage(InPageID, &AdjacentPageID);
+	if (bPageRemoved)
+	{
+		if (InPageID == BuildPageID)
+		{
+			ensureAlwaysMsgf(SetBuildPageID(AdjacentPageID), TEXT("AdjacentPageID returned is always expected to be valid"));
+		}
+
+	}
+
+	return bPageRemoved;
+}
+
 bool FMetaSoundFrontendDocumentBuilder::RemoveInterface(FName InterfaceName)
 {
 	using namespace Metasound::Frontend;
@@ -2996,21 +3165,21 @@ bool FMetaSoundFrontendDocumentBuilder::RemoveInterface(FName InterfaceName)
 	return false;
 }
 
-bool FMetaSoundFrontendDocumentBuilder::RemoveNode(const FGuid& InNodeID)
+bool FMetaSoundFrontendDocumentBuilder::RemoveNode(const FGuid& InNodeID, const FGuid* InPageID)
 {
 	METASOUND_TRACE_CPUPROFILER_EVENT_SCOPE(FMetaSoundFrontendDocumentBuilder::RemoveNode);
 
 	using namespace Metasound::Frontend;
 
-	const IDocumentGraphNodeCache& NodeCache = DocumentCache->GetNodeCache();
-	const IDocumentGraphEdgeCache& EdgeCache = DocumentCache->GetEdgeCache();
+	const FGuid& PageID = InPageID ? *InPageID : BuildPageID;
+	const IDocumentGraphNodeCache& NodeCache = DocumentCache->GetNodeCache(PageID);
+	const IDocumentGraphEdgeCache& EdgeCache = DocumentCache->GetEdgeCache(PageID);
 
 	if (const int32* IndexPtr = NodeCache.FindNodeIndex(InNodeID))
 	{
 		const int32 Index = *IndexPtr; // Copy off indices as the pointer may be modified when notifying the cache below
 
-		FMetasoundFrontendDocument& Document = GetDocumentChecked();
-		FMetasoundFrontendGraph& Graph = FindBuildGraphChecked();
+		FMetasoundFrontendGraph& Graph = GetDocumentChecked().RootGraph.FindGraphChecked(PageID);
 		TArray<FMetasoundFrontendNode>& Nodes = Graph.Nodes;
 		const FMetasoundFrontendNode& Node = Nodes[Index];
 		const FGuid& NodeID = Node.GetID();
@@ -3025,19 +3194,20 @@ bool FMetaSoundFrontendDocumentBuilder::RemoveNode(const FGuid& InNodeID)
 			case EMetasoundFrontendClassType::VariableAccessor:
 			case EMetasoundFrontendClassType::VariableMutator:
 			{
-				const bool bVariableNodeUnlinked = UnlinkVariableNode(NodeID);
+				const bool bVariableNodeUnlinked = UnlinkVariableNode(NodeID, PageID);
 				ensureAlwaysMsgf(bVariableNodeUnlinked, TEXT("Failed to unlink %s node with ID '%s"), LexToString(ClassType), *InNodeID.ToString());
 			}
 			break;
 		}
 
-		RemoveEdges(NodeID);
+		RemoveEdges(NodeID, InPageID);
 		const int32 LastIndex = Nodes.Num() - 1;
-		DocumentDelegates->NodeDelegates.OnRemoveSwappingNode.Broadcast(Index, LastIndex);
+		FNodeModifyDelegates& NodeDelegates = DocumentDelegates->FindNodeDelegatesChecked(PageID);
+		NodeDelegates.OnRemoveSwappingNode.Broadcast(Index, LastIndex);
 		Nodes.RemoveAtSwap(Index, EAllowShrinking::No);
 
 #if WITH_EDITORONLY_DATA
-		Document.Metadata.ModifyContext.AddNodeIDModified(InNodeID);
+		GetDocumentChecked().Metadata.ModifyContext.AddNodeIDModified(InNodeID);
 #endif // WITH_EDITORONLY_DATA
 
 		return true;
@@ -3047,15 +3217,17 @@ bool FMetaSoundFrontendDocumentBuilder::RemoveNode(const FGuid& InNodeID)
 }
 
 #if WITH_EDITOR
-int32 FMetaSoundFrontendDocumentBuilder::RemoveNodeLocation(const FGuid& InNodeID, const FGuid* InLocationGuid)
+int32 FMetaSoundFrontendDocumentBuilder::RemoveNodeLocation(const FGuid& InNodeID, const FGuid* InLocationGuid, const FGuid* InPageID)
 {
 	using namespace Metasound;
 	using namespace Metasound::Frontend;
 
-	const IDocumentGraphNodeCache& NodeCache = DocumentCache->GetNodeCache();
+	const FGuid& PageID = InPageID ? *InPageID : BuildPageID;
+	const IDocumentGraphNodeCache& NodeCache = DocumentCache->GetNodeCache(PageID);
 	if (const int32* NodeIndex = NodeCache.FindNodeIndex(InNodeID))
 	{
-		FMetasoundFrontendNode& Node = FindBuildGraphChecked().Nodes[*NodeIndex];
+		FMetasoundFrontendGraph& Graph = GetDocumentChecked().RootGraph.FindGraphChecked(PageID);
+		FMetasoundFrontendNode& Node = Graph.Nodes[*NodeIndex];
 		FMetasoundFrontendNodeStyle& Style = Node.Style;
 		if (InLocationGuid)
 		{
@@ -3143,7 +3315,8 @@ bool FMetaSoundFrontendDocumentBuilder::SetGraphInputAccessType(FName InputName,
 	if (const int32* Index = DocumentCache->GetInterfaceCache().FindInputIndex(InputName))
 	{
 		FMetasoundFrontendDocument& Document = GetDocumentChecked();
-		FMetasoundFrontendClassInput& GraphInput = Document.RootGraph.Interface.Inputs[*Index];
+		FMetasoundFrontendGraphClass& RootGraph = Document.RootGraph;
+		FMetasoundFrontendClassInput& GraphInput = RootGraph.Interface.Inputs[*Index];
 
 		if (GraphInput.AccessType == AccessType)
 		{
@@ -3153,41 +3326,44 @@ bool FMetaSoundFrontendDocumentBuilder::SetGraphInputAccessType(FName InputName,
 		GraphInput.AccessType = AccessType;
 		if (AccessType == EMetasoundFrontendVertexAccessType::Reference)
 		{
-			const IDocumentGraphNodeCache& NodeCache = DocumentCache->GetNodeCache();
-			if (const int32* NodeIndex = NodeCache.FindNodeIndex(GraphInput.NodeID))
+			RootGraph.IterateGraphPages([this, &GraphInput, &AccessType](FMetasoundFrontendGraph& Graph)
 			{
-				FMetasoundFrontendNode& Node = FindBuildGraphChecked().Nodes[*NodeIndex];
-				const FMetasoundFrontendVertex& NodeOutput = Node.Interface.Outputs.Last();
-				IterateNodesConnectedWithVertex({ GraphInput.NodeID, NodeOutput.VertexID }, [this, &AccessType, &Document](const FMetasoundFrontendEdge& Edge, FMetasoundFrontendNode& ConnectedNode)
+				const IDocumentGraphNodeCache& NodeCache = DocumentCache->GetNodeCache(Graph.PageID);
+				if (const int32* NodeIndex = NodeCache.FindNodeIndex(GraphInput.NodeID))
 				{
-					if (const FMetasoundFrontendClass* ConnectedNodeClass = FindDependency(ConnectedNode.ClassID))
+					FMetasoundFrontendNode& Node = Graph.Nodes[*NodeIndex];
+					const FMetasoundFrontendVertex& NodeOutput = Node.Interface.Outputs.Last();
+					IterateNodesConnectedWithVertex({ GraphInput.NodeID, NodeOutput.VertexID }, [this, &Graph, &AccessType](const FMetasoundFrontendEdge& Edge, FMetasoundFrontendNode& ConnectedNode)
 					{
-						// If connected to an input template node, disconnect the template node from other nodes as the data type is
-						// about to be mismatched.  Otherwise, direct connection to other nodes (i.e. at runtime when template
-						// nodes aren't injected) forcefully remove to avoid data type mismatch.
-						if (ConnectedNodeClass->Metadata.GetClassName() == FInputNodeTemplate::ClassName)
+						if (const FMetasoundFrontendClass* ConnectedNodeClass = FindDependency(ConnectedNode.ClassID))
 						{
-							const FMetasoundFrontendVertex& ConnectedNodeOutput = ConnectedNode.Interface.Outputs.Last();
-							IterateNodesConnectedWithVertex({ Edge.ToNodeID, ConnectedNodeOutput.VertexID }, [this, &AccessType](const FMetasoundFrontendEdge& TempEdge, FMetasoundFrontendNode&)
+							// If connected to an input template node, disconnect the template node from other nodes as the data type is
+							// about to be mismatched.  Otherwise, direct connection to other nodes (i.e. at runtime when template
+							// nodes aren't injected) forcefully remove to avoid data type mismatch.
+							if (ConnectedNodeClass->Metadata.GetClassName() == FInputNodeTemplate::ClassName)
+							{
+								const FMetasoundFrontendVertex& ConnectedNodeOutput = ConnectedNode.Interface.Outputs.Last();
+								IterateNodesConnectedWithVertex({ Edge.ToNodeID, ConnectedNodeOutput.VertexID }, [this, &Graph, &AccessType](const FMetasoundFrontendEdge& TempEdge, FMetasoundFrontendNode&)
+								{
+									const EMetasoundFrontendVertexAccessType ConnectedAccessType = GetNodeInputAccessType(TempEdge.ToNodeID, TempEdge.ToVertexID, &Graph.PageID);
+									if (!FMetasoundFrontendClassVertex::CanConnectVertexAccessTypes(AccessType, ConnectedAccessType))
+									{
+										RemoveEdgeToNodeInput(TempEdge.ToNodeID, TempEdge.ToVertexID, &Graph.PageID);
+									}
+								}, Graph.PageID);
+							}
+							else
 							{
-								const EMetasoundFrontendVertexAccessType ConnectedAccessType = GetNodeInputAccessType(TempEdge.ToNodeID, TempEdge.ToVertexID);
+								const EMetasoundFrontendVertexAccessType ConnectedAccessType = GetNodeInputAccessType(Edge.ToNodeID, Edge.ToVertexID, &Graph.PageID);
 								if (!FMetasoundFrontendClassVertex::CanConnectVertexAccessTypes(AccessType, ConnectedAccessType))
 								{
-									RemoveEdgeToNodeInput(TempEdge.ToNodeID, TempEdge.ToVertexID);
+									RemoveEdgeToNodeInput(Edge.ToNodeID, Edge.ToVertexID, &Graph.PageID);
 								}
-							});
-						}
-						else
-						{
-							const EMetasoundFrontendVertexAccessType ConnectedAccessType = GetNodeInputAccessType(Edge.ToNodeID, Edge.ToVertexID);
-							if (!FMetasoundFrontendClassVertex::CanConnectVertexAccessTypes(AccessType, ConnectedAccessType))
-							{
-								RemoveEdgeToNodeInput(Edge.ToNodeID, Edge.ToVertexID);
 							}
 						}
-					}
-				});
-			}
+					}, Graph.PageID);
+				}
+			});
 
 			const bool bNodeConformed = ConformGraphInputNodeToClass(GraphInput);
 			if (bNodeConformed)
@@ -3213,7 +3389,8 @@ bool FMetaSoundFrontendDocumentBuilder::SetGraphInputDataType(FName InputName, F
 		if (const int32* Index = DocumentCache->GetInterfaceCache().FindInputIndex(InputName))
 		{
 			FMetasoundFrontendDocument& Document = GetDocumentChecked();
-			FMetasoundFrontendClassInput& GraphInput = Document.RootGraph.Interface.Inputs[*Index];
+			FMetasoundFrontendGraphClass& RootGraph = Document.RootGraph;
+			FMetasoundFrontendClassInput& GraphInput = RootGraph.Interface.Inputs[*Index];
 			if (GraphInput.TypeName == DataType)
 			{
 				return true;
@@ -3224,32 +3401,35 @@ bool FMetaSoundFrontendDocumentBuilder::SetGraphInputDataType(FName InputName, F
 			GraphInput.DefaultLiteral = DefaultLiteral;
 			GraphInput.TypeName = DataType;
 
-			const IDocumentGraphNodeCache& NodeCache = DocumentCache->GetNodeCache();
-			if (const int32* NodeIndex = NodeCache.FindNodeIndex(GraphInput.NodeID))
+			RootGraph.IterateGraphPages([this, &DataType, &GraphInput](FMetasoundFrontendGraph& Graph)
 			{
-				FMetasoundFrontendNode& Node = FindBuildGraphChecked().Nodes[*NodeIndex];
-				FMetasoundFrontendVertex& NodeOutput = Node.Interface.Outputs.Last();
-				IterateNodesConnectedWithVertex({ GraphInput.NodeID, NodeOutput.VertexID }, [this, &DataType](const FMetasoundFrontendEdge& Edge, FMetasoundFrontendNode& ConnectedNode)
+				const IDocumentGraphNodeCache& NodeCache = DocumentCache->GetNodeCache();
+				if (const int32* NodeIndex = NodeCache.FindNodeIndex(GraphInput.NodeID))
 				{
-					const FMetasoundFrontendClass* ConnectedNodeClass = FindDependency(ConnectedNode.ClassID);
-					if (ensure(ConnectedNodeClass))
+					FMetasoundFrontendNode& Node = Graph.Nodes[*NodeIndex];
+					FMetasoundFrontendVertex& NodeOutput = Node.Interface.Outputs.Last();
+					IterateNodesConnectedWithVertex({ GraphInput.NodeID, NodeOutput.VertexID }, [this, &Graph, &DataType](const FMetasoundFrontendEdge& Edge, FMetasoundFrontendNode& ConnectedNode)
 					{
-						// If connected to an input template node, disconnect the template node from other nodes as the data type is
-						// about to be mismatched.  Otherwise, direct connection to other nodes (i.e. at runtime when template
-						// nodes aren't injected) forcefully remove to avoid data type mismatch.
-						if (ConnectedNodeClass->Metadata.GetClassName() == FInputNodeTemplate::ClassName)
-						{
-							RemoveEdgesFromNodeOutput(Edge.ToNodeID, ConnectedNode.Interface.Outputs.Last().VertexID);
-							ConnectedNode.Interface.Inputs.Last().TypeName = DataType;
-							ConnectedNode.Interface.Outputs.Last().TypeName = DataType;
-						}
-						else
+						const FMetasoundFrontendClass* ConnectedNodeClass = FindDependency(ConnectedNode.ClassID);
+						if (ensure(ConnectedNodeClass))
 						{
-							RemoveEdgeToNodeInput(Edge.ToNodeID, Edge.ToVertexID);
+							// If connected to an input template node, disconnect the template node from other nodes as the data type is
+							// about to be mismatched.  Otherwise, direct connection to other nodes (i.e. at runtime when template
+							// nodes aren't injected) forcefully remove to avoid data type mismatch.
+							if (ConnectedNodeClass->Metadata.GetClassName() == FInputNodeTemplate::ClassName)
+							{
+								RemoveEdgesFromNodeOutput(Edge.ToNodeID, ConnectedNode.Interface.Outputs.Last().VertexID, &Graph.PageID);
+								ConnectedNode.Interface.Inputs.Last().TypeName = DataType;
+								ConnectedNode.Interface.Outputs.Last().TypeName = DataType;
+							}
+							else
+							{
+								RemoveEdgeToNodeInput(Edge.ToNodeID, Edge.ToVertexID, &Graph.PageID);
+							}
 						}
-					}
-				});
-			}
+					}, Graph.PageID);
+				}
+			});
 
 			const bool bNodeConformed = ConformGraphInputNodeToClass(GraphInput);
 			if (bNodeConformed)
@@ -3334,7 +3514,8 @@ bool FMetaSoundFrontendDocumentBuilder::SetGraphOutputAccessType(FName OutputNam
 	if (const int32* Index = DocumentCache->GetInterfaceCache().FindOutputIndex(OutputName))
 	{
 		FMetasoundFrontendDocument& Document = GetDocumentChecked();
-		FMetasoundFrontendClassOutput& GraphOutput = Document.RootGraph.Interface.Outputs[*Index];
+		FMetasoundFrontendGraphClass& RootGraph = Document.RootGraph;
+		FMetasoundFrontendClassOutput& GraphOutput = RootGraph.Interface.Outputs[*Index];
 		if (GraphOutput.AccessType == AccessType)
 		{
 			return true;
@@ -3343,24 +3524,27 @@ bool FMetaSoundFrontendDocumentBuilder::SetGraphOutputAccessType(FName OutputNam
 		GraphOutput.AccessType = AccessType;
 		if (AccessType == EMetasoundFrontendVertexAccessType::Value)
 		{
-			const IDocumentGraphNodeCache& NodeCache = DocumentCache->GetNodeCache();
-			if (const int32* NodeIndex = NodeCache.FindNodeIndex(GraphOutput.NodeID))
+			RootGraph.IterateGraphPages([this, &GraphOutput, &AccessType](FMetasoundFrontendGraph& Graph)
 			{
-				FMetasoundFrontendNode& Node = FindBuildGraphChecked().Nodes[*NodeIndex];
-				const FMetasoundFrontendVertex& NodeInput = Node.Interface.Inputs.Last();
-				IterateNodesConnectedWithVertex({ GraphOutput.NodeID, NodeInput.VertexID }, [this, &AccessType](const FMetasoundFrontendEdge& Edge, FMetasoundFrontendNode& ConnectedNode)
+				const IDocumentGraphNodeCache& NodeCache = DocumentCache->GetNodeCache();
+				if (const int32* NodeIndex = NodeCache.FindNodeIndex(GraphOutput.NodeID))
 				{
-					if (const FMetasoundFrontendClass* ConnectedNodeClass = FindDependency(ConnectedNode.ClassID))
+					FMetasoundFrontendNode& Node = FindBuildGraphChecked().Nodes[*NodeIndex];
+					const FMetasoundFrontendVertex& NodeInput = Node.Interface.Inputs.Last();
+					IterateNodesConnectedWithVertex({ GraphOutput.NodeID, NodeInput.VertexID }, [this, &Graph, &AccessType](const FMetasoundFrontendEdge& Edge, FMetasoundFrontendNode& ConnectedNode)
 					{
-						const FMetasoundFrontendVertex& ConnectedNodeOutput = ConnectedNode.Interface.Outputs.Last();
-						const EMetasoundFrontendVertexAccessType ConnectedAccessType = GetNodeOutputAccessType(ConnectedNode.GetID(), ConnectedNodeOutput.VertexID);
-						if (!FMetasoundFrontendClassVertex::CanConnectVertexAccessTypes(ConnectedAccessType, AccessType))
+						if (const FMetasoundFrontendClass* ConnectedNodeClass = FindDependency(ConnectedNode.ClassID))
 						{
-							RemoveEdgeToNodeInput(Edge.ToNodeID, Edge.ToVertexID);
+							const FMetasoundFrontendVertex& ConnectedNodeOutput = ConnectedNode.Interface.Outputs.Last();
+							const EMetasoundFrontendVertexAccessType ConnectedAccessType = GetNodeOutputAccessType(ConnectedNode.GetID(), ConnectedNodeOutput.VertexID, &Graph.PageID);
+							if (!FMetasoundFrontendClassVertex::CanConnectVertexAccessTypes(ConnectedAccessType, AccessType))
+							{
+								RemoveEdgeToNodeInput(Edge.ToNodeID, Edge.ToVertexID, &Graph.PageID);
+							}
 						}
-					}
-				});
-			}
+					}, Graph.PageID);
+				}
+			});
 		}
 
 		const bool bNodeConformed = ConformGraphOutputNodeToClass(GraphOutput);
@@ -3386,7 +3570,8 @@ bool FMetaSoundFrontendDocumentBuilder::SetGraphOutputDataType(FName OutputName,
 		if (const int32* Index = DocumentCache->GetInterfaceCache().FindOutputIndex(OutputName))
 		{
 			FMetasoundFrontendDocument& Document = GetDocumentChecked();
-			FMetasoundFrontendClassOutput& GraphOutput = Document.RootGraph.Interface.Outputs[*Index];
+			FMetasoundFrontendGraphClass& RootGraph = Document.RootGraph;
+			FMetasoundFrontendClassOutput& GraphOutput = RootGraph.Interface.Outputs[*Index];
 			if (GraphOutput.TypeName == DataType)
 			{
 				return true;
@@ -3394,20 +3579,22 @@ bool FMetaSoundFrontendDocumentBuilder::SetGraphOutputDataType(FName OutputName,
 
 			GraphOutput.TypeName = DataType;
 
-			const IDocumentGraphNodeCache& NodeCache = DocumentCache->GetNodeCache();
-			if (const int32* NodeIndex = NodeCache.FindNodeIndex(GraphOutput.NodeID))
+			RootGraph.IterateGraphPages([this, &GraphOutput, &DataType](FMetasoundFrontendGraph& Graph)
 			{
-				FMetasoundFrontendNode& Node = FindBuildGraphChecked().Nodes[*NodeIndex];
-
-				FMetasoundFrontendLiteral DefaultLiteral;
-				DefaultLiteral.SetFromLiteral(IDataTypeRegistry::Get().CreateDefaultLiteral(DataType));
-				FMetasoundFrontendVertex& NodeInput = Node.Interface.Inputs.Last();
-				Node.InputLiterals = { FMetasoundFrontendVertexLiteral { NodeInput.VertexID, MoveTemp(DefaultLiteral) } };
+				const IDocumentGraphNodeCache& NodeCache = DocumentCache->GetNodeCache();
+				if (const int32* NodeIndex = NodeCache.FindNodeIndex(GraphOutput.NodeID))
+				{
+					FMetasoundFrontendNode& Node = FindBuildGraphChecked().Nodes[*NodeIndex];
 
-				RemoveEdgeToNodeInput(GraphOutput.NodeID, NodeInput.VertexID);
-				GraphOutput.TypeName = DataType;
+					FMetasoundFrontendLiteral DefaultLiteral;
+					DefaultLiteral.SetFromLiteral(IDataTypeRegistry::Get().CreateDefaultLiteral(DataType));
+					FMetasoundFrontendVertex& NodeInput = Node.Interface.Inputs.Last();
+					Node.InputLiterals = { FMetasoundFrontendVertexLiteral { NodeInput.VertexID, MoveTemp(DefaultLiteral) } };
 
-			}
+					RemoveEdgeToNodeInput(GraphOutput.NodeID, NodeInput.VertexID);
+					GraphOutput.TypeName = DataType;
+				}
+			});
 
 			const bool bNodeConformed = ConformGraphOutputNodeToClass(GraphOutput);
 			if (bNodeConformed)
@@ -3441,9 +3628,9 @@ void FMetaSoundFrontendDocumentBuilder::SetMemberMetadata(UMetaSoundFrontendMemb
 	LiteralMetadata.Add(NewMetadata.MemberID, NewMetadata);
 }
 
-bool FMetaSoundFrontendDocumentBuilder::SetNodeComment(const FGuid& InNodeID, FString&& InNewComment)
+bool FMetaSoundFrontendDocumentBuilder::SetNodeComment(const FGuid& InNodeID, FString&& InNewComment, const FGuid* InPageID)
 {
-	if (FMetasoundFrontendNode* Node = FindNodeInternal(InNodeID))
+	if (FMetasoundFrontendNode* Node = FindNodeInternal(InNodeID, InPageID))
 	{
 		Node->Style.Display.Comment = MoveTemp(InNewComment);
 		return true;
@@ -3452,9 +3639,9 @@ bool FMetaSoundFrontendDocumentBuilder::SetNodeComment(const FGuid& InNodeID, FS
 	return false;
 }
 
-bool FMetaSoundFrontendDocumentBuilder::SetNodeCommentVisible(const FGuid& InNodeID, bool bIsVisible)
+bool FMetaSoundFrontendDocumentBuilder::SetNodeCommentVisible(const FGuid& InNodeID, bool bIsVisible, const FGuid* InPageID)
 {
-	if (FMetasoundFrontendNode* Node = FindNodeInternal(InNodeID))
+	if (FMetasoundFrontendNode* Node = FindNodeInternal(InNodeID, InPageID))
 	{
 		Node->Style.Display.bCommentVisible = bIsVisible;
 		return true;
@@ -3464,14 +3651,15 @@ bool FMetaSoundFrontendDocumentBuilder::SetNodeCommentVisible(const FGuid& InNod
 }
 #endif // WITH_EDITOR
 
-bool FMetaSoundFrontendDocumentBuilder::SetNodeInputDefault(const FGuid& InNodeID, const FGuid& InVertexID, const FMetasoundFrontendLiteral& InLiteral)
+bool FMetaSoundFrontendDocumentBuilder::SetNodeInputDefault(const FGuid& InNodeID, const FGuid& InVertexID, const FMetasoundFrontendLiteral& InLiteral, const FGuid* InPageID)
 {
 	using namespace Metasound::Frontend;
 
-	const IDocumentGraphNodeCache& NodeCache = DocumentCache->GetNodeCache();
+	const FGuid& PageID = InPageID ? *InPageID : BuildPageID;
+	FMetasoundFrontendGraph& Graph = GetDocumentChecked().RootGraph.FindGraphChecked(PageID);
+	const IDocumentGraphNodeCache& NodeCache = DocumentCache->GetNodeCache(PageID);
 	if (const int32* NodeIndex = NodeCache.FindNodeIndex(InNodeID))
 	{
-		FMetasoundFrontendGraph& Graph = FindBuildGraphChecked();
 		FMetasoundFrontendNode& Node = Graph.Nodes[*NodeIndex];
 
 		auto IsVertex = [&InVertexID](const FMetasoundFrontendVertex& Vertex) { return Vertex.VertexID == InVertexID; };
@@ -3494,7 +3682,8 @@ bool FMetaSoundFrontendDocumentBuilder::SetNodeInputDefault(const FGuid& InNodeI
 				Node.InputLiterals[LiteralIndex] = MoveTemp(NewVertexLiteral);
 			}
 
-			const FOnMetaSoundFrontendDocumentMutateNodeInputLiteralArray& OnNodeInputLiteralSet = DocumentDelegates->NodeDelegates.OnNodeInputLiteralSet;
+			FNodeModifyDelegates& NodeDelegates = DocumentDelegates->FindNodeDelegatesChecked(PageID);
+			const FOnMetaSoundFrontendDocumentMutateNodeInputLiteralArray& OnNodeInputLiteralSet = NodeDelegates.OnNodeInputLiteralSet;
 			OnNodeInputLiteralSet.Broadcast(*NodeIndex, VertexIndex, LiteralIndex);
 			return true;
 		}
@@ -3504,9 +3693,9 @@ bool FMetaSoundFrontendDocumentBuilder::SetNodeInputDefault(const FGuid& InNodeI
 }
 
 #if WITH_EDITOR
-bool FMetaSoundFrontendDocumentBuilder::SetNodeLocation(const FGuid& InNodeID, const FVector2D& InLocation, const FGuid* InLocationGuid)
+bool FMetaSoundFrontendDocumentBuilder::SetNodeLocation(const FGuid& InNodeID, const FVector2D& InLocation, const FGuid* InLocationGuid, const FGuid* InPageID)
 {
-	if (FMetasoundFrontendNode* Node = FindNodeInternal(InNodeID))
+	if (FMetasoundFrontendNode* Node = FindNodeInternal(InNodeID, InPageID))
 	{
 		FMetasoundFrontendNodeStyle& Style = Node->Style;
 		if (InLocationGuid)
@@ -3566,34 +3755,35 @@ void FMetaSoundFrontendDocumentBuilder::SetVersionNumber(const FMetasoundFronten
 	GetDocumentChecked().Metadata.Version.Number = InDocumentVersionNumber;
 }
 
-bool FMetaSoundFrontendDocumentBuilder::SpliceVariableNodeFromStack(const FGuid& InNodeID)
+bool FMetaSoundFrontendDocumentBuilder::SpliceVariableNodeFromStack(const FGuid& InNodeID, const FGuid& InPageID)
 {
 	using namespace Metasound;
 	using namespace Metasound::Frontend;
 
-	FMetasoundFrontendGraph& Graph = FindBuildGraphChecked();
-	const IDocumentGraphEdgeCache& EdgeCache = DocumentCache->GetEdgeCache();
+	FMetasoundFrontendDocument& Document = GetDocumentChecked();
+	FMetasoundFrontendGraph& Graph = Document.RootGraph.FindGraphChecked(InPageID);
+	const IDocumentGraphEdgeCache& EdgeCache = DocumentCache->GetEdgeCache(InPageID);
 	FMetasoundFrontendVertexHandle FromVariableVertexHandle;
 	{
-		const FMetasoundFrontendVertex* InputVertex = FindNodeInput(InNodeID, VariableNames::InputVariableName);
+		const FMetasoundFrontendVertex* InputVertex = FindNodeInput(InNodeID, VariableNames::InputVariableName, &InPageID);
 		check(InputVertex);
 		if (const int32* InputEdgeIndex = EdgeCache.FindEdgeIndexToNodeInput(InNodeID, InputVertex->VertexID))
 		{
 			FromVariableVertexHandle = Graph.Edges[*InputEdgeIndex].GetFromVertexHandle();
-			const bool bRemovedEdge = RemoveEdgeToNodeInput(InNodeID, InputVertex->VertexID);
+			const bool bRemovedEdge = RemoveEdgeToNodeInput(InNodeID, InputVertex->VertexID, &InPageID);
 			check(bRemovedEdge);
 		}
 	}
 
 	if (FromVariableVertexHandle.IsSet())
 	{
-		if (const FMetasoundFrontendVertex* OutputVertex = FindNodeOutput(InNodeID, VariableNames::OutputVariableName))
+		if (const FMetasoundFrontendVertex* OutputVertex = FindNodeOutput(InNodeID, VariableNames::OutputVariableName, &InPageID))
 		{
 			TArray<FMetasoundFrontendVertexHandle> ToVertexHandles;
 			const TArrayView<const int32> OutputEdgeIndices = EdgeCache.FindEdgeIndicesFromNodeOutput(InNodeID, OutputVertex->VertexID);
 			Algo::Transform(OutputEdgeIndices, ToVertexHandles, [&Graph](const int32& VertIndex) { return Graph.Edges[VertIndex].GetToVertexHandle(); });
 
-			RemoveEdgesFromNodeOutput(InNodeID, OutputVertex->VertexID);
+			RemoveEdgesFromNodeOutput(InNodeID, OutputVertex->VertexID, &InPageID);
 
 			for (const FMetasoundFrontendVertexHandle& ToHandle : ToVertexHandles)
 			{
@@ -3603,7 +3793,7 @@ bool FMetaSoundFrontendDocumentBuilder::SpliceVariableNodeFromStack(const FGuid&
 					FromVariableVertexHandle.VertexID,
 					ToHandle.NodeID,
 					ToHandle.VertexID
-				});
+				}, &InPageID);
 			}
 			return true;
 		}
@@ -3767,38 +3957,38 @@ bool FMetaSoundFrontendDocumentBuilder::SwapGraphOutput(const FMetasoundFrontend
 	return true;
 }
 
-bool FMetaSoundFrontendDocumentBuilder::UnlinkVariableNode(const FGuid& InNodeID)
+bool FMetaSoundFrontendDocumentBuilder::UnlinkVariableNode(const FGuid& InNodeID, const FGuid& InPageID)
 {
 	auto IsNodeID = [&InNodeID](const FGuid& TestID) { return TestID == InNodeID; };
 
-	FMetasoundFrontendGraph& Graph = FindBuildGraphChecked();
+	FMetasoundFrontendGraph& Graph = GetDocumentChecked().RootGraph.FindGraphChecked(InPageID);
 	for (FMetasoundFrontendVariable& Variable : Graph.Variables)
 	{
 		if (Variable.MutatorNodeID == InNodeID)
 		{
 			Variable.MutatorNodeID = FGuid();
-			SpliceVariableNodeFromStack(InNodeID);
+			SpliceVariableNodeFromStack(InNodeID, InPageID);
 			return true;
 		}
 
 		if (Variable.VariableNodeID == InNodeID)
 		{
 			Variable.VariableNodeID = FGuid();
-			SpliceVariableNodeFromStack(InNodeID);
+			SpliceVariableNodeFromStack(InNodeID, InPageID);
 			return true;
 		}
 
 		const bool bRemovedDeferredNode = Variable.DeferredAccessorNodeIDs.RemoveAllSwap(IsNodeID, EAllowShrinking::No) > 0;
 		if (bRemovedDeferredNode)
 		{
-			SpliceVariableNodeFromStack(InNodeID);
+			SpliceVariableNodeFromStack(InNodeID, InPageID);
 			return true;
 		}
 
 		const bool bRemovedAccessorNode = Variable.AccessorNodeIDs.RemoveAllSwap(IsNodeID, EAllowShrinking::No) > 0;
 		if (bRemovedAccessorNode)
 		{
-			SpliceVariableNodeFromStack(InNodeID);
+			SpliceVariableNodeFromStack(InNodeID, InPageID);
 			return true;
 		}
 	}
@@ -3808,6 +3998,11 @@ bool FMetaSoundFrontendDocumentBuilder::UnlinkVariableNode(const FGuid& InNodeID
 
 bool FMetaSoundFrontendDocumentBuilder::UpdateDependencyClassNames(const TMap<FMetasoundFrontendClassName, FMetasoundFrontendClassName>& OldToNewReferencedClassNames)
 {
+	if (!DocumentDelegates.IsValid())
+	{
+		return false;
+	}
+
 	for (FMetasoundFrontendClass& Dependency : GetDocumentChecked().Dependencies)
 	{
 		const FMetasoundFrontendClassName OldName = Dependency.Metadata.GetClassName();
@@ -3815,7 +4010,7 @@ bool FMetaSoundFrontendDocumentBuilder::UpdateDependencyClassNames(const TMap<FM
 		{
 			const int32* DependencyIndex = DocumentCache->FindDependencyIndex(Dependency.ID);
 			check(DependencyIndex);
-			GetDocumentDelegates().OnRenamingDependencyClass.Broadcast(*DependencyIndex, *NewName);
+			DocumentDelegates->OnRenamingDependencyClass.Broadcast(*DependencyIndex, *NewName);
 			Dependency.Metadata.SetClassName(*NewName);
 		}
 	}
diff --git a/Engine/Plugins/Runtime/Metasound/Source/MetasoundFrontend/Private/MetasoundFrontendDocumentCache.cpp b/Engine/Plugins/Runtime/Metasound/Source/MetasoundFrontend/Private/MetasoundFrontendDocumentCache.cpp
--- a/Engine/Plugins/Runtime/Metasound/Source/MetasoundFrontend/Private/MetasoundFrontendDocumentCache.cpp
+++ b/Engine/Plugins/Runtime/Metasound/Source/MetasoundFrontend/Private/MetasoundFrontendDocumentCache.cpp
@@ -135,8 +135,6 @@ namespace Metasound::Frontend
 		// Factory style constructor as restrictions on construction of shared pointers disallows passing of this document cache's pointer to sub-caches
 		TSharedRef<FDocumentCache> Cache = MakeShared<FDocumentCache>(InDocument, Delegates);
 
-		// Right now the assumption is the whole doc cache is invalidated when swapping pages.  Need to migrate to supporting associated caches per page
-		// (i.e. a node and edge cache per page)
 		Cache->Init(InBuildPageID, bPrimeCache);
 		return Cache;
 	}
@@ -491,9 +489,7 @@ namespace Metasound::Frontend
 
 		// Factory style constructor as restrictions on construction of shared pointers disallows passing of this document cache's pointer to sub-caches
 		TSharedRef<FDocumentGraphNodeCache> Cache = MakeShared<FDocumentGraphNodeCache>(ParentCache, InPageID);
-		Cache->Init(InPageID.IsValid()
-			? OutDelegates.PageNodeDelegates.FindOrAdd(InPageID)
-			: OutDelegates.NodeDelegates);
+		Cache->Init(OutDelegates.FindNodeDelegatesChecked(InPageID));
 		return Cache;
 	}
 
@@ -893,10 +889,7 @@ namespace Metasound::Frontend
 
 		// Factory style constructor as restrictions on construction of shared pointers disallows passing of this document cache's pointer to sub-caches
 		TSharedRef<FDocumentGraphEdgeCache> Cache = MakeShared<FDocumentGraphEdgeCache>(ParentCache, InPageID);
-		Cache->Init(InPageID.IsValid()
-			? OutDelegates.PageEdgeDelegates.FindOrAdd(InPageID)
-			: OutDelegates.EdgeDelegates);
-
+		Cache->Init(OutDelegates.FindEdgeDelegatesChecked(InPageID));
 		return Cache;
 	}
 
diff --git a/Engine/Plugins/Runtime/Metasound/Source/MetasoundFrontend/Private/MetasoundFrontendDocumentModifyDelegates.cpp b/Engine/Plugins/Runtime/Metasound/Source/MetasoundFrontend/Private/MetasoundFrontendDocumentModifyDelegates.cpp
--- a/Engine/Plugins/Runtime/Metasound/Source/MetasoundFrontend/Private/MetasoundFrontendDocumentModifyDelegates.cpp
+++ b/Engine/Plugins/Runtime/Metasound/Source/MetasoundFrontend/Private/MetasoundFrontendDocumentModifyDelegates.cpp
@@ -5,4 +5,27 @@
 
 namespace Metasound::Frontend
 {
+	FDocumentModifyDelegates::FDocumentModifyDelegates(const FMetasoundFrontendDocument& Document)
+	{
+		Document.RootGraph.IterateGraphPages([this](const FMetasoundFrontendGraph& Graph)
+		{
+			AddPageDelegates(Graph.PageID);
+		});
+	}
+
+	void FDocumentModifyDelegates::AddPageDelegates(const FGuid& InPageID)
+	{
+		PageNodeDelegates.Add(InPageID, FNodeModifyDelegates());
+		PageEdgeDelegates.Add(InPageID, FEdgeModifyDelegates());
+
+		PageDelegates.OnPageAdded.Broadcast(FDocumentMutatePageArgs{ InPageID });
+	}
+
+	void FDocumentModifyDelegates::RemovePageDelegates(const FGuid& InPageID)
+	{
+		PageDelegates.OnRemovingPage.Broadcast(FDocumentMutatePageArgs { InPageID });
+
+		PageNodeDelegates.Remove(InPageID);
+		PageEdgeDelegates.Remove(InPageID);
+	}
 } // namespace Metasound::Frontend
diff --git a/Engine/Plugins/Runtime/Metasound/Source/MetasoundFrontend/Public/MetasoundFrontendDocument.h b/Engine/Plugins/Runtime/Metasound/Source/MetasoundFrontend/Public/MetasoundFrontendDocument.h
--- a/Engine/Plugins/Runtime/Metasound/Source/MetasoundFrontend/Public/MetasoundFrontendDocument.h
+++ b/Engine/Plugins/Runtime/Metasound/Source/MetasoundFrontend/Public/MetasoundFrontendDocument.h
@@ -1669,10 +1669,13 @@ public:
 	FMetasoundFrontendGraphClassPresetOptions PresetOptions;
 
 public:
-	void AddNewGraphPage(const FGuid& InPageID);
-	void DuplicateLastGraphPage(const FGuid& InPageID);
+	const FMetasoundFrontendGraph& AddGraphPage(const FGuid& InPageID, bool bDuplicateLastGraph = true, bool bSetAsBuildGraph = true);
 	void RemoveAllGraphPages();
-	bool RemoveGraphPage(const FGuid& InPageID);
+
+	// Removes the page associated with the given PageID.  Returns true if removed, false if not.
+	// If provided an "AdjacentPageID," sets the value at the given pointer to a page ID adjacent to
+	// the removed page. If last page was removed, returns the default graph ID.
+	bool RemoveGraphPage(const FGuid& InPageID, FGuid* OutAdjacentPageID = nullptr);
 
 	bool ContainsGraphPage(const FGuid& InPageID) const;
 	void IterateGraphPages(TFunctionRef<void(FMetasoundFrontendGraph&)> IterFunc);
diff --git a/Engine/Plugins/Runtime/Metasound/Source/MetasoundFrontend/Public/MetasoundFrontendDocumentBuilder.h b/Engine/Plugins/Runtime/Metasound/Source/MetasoundFrontend/Public/MetasoundFrontendDocumentBuilder.h
--- a/Engine/Plugins/Runtime/Metasound/Source/MetasoundFrontend/Public/MetasoundFrontendDocumentBuilder.h
+++ b/Engine/Plugins/Runtime/Metasound/Source/MetasoundFrontend/Public/MetasoundFrontendDocumentBuilder.h
@@ -154,25 +154,26 @@ public:
 	void FinishBuilding();
 
 	const FMetasoundFrontendClass* AddDependency(const FMetasoundFrontendClass& InClass);
-	void AddEdge(FMetasoundFrontendEdge&& InNewEdge);
-	bool AddNamedEdges(const TSet<Metasound::Frontend::FNamedEdge>& ConnectionsToMake, TArray<const FMetasoundFrontendEdge*>* OutEdgesCreated = nullptr, bool bReplaceExistingConnections = true);
-	bool AddEdgesByNodeClassInterfaceBindings(const FGuid& InFromNodeID, const FGuid& InToNodeID, bool bReplaceExistingConnections = true);
-	bool AddEdgesFromMatchingInterfaceNodeOutputsToGraphOutputs(const FGuid& InNodeID, TArray<const FMetasoundFrontendEdge*>& OutEdgesCreated, bool bReplaceExistingConnections = true);
-	bool AddEdgesFromMatchingInterfaceNodeInputsToGraphInputs(const FGuid& InNodeID, TArray<const FMetasoundFrontendEdge*>& OutEdgesCreated, bool bReplaceExistingConnections = true);
-	const FMetasoundFrontendNode* AddGraphInput(const FMetasoundFrontendClassInput& InClassInput);
-	const FMetasoundFrontendNode* DuplicateGraphInput(const FMetasoundFrontendClassInput& InClassInput, FMetasoundFrontendLiteral DefaultValue, const FName InName);
-	const FMetasoundFrontendNode* AddGraphNode(const FMetasoundFrontendGraphClass& InClass, FGuid InNodeID = FGuid::NewGuid());
-	const FMetasoundFrontendNode* AddGraphOutput(const FMetasoundFrontendClassOutput& InClassOutput);
-	const FMetasoundFrontendNode* DuplicateGraphOutput(const FMetasoundFrontendClassOutput& InClassOutput, const FName InName);
+	void AddEdge(FMetasoundFrontendEdge&& InNewEdge, const FGuid* InPageID = nullptr);
+	bool AddNamedEdges(const TSet<Metasound::Frontend::FNamedEdge>& ConnectionsToMake, TArray<const FMetasoundFrontendEdge*>* OutEdgesCreated = nullptr, bool bReplaceExistingConnections = true, const FGuid* InPageID = nullptr);
+	bool AddEdgesByNodeClassInterfaceBindings(const FGuid& InFromNodeID, const FGuid& InToNodeID, bool bReplaceExistingConnections = true, const FGuid* InPageID = nullptr);
+	bool AddEdgesFromMatchingInterfaceNodeOutputsToGraphOutputs(const FGuid& InNodeID, TArray<const FMetasoundFrontendEdge*>& OutEdgesCreated, bool bReplaceExistingConnections = true, const FGuid* InPageID = nullptr);
+	bool AddEdgesFromMatchingInterfaceNodeInputsToGraphInputs(const FGuid& InNodeID, TArray<const FMetasoundFrontendEdge*>& OutEdgesCreated, bool bReplaceExistingConnections = true, const FGuid* InPageID = nullptr);
+	const FMetasoundFrontendNode* AddGraphInput(const FMetasoundFrontendClassInput& InClassInput, const FGuid* InPageID = nullptr);
+	const FMetasoundFrontendNode* AddGraphNode(const FMetasoundFrontendGraphClass& InClass, FGuid InNodeID = FGuid::NewGuid(), const FGuid* InPageID = nullptr);
+	const FMetasoundFrontendNode* AddGraphOutput(const FMetasoundFrontendClassOutput& InClassOutput, const FGuid* InPageID = nullptr);
 
 	bool AddInterface(FName InterfaceName);
 
-	const FMetasoundFrontendNode* AddNodeByClassName(const FMetasoundFrontendClassName& InClassName, int32 InMajorVersion = 1, FGuid InNodeID = FGuid::NewGuid());
-	const FMetasoundFrontendNode* AddNodeByTemplate(const Metasound::Frontend::INodeTemplate& InTemplate, Metasound::Frontend::FNodeTemplateGenerateInterfaceParams Params, FGuid InNodeID = FGuid::NewGuid());
+	const FMetasoundFrontendNode* AddNodeByClassName(const FMetasoundFrontendClassName& InClassName, int32 InMajorVersion = 1, FGuid InNodeID = FGuid::NewGuid(), const FGuid* InPageID = nullptr);
+	const FMetasoundFrontendNode* AddNodeByTemplate(const Metasound::Frontend::INodeTemplate& InTemplate, Metasound::Frontend::FNodeTemplateGenerateInterfaceParams Params, FGuid InNodeID = FGuid::NewGuid(), const FGuid* InPageID = nullptr);
+
+	// Adds a graph page to the given builder's document
+	const FMetasoundFrontendGraph& AddGraphPage(const FGuid& InPageID, bool bDuplicateLastGraph = true, bool bSetAsBuildGraph = true);
 
 	// Returns whether or not the given edge can be added, which requires that its input
 	// is not already connected and the edge is valid (see function 'IsValidEdge').
-	bool CanAddEdge(const FMetasoundFrontendEdge& InEdge) const;
+	bool CanAddEdge(const FMetasoundFrontendEdge& InEdge, const FGuid* InPageID = nullptr) const;
 
 	// Clears document completely of all graph page data (nodes, edges, & member metadata), dependencies,
 	// interfaces, member metadata, preset state, etc. Leaves ClassMetadata intact. Reloads the builder state,
@@ -187,12 +188,15 @@ public:
 #endif // WITH_EDITORONLY_DATA
 
 	bool ContainsDependencyOfType(EMetasoundFrontendClassType ClassType) const;
-	bool ContainsEdge(const FMetasoundFrontendEdge& InEdge) const;
-	bool ContainsNode(const FGuid& InNodeID) const;
+	bool ContainsEdge(const FMetasoundFrontendEdge& InEdge, const FGuid* InPageID = nullptr) const;
+	bool ContainsNode(const FGuid& InNodeID, const FGuid* InPageID = nullptr) const;
 
 	bool ConvertFromPreset();
 	bool ConvertToPreset(const FMetasoundFrontendDocument& InReferencedDocument, TSharedRef<Metasound::Frontend::FDocumentModifyDelegates> ModifyDelegates = { });
 
+	const FMetasoundFrontendNode* DuplicateGraphInput(const FMetasoundFrontendClassInput& InClassInput, FMetasoundFrontendLiteral DefaultValue, const FName InName, const FGuid* InPageID = nullptr);
+	const FMetasoundFrontendNode* DuplicateGraphOutput(const FMetasoundFrontendClassOutput& InClassOutput, const FName InName, const FGuid* InPageID = nullptr);
+
 #if WITH_EDITORONLY_DATA
 	const FMetaSoundFrontendGraphComment* FindGraphComment(const FGuid& InCommentID) const;
 	FMetaSoundFrontendGraphComment* FindGraphComment(const FGuid& InCommentID);
@@ -205,30 +209,30 @@ public:
 
 	const FMetasoundFrontendClass* FindDependency(const FGuid& InClassID) const;
 	const FMetasoundFrontendClass* FindDependency(const FMetasoundFrontendClassMetadata& InMetadata) const;
-	TArray<const FMetasoundFrontendEdge*> FindEdges(const FGuid& InNodeID, const FGuid& InVertexID) const;
+	TArray<const FMetasoundFrontendEdge*> FindEdges(const FGuid& InNodeID, const FGuid& InVertexID, const FGuid* InPageID = nullptr) const;
 
 	const FMetasoundFrontendClassInput* FindGraphInput(FName InputName) const;
-	const FMetasoundFrontendNode* FindGraphInputNode(FName InputName) const;
+	const FMetasoundFrontendNode* FindGraphInputNode(FName InputName, const FGuid* InPageID = nullptr) const;
 	const FMetasoundFrontendClassOutput* FindGraphOutput(FName OutputName) const;
-	const FMetasoundFrontendNode* FindGraphOutputNode(FName OutputName) const;
+	const FMetasoundFrontendNode* FindGraphOutputNode(FName OutputName, const FGuid* InPageID = nullptr) const;
 
-	bool FindInterfaceInputNodes(FName InterfaceName, TArray<const FMetasoundFrontendNode*>& OutInputs) const;
-	bool FindInterfaceOutputNodes(FName InterfaceName, TArray<const FMetasoundFrontendNode*>& OutOutputs) const;
+	bool FindInterfaceInputNodes(FName InterfaceName, TArray<const FMetasoundFrontendNode*>& OutInputs, const FGuid* InPageID = nullptr) const;
+	bool FindInterfaceOutputNodes(FName InterfaceName, TArray<const FMetasoundFrontendNode*>& OutOutputs, const FGuid* InPageID = nullptr) const;
 
 	FMetasoundFrontendGraph& FindBuildGraphChecked();
 	const FMetasoundFrontendGraph& FindConstBuildGraphChecked() const;
 
-	const FMetasoundFrontendNode* FindNode(const FGuid& InNodeID) const;
+	const FMetasoundFrontendNode* FindNode(const FGuid& InNodeID, const FGuid* InPageID = nullptr) const;
 
-	const FMetasoundFrontendVertex* FindNodeInput(const FGuid& InNodeID, const FGuid& InVertexID) const;
-	const FMetasoundFrontendVertex* FindNodeInput(const FGuid& InNodeID, FName InVertexName) const;
-	TArray<const FMetasoundFrontendVertex*> FindNodeInputs(const FGuid& InNodeID, FName TypeName = FName()) const;
-	TArray<const FMetasoundFrontendVertex*> FindNodeInputsConnectedToNodeOutput(const FGuid& InOutputNodeID, const FGuid& InOutputVertexID, TArray<const FMetasoundFrontendNode*>* ConnectedInputNodes = nullptr) const;
+	const FMetasoundFrontendVertex* FindNodeInput(const FGuid& InNodeID, const FGuid& InVertexID, const FGuid* InPageID = nullptr) const;
+	const FMetasoundFrontendVertex* FindNodeInput(const FGuid& InNodeID, FName InVertexName, const FGuid* InPageID = nullptr) const;
+	TArray<const FMetasoundFrontendVertex*> FindNodeInputs(const FGuid& InNodeID, FName TypeName = FName(), const FGuid* InPageID = nullptr) const;
+	TArray<const FMetasoundFrontendVertex*> FindNodeInputsConnectedToNodeOutput(const FGuid& InOutputNodeID, const FGuid& InOutputVertexID, TArray<const FMetasoundFrontendNode*>* ConnectedInputNodes = nullptr, const FGuid* InPageID = nullptr) const;
 
-	const FMetasoundFrontendVertex* FindNodeOutput(const FGuid& InNodeID, const FGuid& InVertexID) const;
-	const FMetasoundFrontendVertex* FindNodeOutput(const FGuid& InNodeID, FName InVertexName) const;
-	TArray<const FMetasoundFrontendVertex*> FindNodeOutputs(const FGuid& InNodeID, FName TypeName = FName()) const;
-	const FMetasoundFrontendVertex* FindNodeOutputConnectedToNodeInput(const FGuid& InInputNodeID, const FGuid& InInputVertexID, const FMetasoundFrontendNode** ConnectedOutputNode = nullptr) const;
+	const FMetasoundFrontendVertex* FindNodeOutput(const FGuid& InNodeID, const FGuid& InVertexID, const FGuid* InPageID = nullptr) const;
+	const FMetasoundFrontendVertex* FindNodeOutput(const FGuid& InNodeID, FName InVertexName, const FGuid* InPageID = nullptr) const;
+	TArray<const FMetasoundFrontendVertex*> FindNodeOutputs(const FGuid& InNodeID, FName TypeName = FName(), const FGuid* InPageID = nullptr) const;
+	const FMetasoundFrontendVertex* FindNodeOutputConnectedToNodeInput(const FGuid& InInputNodeID, const FGuid& InInputVertexID, const FMetasoundFrontendNode** ConnectedOutputNode = nullptr, const FGuid* InPageID = nullptr) const;
 
 	const FMetasoundFrontendDocument& GetConstDocumentChecked() const;
 	const IMetaSoundDocumentInterface& GetConstDocumentInterfaceChecked() const;
@@ -254,7 +258,7 @@ public:
 	// build errors.
 	FMetasoundFrontendClassName GenerateNewClassName();
 
-	const Metasound::Frontend::FDocumentModifyDelegates& GetDocumentDelegates() const;
+	Metasound::Frontend::FDocumentModifyDelegates& GetDocumentDelegates();
 
 	UE_DEPRECATED(5.5, "Use GetConstDocumentInterfaceChecked instead")
 	const IMetaSoundDocumentInterface& GetDocumentInterface() const;
@@ -262,11 +266,11 @@ public:
 
 	int32 GetTransactionCount() const;
 
-	TArray<const FMetasoundFrontendNode*> GetGraphInputTemplateNodes(FName InInputName);
-	EMetasoundFrontendVertexAccessType GetNodeInputAccessType(const FGuid& InNodeID, const FGuid& InVertexID) const;
-	const FMetasoundFrontendLiteral* GetNodeInputClassDefault(const FGuid& InNodeID, const FGuid& InVertexID) const;
-	const FMetasoundFrontendLiteral* GetNodeInputDefault(const FGuid& InNodeID, const FGuid& InVertexID) const;
-	EMetasoundFrontendVertexAccessType GetNodeOutputAccessType(const FGuid& InNodeID, const FGuid& InVertexID) const;
+	TArray<const FMetasoundFrontendNode*> GetGraphInputTemplateNodes(FName InInputName, const FGuid* InPageID = nullptr);
+	EMetasoundFrontendVertexAccessType GetNodeInputAccessType(const FGuid& InNodeID, const FGuid& InVertexID, const FGuid* InPageID = nullptr) const;
+	const FMetasoundFrontendLiteral* GetNodeInputClassDefault(const FGuid& InNodeID, const FGuid& InVertexID, const FGuid* InPageID = nullptr) const;
+	const FMetasoundFrontendLiteral* GetNodeInputDefault(const FGuid& InNodeID, const FGuid& InVertexID, const FGuid* InPageID = nullptr) const;
+	EMetasoundFrontendVertexAccessType GetNodeOutputAccessType(const FGuid& InNodeID, const FGuid& InVertexID, const FGuid* InPageID = nullptr) const;
 
 	// Initializes the builder's document, using the (optional) provided document template, (optional) class name, and (optionally) whether or not to reset the existing class version.
 	void InitDocument(const FMetasoundFrontendDocument* InDocumentTemplate = nullptr, const FMetasoundFrontendClassName* InNewClassName = nullptr, bool bResetVersion = true);
@@ -281,8 +285,8 @@ public:
 	void InvalidateCache() { }
 
 	bool IsDependencyReferenced(const FGuid& InClassID) const;
-	bool IsNodeInputConnected(const FGuid& InNodeID, const FGuid& InVertexID) const;
-	bool IsNodeOutputConnected(const FGuid& InNodeID, const FGuid& InVertexID) const;
+	bool IsNodeInputConnected(const FGuid& InNodeID, const FGuid& InVertexID, const FGuid* InPageID = nullptr) const;
+	bool IsNodeOutputConnected(const FGuid& InNodeID, const FGuid& InVertexID, const FGuid* InPageID = nullptr) const;
 
 	bool IsInterfaceDeclared(FName InInterfaceName) const;
 	bool IsInterfaceDeclared(const FMetasoundFrontendVersion& InInterfaceVersion) const;
@@ -293,10 +297,12 @@ public:
 
 	// Returns whether or not the given edge is valid (i.e. represents an input and output that equate in data and access types) or malformed.
 	// Note that this does not return whether or not the given edge exists, but rather if it could be legally applied to the given edge vertices.
-	Metasound::Frontend::EInvalidEdgeReason IsValidEdge(const FMetasoundFrontendEdge& InEdge) const;
+	Metasound::Frontend::EInvalidEdgeReason IsValidEdge(const FMetasoundFrontendEdge& InEdge, const FGuid* InPageID = nullptr) const;
 
 	bool ModifyInterfaces(Metasound::Frontend::FModifyInterfaceOptions&& InOptions);
 
+	void RemoveAllGraphPages();
+
 	UE_DEPRECATED(5.5,
 		"Cache invalidation may require new copy of delegates. In addition, re-priming is discouraged. "
 		"To enforce this, new recommended pattern is to construct a new builder instead")
@@ -304,32 +310,33 @@ public:
 
 	bool RemoveDependency(const FGuid& InClassID);
 	bool RemoveDependency(EMetasoundFrontendClassType ClassType, const FMetasoundFrontendClassName& InClassName, const FMetasoundFrontendVersionNumber& InClassVersionNumber);
-	bool RemoveEdge(const FMetasoundFrontendEdge& EdgeToRemove);
+	bool RemoveEdge(const FMetasoundFrontendEdge& EdgeToRemove, const FGuid* InPageID = nullptr);
 
 	// Removes all edges connected to an input or output vertex associated with the node of the given ID.
-	bool RemoveEdges(const FGuid& InNodeID);
+	bool RemoveEdges(const FGuid& InNodeID, const FGuid* InPageID = nullptr);
 
-	bool RemoveEdgesByNodeClassInterfaceBindings(const FGuid& InOutputNodeID, const FGuid& InInputNodeID);
-	bool RemoveEdgesFromNodeOutput(const FGuid& InNodeID, const FGuid& InVertexID);
-	bool RemoveEdgeToNodeInput(const FGuid& InNodeID, const FGuid& InVertexID);
+	bool RemoveEdgesByNodeClassInterfaceBindings(const FGuid& InOutputNodeID, const FGuid& InInputNodeID, const FGuid* InPageID = nullptr);
+	bool RemoveEdgesFromNodeOutput(const FGuid& InNodeID, const FGuid& InVertexID, const FGuid* InPageID = nullptr);
+	bool RemoveEdgeToNodeInput(const FGuid& InNodeID, const FGuid& InVertexID, const FGuid* InPageID = nullptr);
 
 #if WITH_EDITOR
-	bool RemoveGraphComment(const FGuid& InCommentID);
+	bool RemoveGraphComment(const FGuid& InCommentID, const FGuid* InPageID = nullptr);
 #endif // WITH_EDITOR
 
 	bool RemoveGraphInput(FName InInputName);
 	bool RemoveGraphOutput(FName InOutputName);
+	bool RemoveGraphPage(const FGuid& InPageID);
 	bool RemoveInterface(FName InName);
-	bool RemoveNamedEdges(const TSet<Metasound::Frontend::FNamedEdge>& InNamedEdgesToRemove, TArray<FMetasoundFrontendEdge>* OutRemovedEdges = nullptr);
-	bool RemoveNode(const FGuid& InNodeID);
+	bool RemoveNamedEdges(const TSet<Metasound::Frontend::FNamedEdge>& InNamedEdgesToRemove, TArray<FMetasoundFrontendEdge>* OutRemovedEdges = nullptr, const FGuid* InPageID = nullptr);
+	bool RemoveNode(const FGuid& InNodeID, const FGuid* InPageID = nullptr);
 
 #if WITH_EDITOR
-	int32 RemoveNodeLocation(const FGuid& InNodeID, const FGuid* InLocationGuid = nullptr);
+	int32 RemoveNodeLocation(const FGuid& InNodeID, const FGuid* InLocationGuid = nullptr, const FGuid* InPageID = nullptr);
 #endif // WITH_EDITOR
 
 	void Reload(TSharedPtr<Metasound::Frontend::FDocumentModifyDelegates> Delegates = {}, bool bPrimeCache = false);
 
-	bool RemoveNodeInputDefault(const FGuid& InNodeID, const FGuid& InVertexID);
+	bool RemoveNodeInputDefault(const FGuid& InNodeID, const FGuid& InVertexID, const FGuid* InPageID = nullptr);
 	bool RemoveUnusedDependencies();
 
 	UE_DEPRECATED(5.5, "Use GenerateNewClassName instead")
@@ -376,25 +383,24 @@ public:
 
 	// Sets the editor-only comment to the provided value.
 	// Returns true if the node was found and the comment was updated, false if not.
-	bool SetNodeComment(const FGuid& InNodeID, FString&& InNewComment);
+	bool SetNodeComment(const FGuid& InNodeID, FString&& InNewComment, const FGuid* InPageID = nullptr);
 
 	// Sets the editor-only comment visibility.
 	// Returns true if the node was found and the visibility was set, false if not.
-	bool SetNodeCommentVisible(const FGuid& InNodeID, bool bIsVisible);
+	bool SetNodeCommentVisible(const FGuid& InNodeID, bool bIsVisible, const FGuid* InPageID = nullptr);
 
 	// Sets the editor-only node location of a node with the given ID to the provided location.
 	// Returns true if the node was found and the location was updated, false if not.
-	bool SetNodeLocation(const FGuid& InNodeID, const FVector2D& InLocation, const FGuid* InLocationGuid = nullptr);
+	bool SetNodeLocation(const FGuid& InNodeID, const FVector2D& InLocation, const FGuid* InLocationGuid = nullptr, const FGuid* InPageID = nullptr);
 
 	// Sets the editor-only Unconnected Pins Hidden for a node with the given ID.
 	bool SetNodeUnconnectedPinsHidden(const FGuid& InNodeID, const bool bUnconnectedPinsHidden);
 
 	// Gets the editor-only style of a node with the given ID.
 	const FMetasoundFrontendNodeStyle* GetNodeStyle(const FGuid& InNodeID);
-
 #endif // WITH_EDITOR
 
-	bool SetNodeInputDefault(const FGuid& InNodeID, const FGuid& InVertexID, const FMetasoundFrontendLiteral& InLiteral);
+	bool SetNodeInputDefault(const FGuid& InNodeID, const FGuid& InVertexID, const FMetasoundFrontendLiteral& InLiteral, const FGuid* InPageID = nullptr);
 
 	// Sets the document's version number.  Should only be called by document versioning.
 	void SetVersionNumber(const FMetasoundFrontendVersionNumber& InDocumentVersionNumber);
@@ -415,7 +421,7 @@ public:
 private:
 	using FFinalizeNodeFunctionRef = TFunctionRef<void(FMetasoundFrontendNode&, const Metasound::Frontend::FNodeRegistryKey&)>;
 
-	FMetasoundFrontendNode* AddNodeInternal(const FMetasoundFrontendClassMetadata& InClassMetadata, Metasound::Frontend::FFinalizeNodeFunctionRef FinalizeNode, FGuid InNodeID = FGuid::NewGuid(), int32* NewNodeIndex = nullptr, const FGuid* InGraphPageID = nullptr);
+	FMetasoundFrontendNode* AddNodeInternal(const FMetasoundFrontendClassMetadata& InClassMetadata, Metasound::Frontend::FFinalizeNodeFunctionRef FinalizeNode, const FGuid& InPageID, FGuid InNodeID = FGuid::NewGuid(), int32* NewNodeIndex = nullptr);
 	void BeginBuilding(TSharedPtr<Metasound::Frontend::FDocumentModifyDelegates> Delegates = {}, bool bPrimeCache = false);
 
 	// Conforms GraphOutput node's ClassID, Access & Data Type with the GraphOutput.
@@ -430,10 +436,10 @@ private:
 	// edges, edges may then be invalid due to access type/DataType incompatibility).
 	bool ConformGraphOutputNodeToClass(const FMetasoundFrontendClassOutput& GraphOutput);
 
-	bool FindNodeClassInterfaces(const FGuid& InNodeID, TSet<FMetasoundFrontendVersion>& OutInterfaces) const;
-	FMetasoundFrontendNode* FindNodeInternal(const FGuid& InNodeID);
+	bool FindNodeClassInterfaces(const FGuid& InNodeID, TSet<FMetasoundFrontendVersion>& OutInterfaces, const FGuid& InPageID) const;
+	FMetasoundFrontendNode* FindNodeInternal(const FGuid& InNodeID, const FGuid* InPageID = nullptr);
 
-	void IterateNodesConnectedWithVertex(const FMetasoundFrontendVertexHandle& Vertex, TFunctionRef<void(const FMetasoundFrontendEdge&, FMetasoundFrontendNode&)> NodeIndexIterFunc);
+	void IterateNodesConnectedWithVertex(const FMetasoundFrontendVertexHandle& Vertex, TFunctionRef<void(const FMetasoundFrontendEdge&, FMetasoundFrontendNode&)> NodeIndexIterFunc, const FGuid& InPageID);
 
 	const FTopLevelAssetPath GetBuilderClassPath() const;
 	FMetasoundFrontendDocument& GetDocumentChecked() const;
@@ -441,12 +447,14 @@ private:
 
 	bool SetGraphInputInheritsDefault(FName InName, bool bInputInheritsDefault);
 
-	bool SpliceVariableNodeFromStack(const FGuid& InNodeID);
-	bool UnlinkVariableNode(const FGuid& InNodeID);
+	bool SpliceVariableNodeFromStack(const FGuid& InNodeID, const FGuid& InPageID);
+	bool UnlinkVariableNode(const FGuid& InNodeID, const FGuid& InPageID);
 
 	UPROPERTY(Transient)
 	TScriptInterface<IMetaSoundDocumentInterface> DocumentInterface;
 
+	// Default page ID to apply build transactions to if no page ID is optionally set directly.
+	UPROPERTY(Transient)
 	FGuid BuildPageID;
 
 	TSharedPtr<Metasound::Frontend::IDocumentCache> DocumentCache;
diff --git a/Engine/Plugins/Runtime/Metasound/Source/MetasoundFrontend/Public/MetasoundFrontendDocumentModifyDelegates.h b/Engine/Plugins/Runtime/Metasound/Source/MetasoundFrontend/Public/MetasoundFrontendDocumentModifyDelegates.h
--- a/Engine/Plugins/Runtime/Metasound/Source/MetasoundFrontend/Public/MetasoundFrontendDocumentModifyDelegates.h
+++ b/Engine/Plugins/Runtime/Metasound/Source/MetasoundFrontend/Public/MetasoundFrontendDocumentModifyDelegates.h
@@ -6,6 +6,7 @@
 #include "Templates/SharedPointer.h"
 
 
+// TODO: Move these to namespace
 DECLARE_MULTICAST_DELEGATE_OneParam(FOnMetaSoundFrontendDocumentMutateArray, int32 /* Index */);
 DECLARE_MULTICAST_DELEGATE_OneParam(FOnMetaSoundFrontendDocumentMutateInterfaceArray, const FMetasoundFrontendInterface& /* Interface */);
 DECLARE_MULTICAST_DELEGATE_TwoParams(FOnMetaSoundFrontendDocumentRemoveSwappingArray, int32 /* Index */, int32 /* LastIndex */);
@@ -15,6 +16,20 @@ DECLARE_MULTICAST_DELEGATE_ThreeParams(FOnMetaSoundFrontendDocumentMutateNodeInp
 
 namespace Metasound::Frontend
 {
+	struct METASOUNDFRONTEND_API FDocumentMutatePageArgs
+	{
+		FGuid PageID;
+	};
+
+	DECLARE_MULTICAST_DELEGATE_OneParam(FOnDocumentPageAdded, const FDocumentMutatePageArgs& /* Args */);
+	DECLARE_MULTICAST_DELEGATE_OneParam(FOnDocumentRemovingPage, const FDocumentMutatePageArgs& /* Args */);
+
+	struct METASOUNDFRONTEND_API FPageModifyDelegates
+	{
+		FOnDocumentPageAdded OnPageAdded;
+		FOnDocumentRemovingPage OnRemovingPage;
+	};
+
 	struct METASOUNDFRONTEND_API FInterfaceModifyDelegates
 	{
 		FOnMetaSoundFrontendDocumentMutateInterfaceArray OnInterfaceAdded;
@@ -45,15 +60,70 @@ namespace Metasound::Frontend
 
 	struct METASOUNDFRONTEND_API FDocumentModifyDelegates : TSharedFromThis<FDocumentModifyDelegates>
 	{
+		FDocumentModifyDelegates() = default;
+		FDocumentModifyDelegates(const FMetasoundFrontendDocument& Document);
+
 		FOnMetaSoundFrontendDocumentMutateArray OnDependencyAdded;
 		FOnMetaSoundFrontendDocumentRemoveSwappingArray OnRemoveSwappingDependency;
 		FOnMetaSoundFrontendDocumentRenameClass OnRenamingDependencyClass;
 
+		FPageModifyDelegates PageDelegates;
 		FInterfaceModifyDelegates InterfaceDelegates;
+
+		UE_DEPRECATED(5.5, "Public exposition of NodeDelegates will be removed in a future build.  Use accessor 'FindNodeDelegates' instead")
 		FNodeModifyDelegates NodeDelegates;
+
+		UE_DEPRECATED(5.5, "Public exposition of EdgeDelegates will be removed in a future build.  Use accessor 'FindEdgeDelegates' instead")
 		FEdgeModifyDelegates EdgeDelegates;
 
+		void AddPageDelegates(const FGuid& InPageID);
+		void RemovePageDelegates(const FGuid& InPageID);
+
+	private:
 		TSortedMap<FGuid, FNodeModifyDelegates> PageNodeDelegates;
 		TSortedMap<FGuid, FEdgeModifyDelegates> PageEdgeDelegates;
+
+	public:
+		FNodeModifyDelegates& FindNodeDelegatesChecked(const FGuid& InPageID)
+		{
+PRAGMA_DISABLE_DEPRECATION_WARNINGS
+			return InPageID.IsValid()
+				? PageNodeDelegates.FindChecked(InPageID)
+				: NodeDelegates;
+PRAGMA_ENABLE_DEPRECATION_WARNINGS
+		}
+
+		FEdgeModifyDelegates& FindEdgeDelegatesChecked(const FGuid& InPageID)
+		{
+PRAGMA_DISABLE_DEPRECATION_WARNINGS
+			return InPageID.IsValid()
+				? PageEdgeDelegates.FindChecked(InPageID)
+				: EdgeDelegates;
+PRAGMA_ENABLE_DEPRECATION_WARNINGS
+		}
+
+		void IterateGraphEdgeDelegates(TFunctionRef<void(FEdgeModifyDelegates&)> Func)
+		{
+PRAGMA_DISABLE_DEPRECATION_WARNINGS
+			Func(EdgeDelegates);
+
+			for (TPair<FGuid, FEdgeModifyDelegates>& Delegates : PageEdgeDelegates)
+			{
+				Func(Delegates.Value);
+			}
+PRAGMA_ENABLE_DEPRECATION_WARNINGS
+		}
+
+		void IterateGraphNodeDelegates(TFunctionRef<void(FNodeModifyDelegates&)> Func)
+		{
+PRAGMA_DISABLE_DEPRECATION_WARNINGS
+			Func(NodeDelegates);
+
+			for (TPair<FGuid, FNodeModifyDelegates>& Delegates : PageNodeDelegates)
+			{
+				Func(Delegates.Value);
+			}
+PRAGMA_ENABLE_DEPRECATION_WARNINGS
+		}
 	};
 } // namespace Metasound::Frontend
