Fix for Clang compiler warning. Minor feedback for const declaration of default page ID.
#rnx

diff --git a/Engine/Plugins/Runtime/Metasound/Source/MetasoundFrontend/Private/MetasoundFrontendDocument.cpp b/Engine/Plugins/Runtime/Metasound/Source/MetasoundFrontend/Private/MetasoundFrontendDocument.cpp
--- a/Engine/Plugins/Runtime/Metasound/Source/MetasoundFrontend/Private/MetasoundFrontendDocument.cpp
+++ b/Engine/Plugins/Runtime/Metasound/Source/MetasoundFrontend/Private/MetasoundFrontendDocument.cpp
@@ -879,7 +879,7 @@ FMetasoundFrontendGraphClass::FMetasoundFrontendGraphClass()
 
 const FMetasoundFrontendGraph& FMetasoundFrontendGraphClass::AddGraphPage(const FGuid& InPageID, bool bDuplicateLastGraph, bool bSetAsBuildGraph)
 {
-	checkf(InPageID.IsValid(), TEXT("Cannot add graph with invalid PageID"))
+	checkf(InPageID != Metasound::Frontend::DefaultGraphPageID, TEXT("Cannot add graph with DefaultPageID"))
 	checkf(!ContainsGraphPage(InPageID), TEXT("Cannot add new graph page with existing PageID"));
 
 	FMetasoundFrontendGraph* NewGraph = nullptr;
diff --git a/Engine/Plugins/Runtime/Metasound/Source/MetasoundFrontend/Private/MetasoundFrontendDocumentBuilder.cpp b/Engine/Plugins/Runtime/Metasound/Source/MetasoundFrontend/Private/MetasoundFrontendDocumentBuilder.cpp
--- a/Engine/Plugins/Runtime/Metasound/Source/MetasoundFrontend/Private/MetasoundFrontendDocumentBuilder.cpp
+++ b/Engine/Plugins/Runtime/Metasound/Source/MetasoundFrontend/Private/MetasoundFrontendDocumentBuilder.cpp
@@ -2599,14 +2599,17 @@ void FMetaSoundFrontendDocumentBuilder::BeginBuilding(TSharedPtr<Metasound::Fron
 	{
 		DocumentDelegates = Delegates;
 	}
-	else if (DocumentInterface)
-	{
-		const FMetasoundFrontendDocument& Document = GetConstDocumentChecked();
-		DocumentDelegates = Delegates.IsValid() ? Delegates : MakeShared<FDocumentModifyDelegates>(Document);
-	}
 	else
 	{
-		DocumentDelegates = MakeShared<FDocumentModifyDelegates>();
+		if (DocumentInterface)
+		{
+			const FMetasoundFrontendDocument& Document = GetConstDocumentChecked();
+			DocumentDelegates = MakeShared<FDocumentModifyDelegates>(Document);
+		}
+		else
+		{
+			DocumentDelegates = MakeShared<FDocumentModifyDelegates>();
+		}
 	}
 
 	if (DocumentInterface)
@@ -2639,13 +2642,13 @@ void FMetaSoundFrontendDocumentBuilder::RemoveAllGraphPages()
 	{
 		using namespace Metasound::Frontend;
 
-		if (Graph.PageID.IsValid())
+		if (Graph.PageID != Metasound::Frontend::DefaultGraphPageID)
 		{
-			DocumentDelegates->PageDelegates.OnRemovingPage.Broadcast(FDocumentMutatePageArgs { Graph.PageID });
+			DocumentDelegates->PageDelegates.OnRemovingPage.Broadcast(Metasound::Frontend::FDocumentMutatePageArgs { Graph.PageID });
 		}
 	});
 	RootGraph.RemoveAllGraphPages();
-	if (BuildPageID.IsValid())
+	if (BuildPageID == Metasound::Frontend::DefaultGraphPageID)
 	{
 		SetBuildPageID(FGuid());
 	}
diff --git a/Engine/Plugins/Runtime/Metasound/Source/MetasoundFrontend/Private/MetasoundFrontendDocumentModifyDelegates.cpp b/Engine/Plugins/Runtime/Metasound/Source/MetasoundFrontend/Private/MetasoundFrontendDocumentModifyDelegates.cpp
--- a/Engine/Plugins/Runtime/Metasound/Source/MetasoundFrontend/Private/MetasoundFrontendDocumentModifyDelegates.cpp
+++ b/Engine/Plugins/Runtime/Metasound/Source/MetasoundFrontend/Private/MetasoundFrontendDocumentModifyDelegates.cpp
@@ -5,7 +5,19 @@
 
 namespace Metasound::Frontend
 {
+	FDocumentModifyDelegates::FDocumentModifyDelegates()
+PRAGMA_DISABLE_DEPRECATION_WARNINGS
+		: NodeDelegates()
+		, EdgeDelegates()
+PRAGMA_ENABLE_DEPRECATION_WARNINGS
+	{
+	}
+
 	FDocumentModifyDelegates::FDocumentModifyDelegates(const FMetasoundFrontendDocument& Document)
+PRAGMA_DISABLE_DEPRECATION_WARNINGS
+		: NodeDelegates()
+		, EdgeDelegates()
+PRAGMA_ENABLE_DEPRECATION_WARNINGS
 	{
 		Document.RootGraph.IterateGraphPages([this](const FMetasoundFrontendGraph& Graph)
 		{
diff --git a/Engine/Plugins/Runtime/Metasound/Source/MetasoundFrontend/Public/MetasoundFrontendDocument.h b/Engine/Plugins/Runtime/Metasound/Source/MetasoundFrontend/Public/MetasoundFrontendDocument.h
--- a/Engine/Plugins/Runtime/Metasound/Source/MetasoundFrontend/Public/MetasoundFrontendDocument.h
+++ b/Engine/Plugins/Runtime/Metasound/Source/MetasoundFrontend/Public/MetasoundFrontendDocument.h
@@ -35,6 +35,8 @@ namespace Metasound
 
 	namespace Frontend
 	{
+		constexpr FGuid DefaultGraphPageID(0, 0, 0, 0);
+
 		namespace DisplayStyle
 		{
 			namespace EdgeAnimation
diff --git a/Engine/Plugins/Runtime/Metasound/Source/MetasoundFrontend/Public/MetasoundFrontendDocumentModifyDelegates.h b/Engine/Plugins/Runtime/Metasound/Source/MetasoundFrontend/Public/MetasoundFrontendDocumentModifyDelegates.h
--- a/Engine/Plugins/Runtime/Metasound/Source/MetasoundFrontend/Public/MetasoundFrontendDocumentModifyDelegates.h
+++ b/Engine/Plugins/Runtime/Metasound/Source/MetasoundFrontend/Public/MetasoundFrontendDocumentModifyDelegates.h
@@ -60,7 +60,7 @@ namespace Metasound::Frontend
 
 	struct METASOUNDFRONTEND_API FDocumentModifyDelegates : TSharedFromThis<FDocumentModifyDelegates>
 	{
-		FDocumentModifyDelegates() = default;
+		FDocumentModifyDelegates();
 		FDocumentModifyDelegates(const FMetasoundFrontendDocument& Document);
 
 		FOnMetaSoundFrontendDocumentMutateArray OnDependencyAdded;
@@ -87,18 +87,18 @@ namespace Metasound::Frontend
 		FNodeModifyDelegates& FindNodeDelegatesChecked(const FGuid& InPageID)
 		{
 PRAGMA_DISABLE_DEPRECATION_WARNINGS
-			return InPageID.IsValid()
-				? PageNodeDelegates.FindChecked(InPageID)
-				: NodeDelegates;
+			return InPageID == DefaultGraphPageID
+				? NodeDelegates
+				: PageNodeDelegates.FindChecked(InPageID);
 PRAGMA_ENABLE_DEPRECATION_WARNINGS
 		}
 
 		FEdgeModifyDelegates& FindEdgeDelegatesChecked(const FGuid& InPageID)
 		{
 PRAGMA_DISABLE_DEPRECATION_WARNINGS
-			return InPageID.IsValid()
-				? PageEdgeDelegates.FindChecked(InPageID)
-				: EdgeDelegates;
+			return InPageID == DefaultGraphPageID
+				? EdgeDelegates
+				: PageEdgeDelegates.FindChecked(InPageID);
 PRAGMA_ENABLE_DEPRECATION_WARNINGS
 		}
 
@@ -126,4 +126,12 @@ PRAGMA_DISABLE_DEPRECATION_WARNINGS
 PRAGMA_ENABLE_DEPRECATION_WARNINGS
 		}
 	};
+
+	class METASOUNDFRONTEND_API IDocumentBuilderListener : public TSharedFromThis<IDocumentBuilderListener>
+	{
+	public:
+		virtual ~IDocumentBuilderListener() = default;
+
+		virtual void AddDocumentBuilderDelegates(FDocumentModifyDelegates& OutDelegates) = 0;
+	};
 } // namespace Metasound::Frontend
