Fix round 2 for disabling deprecation warnings
#rnx

diff --git a/Engine/Plugins/Runtime/Metasound/Source/MetasoundFrontend/Private/MetasoundFrontendDocumentModifyDelegates.cpp b/Engine/Plugins/Runtime/Metasound/Source/MetasoundFrontend/Private/MetasoundFrontendDocumentModifyDelegates.cpp
--- a/Engine/Plugins/Runtime/Metasound/Source/MetasoundFrontend/Private/MetasoundFrontendDocumentModifyDelegates.cpp
+++ b/Engine/Plugins/Runtime/Metasound/Source/MetasoundFrontend/Private/MetasoundFrontendDocumentModifyDelegates.cpp
@@ -13,6 +13,48 @@ PRAGMA_ENABLE_DEPRECATION_WARNINGS
 	{
 	}
 
+	FDocumentModifyDelegates::FDocumentModifyDelegates(const FDocumentModifyDelegates& InModifyDelegates)
+PRAGMA_DISABLE_DEPRECATION_WARNINGS
+		: NodeDelegates(InModifyDelegates.NodeDelegates)
+		, EdgeDelegates(InModifyDelegates.EdgeDelegates)
+PRAGMA_ENABLE_DEPRECATION_WARNINGS
+		, PageNodeDelegates(InModifyDelegates.PageNodeDelegates)
+		, PageEdgeDelegates(InModifyDelegates.PageEdgeDelegates)
+	{
+	}
+
+	FDocumentModifyDelegates::FDocumentModifyDelegates(FDocumentModifyDelegates&& InModifyDelegates)
+PRAGMA_DISABLE_DEPRECATION_WARNINGS
+		: NodeDelegates(MoveTemp(InModifyDelegates.NodeDelegates))
+		, EdgeDelegates(MoveTemp(InModifyDelegates.EdgeDelegates))
+PRAGMA_ENABLE_DEPRECATION_WARNINGS
+		, PageNodeDelegates(MoveTemp(InModifyDelegates.PageNodeDelegates))
+		, PageEdgeDelegates(MoveTemp(InModifyDelegates.PageEdgeDelegates))
+	{
+	}
+
+	FDocumentModifyDelegates& FDocumentModifyDelegates::operator=(const FDocumentModifyDelegates& InModifyDelegates)
+	{
+PRAGMA_DISABLE_DEPRECATION_WARNINGS
+		NodeDelegates = InModifyDelegates.NodeDelegates;
+		EdgeDelegates = InModifyDelegates.EdgeDelegates;
+PRAGMA_ENABLE_DEPRECATION_WARNINGS
+		PageNodeDelegates = InModifyDelegates.PageNodeDelegates;
+		PageEdgeDelegates = InModifyDelegates.PageEdgeDelegates;
+		return *this;
+	}
+
+	FDocumentModifyDelegates& FDocumentModifyDelegates::operator=(FDocumentModifyDelegates&& InModifyDelegates)
+	{
+PRAGMA_DISABLE_DEPRECATION_WARNINGS
+		NodeDelegates = MoveTemp(InModifyDelegates.NodeDelegates);
+		EdgeDelegates = MoveTemp(InModifyDelegates.EdgeDelegates);
+PRAGMA_ENABLE_DEPRECATION_WARNINGS
+		PageNodeDelegates = MoveTemp(InModifyDelegates.PageNodeDelegates);
+		PageEdgeDelegates = MoveTemp(InModifyDelegates.PageEdgeDelegates);
+		return *this;
+	}
+
 	FDocumentModifyDelegates::FDocumentModifyDelegates(const FMetasoundFrontendDocument& Document)
 PRAGMA_DISABLE_DEPRECATION_WARNINGS
 		: NodeDelegates()
@@ -33,6 +75,24 @@ PRAGMA_ENABLE_DEPRECATION_WARNINGS
 		PageDelegates.OnPageAdded.Broadcast(FDocumentMutatePageArgs{ InPageID });
 	}
 
+	FNodeModifyDelegates& FDocumentModifyDelegates::FindNodeDelegatesChecked(const FGuid& InPageID)
+	{
+PRAGMA_DISABLE_DEPRECATION_WARNINGS
+		return InPageID == DefaultGraphPageID
+		? NodeDelegates
+		: PageNodeDelegates.FindChecked(InPageID);
+PRAGMA_ENABLE_DEPRECATION_WARNINGS
+	}
+
+	FEdgeModifyDelegates& FDocumentModifyDelegates::FindEdgeDelegatesChecked(const FGuid& InPageID)
+	{
+PRAGMA_DISABLE_DEPRECATION_WARNINGS
+		return InPageID == DefaultGraphPageID
+		? EdgeDelegates
+		: PageEdgeDelegates.FindChecked(InPageID);
+PRAGMA_ENABLE_DEPRECATION_WARNINGS
+	}
+
 	void FDocumentModifyDelegates::RemovePageDelegates(const FGuid& InPageID)
 	{
 		PageDelegates.OnRemovingPage.Broadcast(FDocumentMutatePageArgs { InPageID });
diff --git a/Engine/Plugins/Runtime/Metasound/Source/MetasoundFrontend/Public/MetasoundFrontendDocumentModifyDelegates.h b/Engine/Plugins/Runtime/Metasound/Source/MetasoundFrontend/Public/MetasoundFrontendDocumentModifyDelegates.h
--- a/Engine/Plugins/Runtime/Metasound/Source/MetasoundFrontend/Public/MetasoundFrontendDocumentModifyDelegates.h
+++ b/Engine/Plugins/Runtime/Metasound/Source/MetasoundFrontend/Public/MetasoundFrontendDocumentModifyDelegates.h
@@ -61,7 +61,12 @@ namespace Metasound::Frontend
 	struct METASOUNDFRONTEND_API FDocumentModifyDelegates : TSharedFromThis<FDocumentModifyDelegates>
 	{
 		FDocumentModifyDelegates();
+		FDocumentModifyDelegates(const FDocumentModifyDelegates& InModifyDelegates);
+		FDocumentModifyDelegates(FDocumentModifyDelegates&&);
 		FDocumentModifyDelegates(const FMetasoundFrontendDocument& Document);
+		FDocumentModifyDelegates& operator=(const FDocumentModifyDelegates&);
+		FDocumentModifyDelegates& operator=(FDocumentModifyDelegates&&);
+
 
 		FOnMetaSoundFrontendDocumentMutateArray OnDependencyAdded;
 		FOnMetaSoundFrontendDocumentRemoveSwappingArray OnRemoveSwappingDependency;
@@ -84,23 +89,8 @@ namespace Metasound::Frontend
 		TSortedMap<FGuid, FEdgeModifyDelegates> PageEdgeDelegates;
 
 	public:
-		FNodeModifyDelegates& FindNodeDelegatesChecked(const FGuid& InPageID)
-		{
-PRAGMA_DISABLE_DEPRECATION_WARNINGS
-			return InPageID == DefaultGraphPageID
-				? NodeDelegates
-				: PageNodeDelegates.FindChecked(InPageID);
-PRAGMA_ENABLE_DEPRECATION_WARNINGS
-		}
-
-		FEdgeModifyDelegates& FindEdgeDelegatesChecked(const FGuid& InPageID)
-		{
-PRAGMA_DISABLE_DEPRECATION_WARNINGS
-			return InPageID == DefaultGraphPageID
-				? EdgeDelegates
-				: PageEdgeDelegates.FindChecked(InPageID);
-PRAGMA_ENABLE_DEPRECATION_WARNINGS
-		}
+		FNodeModifyDelegates& FindNodeDelegatesChecked(const FGuid& InPageID);
+		FEdgeModifyDelegates& FindEdgeDelegatesChecked(const FGuid& InPageID);
 
 		void IterateGraphEdgeDelegates(TFunctionRef<void(FEdgeModifyDelegates&)> Func)
 		{
