Initial MetaSound Editor Page Support (cvar disabled until all transactions are supported and tested)
#rnx
#tests PIE, MetaSound Asset Editor, Engine Automation tests

diff --git a/Engine/Plugins/Runtime/Metasound/Source/MetasoundEditor/Private/MetasoundDetailCustomization.cpp b/Engine/Plugins/Runtime/Metasound/Source/MetasoundEditor/Private/MetasoundDetailCustomization.cpp
--- a/Engine/Plugins/Runtime/Metasound/Source/MetasoundEditor/Private/MetasoundDetailCustomization.cpp
+++ b/Engine/Plugins/Runtime/Metasound/Source/MetasoundEditor/Private/MetasoundDetailCustomization.cpp
@@ -21,6 +21,7 @@
 #include "MetasoundFrontend.h"
 #include "MetasoundFrontendController.h"
 #include "MetasoundFrontendSearchEngine.h"
+#include "MetasoundSettings.h"
 #include "MetasoundSource.h"
 #include "MetasoundUObjectRegistry.h"
 #include "PropertyCustomizationHelpers.h"
@@ -98,19 +99,19 @@ namespace Metasound::Editor
 	{
 	}
 
-	FName FMetasoundDetailCustomization::GetInterfaceVersionsPath() const
+	FName FMetasoundDetailCustomization::GetInterfaceVersionsPropertyPath() const
 	{
 		return BuildChildPath(DocumentPropertyName, GET_MEMBER_NAME_CHECKED(FMetasoundFrontendDocument, Interfaces));
 	}
 
-	FName FMetasoundDetailCustomization::GetMetadataRootClassPath() const
+	FName FMetasoundDetailCustomization::GetRootClassPropertyPath() const
 	{
 		return BuildChildPath(DocumentPropertyName, GET_MEMBER_NAME_CHECKED(FMetasoundFrontendDocument, RootGraph));
 	}
 
 	FName FMetasoundDetailCustomization::GetMetadataPropertyPath() const
 	{
-		const FName RootClass = FName(GetMetadataRootClassPath());
+		const FName RootClass = FName(GetRootClassPropertyPath());
 		return BuildChildPath(RootClass, GET_MEMBER_NAME_CHECKED(FMetasoundFrontendClass, Metadata));
 	}
 
@@ -151,19 +152,19 @@ namespace Metasound::Editor
 				IDetailCategoryBuilder& GeneralCategoryBuilder = DetailLayout.EditCategory("MetaSound");
 				const FName AuthorPropertyPath = BuildChildPath(GetMetadataPropertyPath(), FMetasoundFrontendClassMetadata::GetAuthorPropertyName());
 				const FName CategoryHierarchyPropertyPath = BuildChildPath(GetMetadataPropertyPath(), FMetasoundFrontendClassMetadata::GetCategoryHierarchyPropertyName());
+				const FName ClassNamePropertyPath = BuildChildPath(GetMetadataPropertyPath(), FMetasoundFrontendClassMetadata::GetClassNamePropertyName());
 				const FName DescPropertyPath = BuildChildPath(GetMetadataPropertyPath(), FMetasoundFrontendClassMetadata::GetDescriptionPropertyName());
 				const FName DisplayNamePropertyPath = BuildChildPath(GetMetadataPropertyPath(), FMetasoundFrontendClassMetadata::GetDisplayNamePropertyName());
 				const FName KeywordsPropertyPath = BuildChildPath(GetMetadataPropertyPath(), FMetasoundFrontendClassMetadata::GetKeywordsPropertyName());
 				const FName IsDeprecatedPropertyPath = BuildChildPath(GetMetadataPropertyPath(), FMetasoundFrontendClassMetadata::GetIsDeprecatedPropertyName());
+				const FName VersionPropertyPath = BuildChildPath(GetMetadataPropertyPath(), FMetasoundFrontendClassMetadata::GetVersionPropertyName());
 
-				const FName ClassNamePropertyPath = BuildChildPath(GetMetadataPropertyPath(), FMetasoundFrontendClassMetadata::GetClassNamePropertyName());
 				const FName ClassNameNamePropertyPath = BuildChildPath(ClassNamePropertyPath, GET_MEMBER_NAME_CHECKED(FMetasoundFrontendClassName, Name));
 
-				const FName VersionPropertyPath = BuildChildPath(GetMetadataPropertyPath(), FMetasoundFrontendClassMetadata::GetVersionPropertyName());
 				const FName MajorVersionPropertyPath = BuildChildPath(VersionPropertyPath, GET_MEMBER_NAME_CHECKED(FMetasoundFrontendVersionNumber, Major));
 				const FName MinorVersionPropertyPath = BuildChildPath(VersionPropertyPath, GET_MEMBER_NAME_CHECKED(FMetasoundFrontendVersionNumber, Minor));
 
-				const FName InterfaceVersionsPropertyPath = GetInterfaceVersionsPath();
+				const FName InterfaceVersionsPropertyPath = GetInterfaceVersionsPropertyPath();
 
 				TSharedPtr<IPropertyHandle> AuthorHandle = DetailLayout.GetProperty(AuthorPropertyPath);
 				TSharedPtr<IPropertyHandle> CategoryHierarchyHandle = DetailLayout.GetProperty(CategoryHierarchyPropertyPath);
@@ -338,6 +339,330 @@ namespace Metasound::Editor
 		DetailLayout.HideCategory("Waveform Processing");
 	}
 
+	FMetasoundPagesDetailCustomization::FMetasoundPagesDetailCustomization()
+		: ItemName(TEXT("Page"))
+	{
+	}
+
+	void FMetasoundPagesDetailCustomization::FPageListener::OnBuilderReloaded(Frontend::FDocumentModifyDelegates& OutDelegates)
+	{
+		OutDelegates.PageDelegates.OnPageAdded.AddSP(this, &FPageListener::OnPageAdded);
+		OutDelegates.PageDelegates.OnRemovingPage.AddSP(this, &FPageListener::OnRemovingPage);
+	}
+
+	void FMetasoundPagesDetailCustomization::FPageListener::OnPageAdded(const Frontend::FDocumentMutatePageArgs& Args)
+	{
+		if (TSharedPtr<FMetasoundPagesDetailCustomization> ParentPtr = Parent.Pin())
+		{
+			const UMetaSoundSettings* Settings = GetDefault<UMetaSoundSettings>();
+			check(Settings);
+			if (const FMetaSoundPageSettings* PageSettings = Settings->FindPageSettings(Args.PageID))
+			{
+				if (PageSettings->Name != ParentPtr->BuildPageName)
+				{
+					ParentPtr->BuildPageName = PageSettings->Name;
+					FGraphBuilder::RegisterGraphWithFrontend(ParentPtr->GetMetaSound());
+				}
+
+				ParentPtr->AddableItems.RemoveAll([&PageSettings](const TSharedPtr<FString>& Item) { return Item->Compare(PageSettings->Name.ToString()) == 0; });
+				ParentPtr->ImplementedNames.Add(PageSettings->Name);
+				ParentPtr->ComboBox->RefreshOptions();
+				ParentPtr->RebuildImplemented();
+			}
+		}
+	}
+
+	void FMetasoundPagesDetailCustomization::FPageListener::OnRemovingPage(const Frontend::FDocumentMutatePageArgs& Args)
+	{
+		if (TSharedPtr<FMetasoundPagesDetailCustomization> ParentPtr = Parent.Pin())
+		{
+			const UMetaSoundSettings* Settings = GetDefault<UMetaSoundSettings>();
+			check(Settings);
+			if (const FMetaSoundPageSettings* PageSettings = Settings->FindPageSettings(Args.PageID))
+			{
+				if (PageSettings->Name != ParentPtr->BuildPageName)
+				{
+					ParentPtr->BuildPageName = PageSettings->Name;
+					FGraphBuilder::RegisterGraphWithFrontend(ParentPtr->GetMetaSound());
+				}
+
+				ParentPtr->AddableItems.Add(MakeShared<FString>(PageSettings->Name.ToString()));
+				ParentPtr->ImplementedNames.Remove(PageSettings->Name);
+				ParentPtr->ComboBox->RefreshOptions();
+				ParentPtr->RebuildImplemented();
+			}
+		}
+	}
+
+	void FMetasoundPagesDetailCustomization::CustomizeDetails(IDetailLayoutBuilder& DetailLayout)
+	{
+		using namespace Engine;
+		using namespace Frontend;
+
+		TArray<TWeakObjectPtr<UObject>> Objects;
+		DetailLayout.GetObjectsBeingCustomized(Objects);
+
+		// Only support modifying a single MetaSound at a time (Multiple
+		// MetaSound editing will be covered most likely by separate tool).
+		if (Objects.Num() > 1)
+		{
+			return;
+		}
+		if (UMetasoundEditorViewBase* View = CastChecked<UMetasoundEditorViewBase>(Objects.Last()))
+		{
+			if (UObject* MetaSound = View->GetMetasound())
+			{
+				InitBuilder(*MetaSound);
+				PageListener = MakeShared<FPageListener>(StaticCastSharedRef<FMetasoundPagesDetailCustomization>(AsShared()));
+				Builder->AddTransactionListener(PageListener->AsShared());
+			}
+		}
+
+		if (UMetaSoundSettings* Settings = GetMutableDefault<UMetaSoundSettings>())
+		{
+			Settings->OnPageSettingsUpdated.AddSPLambda(this, [this]()
+			{
+				UpdateItemNames();
+				if (ComboBox.IsValid())
+				{
+					ComboBox->RefreshOptions();
+				}
+			});
+		}
+
+		UpdateItemNames();
+
+		SAssignNew(ComboBox, SSearchableComboBox)
+			.OptionsSource(&AddableItems)
+			.OnGenerateWidget_Lambda([](TSharedPtr<FString> InItem)
+			{
+				return SNew(STextBlock).Text(FText::FromString(*InItem));
+			})
+			.OnSelectionChanged_Lambda([this](TSharedPtr<FString> NameToAdd, ESelectInfo::Type InSelectInfo)
+			{
+				using namespace Engine;
+				using namespace Frontend;
+
+				if (InSelectInfo != ESelectInfo::OnNavigation)
+				{
+					UObject& MetaSound = GetMetaSound();
+
+					const FScopedTransaction Transaction(FText::Format(LOCTEXT("AddPageTransactionFormat", "Add MetaSound Page '{0}'"), FText::FromString(*NameToAdd)));
+					MetaSound.Modify();
+
+					// Underlying DocBuilder's pageID is a property that is tracked by transaction stack, so signal as modifying behavior
+					Builder->Modify();
+
+					constexpr bool bDuplicateLastGraph = true;
+					constexpr bool bSetAsBuildGraph = true;
+
+					EMetaSoundBuilderResult Result = EMetaSoundBuilderResult::Failed;
+					Builder->AddGraphPage(FName(*NameToAdd.Get()), bDuplicateLastGraph, bSetAsBuildGraph, Result);
+				}
+			})
+			.Content()
+			[
+				SNew(STextBlock)
+				.Text(LOCTEXT("UpdatePageAction", "Add Page..."))
+				.IsEnabled(TAttribute<bool>::Create(TAttribute<bool>::FGetter::CreateSP(this, &FMetaSoundDetailCustomizationBase::IsGraphEditable)))
+			];
+
+		TSharedRef<SWidget> Utilities = SNew(SHorizontalBox)
+			+ SHorizontalBox::Slot()
+			.Padding(2.0f)
+			.HAlign(HAlign_Left)
+			.VAlign(VAlign_Center)
+			.AutoWidth()
+			[
+				ComboBox->AsShared()
+			]
+			+ SHorizontalBox::Slot()
+			.Padding(2.0f)
+			.HAlign(HAlign_Left)
+			.VAlign(VAlign_Center)
+			.AutoWidth()
+			[
+				PropertyCustomizationHelpers::MakeDeleteButton(FSimpleDelegate::CreateLambda([this]()
+				{
+					using namespace Frontend;
+					UObject& MetaSound = GetMetaSound();
+
+					const FScopedTransaction Transaction(LOCTEXT("RemoveAllPagesTransaction", "Remove All MetaSound Pages"));
+					MetaSound.Modify();
+
+					// Underlying DocBuilder's pageID is a property that is tracked by transaction stack, so signal as modifying behavior
+					Builder->Modify();
+					Builder->RemoveAllGraphPages();
+
+					UpdateItemNames();
+					ComboBox->RefreshOptions();
+					FGraphBuilder::RegisterGraphWithFrontend(MetaSound);
+				}), LOCTEXT("RemoveImplementablePagesTooltip1", "Removes all pages from the given MetaSound."))
+			];
+
+		Utilities->SetEnabled(TAttribute<bool>::Create(TAttribute<bool>::FGetter::CreateSP(this, &FMetaSoundDetailCustomizationBase::IsGraphEditable)));
+
+		const FText HeaderName = LOCTEXT("PagesGroupDisplayName", "Pages");
+		IDetailCategoryBuilder& Category = DetailLayout.EditCategory(FName(ItemName), HeaderName);
+		Category.AddCustomRow(HeaderName) [ Utilities ];
+
+		const UMetaSoundSettings* Settings = GetDefault<UMetaSoundSettings>();
+		check(Settings);
+		if (const FMetaSoundPageSettings* PageSettings = Settings->FindPageSettings(Builder->GetBuilder().GetBuildPageID()))
+		{
+			BuildPageName = PageSettings->Name;
+		}
+		else
+		{
+			BuildPageName = Metasound::Frontend::DefaultGraphPageName;
+		}
+
+		EntryWidgets = SNew(SVerticalBox);
+		RebuildImplemented();
+		Category.AddCustomRow(LOCTEXT("ImplementedPagesLabel", "Implemented Pages"))
+		[
+			EntryWidgets->AsShared()
+		];
+	}
+
+	UObject& FMetasoundPagesDetailCustomization::GetMetaSound() const
+	{
+		return Builder->GetBuilder().CastDocumentObjectChecked<UObject>();
+	}
+
+	void FMetasoundPagesDetailCustomization::RebuildImplemented()
+	{
+		EntryWidgets->ClearChildren();
+
+		auto CreateEntryWidget = [this](bool bIsDefault, FName InName) -> TSharedRef<SWidget>
+		{
+			using namespace Frontend;
+
+			TSharedRef<SWidget> SelectButtonWidget = PropertyCustomizationHelpers::MakeUseSelectedButton(FSimpleDelegate::CreateLambda([this, InName]()
+			{
+				using namespace Frontend;
+				const UMetaSoundSettings* Settings = GetDefault<UMetaSoundSettings>();
+				check(Settings);
+				if (const FMetaSoundPageSettings* PageSettings = Settings->FindPageSettings(InName))
+				{
+					if (Builder->GetBuilder().SetBuildPageID(PageSettings->UniqueId))
+					{
+						BuildPageName = InName;
+						FGraphBuilder::RegisterGraphWithFrontend(GetMetaSound());
+					}
+				}
+			}),
+			LOCTEXT("SetPageTooltip", "Sets the actively displayed graph page of the MetaSound."),
+			TAttribute<bool>::Create([this, InName]()
+			{
+				return BuildPageName != InName;
+			}));
+
+			TSharedRef<SHorizontalBox> EntryWidget = SNew(SHorizontalBox)
+				+ SHorizontalBox::Slot()
+				.Padding(2.0f)
+				.HAlign(HAlign_Left)
+				.VAlign(VAlign_Center)
+				.AutoWidth()
+				[
+					SNew(STextBlock)
+					.Text(FText::FromName(InName))
+				];
+
+			if (!bIsDefault)
+			{
+				TSharedRef<SWidget> RemoveButtonWidget = PropertyCustomizationHelpers::MakeDeleteButton(FSimpleDelegate::CreateLambda([this, InName]()
+				{
+					using namespace Frontend;
+					const FScopedTransaction Transaction(FText::Format(LOCTEXT("RemovePageTransactionFormat", "Remove MetaSound Page '{0}'"), FText::FromName(InName)));
+					UObject& MetaSound = GetMetaSound();
+					MetaSound.Modify();
+
+					// Removal may modify the builder's build page ID if it is the currently set value
+					Builder->Modify();
+
+					EMetaSoundBuilderResult Result = EMetaSoundBuilderResult::Failed;
+					Builder->RemoveGraphPage(InName, Result);
+					if (Result == EMetaSoundBuilderResult::Succeeded)
+					{
+						UpdateItemNames();
+						ComboBox->RefreshOptions();
+						if (InName == BuildPageName)
+						{
+							FGraphBuilder::RegisterGraphWithFrontend(MetaSound);
+						}
+					}
+				}), LOCTEXT("RemovePageTooltip2", "Removes the associated page from the MetaSound."));
+				EntryWidget->AddSlot()
+					.Padding(2.0f)
+					.HAlign(HAlign_Center)
+					.VAlign(VAlign_Center)
+					.AutoWidth()[RemoveButtonWidget];
+			}
+
+			EntryWidget->AddSlot()
+				.Padding(2.0f)
+				.HAlign(HAlign_Center)
+				.VAlign(VAlign_Center)
+				.AutoWidth() [ SelectButtonWidget ];
+
+			EntryWidget->SetEnabled(TAttribute<bool>::Create(TAttribute<bool>::FGetter::CreateSP(this, &FMetaSoundDetailCustomizationBase::IsGraphEditable)));
+			return EntryWidget;
+		};
+
+		const UMetaSoundSettings* Settings = GetDefault<UMetaSoundSettings>();
+		check(Settings);
+		for (const FMetaSoundPageSettings& PageSettings : Settings->GetPageSettings())
+		{
+			if (ImplementedNames.Contains(PageSettings.Name))
+			{
+				EntryWidgets->AddSlot()
+				.HAlign(HAlign_Left)
+				.VAlign(VAlign_Center)
+				.AutoHeight()
+				[
+					CreateEntryWidget(!PageSettings.UniqueId.IsValid(), PageSettings.Name)
+				];
+			}
+		}
+	}
+
+	void FMetasoundPagesDetailCustomization::UpdateItemNames()
+	{
+		using namespace Frontend;
+
+		AddableItems.Reset();
+		ImplementedNames.Reset();
+
+		const UMetaSoundSettings* Settings = GetDefault<UMetaSoundSettings>();
+		check(Settings);
+
+		TSet<FGuid> ImplementedGuids { FGuid() }; // Default "no guid" is always implemented (base graph for back compat & cook safety to ensure there's always at least one)
+
+		const FMetasoundFrontendDocument& Document = Builder->GetBuilder().GetConstDocumentChecked();
+		Document.RootGraph.IterateGraphPages([&ImplementedGuids](const FMetasoundFrontendGraph& Graph)
+		{
+			ImplementedGuids.Add(Graph.PageID);
+		});
+
+		auto GetStringName = [](const FMetaSoundPageSettings& Page) -> TSharedPtr<FString> { return MakeShared<FString>(Page.Name.ToString()); };
+		auto RemoveImplementedItem = [&ImplementedGuids](const FMetaSoundPageSettings& Page) { return !ImplementedGuids.Contains(Page.UniqueId); };
+		Algo::TransformIf(Settings->GetPageSettings(), AddableItems, RemoveImplementedItem, GetStringName);
+
+		auto GetPageName = [&Settings](const FGuid& PageID)
+		{
+			const FMetaSoundPageSettings* Page = Settings->FindPageSettings(PageID);
+			if (ensure(Page))
+			{
+				return Page->Name;
+			}
+
+			return FName();
+		};
+
+		Algo::Transform(ImplementedGuids, ImplementedNames, GetPageName);
+	}
+
 	void FMetasoundInterfacesDetailCustomization::CustomizeDetails(IDetailLayoutBuilder& DetailLayout)
 	{
 		TArray<TWeakObjectPtr<UObject>> Objects;
@@ -351,7 +676,7 @@ namespace Metasound::Editor
 		}
 		if (UMetasoundInterfacesView* InterfacesView = CastChecked<UMetasoundInterfacesView>(Objects.Last()))
 		{
-			if (UObject* MetaSound = InterfacesView->GetMetasound().Get())
+			if (UObject* MetaSound = InterfacesView->GetMetasound())
 			{
 				InitBuilder(*MetaSound);
 			}
diff --git a/Engine/Plugins/Runtime/Metasound/Source/MetasoundEditor/Private/MetasoundDetailCustomization.h b/Engine/Plugins/Runtime/Metasound/Source/MetasoundEditor/Private/MetasoundDetailCustomization.h
--- a/Engine/Plugins/Runtime/Metasound/Source/MetasoundEditor/Private/MetasoundDetailCustomization.h
+++ b/Engine/Plugins/Runtime/Metasound/Source/MetasoundEditor/Private/MetasoundDetailCustomization.h
@@ -2,6 +2,7 @@
 #pragma once
 
 #include "IDetailCustomization.h"
+#include "MetasoundFrontendDocumentModifyDelegates.h"
 #include "SGraphActionMenu.h"
 #include "SSearchableComboBox.h"
 #include "Templates/SharedPointer.h"
@@ -44,13 +45,56 @@ namespace Metasound
 			// End of IDetailCustomization interface
 
 		private:
-			FName GetInterfaceVersionsPath() const;
-			FName GetMetadataRootClassPath() const;
+			FName GetInterfaceVersionsPropertyPath() const;
+			FName GetRootClassPropertyPath() const;
 			FName GetMetadataPropertyPath() const;
 
 			FName DocumentPropertyName;
 		};
 
+		class FMetasoundPagesDetailCustomization : public FMetaSoundDetailCustomizationBase
+		{
+		public:
+			FMetasoundPagesDetailCustomization();
+
+			// IDetailCustomization interface
+			virtual void CustomizeDetails(IDetailLayoutBuilder& DetailLayout) override;
+			// End of IDetailCustomization interface
+
+		private:
+			UObject& GetMetaSound() const;
+			void RebuildImplemented();
+			void UpdateItemNames();
+
+			TArray<TSharedPtr<FString>> AddableItems;
+
+			TSet<FName> ImplementedNames;
+			TSharedPtr<SSearchableComboBox> ComboBox;
+			TSharedPtr<SVerticalBox> EntryWidgets;
+			FString ItemName;
+			FName BuildPageName;
+
+			class FPageListener : public Frontend::IDocumentBuilderTransactionListener
+			{
+				TWeakPtr<FMetasoundPagesDetailCustomization> Parent;
+
+			public:
+				FPageListener() = default;
+				FPageListener(TSharedRef<FMetasoundPagesDetailCustomization> InParent)
+					: Parent(InParent)
+				{
+				}
+
+				virtual ~FPageListener() = default;
+
+			private:
+				virtual void OnBuilderReloaded(Frontend::FDocumentModifyDelegates& OutDelegates) override;
+				void OnPageAdded(const Frontend::FDocumentMutatePageArgs& Args);
+				void OnRemovingPage(const Frontend::FDocumentMutatePageArgs& Args);
+			};
+			TSharedPtr<FPageListener> PageListener;
+		};
+
 		class FMetasoundInterfacesDetailCustomization : public FMetaSoundDetailCustomizationBase
 		{
 		public:
diff --git a/Engine/Plugins/Runtime/Metasound/Source/MetasoundEditor/Private/MetasoundEditor.cpp b/Engine/Plugins/Runtime/Metasound/Source/MetasoundEditor/Private/MetasoundEditor.cpp
--- a/Engine/Plugins/Runtime/Metasound/Source/MetasoundEditor/Private/MetasoundEditor.cpp
+++ b/Engine/Plugins/Runtime/Metasound/Source/MetasoundEditor/Private/MetasoundEditor.cpp
@@ -99,6 +99,30 @@ namespace Metasound
 {
 	namespace Editor
 	{
+		namespace AssetEditorPrivate
+		{
+			int32 EnablePageEditor = 0;
+
+			FAutoConsoleVariableRef CVarMetaSoundEnablePageEditor(
+				TEXT("au.MetaSound.Experimental.EnablePagesEditor"),
+				EnablePageEditor,
+				TEXT("Enables Page Editor in MetaSound Asset Editor.\n")
+				TEXT("Default: 0"),
+				ECVF_Default);
+		} // namespace AssetEditorPrivate
+
+		namespace TabNamesPrivate
+		{
+			const FName Analyzers = "MetasoundEditor_Analyzers";
+			const FName Details = "MetasoundEditor_Details";
+			const FName GraphCanvas = "MetasoundEditor_GraphCanvas";
+			const FName Members = "MetasoundEditor_Members";
+			const FName Palette = "MetasoundEditor_Palette";
+			const FName Interfaces = "MetasoundEditor_Interfaces";
+			const FName Pages = "MetasoundEditor_Pages";
+			const FName Find = "MetasoundEditor_Find";
+		} // namespace TabNamesPrivate
+
 		static const TArray<FText> NodeSectionNames
 		{
 			LOCTEXT("NodeSectionName_Invalid", "INVALID"),
@@ -819,9 +843,15 @@ namespace Metasound
 
 			FAssetEditorToolkit::RegisterTabSpawners(InTabManager);
 
-			InTabManager->RegisterTabSpawner(TabFactory::Names::GraphCanvas, FOnSpawnTab::CreateLambda([InPlayTimeWidget = PlayTimeWidget, InRenderStatsWidget = RenderStatsWidget, InMetasoundGraphEditor = MetasoundGraphEditor](const FSpawnTabArgs& Args)
+			InTabManager->RegisterTabSpawner(TabNamesPrivate::GraphCanvas, FOnSpawnTab::CreateLambda(
+				[
+					InPlayTimeWidget = PlayTimeWidget,
+					InMetasoundGraphEditor = MetasoundGraphEditor,
+					InRenderStatsWidget = RenderStatsWidget
+				](const FSpawnTabArgs& Args)
 			{
-				return TabFactory::CreateGraphCanvasTab(SNew(SOverlay)
+				TSharedRef<SDockTab> SpawnedTab = SNew(SDockTab).Label(LOCTEXT("MetasoundGraphCanvasTitle", "MetaSound Graph"));
+				SpawnedTab->SetContent(SNew(SOverlay)
 					+ SOverlay::Slot()
 					[
 						InMetasoundGraphEditor.ToSharedRef()
@@ -836,47 +866,70 @@ namespace Metasound
 						InRenderStatsWidget.ToSharedRef()
 					]
 					.Padding(5.0f, 5.0f)
-				, Args);
+				);
+				return SpawnedTab;
 			}))
 			.SetDisplayName(LOCTEXT("GraphCanvasTab", "Viewport"))
 			.SetGroup(WorkspaceMenuCategoryRef)
 			.SetIcon(FSlateIcon(FAppStyle::GetAppStyleSetName(), "GraphEditor.EventGraph_16x"));
 
-			InTabManager->RegisterTabSpawner(TabFactory::Names::Details, FOnSpawnTab::CreateLambda([InMetasoundDetails = MetasoundDetails](const FSpawnTabArgs& Args)
+			InTabManager->RegisterTabSpawner(TabNamesPrivate::Details, FOnSpawnTab::CreateLambda([InMetasoundDetails = MetasoundDetails](const FSpawnTabArgs& Args)
 			{
-				return TabFactory::CreateDetailsTab(InMetasoundDetails, Args);
+				return SNew(SDockTab).Label(LOCTEXT("MetaSoundDetailsTitle", "Details"))[ InMetasoundDetails.ToSharedRef() ];
 			}))
 			.SetDisplayName(LOCTEXT("DetailsTab", "Details"))
 			.SetGroup(WorkspaceMenuCategoryRef)
 			.SetIcon(FSlateIcon(FAppStyle::GetAppStyleSetName(), "LevelEditor.Tabs.Details"));
 
-			InTabManager->RegisterTabSpawner(TabFactory::Names::Members, FOnSpawnTab::CreateLambda([InGraphMembersMenu = GraphMembersMenu](const FSpawnTabArgs& Args)
+			InTabManager->RegisterTabSpawner(TabNamesPrivate::Members, FOnSpawnTab::CreateLambda([InGraphMembersMenu = GraphMembersMenu](const FSpawnTabArgs& Args)
 			{
-				return TabFactory::CreateMembersTab(InGraphMembersMenu, Args);
+				TSharedRef<SDockTab> NewTab = SNew(SDockTab)
+				.Label(LOCTEXT("GraphMembersMenulTitle", "Members"))
+				[
+					InGraphMembersMenu.ToSharedRef()
+				];
+
+				if (const ISlateStyle* MetasoundStyle = FSlateStyleRegistry::FindSlateStyle("MetaSoundStyle"))
+				{
+					NewTab->SetTabIcon(MetasoundStyle->GetBrush("MetasoundEditor.Metasound.Icon"));
+				}
+
+				return NewTab;
 			}))
 			.SetDisplayName(LOCTEXT("MembersTab", "Members"))
 			.SetGroup(WorkspaceMenuCategoryRef)
 			.SetIcon(FSlateIcon("MetaSoundStyle", "MetasoundEditor.Metasound.Icon"));
 
-			InTabManager->RegisterTabSpawner(TabFactory::Names::Analyzers, FOnSpawnTab::CreateLambda([InAnalyzerWidget = BuildAnalyzerWidget()](const FSpawnTabArgs& Args)
+			InTabManager->RegisterTabSpawner(TabNamesPrivate::Analyzers, FOnSpawnTab::CreateLambda([InAnalyzerWidget = BuildAnalyzerWidget()](const FSpawnTabArgs&)
 			{
-				return TabFactory::CreateAnalyzersTab(InAnalyzerWidget, Args);
+				return SNew(SDockTab).Label(LOCTEXT("MetasoundAnalyzersTitle", "Analyzers")) [ InAnalyzerWidget.ToSharedRef() ];
 			}))
 			.SetDisplayName(LOCTEXT("AnalyzersTab", "Analyzers"))
 			.SetGroup(WorkspaceMenuCategoryRef)
 			.SetIcon(FSlateIcon(FAppStyle::GetAppStyleSetName(), "Kismet.Tabs.Palette"));
 
-			InTabManager->RegisterTabSpawner(TabFactory::Names::Interfaces, FOnSpawnTab::CreateLambda([InInterfacesDetails = InterfacesDetails](const FSpawnTabArgs& Args)
+			if (AssetEditorPrivate::EnablePageEditor)
 			{
-				return TabFactory::CreateInterfacesTab(InInterfacesDetails, Args);
+				InTabManager->RegisterTabSpawner(TabNamesPrivate::Pages, FOnSpawnTab::CreateLambda([InPagesDetails = PagesDetails](const FSpawnTabArgs&)
+				{
+					return SNew(SDockTab).Label(LOCTEXT("MetasoundPagesDetailsTitle", "Pages"))[InPagesDetails.ToSharedRef()];
+				}))
+				.SetDisplayName(LOCTEXT("PagesTab", "Pages"))
+					.SetGroup(WorkspaceMenuCategoryRef)
+					.SetIcon(FSlateIcon(FAppStyle::GetAppStyleSetName(), "Kismet.Tabs.Palette"));
+			}
+
+			InTabManager->RegisterTabSpawner(TabNamesPrivate::Interfaces, FOnSpawnTab::CreateLambda([InInterfacesDetails = InterfacesDetails](const FSpawnTabArgs&)
+			{
+				return SNew(SDockTab).Label(LOCTEXT("MetasoundInterfacesDetailsTitle", "Interfaces")) [ InInterfacesDetails.ToSharedRef() ];
 			}))
 			.SetDisplayName(LOCTEXT("InterfacesTab", "Interfaces"))
 			.SetGroup(WorkspaceMenuCategoryRef)
 			.SetIcon(FSlateIcon(FAppStyle::GetAppStyleSetName(), "ClassIcon.Interface"));
 
-			InTabManager->RegisterTabSpawner(TabFactory::Names::Find, FOnSpawnTab::CreateLambda([InFindWidget = FindWidget](const FSpawnTabArgs& Args)
+			InTabManager->RegisterTabSpawner(TabNamesPrivate::Find, FOnSpawnTab::CreateLambda([InFindWidget = FindWidget](const FSpawnTabArgs&)
 			{
-				return TabFactory::CreateFindTab(InFindWidget, Args);
+				return SNew(SDockTab).Label(LOCTEXT("MetasoundFindTitle", "Find Results")) [ InFindWidget.ToSharedRef() ];
 			}))
 			.SetDisplayName(LOCTEXT("FindTab", "Find in MetaSound"))
 			.SetGroup(WorkspaceMenuCategoryRef)
@@ -889,12 +942,13 @@ namespace Metasound
 
 			FAssetEditorToolkit::UnregisterTabSpawners(InTabManager);
 
-			InTabManager->UnregisterTabSpawner(TabFactory::Names::Analyzers);
-			InTabManager->UnregisterTabSpawner(TabFactory::Names::GraphCanvas);
-			InTabManager->UnregisterTabSpawner(TabFactory::Names::Details);
-			InTabManager->UnregisterTabSpawner(TabFactory::Names::Members);
-			InTabManager->UnregisterTabSpawner(TabFactory::Names::Interfaces);
-			InTabManager->UnregisterTabSpawner(TabFactory::Names::Find);
+			InTabManager->UnregisterTabSpawner(TabNamesPrivate::Analyzers);
+			InTabManager->UnregisterTabSpawner(TabNamesPrivate::GraphCanvas);
+			InTabManager->UnregisterTabSpawner(TabNamesPrivate::Details);
+			InTabManager->UnregisterTabSpawner(TabNamesPrivate::Members);
+			InTabManager->UnregisterTabSpawner(TabNamesPrivate::Pages);
+			InTabManager->UnregisterTabSpawner(TabNamesPrivate::Interfaces);
+			InTabManager->UnregisterTabSpawner(TabNamesPrivate::Find);
 		}
 
 		TSharedPtr<SWidget> FEditor::BuildAnalyzerWidget() const
@@ -977,6 +1031,7 @@ namespace Metasound
 			}
 
 			GraphConnectionManager.Reset();
+			PagesView.Reset();
 			InterfacesView.Reset();
 			DestroyAnalyzers();
 			check(GEditor);
@@ -1058,21 +1113,21 @@ namespace Metasound
 							FTabManager::NewStack()
 							->SetSizeCoefficient(0.25f)
 							->SetHideTabWell(false)
-							->AddTab(TabFactory::Names::Members, ETabState::OpenedTab)
+							->AddTab(TabNamesPrivate::Members, ETabState::OpenedTab)
 						)
 						->Split
 						(
 							FTabManager::NewStack()
 							->SetSizeCoefficient(0.1f)
 							->SetHideTabWell(true)
-							->AddTab(TabFactory::Names::Interfaces, ETabState::OpenedTab)
+							->AddTab(TabNamesPrivate::Interfaces, ETabState::OpenedTab)
 						)
 						->Split
 						(
 							FTabManager::NewStack()
 							->SetSizeCoefficient(0.50f)
 							->SetHideTabWell(false)
-							->AddTab(TabFactory::Names::Details, ETabState::OpenedTab)
+							->AddTab(TabNamesPrivate::Details, ETabState::OpenedTab)
 						)
 					)
 					->Split
@@ -1085,14 +1140,14 @@ namespace Metasound
 							FTabManager::NewStack()
 							->SetSizeCoefficient(0.8f)
 							->SetHideTabWell(true)
-							->AddTab(TabFactory::Names::GraphCanvas, ETabState::OpenedTab)
+							->AddTab(TabNamesPrivate::GraphCanvas, ETabState::OpenedTab)
 						)
 						->Split
 						(
 							FTabManager::NewStack()
 							->SetSizeCoefficient(0.2f)
 							->SetHideTabWell(true)
-							->AddTab(TabFactory::Names::Find, ETabState::OpenedTab)
+							->AddTab(TabNamesPrivate::Find, ETabState::OpenedTab)
 						)
 					)
 
@@ -1101,7 +1156,7 @@ namespace Metasound
 						FTabManager::NewStack()
 						->SetSizeCoefficient(0.08f)
 						->SetHideTabWell(true)
-						->AddTab(TabFactory::Names::Analyzers, ETabState::OpenedTab)
+						->AddTab(TabNamesPrivate::Analyzers, ETabState::OpenedTab)
 					)
 				)
 			);
@@ -1449,6 +1504,20 @@ namespace Metasound
 				InterfacesDetails->HideFilterArea(true);
 			}
 
+			if (AssetEditorPrivate::EnablePageEditor)
+			{
+				PagesDetails = PropertyModule.CreateDetailView(Args);
+				if (PagesDetails.IsValid())
+				{
+					PagesView = TStrongObjectPtr(NewObject<UMetasoundPagesView>());
+					PagesView->SetMetasound(&MetaSound);
+					const TArray<UObject*> PagesViewObj{ PagesView.Get() };
+
+					PagesDetails->SetObjects(PagesViewObj);
+					PagesDetails->HideFilterArea(true);
+				}
+			}
+
 			Palette = SNew(SMetasoundPalette);
 
 			FindWidget = SNew(SFindInMetasound, SharedThis(this));
@@ -3149,7 +3218,17 @@ namespace Metasound
 			}
 		}
 
-		void FEditor::RefreshInterfaces()
+		void FEditor::RefreshPagesView()
+		{
+			TRACE_CPUPROFILER_EVENT_SCOPE(Metasound::Editor::FEditor::RefreshPages);
+
+			if (PagesDetails.IsValid())
+			{
+				PagesDetails->ForceRefresh();
+			}
+		}
+
+		void FEditor::RefreshInterfaceView()
 		{
 			TRACE_CPUPROFILER_EVENT_SCOPE(Metasound::Editor::FEditor::RefreshInterfaces);
 
@@ -3860,7 +3939,7 @@ namespace Metasound
 
 				if (!InterfacesModified.IsEmpty() || bForceRefreshViews)
 				{
-					RefreshInterfaces();
+					RefreshInterfaceView();
 				}
 
 				HighestMessageSeverity = Graph->GetHighestMessageSeverity();
diff --git a/Engine/Plugins/Runtime/Metasound/Source/MetasoundEditor/Private/MetasoundEditor.h b/Engine/Plugins/Runtime/Metasound/Source/MetasoundEditor/Private/MetasoundEditor.h
--- a/Engine/Plugins/Runtime/Metasound/Source/MetasoundEditor/Private/MetasoundEditor.h
+++ b/Engine/Plugins/Runtime/Metasound/Source/MetasoundEditor/Private/MetasoundEditor.h
@@ -15,6 +15,7 @@
 #include "MetasoundEditorGraphNode.h"
 #include "MetasoundFrontend.h"
 #include "MetasoundFrontendController.h"
+#include "MetasoundFrontendDocumentModifyDelegates.h"
 #include "Meter.h"
 #include "Misc/NotifyHook.h"
 #include "SAudioMeter.h"
@@ -69,20 +70,35 @@ struct FMeterResults;
 struct FPropertyChangedEvent;
 
 
-// Simple class for the interfaces details tab to keep track of its corresponding Metasound 
+// Base implementation for editor wrappers that provide edit customizations
 UCLASS(MinimalAPI)
-class UMetasoundInterfacesView : public UObject
+class UMetasoundEditorViewBase : public UObject
 {
 	GENERATED_BODY()
 public:
-	UMetasoundInterfacesView() = default;
 
 	void SetMetasound(UObject* InMetasound) { Metasound = InMetasound; }
-	const TWeakObjectPtr<UObject> GetMetasound() { return Metasound; }
+	const UObject* GetMetasound() const { return Metasound.Get(); }
+	UObject* GetMetasound() { return Metasound.Get(); }
+
 private:
 	TWeakObjectPtr<UObject> Metasound;
 };
 
+
+// Simple class for the interfaces details tab to keep track of its corresponding Metasound 
+UCLASS(MinimalAPI)
+class UMetasoundInterfacesView : public UMetasoundEditorViewBase
+{
+	GENERATED_BODY()
+};
+
+UCLASS(MinimalAPI)
+class UMetasoundPagesView : public UMetasoundEditorViewBase
+{
+	GENERATED_BODY()
+};
+
 namespace Metasound
 {
 	namespace Editor
@@ -101,6 +117,7 @@ namespace Metasound
 			COUNT
 		};
 
+
 		class FEditor : public IMetasoundEditor, public FNotifyHook, public FEditorUndoClient, public FTickableEditorObject
 		{
 		public:
@@ -175,9 +192,6 @@ namespace Metasound
 			/** Forces all UX pertaining to the root graph's interface to be refreshed, returning the first selected member. */
 			UMetasoundEditorGraphMember* RefreshGraphMemberMenu();
 
-			/** Forces refresh of interfaces list. */
-			void RefreshInterfaces();
-
 			/** Updates selected node classes to highest class found in the MetaSound Class Registry. */
 			void UpdateSelectedNodeClasses();
 
@@ -356,6 +370,12 @@ namespace Metasound
 			void ShowUnconnectedPins();
 
 		private:
+			/** Forces refresh of pages view. */
+			void RefreshPagesView();
+
+			/** Forces refresh of interfaces view. */
+			void RefreshInterfaceView();
+
 			void RemoveInvalidSelection();
 
 			void SetPreviewID(uint32 InPreviewID);
@@ -439,6 +459,10 @@ namespace Metasound
 			/** Details tab */
 			TSharedPtr<IDetailsView> MetasoundDetails;
 
+			/** Pages tab */
+			TSharedPtr<IDetailsView> PagesDetails;
+			TStrongObjectPtr<UMetasoundPagesView> PagesView;
+
 			/** Interfaces tab */
 			TSharedPtr<IDetailsView> InterfacesDetails;
 			TStrongObjectPtr<UMetasoundInterfacesView> InterfacesView;
diff --git a/Engine/Plugins/Runtime/Metasound/Source/MetasoundEditor/Private/MetasoundEditorModule.cpp b/Engine/Plugins/Runtime/Metasound/Source/MetasoundEditor/Private/MetasoundEditorModule.cpp
--- a/Engine/Plugins/Runtime/Metasound/Source/MetasoundEditor/Private/MetasoundEditorModule.cpp
+++ b/Engine/Plugins/Runtime/Metasound/Source/MetasoundEditor/Private/MetasoundEditorModule.cpp
@@ -9,7 +9,7 @@
 #include "EdGraph/EdGraphNode.h"
 #include "EdGraph/EdGraphPin.h"
 #include "EdGraphUtilities.h"
-#include "Styling/AppStyle.h"
+#include "Framework/Notifications/NotificationManager.h"
 #include "HAL/IConsoleManager.h"
 #include "IDetailCustomization.h"
 #include "ISettingsModule.h"
@@ -18,6 +18,7 @@
 #include "MetasoundAudioBuffer.h"
 #include "MetasoundBuilderSubsystem.h"
 #include "MetasoundDetailCustomization.h"
+#include "MetasoundDocumentBuilderRegistry.h"
 #include "MetasoundDocumentInterface.h"
 #include "MetasoundEditorGraph.h"
 #include "MetasoundEditorGraphBuilder.h"
@@ -42,15 +43,17 @@
 #include "PackageMigrationContext.h"
 #include "PropertyEditorDelegates.h"
 #include "PropertyEditorModule.h"
+#include "Styling/AppStyle.h"
 #include "Styling/CoreStyle.h"
-#include "Styling/StyleColors.h"
 #include "Styling/SlateStyle.h"
 #include "Styling/SlateStyleMacros.h"
 #include "Styling/SlateStyleRegistry.h"
 #include "Styling/SlateTypes.h"
+#include "Styling/StyleColors.h"
 #include "Templates/SharedPointer.h"
 #include "UObject/UObjectGlobals.h"
 #include "UObject/UObjectIterator.h"
+#include "Widgets/Notifications/SNotificationList.h"
 
 
 DEFINE_LOG_CATEGORY(LogMetasoundEditor);
@@ -64,8 +67,7 @@ FAutoConsoleVariableRef CVarMetaSoundEditorAsyncRegistrationEnabled(
 	TEXT("0: Disabled, !0: Enabled (default)"),
 	ECVF_Default);
 
-// Forward declarations 
-class UMetasoundInterfacesView;
+#define LOCTEXT_NAMESPACE "MetaSounds"
 
 namespace Metasound
 {
@@ -831,6 +833,10 @@ namespace Metasound
 					UMetasoundInterfacesView::StaticClass()->GetFName(),
 					FOnGetDetailCustomizationInstance::CreateLambda([]() { return MakeShared<FMetasoundInterfacesDetailCustomization>(); }));
 
+				PropertyModule.RegisterCustomClassLayout(
+					UMetasoundPagesView::StaticClass()->GetFName(),
+					FOnGetDetailCustomizationInstance::CreateLambda([]() { return MakeShared<FMetasoundPagesDetailCustomization>(); }));
+
 				PropertyModule.RegisterCustomClassLayout(
 					UMetasoundEditorGraphInput::StaticClass()->GetFName(),
 					FOnGetDetailCustomizationInstance::CreateLambda([]() { return MakeShared<FMetasoundInputDetailCustomization>(); }));
@@ -877,8 +883,8 @@ namespace Metasound
 				ISettingsModule& SettingsModule = FModuleManager::LoadModuleChecked<ISettingsModule>("Settings");
 
 				SettingsModule.RegisterSettings("Editor", "ContentEditors", "MetaSound Editor",
-					NSLOCTEXT("MetaSoundsEditor", "MetaSoundEditorSettingsName", "MetaSound Editor"),
-					NSLOCTEXT("MetaSoundsEditor", "MetaSoundEditorSettingsDescription", "Customize MetaSound Editor."),
+					LOCTEXT("MetaSoundEditorSettingsName", "MetaSound Editor"),
+					LOCTEXT("MetaSoundEditorSettingsDescription", "Customize MetaSound Editor."),
 					GetMutableDefault<UMetasoundEditorSettings>()
 				);
 
@@ -907,6 +913,18 @@ namespace Metasound
 				});
 
 				AssetTools.GetOnPackageMigration().AddRaw(this, &FModule::OnPackageMigration);
+
+				if (UMetaSoundSettings* Settings = GetMutableDefault<UMetaSoundSettings>())
+				{
+					Settings->OnDefaultConformed.AddLambda([]()
+					{
+						FNotificationInfo Info(LOCTEXT("MetaSoundSettings_CannotDeleteDefaultPage", "Cannot change name of nor delete 'Default' MetaSound Page"));
+						Info.bFireAndForget = true;
+						Info.ExpireDuration = 2.0f;
+						Info.bUseThrobber = true;
+						FSlateNotificationManager::Get().AddNotification(Info);
+					});
+				}
 			}
 
 			virtual void ShutdownModule() override
@@ -1024,3 +1042,5 @@ namespace Metasound
 } // namespace Metasound
 
 IMPLEMENT_MODULE(Metasound::Editor::FModule, MetasoundEditor);
+
+#undef LOCTEXT_NAMESPACE
\ No newline at end of file
diff --git a/Engine/Plugins/Runtime/Metasound/Source/MetasoundEngine/Private/Metasound.cpp b/Engine/Plugins/Runtime/Metasound/Source/MetasoundEngine/Private/Metasound.cpp
--- a/Engine/Plugins/Runtime/Metasound/Source/MetasoundEngine/Private/Metasound.cpp
+++ b/Engine/Plugins/Runtime/Metasound/Source/MetasoundEngine/Private/Metasound.cpp
@@ -67,13 +67,11 @@ Metasound::Frontend::FDocumentAccessPtr UMetaSoundPatch::GetDocumentAccessPtr()
 
 	// Mutation of a document via the soft deprecated access ptr/controller system is not tracked by
 	// the builder registry, so the document cache is invalidated here.
-	if (UMetaSoundBuilderSubsystem* BuilderSubsystem = UMetaSoundBuilderSubsystem::Get())
+	if (IDocumentBuilderRegistry* BuilderRegistry = IDocumentBuilderRegistry::Get())
 	{
-		if (UMetaSoundBuilderBase* Builder = BuilderSubsystem->FindBuilderOfDocument(this))
-		{
-			Builder->Reload();
-		}
+		BuilderRegistry->ReloadBuilder(RootMetaSoundDocument.RootGraph.Metadata.GetClassName());
 	}
+
 	// Return document using FAccessPoint to inform the TAccessPtr when the 
 	// object is no longer valid.
 	return MakeAccessPtr<FDocumentAccessPtr>(RootMetaSoundDocument.AccessPoint, RootMetaSoundDocument);
diff --git a/Engine/Plugins/Runtime/Metasound/Source/MetasoundEngine/Private/MetasoundSettings.cpp b/Engine/Plugins/Runtime/Metasound/Source/MetasoundEngine/Private/MetasoundSettings.cpp
--- a/Engine/Plugins/Runtime/Metasound/Source/MetasoundEngine/Private/MetasoundSettings.cpp
+++ b/Engine/Plugins/Runtime/Metasound/Source/MetasoundEngine/Private/MetasoundSettings.cpp
@@ -1,108 +1,207 @@
 // Copyright Epic Games, Inc. All Rights Reserved.
 
 #include "MetasoundSettings.h"
+
 #include "Algo/Count.h"
+#include "MetasoundFrontendDocument.h"
+
+#define LOCTEXT_NAMESPACE "MetaSound"
 
-#if WITH_EDITORONLY_DATA
 
-namespace MetaSoundSettingsPrivate
+#if WITH_EDITORONLY_DATA
+namespace Metasound::SettingsPrivate
 {
+	template<typename SettingsStructType>
+	TSet<FName> GetStructNames(const TArray<SettingsStructType>& InSettings, int32 IgnoreIndex = INDEX_NONE)
+	{
+		TSet<FName> Names;
+		for (int32 Index = 0; Index < InSettings.Num(); ++Index)
+		{
+			if (Index != IgnoreIndex)
+			{
+				Names.Add(InSettings[Index].Name);
+			}
+		}
+		return Names;
+	}
+
 	/** Generate new name for the item. **/
-	static FName GenerateNextName(const TCHAR* InBaseName)
+	static FName GenerateUniqueName(const TSet<FName>& Names, const TCHAR* InBaseName)
 	{
-		const TArray<FName> Names = UMetaSoundQualityHelper::GetQualityList();	
 		FString NewName = InBaseName;
-		
-		// Start adding a postfix starting at 1.
-		for (int32 Postfix=1; Names.Contains(*NewName); ++Postfix)
+		for (int32 Postfix = 1; Names.Contains(*NewName); ++Postfix)
 		{
-			NewName = FString::Format(TEXT("{0} {1}"), { InBaseName, Postfix });
+			NewName = FString::Format(TEXT("{0}_{1}"), { InBaseName, Postfix });
 		}
 		return FName(*NewName);
 	}
-	
-	/** When adding a quality new entry in the editor call this **/	
-	static void OnCreateNewQualitySettings(FMetaSoundQualitySettings& InNewItem)
-	{
-		InNewItem.Name = GenerateNextName(TEXT("New Quality"));
-		InNewItem.UniqueId = FGuid::NewGuid();
-	}
 
-	static void OnDuplicateQualitySettings(FMetaSoundQualitySettings& InNewlyDuplicatedItem)
+	template<typename SettingsStructType>
+	void OnCreateNewSettingsStruct(const TArray<SettingsStructType>& InSettings, const FString& InBaseName, SettingsStructType& OutNewItem)
 	{
-		InNewlyDuplicatedItem.Name = GenerateNextName(*FString::Printf(TEXT("%s (Duplicate)"), *InNewlyDuplicatedItem.Name.ToString()));
-		InNewlyDuplicatedItem.UniqueId = FGuid::NewGuid();
+		const TSet<FName> Names = GetStructNames(InSettings);
+		OutNewItem.Name = GenerateUniqueName(Names, *InBaseName);
+		OutNewItem.UniqueId = FGuid::NewGuid();
 	}
 
-	static void OnPasteQualitySettings(FMetaSoundQualitySettings& InNewlyPastedItem)
-	{
-		InNewlyPastedItem.Name = GenerateNextName(*FString::Printf(TEXT("%s (Copied)"), *InNewlyPastedItem.Name.ToString()));
-		InNewlyPastedItem.UniqueId = FGuid::NewGuid();
-	}
-	static void OnRenameQualitySettings(FMetaSoundQualitySettings& InRenamed)
+	template<typename SettingsStructType>
+	void OnRenameSettingsStruct(const TArray<SettingsStructType>& InSettings, int32 Index, const FString& InBaseName, SettingsStructType& OutRenamed)
 	{
-		// Prevent 'None' as an option.
-		if (InRenamed.Name.IsNone())
+		if (OutRenamed.Name.IsNone())
 		{
-			InRenamed.Name = GenerateNextName(TEXT("New Quality"));
+			const TSet<FName> Names = GetStructNames(InSettings);
+			OutRenamed.Name = GenerateUniqueName(Names, *InBaseName);
 		}
-		
-		// More than one?
-		else if (Algo::Count(UMetaSoundQualityHelper::GetQualityList(), InRenamed.Name) > 1)
+		else
 		{
-			// add something unique to the name.
-			InRenamed.Name = GenerateNextName(*InRenamed.Name.ToString());
-		}	
+			const TSet<FName> Names = GetStructNames(InSettings, Index);
+			if (Names.Contains(OutRenamed.Name))
+			{
+				OutRenamed.Name = GenerateUniqueName(Names, *OutRenamed.Name.ToString());
+			}
+		}
 	}
-	
-}//MetaSoundSettingsPrivate
 
-void UMetaSoundSettings::PostEditChangeChainProperty(FPropertyChangedChainEvent& PostEditChangeChainProperty)
-{
-	const int32 QualityItemIndex = PostEditChangeChainProperty.GetArrayIndex(GET_MEMBER_NAME_STRING_CHECKED(UMetaSoundSettings, QualitySettings));
+	template<typename SettingsStructType>
+	const SettingsStructType* FindSettingsStruct(const TArray<SettingsStructType>& Settings, const FGuid& InUniqueID)
+	{
+		auto MatchesIDPredicate = [&InUniqueID](const SettingsStructType& Struct) { return Struct.UniqueId == InUniqueID; };
+		return Settings.FindByPredicate(MatchesIDPredicate);
+	}
+
+	template<typename SettingsStructType>
+	const SettingsStructType* FindSettingsStruct(const TArray<SettingsStructType>& Settings, FName Name)
+	{
+		auto MatchesNamePredicate = [Name](const SettingsStructType& Struct) { return Struct.Name == Name; };
+		return Settings.FindByPredicate(MatchesNamePredicate);
+	}
 
-	// Quality item changed..
-	if (QualityItemIndex != INDEX_NONE && QualitySettings.IsValidIndex(QualityItemIndex))
+	template<typename SettingsStructType>
+	void PostEditChainChangedStructMember(FPropertyChangedChainEvent& PostEditChangeChainProperty, TArray<SettingsStructType>& StructSettings, FName PropertyName, const FString& NewItemName)
 	{
-		FMetaSoundQualitySettings& Item = QualitySettings[QualityItemIndex];
+		const int32 ItemIndex = PostEditChangeChainProperty.GetArrayIndex(PropertyName.ToString());
 
-		// Name has changed.
-		if (PostEditChangeChainProperty.GetPropertyName() == GET_MEMBER_NAME_CHECKED(FMetaSoundQualitySettings, Name))
+		if (TDoubleLinkedList<FProperty*>::TDoubleLinkedListNode* HeadNode = PostEditChangeChainProperty.PropertyChain.GetHead())
 		{
-			MetaSoundSettingsPrivate::OnRenameQualitySettings(Item);
+			const FProperty* Prop = HeadNode->GetValue();
+			if (Prop->GetName() != PropertyName)
+			{
+				return;
+			}
 		}
-		
-		// Array change.
-		else if (PostEditChangeChainProperty.GetPropertyName() == GET_MEMBER_NAME_CHECKED(UMetaSoundSettings, QualitySettings))
+
+		// Item changed..
+		if (ItemIndex != INDEX_NONE && StructSettings.IsValidIndex(ItemIndex))
 		{
-			// Add
-			if (PostEditChangeChainProperty.ChangeType == EPropertyChangeType::ArrayAdd)
+			SettingsStructType& Item = StructSettings[ItemIndex];
+			if (PostEditChangeChainProperty.GetPropertyName() == "Name")
 			{
-				MetaSoundSettingsPrivate::OnCreateNewQualitySettings(Item);
+				OnRenameSettingsStruct<SettingsStructType>(StructSettings, ItemIndex, NewItemName, Item);
 			}
-			// Duplicate.
-			else if (PostEditChangeChainProperty.ChangeType == EPropertyChangeType::Duplicate)
+			else if (PostEditChangeChainProperty.GetPropertyName() == PropertyName)
 			{
-				MetaSoundSettingsPrivate::OnDuplicateQualitySettings(Item);
+				// Array change add or duplicate
+				if (PostEditChangeChainProperty.ChangeType == EPropertyChangeType::ArrayAdd
+					|| PostEditChangeChainProperty.ChangeType == EPropertyChangeType::Duplicate)
+				{
+					OnCreateNewSettingsStruct<SettingsStructType>(StructSettings, NewItemName, Item);
+				}
 			}
-		}	
+		}
+
+		// Handle pasting separately as we might not have a valid index in the case of pasting when array is empty.
+		if (PostEditChangeChainProperty.GetPropertyName() == PropertyName)
+		{
+			// Paste...
+			if (PostEditChangeChainProperty.ChangeType == EPropertyChangeType::ValueSet)
+			{
+				const int32 IndexOfPastedItem = ItemIndex != INDEX_NONE ? ItemIndex : 0;
+				if (StructSettings.IsValidIndex(IndexOfPastedItem))
+				{
+					SettingsStructType& Item = StructSettings[IndexOfPastedItem];
+					OnCreateNewSettingsStruct<SettingsStructType>(StructSettings, NewItemName, Item);
+				}
+			}
+		}
 	}
-	
-	// Handle pasting separately as we might not have a valid index in the case of pasting when quality array is empty.
-	if (PostEditChangeChainProperty.GetPropertyName() == GET_MEMBER_NAME_CHECKED(UMetaSoundSettings, QualitySettings))
+} // namespace Metasound::SettingsPrivate
+
+#if WITH_EDITORONLY_DATA
+void UMetaSoundSettings::ConformPageSettingsDefault(bool bNotifyDefaultConformed)
+{
+	using namespace Metasound;
+
+	bool bContainsPageDefault = false;
+	bool bDefaultConformed = false;
+	for (int32 Index = PageSettings.Num() - 1; Index >= 0; --Index)
 	{
-		// Paste...
-		if (PostEditChangeChainProperty.ChangeType == EPropertyChangeType::ValueSet)
+		FMetaSoundPageSettings& Page = PageSettings[Index];
+		const bool bIsDefaultName = Page.Name == Frontend::DefaultGraphPageName;
+		if (bIsDefaultName)
+		{
+			if (Page.UniqueId != Frontend::DefaultGraphPageID)
+			{
+				Page.UniqueId = { };
+				bDefaultConformed = true;
+			}
+
+			bContainsPageDefault = true;
+		}
+		else
 		{
-			const int32 IndexOfPastedItem = QualityItemIndex != INDEX_NONE ? QualityItemIndex : 0;
-			if (QualitySettings.IsValidIndex(IndexOfPastedItem))
+			if (Page.UniqueId == Frontend::DefaultGraphPageID)
 			{
-				FMetaSoundQualitySettings& PastedItem = QualitySettings[IndexOfPastedItem];
-				MetaSoundSettingsPrivate::OnPasteQualitySettings(PastedItem);
+				Page.UniqueId = FGuid::NewGuid();
+				bDefaultConformed = true;
 			}
 		}
 	}
 
+	if (!bContainsPageDefault)
+	{
+		FMetaSoundPageSettings DefaultSettings;
+		DefaultSettings.Name = Frontend::DefaultGraphPageName;
+		PageSettings.Insert(MoveTemp(DefaultSettings), 0);
+		bDefaultConformed = true;
+	}
+
+	if (bNotifyDefaultConformed && bDefaultConformed)
+	{
+		OnDefaultConformed.Broadcast();
+	}
+}
+
+const FMetaSoundPageSettings* UMetaSoundSettings::FindPageSettings(FName Name) const
+{
+	return Metasound::SettingsPrivate::FindSettingsStruct(PageSettings, Name);
+}
+
+const FMetaSoundPageSettings* UMetaSoundSettings::FindPageSettings(const FGuid& InPageID) const
+{
+	return Metasound::SettingsPrivate::FindSettingsStruct(PageSettings, InPageID);
+}
+
+const FMetaSoundQualitySettings* UMetaSoundSettings::FindQualitySettings(FName Name) const
+{
+	return Metasound::SettingsPrivate::FindSettingsStruct(QualitySettings, Name);
+}
+
+const FMetaSoundQualitySettings* UMetaSoundSettings::FindQualitySettings(const FGuid& InQualityID) const
+{
+	return Metasound::SettingsPrivate::FindSettingsStruct(QualitySettings, InQualityID);
+}
+#endif // WITH_EDITORONLY_DATA
+
+void UMetaSoundSettings::PostEditChangeChainProperty(FPropertyChangedChainEvent& PostEditChangeChainProperty)
+{
+	using namespace Metasound::SettingsPrivate;
+
+	PostEditChainChangedStructMember(PostEditChangeChainProperty, PageSettings, GetPageSettingPropertyName(), TEXT("New Page"));
+	PostEditChainChangedStructMember(PostEditChangeChainProperty, QualitySettings, GetQualitySettingPropertyName(), TEXT("New Quality"));
+
+	constexpr bool bNotifyDefaultConformed = true;
+	ConformPageSettingsDefault(bNotifyDefaultConformed);
+
 	Super::PostEditChangeChainProperty(PostEditChangeChainProperty);
 }
 
@@ -110,9 +209,27 @@ void UMetaSoundSettings::PostEditChangeProperty(FPropertyChangedEvent& PropertyC
 {	
 	Super::PostEditChangeProperty(PropertyChangedEvent);
 	
+	if (PropertyChangedEvent.MemberProperty->GetName() == GetPageSettingPropertyName())
+	{
+		OnPageSettingsUpdated.Broadcast();
+	}
+
 	DenyListCacheChangeID++;
 }
 
+void UMetaSoundSettings::PostInitProperties()
+{
+	constexpr bool bNotifyDefaultConformed = false;
+	ConformPageSettingsDefault(bNotifyDefaultConformed);
+
+	Super::PostInitProperties();
+}
+
+FName UMetaSoundSettings::GetPageSettingPropertyName()
+{
+	return GET_MEMBER_NAME_CHECKED(UMetaSoundSettings, PageSettings);
+}
+
 FName UMetaSoundSettings::GetQualitySettingPropertyName()
 {
 	return GET_MEMBER_NAME_CHECKED(UMetaSoundSettings, QualitySettings);
@@ -136,3 +253,4 @@ TArray<FName> UMetaSoundQualityHelper::GetQualityList()
 
 	return Names;
 }
+#undef LOCTEXT_NAMESPACE // MetaSound
diff --git a/Engine/Plugins/Runtime/Metasound/Source/MetasoundEngine/Private/MetasoundSource.cpp b/Engine/Plugins/Runtime/Metasound/Source/MetasoundEngine/Private/MetasoundSource.cpp
--- a/Engine/Plugins/Runtime/Metasound/Source/MetasoundEngine/Private/MetasoundSource.cpp
+++ b/Engine/Plugins/Runtime/Metasound/Source/MetasoundEngine/Private/MetasoundSource.cpp
@@ -945,12 +945,9 @@ Metasound::Frontend::FDocumentAccessPtr UMetaSoundSource::GetDocumentAccessPtr()
 
 	// Mutation of a document via the soft deprecated access ptr/controller system is not tracked by
 	// the builder registry, so the document cache is invalidated here.
-	if (UMetaSoundBuilderSubsystem* BuilderSubsystem = UMetaSoundBuilderSubsystem::Get())
+	if (IDocumentBuilderRegistry* BuilderRegistry = IDocumentBuilderRegistry::Get())
 	{
-		if (UMetaSoundBuilderBase* Builder = BuilderSubsystem->FindBuilderOfDocument(this))
-		{
-			Builder->Reload();
-		}
+		BuilderRegistry->ReloadBuilder(RootMetasoundDocument.RootGraph.Metadata.GetClassName());
 	}
 
 	// Return document using FAccessPoint to inform the TAccessPtr when the 
diff --git a/Engine/Plugins/Runtime/Metasound/Source/MetasoundEngine/Public/MetasoundBuilderBase.cpp b/Engine/Plugins/Runtime/Metasound/Source/MetasoundEngine/Public/MetasoundBuilderBase.cpp
--- a/Engine/Plugins/Runtime/Metasound/Source/MetasoundEngine/Public/MetasoundBuilderBase.cpp
+++ b/Engine/Plugins/Runtime/Metasound/Source/MetasoundEngine/Public/MetasoundBuilderBase.cpp
@@ -23,6 +23,7 @@
 #include "MetasoundFrontendSearchEngine.h"
 #include "MetasoundFrontendTransform.h"
 #include "MetasoundLog.h"
+#include "MetasoundSettings.h"
 #include "MetasoundSource.h"
 #include "MetasoundTrace.h"
 #include "MetasoundUObjectRegistry.h"
@@ -151,12 +152,33 @@ FMetaSoundBuilderNodeInputHandle UMetaSoundBuilderBase::AddGraphOutputNode(FName
 	return NewHandle;
 }
 
+void UMetaSoundBuilderBase::AddGraphPage(FName PageName, bool bDuplicateLastGraph, bool bSetAsBuildGraph, EMetaSoundBuilderResult& OutResult)
+{
+	if (const UMetaSoundSettings* Settings = GetDefault<UMetaSoundSettings>())
+	{
+		if (const FMetaSoundPageSettings* PageSettings = Settings->FindPageSettings(PageName))
+		{
+			Builder.AddGraphPage(PageSettings->UniqueId, bDuplicateLastGraph, bSetAsBuildGraph);
+			OutResult = EMetaSoundBuilderResult::Succeeded;
+			return;
+		}
+	}
+
+	OutResult = EMetaSoundBuilderResult::Failed;
+}
+
 void UMetaSoundBuilderBase::AddInterface(FName InterfaceName, EMetaSoundBuilderResult& OutResult)
 {
 	const bool bInterfaceAdded = Builder.AddInterface(InterfaceName);
 	OutResult = bInterfaceAdded ? EMetaSoundBuilderResult::Succeeded : EMetaSoundBuilderResult::Failed;
 }
 
+void UMetaSoundBuilderBase::AddTransactionListener(TSharedRef<Metasound::Frontend::IDocumentBuilderTransactionListener> BuilderListener)
+{
+	BuilderListener->OnBuilderReloaded(GetBuilderDelegates());
+	BuilderReloadDelegate.AddSP(BuilderListener, &Metasound::Frontend::IDocumentBuilderTransactionListener::OnBuilderReloaded);
+}
+
 FMetaSoundNodeHandle UMetaSoundBuilderBase::AddNode(const TScriptInterface<IMetaSoundDocumentInterface>& NodeClass, EMetaSoundBuilderResult& OutResult)
 {
 	using namespace Metasound;
@@ -715,6 +737,11 @@ FMetaSoundFrontendDocumentBuilder& UMetaSoundBuilderBase::GetBuilder()
 	return Builder;
 }
 
+Metasound::Frontend::FDocumentModifyDelegates& UMetaSoundBuilderBase::GetBuilderDelegates()
+{
+	return Builder.GetDocumentDelegates();
+}
+
 const FMetaSoundFrontendDocumentBuilder& UMetaSoundBuilderBase::GetConstBuilder() const
 {
 	return Builder;
@@ -799,7 +826,7 @@ bool UMetaSoundBuilderBase::InterfaceIsDeclared(FName InterfaceName) const
 
 void UMetaSoundBuilderBase::InvalidateCache(bool bPrimeCache)
 {
-	Reload(bPrimeCache);
+	Reload({ }, bPrimeCache);
 }
 
 bool UMetaSoundBuilderBase::IsPreset() const
@@ -825,6 +852,7 @@ bool UMetaSoundBuilderBase::NodeOutputIsConnected(const FMetaSoundBuilderNodeOut
 
 void UMetaSoundBuilderBase::InitDelegates(Metasound::Frontend::FDocumentModifyDelegates& OutDocumentDelegates)
 {
+	BuilderReloadDelegate.Broadcast(OutDocumentDelegates);
 	OutDocumentDelegates.OnDependencyAdded.AddUObject(this, &UMetaSoundBuilderBase::OnDependencyAdded);
 	OutDocumentDelegates.OnRemoveSwappingDependency.AddUObject(this, &UMetaSoundBuilderBase::OnRemoveSwappingDependency);
 }
@@ -896,7 +924,7 @@ void UMetaSoundBuilderBase::RegisterGraphIfOutstandingTransactions(UObject& InMe
 	MetaSoundAsset->UpdateAndRegisterForExecution(Options);
 }
 
-void UMetaSoundBuilderBase::Reload(bool bPrimeCache)
+void UMetaSoundBuilderBase::Reload(TScriptInterface<IMetaSoundDocumentInterface> NewMetaSound, bool bPrimeCache)
 {
 	using namespace Metasound::Frontend;
 
@@ -907,7 +935,12 @@ void UMetaSoundBuilderBase::Reload(bool bPrimeCache)
 
 void UMetaSoundBuilderBase::ReloadCache(bool bPrimeCache)
 {
-	Reload(bPrimeCache);
+	Reload({ }, bPrimeCache);
+}
+
+void UMetaSoundBuilderBase::RemoveAllGraphPages()
+{
+	Builder.RemoveAllGraphPages();
 }
 
 #if WITH_EDITOR
@@ -929,6 +962,20 @@ void UMetaSoundBuilderBase::RemoveGraphOutput(FName Name, EMetaSoundBuilderResul
 	OutResult = bRemoved ? EMetaSoundBuilderResult::Succeeded : EMetaSoundBuilderResult::Failed;
 }
 
+void UMetaSoundBuilderBase::RemoveGraphPage(FName Name, EMetaSoundBuilderResult& OutResult)
+{
+	const UMetaSoundSettings* Settings = GetDefault<UMetaSoundSettings>();
+	check(Settings);
+	if (const FMetaSoundPageSettings* PageSettings = Settings->FindPageSettings(Name))
+	{
+		Builder.RemoveGraphPage(PageSettings->UniqueId);
+		OutResult = EMetaSoundBuilderResult::Succeeded;
+		return;
+	}
+
+	OutResult = EMetaSoundBuilderResult::Failed;
+}
+
 void UMetaSoundBuilderBase::RemoveInterface(FName InterfaceName, EMetaSoundBuilderResult& OutResult)
 {
 	const bool bInterfaceRemoved = Builder.RemoveInterface(InterfaceName);
@@ -951,6 +998,11 @@ void UMetaSoundBuilderBase::RemoveNodeInputDefault(const FMetaSoundBuilderNodeIn
 	OutResult = bInputDefaultRemoved ? EMetaSoundBuilderResult::Succeeded : EMetaSoundBuilderResult::Failed;
 }
 
+void UMetaSoundBuilderBase::RemoveTransactionListener(FDelegateHandle BuilderListenerDelegateHandle)
+{
+	BuilderReloadDelegate.Remove(BuilderListenerDelegateHandle);
+}
+
 void UMetaSoundBuilderBase::RemoveUnusedDependencies()
 {
 	Builder.RemoveUnusedDependencies();
diff --git a/Engine/Plugins/Runtime/Metasound/Source/MetasoundEngine/Public/MetasoundBuilderBase.h b/Engine/Plugins/Runtime/Metasound/Source/MetasoundEngine/Public/MetasoundBuilderBase.h
--- a/Engine/Plugins/Runtime/Metasound/Source/MetasoundEngine/Public/MetasoundBuilderBase.h
+++ b/Engine/Plugins/Runtime/Metasound/Source/MetasoundEngine/Public/MetasoundBuilderBase.h
@@ -21,6 +21,14 @@
 struct FMetasoundFrontendClassName;
 struct FMetasoundFrontendVersion;
 
+namespace Metasound::Engine
+{
+	// Forward Declarations
+	class FDocumentBuilderRegistry;
+
+	DECLARE_MULTICAST_DELEGATE_OneParam(FOnBuilderReload, Metasound::Frontend::FDocumentModifyDelegates& /* OutDelegates */)
+} // namespace Metasound::Engine
+
 
 USTRUCT(BlueprintType, meta = (DisplayName = "MetaSound Node Input Handle"))
 struct METASOUNDENGINE_API FMetaSoundBuilderNodeInputHandle : public FMetasoundFrontendVertexHandle
@@ -127,6 +135,11 @@ public:
 	UFUNCTION(BlueprintCallable, Category = "Audio|MetaSound|Builder", meta = (ExpandEnumAsExecs = "OutResult", AdvancedDisplay = "3"))
 	UPARAM(DisplayName = "Input Handle") FMetaSoundBuilderNodeInputHandle AddGraphOutputNode(FName Name, FName DataType, FMetasoundFrontendLiteral DefaultValue, EMetaSoundBuilderResult& OutResult, bool bIsConstructorOutput = false);
 
+	// Adds a graph page to the given builder's document. Fails if the page is not a valid page registered with MetaSoundSettings
+	// or if the document already contains a page with the given name. No check is done here to determine cook eligibility (i.e.
+	// pages can be added even if set to be stripped for the active platform).
+	void AddGraphPage(FName PageName, bool bDuplicateLastGraph, bool bSetAsBuildGraph, EMetaSoundBuilderResult& OutResult);
+
 	// Adds an interface registered with the given name to the graph, adding associated input and output nodes.
 	UFUNCTION(BlueprintCallable, Category = "Audio|MetaSound|Builder", meta = (ExpandEnumAsExecs = "OutResult"))
 	void AddInterface(FName InterfaceName, EMetaSoundBuilderResult& OutResult);
@@ -144,6 +157,11 @@ public:
 	UE_DEPRECATED(5.4, "This version of AddNodeByClassName is deprecated. Use the one with a default MajorVersion of 1.")
 	UPARAM(DisplayName = "Node Handle") FMetaSoundNodeHandle AddNodeByClassName(const FMetasoundFrontendClassName& ClassName, int32 MajorVersion, EMetaSoundBuilderResult& OutResult);
 
+	// Adds transaction listener which allows objects to respond to when certain graph operations are applied from anywhere (adding or removing nodes, edges, pages, etc.)
+	// Currently there is no guarantee all transactions will be represented until the Controller API is fully deprecated! (ex. if a node or edge is added or removed via a
+	// controller API call, the transaction will be missed). OnBuilderReloaded is however guaranteed to be called on mutable controller creation.
+	void AddTransactionListener(TSharedRef<Metasound::Frontend::IDocumentBuilderTransactionListener> BuilderListener);
+
 #if WITH_EDITOR
 	bool ClearMemberMetadata(const FGuid& InMemberID);
 #endif // WITH_EDITOR
@@ -318,6 +336,9 @@ public:
 	UFUNCTION(BlueprintCallable, Category = "Audio|MetaSound|Builder", meta = (ExpandEnumAsExecs = "OutResult"))
 	void ConvertToPreset(const TScriptInterface<IMetaSoundDocumentInterface>& ReferencedNodeClass, EMetaSoundBuilderResult& OutResult);
 
+	// Removes all graph pages, leaving just the default.
+	void RemoveAllGraphPages();
+
 	// Removes graph input if it exists; sets result to succeeded if it was removed and failed if it was not.
 	UFUNCTION(BlueprintCallable, Category = "Audio|MetaSound|Builder", meta = (ExpandEnumAsExecs = "OutResult"))
 	void RemoveGraphInput(FName Name, EMetaSoundBuilderResult& OutResult);
@@ -326,6 +347,10 @@ public:
 	UFUNCTION(BlueprintCallable, Category = "Audio|MetaSound|Builder", meta = (ExpandEnumAsExecs = "OutResult"))
 	void RemoveGraphOutput(FName Name, EMetaSoundBuilderResult& OutResult);
 
+	// Removes a graph page with the given name, setting result to failed if the name was not found, was invalid,
+	// or was the default (which cannot be removed).
+	void RemoveGraphPage(FName Name, EMetaSoundBuilderResult& OutResult);
+
 	// Removes the interface with the given name from the builder's MetaSound. Removes any graph inputs
 	// and outputs associated with the given interface and their respective connections (if any).
 	UFUNCTION(BlueprintCallable, Category = "Audio|MetaSound|Builder", meta = (ExpandEnumAsExecs = "OutResult"))
@@ -340,6 +365,10 @@ public:
 	UFUNCTION(BlueprintCallable, Category = "Audio|MetaSound|Builder", meta = (ExpandEnumAsExecs = "OutResult"))
 	void RemoveNodeInputDefault(const FMetaSoundBuilderNodeInputHandle& InputHandle, EMetaSoundBuilderResult& OutResult);
 
+	// Explicitly remove transaction listener from builder (see corresponding 'AddTransactionListener' function).
+	// (If listener provided with `AddTransactionListener` is destroyed, handled automatically)
+	void RemoveTransactionListener(FDelegateHandle BuilderListenerDelegateHandle);
+
 	// Removes dependencies in document that are no longer referenced by nodes
 	UFUNCTION(BlueprintCallable, Category = "Audio|MetaSound|Builder")
 	void RemoveUnusedDependencies();
@@ -347,14 +376,10 @@ public:
 	UE_DEPRECATED(5.5, "Use IDocumentBuilderRegistry::GenerateNewClassName instead to maintain registry mappings.")
 	void RenameRootGraphClass(const FMetasoundFrontendClassName& InName);
 
-	UE_DEPRECATED(5.5, "Moved to 'Reload' to discourage priming by defaulting to false. "
-		"In addition, internal delegates and respective bindings are now regenerated when reloading, "
-		"effectively resetting the builder's entire state and not just strictly the underlying cache.")
+	UE_DEPRECATED(5.5, "Moved to internal implementation and only accessible via registry to ensure delegates are properly reloaded, "
+		"path keys kept aligned, and priming managed internally")
 	void ReloadCache(bool bPrimeCache = true);
 
-	// Reloads the builder, freeing the internal cache and rebuilding delegate bindings
-	void Reload(bool bPrimeCache = false);
-
 #if WITH_EDITOR
 	// Sets the author of the MetaSound.
 	void SetAuthor(const FString& InAuthor);
@@ -403,6 +428,8 @@ public:
 	UE_DEPRECATED(5.5, "Renamed to 'GetBaseMetaSoundUClass' for consistency")
 	virtual const UClass& GetBuilderUClass() const { return GetBaseMetaSoundUClass(); }
 
+	Metasound::Frontend::FDocumentModifyDelegates& GetBuilderDelegates();
+
 	UE_DEPRECATED(5.4, "Moved to Initialize and should only be called by internal implementation (i.e. IDocumentBuilderRegistry'")
 	virtual void InitFrontendBuilder();
 
@@ -436,7 +463,7 @@ public:
 	void Initialize();
 
 protected:
-	virtual void InitDelegates(Metasound::Frontend::FDocumentModifyDelegates & OutDocumentDelegates);
+	virtual void InitDelegates(Metasound::Frontend::FDocumentModifyDelegates& OutDocumentDelegates);
 
 	virtual void OnAssetReferenceAdded(TScriptInterface<IMetaSoundDocumentInterface> DocInterface) PURE_VIRTUAL(UMetaSoundBuilderBase::OnAssetReferenceAdded, );
 	virtual void OnRemovingAssetReference(TScriptInterface<IMetaSoundDocumentInterface> DocInterface) PURE_VIRTUAL(UMetaSoundBuilderBase::OnRemovingAssetReference, );
@@ -531,9 +558,14 @@ private:
 	void OnRemoveSwappingDependency(int32 Index, int32 LastIndex);
 
 private:
+	// Reloads the builder, freeing the internal cache and rebuilding delegate bindings. Optionally,
+	// can be associated with a new MetaSound (ex. during rename. Otherwise it reuses the existing
+	// document object reference) or can have its cache primed.
+	void Reload(TScriptInterface<IMetaSoundDocumentInterface> NewMetaSound = { }, bool bPrimeCache = false);
+
+	Metasound::Engine::FOnBuilderReload BuilderReloadDelegate;
+
 	int32 LastTransactionRegistered = 0;
 
-	// Friending allows for swapping the builder in certain circumstances where desired (eg. attaching a builder to an existing asset)
-	// or performing editor-only manipulations to the builder (ex. adding template nodes) and versioning the document data where necessary.
-	friend class UMetaSoundBuilderSubsystem;
+	friend class Metasound::Engine::FDocumentBuilderRegistry;
 };
\ No newline at end of file
diff --git a/Engine/Plugins/Runtime/Metasound/Source/MetasoundEngine/Public/MetasoundSettings.h b/Engine/Plugins/Runtime/Metasound/Source/MetasoundEngine/Public/MetasoundSettings.h
--- a/Engine/Plugins/Runtime/Metasound/Source/MetasoundEngine/Public/MetasoundSettings.h
+++ b/Engine/Plugins/Runtime/Metasound/Source/MetasoundEngine/Public/MetasoundSettings.h
@@ -1,7 +1,7 @@
 // Copyright Epic Games, Inc. All Rights Reserved.
 #pragma once
 
-#include "CoreMinimal.h"
+#include "Delegates/DelegateCombinations.h"
 #include "Engine/DeveloperSettings.h"
 #include "MetasoundFrontendDocument.h"
 #include "UObject/NoExportTypes.h"
@@ -16,6 +16,14 @@
 struct FMetasoundFrontendClassName;
 struct FPropertyChangedChainEvent;
 
+#if WITH_EDITORONLY_DATA
+namespace Metasound::Engine
+{
+	DECLARE_MULTICAST_DELEGATE(FOnSettingsDefaultConformed);
+	DECLARE_MULTICAST_DELEGATE(FOnPageSettingsUpdated);
+} // namespace Metasound::Engine
+#endif // WITH_EDITORONLY_DATA
+
 UENUM()
 enum class EMetaSoundMessageLevel : uint8
 {
@@ -48,22 +56,40 @@ public:
 };
 
 USTRUCT()
-struct METASOUNDENGINE_API FMetaSoundQualitySettings
+struct METASOUNDENGINE_API FMetaSoundPageSettings
 {
 	GENERATED_BODY()
 	
+	UPROPERTY()
+	FGuid UniqueId;
+
+	/** Name of this page's setting. This will appear in the MetaSound Asset Editor's 'Page Selector'.
+		The names should be unique and adequately describe the Entry. "High", "Low" etc. **/
+	UPROPERTY(EditAnywhere, Category = "Pages")
+	FName Name;
+
 #if WITH_EDITORONLY_DATA
+	// When true, page data defined on serialized MetaSounds are included in cook for the assigned platform(s).
+	UPROPERTY(EditAnywhere, config, Category = "Pages")
+	FPerPlatformBool CookForPlatform = true;
+#endif //WITH_EDITORONLY_DATA
+};
 
+USTRUCT()
+struct METASOUNDENGINE_API FMetaSoundQualitySettings
+{
+	GENERATED_BODY()
+	
+#if WITH_EDITORONLY_DATA
 	/** A hidden GUID that will be generated once when adding a new entry. This prevents orphaning of renamed entries. **/
 	UPROPERTY()
-	FGuid UniqueId = {};
+	FGuid UniqueId;
 
 	/** Name of this quality setting. This will appear in the quality dropdown list.
 		The names should be unique and adequately describe the Entry. "High", "Low" etc. **/
 	UPROPERTY(EditAnywhere, Category = "Quality")
-	FName Name = {};
-	
-#endif //WITH_EDITORONLY_DATA	
+	FName Name;
+#endif // WITH_EDITORONLY_DATA
 
 	/** Sample Rate (in Hz). NOTE: A Zero value will have no effect and use the Device Rate. **/
 	UPROPERTY(EditAnywhere, config, Category = "Quality", meta = (ClampMin = "0", ClampMax="96000"))
@@ -110,24 +136,41 @@ public:
 	int32 DenyListCacheChangeID = 0;	
 
 #if WITH_EDITORONLY_DATA
-	const TArray<FMetaSoundQualitySettings>& GetQualitySettings() const { return QualitySettings; }
-	static FName GetQualitySettingPropertyName(); 
+	Metasound::Engine::FOnSettingsDefaultConformed OnDefaultConformed;
+	Metasound::Engine::FOnPageSettingsUpdated OnPageSettingsUpdated;
 #endif //WITH_EDITORONLY_DATA
 
 private:
+	/** Array of possible page settings that can be added to a MetaSound object. */
+	UPROPERTY(EditAnywhere, config, Category = Pages)
+	TArray<FMetaSoundPageSettings> PageSettings;
 
 	/** Array of possible quality settings for Metasounds to chose from */
-	// NOTE: Ideally this would be wrapped with WITH_EDITORONLY_DATA, but standalone "-game" requires
-	// it to exist. Access is limited to the accessor above, which enforces it correctly.
 	UPROPERTY(EditAnywhere, config, Category = Quality)
 	TArray<FMetaSoundQualitySettings> QualitySettings;
 
+public:
+	const FMetaSoundPageSettings* FindPageSettings(FName Name) const;
+	const FMetaSoundPageSettings* FindPageSettings(const FGuid& InPageID) const;
+
+	const FMetaSoundQualitySettings* FindQualitySettings(FName Name) const;
+	const FMetaSoundQualitySettings* FindQualitySettings(const FGuid& InQualityID) const;
+
+	const TArray<FMetaSoundPageSettings>& GetPageSettings() const { return PageSettings; }
+	const TArray<FMetaSoundQualitySettings>& GetQualitySettings() const { return QualitySettings; }
+
+#if WITH_EDITORONLY_DATA
+	static FName GetPageSettingPropertyName();
+	static FName GetQualitySettingPropertyName();
+#endif // WITH_EDITORONLY_DATA
+
 #if WITH_EDITOR
 private:
+	void ConformPageSettingsDefault(bool bNotifyDefaultConformed);
 
 	virtual void PostEditChangeChainProperty(FPropertyChangedChainEvent& PropertyChangedEvent) override;
-
 	virtual void PostEditChangeProperty(FPropertyChangedEvent& PropertyChangedEvent) override;
+	virtual void PostInitProperties() override;
 #endif // WITH_EDITOR
 };
 
diff --git a/Engine/Plugins/Runtime/Metasound/Source/MetasoundFrontend/Private/MetasoundFrontendDocument.cpp b/Engine/Plugins/Runtime/Metasound/Source/MetasoundFrontend/Private/MetasoundFrontendDocument.cpp
--- a/Engine/Plugins/Runtime/Metasound/Source/MetasoundFrontend/Private/MetasoundFrontendDocument.cpp
+++ b/Engine/Plugins/Runtime/Metasound/Source/MetasoundFrontend/Private/MetasoundFrontendDocument.cpp
@@ -15,6 +15,10 @@
 #include "MetasoundParameterTransmitter.h"
 #include "MetasoundVertex.h"
 
+#if WITH_EDITORONLY_DATA
+#include "Internationalization/Internationalization.h"
+#endif // WITH_EDITORONLY_DATA
+
 #include UE_INLINE_GENERATED_CPP_BY_NAME(MetasoundFrontendDocument)
 
 namespace Metasound
@@ -23,6 +27,10 @@ namespace Metasound
 
 	namespace Frontend
 	{
+#if WITH_EDITORONLY_DATA
+		const FText DefaultGraphPageDisplayName = NSLOCTEXT("MetasoundFrontend", "DefaultGraphPageDisplayName", "Default");
+#endif // WITH_EDITORONLY_DATA
+
 		namespace DisplayStyle
 		{
 			namespace EdgeAnimation
diff --git a/Engine/Plugins/Runtime/Metasound/Source/MetasoundFrontend/Public/MetasoundFrontendDocument.h b/Engine/Plugins/Runtime/Metasound/Source/MetasoundFrontend/Public/MetasoundFrontendDocument.h
--- a/Engine/Plugins/Runtime/Metasound/Source/MetasoundFrontend/Public/MetasoundFrontendDocument.h
+++ b/Engine/Plugins/Runtime/Metasound/Source/MetasoundFrontend/Public/MetasoundFrontendDocument.h
@@ -36,6 +36,11 @@ namespace Metasound
 	namespace Frontend
 	{
 		constexpr FGuid DefaultGraphPageID(0, 0, 0, 0);
+		constexpr TCHAR DefaultGraphPageName[] = TEXT("Default");
+
+#if WITH_EDITORONLY_DATA
+		extern const FText METASOUNDFRONTEND_API DefaultGraphPageDisplayName;
+#endif // WITH_EDITORONLY_DATA
 
 		namespace DisplayStyle
 		{
