Missing file
#rnx

diff --git a/Engine/Plugins/Runtime/Metasound/Source/MetasoundFrontend/Public/MetasoundFrontendDocumentModifyDelegates.h b/Engine/Plugins/Runtime/Metasound/Source/MetasoundFrontend/Public/MetasoundFrontendDocumentModifyDelegates.h
--- a/Engine/Plugins/Runtime/Metasound/Source/MetasoundFrontend/Public/MetasoundFrontendDocumentModifyDelegates.h
+++ b/Engine/Plugins/Runtime/Metasound/Source/MetasoundFrontend/Public/MetasoundFrontendDocumentModifyDelegates.h
@@ -117,11 +117,12 @@ PRAGMA_ENABLE_DEPRECATION_WARNINGS
 		}
 	};
 
-	class METASOUNDFRONTEND_API IDocumentBuilderListener : public TSharedFromThis<IDocumentBuilderListener>
+	class METASOUNDFRONTEND_API IDocumentBuilderTransactionListener : public TSharedFromThis<IDocumentBuilderTransactionListener>
 	{
 	public:
-		virtual ~IDocumentBuilderListener() = default;
+		virtual ~IDocumentBuilderTransactionListener() = default;
 
-		virtual void AddDocumentBuilderDelegates(FDocumentModifyDelegates& OutDelegates) = 0;
+		// Called when the builder is reloaded, at which point the document cache and delegates are refreshed
+		virtual void OnBuilderReloaded(FDocumentModifyDelegates& OutDelegates) = 0;
 	};
 } // namespace Metasound::Frontend
