Add MetaSound asset validation for assets with duplicate class names
- Move asset registration to engine module from editor
- Separate call for MetaSound asset manager prime and node class registry prime
- Fix bug causing MetaSound asset editor registration notification to close early
- Make editor asset registration async by default (remove au.MetaSound.Editor.AsyncRegistrationEnabled cvar)

diff --git a/Engine/Plugins/Runtime/Metasound/Source/MetasoundEditor/Private/MetasoundEditor.cpp b/Engine/Plugins/Runtime/Metasound/Source/MetasoundEditor/Private/MetasoundEditor.cpp
--- a/Engine/Plugins/Runtime/Metasound/Source/MetasoundEditor/Private/MetasoundEditor.cpp
+++ b/Engine/Plugins/Runtime/Metasound/Source/MetasoundEditor/Private/MetasoundEditor.cpp
@@ -31,6 +31,7 @@
 #include "HAL/PlatformApplicationMisc.h"
 #include "IAudioExtensionPlugin.h"
 #include "IDetailsView.h"
+#include "IMetasoundEngineModule.h"
 #include "Kismet2/BlueprintEditorUtils.h"
 #include "Kismet2/KismetEditorUtilities.h"
 #include "Logging/TokenizedMessage.h"
@@ -1041,15 +1042,16 @@ namespace Metasound
 		void FEditor::InitMetasoundEditor(const EToolkitMode::Type Mode, const TSharedPtr<IToolkitHost>& InitToolkitHost, UObject* ObjectToEdit)
 		{
 			using namespace Metasound::Frontend;
+			using namespace Metasound::Engine;
 
 			check(ObjectToEdit);
 			checkf(IMetasoundUObjectRegistry::Get().IsRegisteredClass(ObjectToEdit), TEXT("Object passed in was not registered as a valid metasound interface!"));
 
-			IMetasoundEditorModule& MetaSoundEditorModule = FModuleManager::GetModuleChecked<IMetasoundEditorModule>("MetaSoundEditor");
-			bPrimingRegistry = MetaSoundEditorModule.GetAssetRegistryPrimeStatus() <= EAssetPrimeStatus::InProgress;
-			if (MetaSoundEditorModule.GetAssetRegistryPrimeStatus() < EAssetPrimeStatus::InProgress)
+			IMetasoundEngineModule& MetaSoundEngineModule = FModuleManager::GetModuleChecked<IMetasoundEngineModule>("MetaSoundEngine");
+			bPrimingRegistry = MetaSoundEngineModule.GetNodeClassRegistryPrimeStatus() <= ENodeClassRegistryPrimeStatus::InProgress;
+			if (MetaSoundEngineModule.GetNodeClassRegistryPrimeStatus() < ENodeClassRegistryPrimeStatus::InProgress)
 			{
-				MetaSoundEditorModule.PrimeAssetRegistryAsync();
+				MetaSoundEngineModule.PrimeAssetRegistryAsync();
 			}
 
 			// Support undo/redo
@@ -3809,6 +3811,7 @@ namespace Metasound
 
 		void FEditor::Tick(float DeltaTime)
 		{
+			using namespace Metasound::Engine;
 			UObject* MetaSound = GetMetasoundObject();
 			if (!MetaSound)
 			{
@@ -3817,10 +3820,10 @@ namespace Metasound
 
 			if (bPrimingRegistry)
 			{
-				IMetasoundEditorModule& MetaSoundEditorModule = FModuleManager::GetModuleChecked<IMetasoundEditorModule>("MetaSoundEditor");
-				EAssetPrimeStatus PrimeStatus = MetaSoundEditorModule.GetAssetRegistryPrimeStatus();
-				EAssetScanStatus ScanStatus = MetaSoundEditorModule.GetAssetRegistryScanStatus();
-				if (PrimeStatus == EAssetPrimeStatus::Complete || (PrimeStatus == EAssetPrimeStatus::InProgress && ScanStatus == EAssetScanStatus::Complete))
+				IMetasoundEngineModule& MetaSoundEngineModule = FModuleManager::GetModuleChecked<Metasound::Engine::IMetasoundEngineModule>("MetaSoundEngine");
+				ENodeClassRegistryPrimeStatus PrimeStatus = MetaSoundEngineModule.GetNodeClassRegistryPrimeStatus();
+				EAssetScanStatus ScanStatus = MetaSoundEngineModule.GetAssetRegistryScanStatus();
+				if (PrimeStatus == Metasound::Engine::ENodeClassRegistryPrimeStatus::Complete)
 				{
 					bPrimingRegistry = false;
 					NotifyAssetPrimeComplete();
diff --git a/Engine/Plugins/Runtime/Metasound/Source/MetasoundEditor/Private/MetasoundEditorModule.cpp b/Engine/Plugins/Runtime/Metasound/Source/MetasoundEditor/Private/MetasoundEditorModule.cpp
--- a/Engine/Plugins/Runtime/Metasound/Source/MetasoundEditor/Private/MetasoundEditorModule.cpp
+++ b/Engine/Plugins/Runtime/Metasound/Source/MetasoundEditor/Private/MetasoundEditorModule.cpp
@@ -1,9 +1,6 @@
 // Copyright Epic Games, Inc. All Rights Reserved.
 #include "MetasoundEditorModule.h"
 
-#include "AssetRegistry/AssetData.h"
-#include "AssetRegistry/AssetRegistryModule.h"
-#include "AssetRegistry/AssetRegistryModule.h"
 #include "AssetTypeActions_Base.h"
 #include "Brushes/SlateImageBrush.h"
 #include "EdGraph/EdGraphNode.h"
@@ -12,6 +9,7 @@
 #include "Framework/Notifications/NotificationManager.h"
 #include "HAL/IConsoleManager.h"
 #include "IDetailCustomization.h"
+#include "IMetasoundEngineModule.h"
 #include "ISettingsModule.h"
 #include "Metasound.h"
 #include "MetasoundAssetSubsystem.h"
@@ -30,7 +28,6 @@
 #include "MetasoundFrontendDocument.h"
 #include "MetasoundFrontendRegistries.h"
 #include "MetasoundFrontendTransform.h"
-#include "MetasoundGlobals.h"
 #include "MetasoundNodeDetailCustomization.h"
 #include "MetasoundSettings.h"
 #include "MetasoundSource.h"
@@ -38,7 +35,6 @@
 #include "MetasoundTrace.h"
 #include "MetasoundTrigger.h"
 #include "MetasoundUObjectRegistry.h"
-#include "Modules/ModuleInterface.h"
 #include "Modules/ModuleManager.h"
 #include "PackageMigrationContext.h"
 #include "PropertyEditorDelegates.h"
@@ -58,15 +54,6 @@
 
 DEFINE_LOG_CATEGORY(LogMetasoundEditor);
 
-
-static int32 MetaSoundEditorAsyncRegistrationEnabledCVar = 1;
-FAutoConsoleVariableRef CVarMetaSoundEditorAsyncRegistrationEnabled(
-	TEXT("au.MetaSound.Editor.AsyncRegistrationEnabled"),
-	MetaSoundEditorAsyncRegistrationEnabledCVar,
-	TEXT("Enable registering all MetaSound asset classes asyncronously on editor load.\n")
-	TEXT("0: Disabled, !0: Enabled (default)"),
-	ECVF_Default);
-
 #define LOCTEXT_NAMESPACE "MetaSounds"
 
 namespace Metasound
@@ -241,262 +228,6 @@ namespace Metasound
 
 		class FModule : public IMetasoundEditorModule
 		{
-			void LoadAndRegisterAsset(const FAssetData& InAssetData)
-			{
-				using namespace Metasound;
-			
-				// Ignore requests if graphs cannot be executed, as registration that results in IGraph generation is not supported.
-				if (!Metasound::CanEverExecuteGraph())
-				{
-					return;
-				}
-
-				Frontend::FMetaSoundAssetRegistrationOptions RegOptions;
-				RegOptions.bForceReregister = false;
-				if (const UMetaSoundSettings* Settings = GetDefault<UMetaSoundSettings>())
-				{
-					RegOptions.bAutoUpdateLogWarningOnDroppedConnection = Settings->bAutoUpdateLogWarningOnDroppedConnection;
-				}
-
-				if (InAssetData.IsAssetLoaded())
-				{
-					if (UObject* AssetObject = InAssetData.GetAsset())
-					{
-						FGraphBuilder::RegisterGraphWithFrontend(*AssetObject);
-					}
-				}
-				else
-				{
-					if (!MetaSoundEditorAsyncRegistrationEnabledCVar)
-					{
-						return;
-					}
-
-					if (AssetPrimeStatus == EAssetPrimeStatus::NotRequested || AssetPrimeStatus == EAssetPrimeStatus::Canceled)
-					{
-						return;
-					}
-
-					ActiveAsyncAssetLoadRequests++;
-
-					FSoftObjectPath AssetPath = InAssetData.ToSoftObjectPath();
-					auto LoadAndRegister = [this, ObjectPath = AssetPath, RegOptions](const FName& PackageName, UPackage* Package, EAsyncLoadingResult::Type Result)
-					{
-						if (AssetPrimeStatus == EAssetPrimeStatus::Canceled)
-						{
-							return;
-						}
-
-						if (Result == EAsyncLoadingResult::Succeeded)
-						{
-							UObject* MetaSoundObj = ObjectPath.ResolveObject();
-							FMetasoundAssetBase* MetaSoundAsset = IMetasoundUObjectRegistry::Get().GetObjectAsAssetBase(MetaSoundObj);
-							check(MetaSoundAsset);
-							if (!MetaSoundAsset->IsRegistered())
-							{
-								FGraphBuilder::RegisterGraphWithFrontend(*MetaSoundObj);
-							}
-						}
-
-						ActiveAsyncAssetLoadRequests--;
-						if (AssetPrimeStatus == EAssetPrimeStatus::InProgress && ActiveAsyncAssetLoadRequests == 0)
-						{
-							AssetPrimeStatus = EAssetPrimeStatus::Complete;
-						}
-					};
-					LoadPackageAsync(AssetPath.GetLongPackageName(), FLoadPackageAsyncDelegate::CreateLambda(LoadAndRegister));
-				}
-			}
-
-			void AddClassRegistryAsset(const FAssetData& InAssetData)
-			{
-				using namespace Frontend;
-
-				// If an object's class could not be found, ignore this asset.  This can hit for non-MetaSound assets
-				// and it is up to the system in charge of interacting with that asset or the loading behavior to
-				// report the failed load of the class.
-				if (const UClass* AssetClass = InAssetData.GetClass())
-				{
-					const bool bIsRegisteredClass = IMetasoundUObjectRegistry::Get().IsRegisteredClass(*AssetClass);
-					if (bIsRegisteredClass)
-					{
-						const FNodeRegistryKey RegistryKey = IMetaSoundAssetManager::GetChecked().AddOrUpdateAsset(InAssetData);
-
-						// Can be invalid if being called for the first time on an asset before its class name is generated
-						if (RegistryKey.IsValid())
-						{
-							const bool bPrimeRequested = AssetPrimeStatus > EAssetPrimeStatus::NotRequested;
-							const bool bIsRegistered = FMetasoundFrontendRegistryContainer::Get()->IsNodeRegistered(RegistryKey);
-							if (bPrimeRequested && !bIsRegistered)
-							{
-								LoadAndRegisterAsset(InAssetData);
-							}
-						}
-					}
-				}
-			}
-
-			void UpdateClassRegistryAsset(const FAssetData& InAssetData)
-			{
-				using namespace Frontend;
-
-				// If an object's class could not be found, ignore this asset.  This can hit for non-MetaSound assets
-				// and it is up to the system in charge of interacting with that asset or the loading behavior to
-				// report the failed load of the class.
-				if (const UClass* AssetClass = InAssetData.GetClass())
-				{
-					const bool bIsRegisteredClass = IMetasoundUObjectRegistry::Get().IsRegisteredClass(*AssetClass);
-					if (bIsRegisteredClass)
-					{
-						const FNodeRegistryKey RegistryKey = IMetaSoundAssetManager::GetChecked().AddOrUpdateAsset(InAssetData);
-						const bool bPrimeRequested = AssetPrimeStatus > EAssetPrimeStatus::NotRequested;
-						const bool bIsRegistered = FMetasoundFrontendRegistryContainer::Get()->IsNodeRegistered(RegistryKey);
-
-						// Have to re-register even if prime was not requested to avoid registry desync.
-						if (bPrimeRequested || bIsRegistered)
-						{
-							LoadAndRegisterAsset(InAssetData);
-						}
-					}	
-				}
-			}
-
-			void OnPackageReloaded(const EPackageReloadPhase InPackageReloadPhase, FPackageReloadedEvent* InPackageReloadedEvent)
-			{
-				using namespace Metasound;
-				using namespace Metasound::Editor;
-				using namespace Metasound::Frontend;
-
-				if (!InPackageReloadedEvent)
-				{
-					return;
-				}
-
-				if (InPackageReloadPhase != EPackageReloadPhase::OnPackageFixup)
-				{
-					return;
-				}
-
-				auto IsAssetMetaSound = [](const UObject* Obj)
-				{
-					check(Obj);
-					if (const UClass* AssetClass = Obj->GetClass())
-					{
-						return IMetasoundUObjectRegistry::Get().IsRegisteredClass(*AssetClass);
-					}
-					
-					return false;
-				};
-
-				for (const TPair<UObject*, UObject*>& Pair : InPackageReloadedEvent->GetRepointedObjects())
-				{
-					if (UObject* Obj = Pair.Key)
-					{
-						if (IsAssetMetaSound(Obj))
-						{
-							IMetaSoundAssetManager::GetChecked().RemoveAsset(*Pair.Key);
-
-							// Use the editor version of UnregisterWithFrontend so it refreshes any open MetaSound editors
-							FGraphBuilder::UnregisterGraphWithFrontend(*Pair.Key);
-						}
-					}
-
-					if (UObject* Obj = Pair.Value)
-					{
-						if (IsAssetMetaSound(Obj))
-						{
-							IMetaSoundAssetManager::GetChecked().AddOrUpdateAsset(*Pair.Value);
-
-							// Use the editor version of RegisterWithFrontend so it refreshes any open MetaSound editors
-							FGraphBuilder::RegisterGraphWithFrontend(*Pair.Value);
-						}
-					}
-				}
-			}
-
-			void OnAssetScanFinished()
-			{
-				if (IsRunningCookCommandlet())
-				{
-					return;
-				}
-
-				AssetScanStatus = EAssetScanStatus::Complete;
-
-				if (AssetPrimeStatus == EAssetPrimeStatus::Requested)
-				{
-					PrimeAssetRegistryAsync();
-				}
-
-				FAssetRegistryModule& AssetRegistryModule = FModuleManager::LoadModuleChecked<FAssetRegistryModule>("AssetRegistry");
-				AssetRegistryModule.Get().OnAssetAdded().AddRaw(this, &FModule::AddClassRegistryAsset);
-				AssetRegistryModule.Get().OnAssetUpdated().AddRaw(this, &FModule::UpdateClassRegistryAsset);
-				AssetRegistryModule.Get().OnAssetRemoved().AddRaw(this, &FModule::RemoveAssetFromClassRegistry);
-				AssetRegistryModule.Get().OnAssetRenamed().AddRaw(this, &FModule::RenameAssetInClassRegistry);
-
-				AssetRegistryModule.Get().OnFilesLoaded().RemoveAll(this);
-
-				FCoreUObjectDelegates::OnPackageReloaded.AddRaw(this, &FModule::OnPackageReloaded);
-			}
-
-			void RemoveAssetFromClassRegistry(const FAssetData& InAssetData)
-			{
-				using namespace Frontend;
-				if (const UClass* AssetClass = InAssetData.GetClass())
-				{
-					const bool bIsRegisteredClass = IMetasoundUObjectRegistry::Get().IsRegisteredClass(*AssetClass);
-					if (bIsRegisteredClass)
-					{
-						// Use the editor version of UnregisterWithFrontend so it refreshes any open MetaSound editors
-						// Doesn't use AssetData::GetAsset() as this can result in attempting to reload the object.
-						// If this call is hit after the asset is removed, the assumption is unregistration already
-						// occurred on object destroy.
-						if (UObject* AssetObject = InAssetData.GetSoftObjectPath().ResolveObject())
-						{
-							FGraphBuilder::UnregisterGraphWithFrontend(*AssetObject);
-						}
-
-						IMetaSoundAssetManager::GetChecked().RemoveAsset(InAssetData);
-					}
-				}
-			}
-
-			void RenameAssetInClassRegistry(const FAssetData& InAssetData, const FString& InOldObjectPath)
-			{
-				using namespace Frontend;
-
-				if (const UClass* AssetClass = InAssetData.GetClass())
-				{
-					const bool bIsRegisteredClass = IMetasoundUObjectRegistry::Get().IsRegisteredClass(*AssetClass);
-					if (bIsRegisteredClass)
-					{
-						IMetaSoundAssetManager& AssetManager = IMetaSoundAssetManager::GetChecked();
-
-						// Unregister using the new asset data even though the old object was last to be registered
-						// as the old asset is no longer accessible by the time rename is called. The asset at this
-						// point is identical however to its prior counterpart.
-						UObject* AssetObject = InAssetData.GetAsset();
-						check(AssetObject);
-
-						FMetasoundAssetBase* AssetBase = AssetManager.GetAsAsset(*AssetObject);
-						check(AssetBase);
-						bool bIsRegistered = AssetBase->IsRegistered();
-						if (bIsRegistered)
-						{
-							FGraphBuilder::UnregisterGraphWithFrontend(*AssetObject);
-						}
-
-						IMetaSoundAssetManager::GetChecked().RenameAsset(InAssetData, InOldObjectPath);
-
-						if (bIsRegistered)
-						{
-							constexpr bool bForceViewSynchronization = true;
-							FGraphBuilder::RegisterGraphWithFrontend(*AssetObject, bForceViewSynchronization);
-						}
-					}
-				}
-			}
-
 			void RegisterInputDefaultClasses()
 			{
 				TSubclassOf<UMetasoundEditorGraphMemberDefaultLiteral> NodeClass;
@@ -666,60 +397,17 @@ namespace Metasound
 
 			}
 
-			void ShutdownAssetClassRegistry()
-			{
-				if (FAssetRegistryModule* AssetRegistryModule = static_cast<FAssetRegistryModule*>(FModuleManager::Get().GetModule("AssetRegistry")))
-				{
-					AssetRegistryModule->Get().OnAssetAdded().RemoveAll(this);
-					AssetRegistryModule->Get().OnAssetUpdated().RemoveAll(this);
-					AssetRegistryModule->Get().OnAssetRemoved().RemoveAll(this);
-					AssetRegistryModule->Get().OnAssetRenamed().RemoveAll(this);
-					AssetRegistryModule->Get().OnFilesLoaded().RemoveAll(this);
-
-					FCoreUObjectDelegates::OnPackageReloaded.RemoveAll(this);
-				}
-			}
-
-			virtual void PrimeAssetRegistryAsync() override
-			{
-				// Ignore step if still loading assets from initial scan but set prime status as requested.
-				if (AssetScanStatus <= EAssetScanStatus::InProgress)
-				{
-					AssetPrimeStatus = EAssetPrimeStatus::Requested;
-					return;
-				}
-
-				if (AssetPrimeStatus != EAssetPrimeStatus::InProgress)
-				{
-					AssetPrimeStatus = EAssetPrimeStatus::InProgress;
-
-					TArray<FTopLevelAssetPath> ClassNames;
-					IMetasoundUObjectRegistry::Get().IterateRegisteredUClasses([&ClassNames](UClass& InClass)
-					{
-						ClassNames.Add(InClass.GetClassPathName());
-					});
-
-					FARFilter Filter;
-					Filter.ClassPaths = ClassNames;
-
-					FAssetRegistryModule& AssetRegistryModule = FModuleManager::LoadModuleChecked<FAssetRegistryModule>("AssetRegistry");
-					AssetRegistryModule.Get().EnumerateAssets(Filter, [this](const FAssetData& AssetData)
-					{
-						AddClassRegistryAsset(AssetData);
-						return true;
-					});
-				}
-			}
+			virtual void PrimeAssetRegistryAsync() override {}
 
 			virtual EAssetPrimeStatus GetAssetRegistryPrimeStatus() const override
 			{
-				return AssetPrimeStatus;
+				return EAssetPrimeStatus::NotRequested;
 			}
 
 			virtual EAssetScanStatus GetAssetRegistryScanStatus() const override
 			{
-				return AssetScanStatus;
-			}
+				return EAssetScanStatus::NotRequested;
+			}			
 
 			virtual void RegisterExplicitProxyClass(const UClass& InClass) override
 			{
@@ -815,6 +503,7 @@ namespace Metasound
 
 			virtual void StartupModule() override
 			{
+				using namespace Metasound::Engine;
 				METASOUND_LLM_SCOPE;
 				// Register Metasound asset type actions
 				IAssetTools& AssetTools = FModuleManager::LoadModuleChecked<FAssetToolsModule>(AssetToolName).Get();
@@ -888,21 +577,20 @@ namespace Metasound
 					GetMutableDefault<UMetasoundEditorSettings>()
 				);
 
-				FAssetRegistryModule& AssetRegistryModule = FModuleManager::LoadModuleChecked<FAssetRegistryModule>("AssetRegistry");
-				if (AssetRegistryModule.Get().IsLoadingAssets())
-				{
-					AssetScanStatus = EAssetScanStatus::InProgress;
-					AssetRegistryModule.Get().OnFilesLoaded().AddRaw(this, &FModule::OnAssetScanFinished);
-				}
-				else
-				{
-					AssetScanStatus = EAssetScanStatus::Complete;
-				}
-
 				// Metasound Engine registers USoundWave as a proxy class in the
 				// Metasound Frontend. The frontend registration must occur before
 				// the Metasound Editor registration of a USoundWave.
-				FModuleManager::LoadModuleChecked<IModuleInterface>("MetasoundEngine");
+				IMetasoundEngineModule& MetaSoundEngineModule = FModuleManager::LoadModuleChecked<IMetasoundEngineModule>("MetasoundEngine");
+				
+				// Bind delegates for MetaSound registration in the asset registry
+				MetaSoundEngineModule.GetOnGraphRegisteredDelegate().BindLambda([](UObject& InMetaSound, bool bForceViewSynchronization)
+				{
+					FGraphBuilder::RegisterGraphWithFrontend(InMetaSound, bForceViewSynchronization);
+				});
+				MetaSoundEngineModule.GetOnGraphUnregisteredDelegate().BindLambda([](UObject& InMetaSound)
+				{
+					FGraphBuilder::UnregisterGraphWithFrontend(InMetaSound);
+				});
 
 				// Required to ensure logic to order nodes for presets exclusive to
 				// editor is propagated to transform instances while editing in editor.
@@ -931,11 +619,6 @@ namespace Metasound
 			{
 				METASOUND_LLM_SCOPE;
 
-				if (AssetPrimeStatus != EAssetPrimeStatus::NotRequested && AssetPrimeStatus != EAssetPrimeStatus::Complete)
-				{
-					AssetPrimeStatus = EAssetPrimeStatus::Canceled;
-				}
-
 				if (ISettingsModule* SettingsModule = FModuleManager::GetModulePtr<ISettingsModule>("Settings"))
 				{
 					SettingsModule->UnregisterSettings("Editor", "Audio", "MetaSound Editor");
@@ -968,8 +651,6 @@ namespace Metasound
 					GraphPanelPinFactory.Reset();
 				}
 
-				ShutdownAssetClassRegistry();
-
 				AssetActions.Reset();
 				PinTypes.Reset();
 			}
@@ -1033,10 +714,6 @@ namespace Metasound
 			TSharedPtr<FSlateStyleSet> StyleSet;
 
 			TSet<const UClass*> ExplicitProxyClasses;
-
-			EAssetPrimeStatus AssetPrimeStatus = EAssetPrimeStatus::NotRequested;
-			EAssetScanStatus AssetScanStatus = EAssetScanStatus::NotRequested;
-			int32 ActiveAsyncAssetLoadRequests = 0;
 		};
 	} // namespace Editor
 } // namespace Metasound
diff --git a/Engine/Plugins/Runtime/Metasound/Source/MetasoundEditor/Public/MetasoundEditorModule.h b/Engine/Plugins/Runtime/Metasound/Source/MetasoundEditor/Public/MetasoundEditorModule.h
--- a/Engine/Plugins/Runtime/Metasound/Source/MetasoundEditor/Public/MetasoundEditorModule.h
+++ b/Engine/Plugins/Runtime/Metasound/Source/MetasoundEditor/Public/MetasoundEditorModule.h
@@ -18,6 +18,11 @@ class IDetailPropertyRow;
 class UEdGraphPin;
 class UMetasoundEditorGraph;
 class UMetasoundEditorGraphMemberDefaultLiteral;
+namespace Metasound::Engine 
+{
+	enum class EAssetScanStatus : uint8;
+	enum class ENodeClassRegistryPrimeStatus : uint8;
+}
 
 DECLARE_LOG_CATEGORY_EXTERN(LogMetasoundEditor, Log, All);
 
@@ -32,25 +37,8 @@ namespace Metasound
 			METASOUNDEDITOR_API const FSlateBrush& GetSlateBrushSafe(FName InName);
 		} // namespace Style
 
-		// Status of initial asset scan when editor loads up.
-		enum class EAssetScanStatus : uint8
-		{
-			NotRequested = 0,
-			InProgress = 2,
-			Complete = 3
-		};
-
-		// Primes status of MetaSound assets.  Priming an asset
-		// effectively loading the asset asynchronously (if not already loaded)
-		// & registers it with the MetaSound Class Registry.
-		enum class EAssetPrimeStatus : uint8
-		{
-			NotRequested = 0,
-			Requested = 1,
-			InProgress = 2,
-			Complete = 3,
-			Canceled = 4
-		};
+		using EAssetScanStatus = Metasound::Engine::EAssetScanStatus;
+		using EAssetPrimeStatus = Metasound::Engine::ENodeClassRegistryPrimeStatus;
 
 		struct FGraphPinParams
 		{
@@ -99,10 +87,13 @@ namespace Metasound
 			UE_DEPRECATED(5.3, "IsMetaSoundAssetClass is deprecated, use IMetasoundUObjectRegistry::IsRegisteredClass")
 			virtual bool IsMetaSoundAssetClass(const FTopLevelAssetPath& InClassName) const = 0;
 
-			// Primes MetaSound assets, effectively loading the asset asynchronously (if not already
-			// loaded) & registers them if not already registered with the MetaSound Class Registry.
+			UE_DEPRECATED(5.5, "Use PrimeAssetRegistryAsync in MetaSoundEngineModule.")
 			virtual void PrimeAssetRegistryAsync() = 0;
+
+			UE_DEPRECATED(5.5, "Use GetNodeClassRegistryPrimeStatus in MetaSoundEngineModule.")
 			virtual EAssetPrimeStatus GetAssetRegistryPrimeStatus() const = 0;
+
+			UE_DEPRECATED(5.5, "Use the same function in MetaSoundEngineModule.")
 			virtual EAssetScanStatus GetAssetRegistryScanStatus() const = 0;
 
 			virtual TUniquePtr<FMetasoundDefaultLiteralCustomizationBase> CreateMemberDefaultLiteralCustomization(UClass& InClass, IDetailCategoryBuilder& DefaultCategoryBuilder) const = 0;
diff --git a/Engine/Plugins/Runtime/Metasound/Source/MetasoundEngine/MetasoundEngine.Build.cs b/Engine/Plugins/Runtime/Metasound/Source/MetasoundEngine/MetasoundEngine.Build.cs
--- a/Engine/Plugins/Runtime/Metasound/Source/MetasoundEngine/MetasoundEngine.Build.cs
+++ b/Engine/Plugins/Runtime/Metasound/Source/MetasoundEngine/MetasoundEngine.Build.cs
@@ -21,7 +21,7 @@ namespace UnrealBuildTool.Rules
 			(
 				new string[]
 				{
-					"Core",
+					"AssetRegistry",
 					"AudioExtensions",
 					"AudioMixer",
 					"MetasoundGraphCore",
@@ -34,6 +34,7 @@ namespace UnrealBuildTool.Rules
 			(
 				new string[]
 				{
+					"Core",
 					"CoreUObject",
 					"DeveloperSettings",
 					"Engine",
diff --git a/Engine/Plugins/Runtime/Metasound/Source/MetasoundEngine/Private/Metasound.cpp b/Engine/Plugins/Runtime/Metasound/Source/MetasoundEngine/Private/Metasound.cpp
--- a/Engine/Plugins/Runtime/Metasound/Source/MetasoundEngine/Private/Metasound.cpp
+++ b/Engine/Plugins/Runtime/Metasound/Source/MetasoundEngine/Private/Metasound.cpp
@@ -118,6 +118,13 @@ void UMetaSoundPatch::PostEditUndo()
 	Super::PostEditUndo();
 	Metasound::Engine::FAssetHelper::PostEditUndo(*this);
 }
+
+EDataValidationResult UMetaSoundPatch::IsDataValid(FDataValidationContext& Context) const
+{
+	const EDataValidationResult Result = Metasound::Engine::FAssetHelper::IsDataValid(*this, Context);
+	return CombineDataValidationResults(Result, Super::IsDataValid(Context));
+}
+
 #endif // WITHEDITOR
 
 void UMetaSoundPatch::BeginDestroy()
diff --git a/Engine/Plugins/Runtime/Metasound/Source/MetasoundEngine/Private/MetasoundAssetSubsystem.cpp b/Engine/Plugins/Runtime/Metasound/Source/MetasoundEngine/Private/MetasoundAssetSubsystem.cpp
--- a/Engine/Plugins/Runtime/Metasound/Source/MetasoundEngine/Private/MetasoundAssetSubsystem.cpp
+++ b/Engine/Plugins/Runtime/Metasound/Source/MetasoundEngine/Private/MetasoundAssetSubsystem.cpp
@@ -160,6 +160,7 @@ namespace Metasound::Engine
 		virtual FMetasoundAssetBase* FindAsset(const FAssetKey& InKey) const override;
 		virtual TScriptInterface<IMetaSoundDocumentInterface> FindAssetAsDocumentInterface(const Frontend::FAssetKey& InKey) const override;
 		virtual const FTopLevelAssetPath* FindAssetPath(const FAssetKey& InKey) const override;
+		virtual const TArray<FTopLevelAssetPath>* FindAssetPaths(const FAssetKey& InKey) const override;
 		virtual FMetasoundAssetBase* GetAsAsset(UObject& InObject) const override;
 		virtual const FMetasoundAssetBase* GetAsAsset(const UObject& InObject) const override;
 #if WITH_EDITOR
@@ -438,6 +439,17 @@ namespace Metasound::Engine
 		return nullptr;
 	}
 
+	const TArray<FTopLevelAssetPath>* FMetaSoundAssetManager::FindAssetPaths(const Metasound::Frontend::FAssetKey& InKey) const
+	{
+		check(IsInGameThread());
+		if (const TArray<FTopLevelAssetPath>* Paths = PathMap.Find(InKey))
+		{
+			return Paths;
+		}
+
+		return { };
+	}
+
 	FMetasoundAssetBase* FMetaSoundAssetManager::GetAsAsset(UObject& InObject) const
 	{
 		return IMetasoundUObjectRegistry::Get().GetObjectAsAssetBase(&InObject);
diff --git a/Engine/Plugins/Runtime/Metasound/Source/MetasoundEngine/Private/MetasoundEngineAsset.h b/Engine/Plugins/Runtime/Metasound/Source/MetasoundEngine/Private/MetasoundEngineAsset.h
--- a/Engine/Plugins/Runtime/Metasound/Source/MetasoundEngine/Private/MetasoundEngineAsset.h
+++ b/Engine/Plugins/Runtime/Metasound/Source/MetasoundEngine/Private/MetasoundEngineAsset.h
@@ -5,21 +5,25 @@
 #include "Metasound.h"
 #include "MetasoundAssetManager.h"
 #include "MetasoundDocumentBuilderRegistry.h"
+#include "MetasoundEngineModule.h"
 #include "MetasoundFrontendDocumentIdGenerator.h"
 #include "MetasoundFrontendRegistryKey.h"
 #include "MetasoundGlobals.h"
 #include "MetasoundUObjectRegistry.h"
 #include "Misc/App.h"
+#include "Modules/ModuleManager.h"
 #include "Serialization/Archive.h"
 
 #if WITH_EDITORONLY_DATA
 #include "Algo/Transform.h"
 #include "MetasoundFrontendRegistryContainer.h"
+#include "Misc/DataValidation.h"
 #include "UObject/GarbageCollection.h"
 #include "UObject/ObjectMacros.h"
 #include "UObject/StrongObjectPtrTemplates.h"
 #endif // WITH_EDITORONLY_DATA
 
+#define LOCTEXT_NAMESPACE "MetasoundEngine"
 
 namespace Metasound::Engine
 {
@@ -99,6 +103,57 @@ namespace Metasound::Engine
 				}
 			}
 		}
+
+		template <typename TMetaSoundObject>
+		static EDataValidationResult IsClassNameUnique(const TMetaSoundObject& InMetaSound, FDataValidationContext& InOutContext)
+		{
+			using namespace Metasound::Frontend;
+			using namespace Metasound::Engine;
+
+			EDataValidationResult Result = EDataValidationResult::Valid;
+
+			// Need to prime asset registry to look for duplicate class names
+			IMetasoundEngineModule& MetaSoundEngineModule = FModuleManager::GetModuleChecked<IMetasoundEngineModule>("MetaSoundEngine");
+			// Checking for duplicate class names only requires the asset manager to be primed, but not for assets to be loaded. 
+			if (!MetaSoundEngineModule.IsAssetManagerPrimed())
+			{
+				MetaSoundEngineModule.PrimeAssetManager();
+				// Check again, as priming relies on the asset registry being loaded so may not be complete
+				if (!MetaSoundEngineModule.IsAssetManagerPrimed())
+				{
+					Result = EDataValidationResult::Invalid;
+					InOutContext.AddError(LOCTEXT("UniqueClassNameAssetManagerNotReady",
+						"MetaSound Asset Manager was unable to be primed to check for unique class names. This may be because the asset registry has not finished loading assets. Please try again later."));
+					return Result;
+				}
+			}
+
+			// Add error for multiple assets with the same class name
+			const IMetaSoundAssetManager& AssetManager = IMetaSoundAssetManager::GetChecked();
+			const FAssetKey Key(InMetaSound.GetConstDocument().RootGraph.Metadata);
+			const TArray<FTopLevelAssetPath>* AssetPaths = AssetManager.FindAssetPaths(Key);
+			if (AssetPaths && AssetPaths->Num() > 1)
+			{
+				Result = EDataValidationResult::Invalid;
+
+				TArray<FText> PathStrings;
+				Algo::Transform(*AssetPaths, PathStrings, [](const FTopLevelAssetPath& Path) { return FText::FromString(Path.ToString()); });
+				InOutContext.AddError(FText::Format(LOCTEXT("UniqueClassNameValidation",
+					"Multiple assets use the same class name which may result in unintended behavior. This may happen when an asset is moved, then the move is reverted in revision control without removing the newly created asset. Please remove the offending asset or duplicate it to automatically generate a new class name." \
+					"\nConflicting Asset Paths:\n{0}"), FText::Join(FText::FromString(TEXT("\n")), PathStrings)));
+			}
+
+			// Success
+			return Result;
+		}
+
+		template <typename TMetaSoundObject>
+		static EDataValidationResult IsDataValid(const TMetaSoundObject& InMetaSound, FDataValidationContext& InOutContext)
+		{
+			const EDataValidationResult Result = IsClassNameUnique(InMetaSound, InOutContext);
+			return Result;
+		}
+
 #endif // WITH_EDITOR
 
 		template <typename TMetaSoundObject>
@@ -277,3 +332,4 @@ namespace Metasound::Engine
 #endif // WITH_EDITORONLY_DATA
 	};
 } // namespace Metasound::Engine
+#undef LOCTEXT_NAMESPACE // MetasoundEngine
diff --git a/Engine/Plugins/Runtime/Metasound/Source/MetasoundEngine/Private/MetasoundEngineModule.cpp b/Engine/Plugins/Runtime/Metasound/Source/MetasoundEngine/Private/MetasoundEngineModule.cpp
--- a/Engine/Plugins/Runtime/Metasound/Source/MetasoundEngine/Private/MetasoundEngineModule.cpp
+++ b/Engine/Plugins/Runtime/Metasound/Source/MetasoundEngine/Private/MetasoundEngineModule.cpp
@@ -1,5 +1,9 @@
 // Copyright Epic Games, Inc. All Rights Reserved.
 
+#include "MetasoundEngineModule.h"
+
+#include "AssetRegistry/AssetData.h"
+#include "AssetRegistry/AssetRegistryModule.h"
 #include "Metasound.h"
 #include "MetasoundAssetSubsystem.h"
 #include "MetasoundAudioBus.h"
@@ -9,8 +13,10 @@
 #include "MetasoundDocumentBuilderRegistry.h"
 #include "MetasoundDocumentInterface.h"
 #include "MetasoundFrontendDocumentBuilder.h"
+#include "MetasoundGlobals.h"
 #include "MetasoundLog.h"
 #include "MetasoundOutputSubsystem.h"
+#include "MetasoundSettings.h"
 #include "MetasoundSource.h"
 #include "MetasoundTrace.h"
 #include "MetasoundUObjectRegistry.h"
@@ -29,7 +35,6 @@
 #include "Modules/ModuleManager.h"
 #include "Sound/AudioSettings.h"
 
-
 namespace Metasound
 {
 	// Enable send/receive node registration for data types which existed before
@@ -46,132 +51,474 @@ REGISTER_METASOUND_DATATYPE(Metasound::FWaveAsset, "WaveAsset", Metasound::ELite
 REGISTER_METASOUND_DATATYPE(WaveTable::FWaveTable, "WaveTable", Metasound::ELiteralType::FloatArray)
 REGISTER_METASOUND_DATATYPE(Metasound::FWaveTableBankAsset, "WaveTableBankAsset", Metasound::ELiteralType::UObjectProxy, UWaveTableBank);
 
-
 namespace Metasound::Engine
 {
-	class FModule : public IModuleInterface
+	void FModule::FModule::StartupModule() 
+	{
+		using namespace Frontend;
+
+		METASOUND_LLM_SCOPE;
+		FModuleManager::Get().LoadModuleChecked("MetasoundGraphCore");
+		FModuleManager::Get().LoadModuleChecked("MetasoundFrontend");
+		FModuleManager::Get().LoadModuleChecked("MetasoundStandardNodes");
+		FModuleManager::Get().LoadModuleChecked("MetasoundGenerator");
+		FModuleManager::Get().LoadModuleChecked("WaveTable");
+		
+		InitializeAssetManager();
+		IDocumentBuilderRegistry::Initialize(MakeUnique<FDocumentBuilderRegistry>());
+
+		// Set GCObject referencer for metasound frontend node registry. The MetaSound
+		// frontend does not have access to Engine GC tools and must have them 
+		// supplied externally.
+		FMetasoundFrontendRegistryContainer::Get()->SetObjectReferencer(MakeUnique<FObjectReferencer>());
+
+		// Register engine-level parameter interfaces if not done already.
+		// (Potentially not already called if plugin is loaded while cooking.)
+		UAudioSettings* AudioSettings = GetMutableDefault<UAudioSettings>();
+		check(AudioSettings);
+		AudioSettings->RegisterParameterInterfaces();
+
+		IMetasoundUObjectRegistry::Get().RegisterUClass(MakeUnique<TMetasoundUObjectRegistryEntry<UMetaSoundBuilderDocument>>());
+		IMetasoundUObjectRegistry::Get().RegisterUClass(MakeUnique<TMetasoundUObjectRegistryEntry<UMetaSoundPatch>>());
+		IMetasoundUObjectRegistry::Get().RegisterUClass(MakeUnique<TMetasoundUObjectRegistryEntry<UMetaSoundSource>>());
+
+		RegisterDeprecatedInterfaces();
+		RegisterInterfaces();
+		RegisterInternalInterfaceBindings();
+
+		// Flush node registration queue
+		FMetasoundFrontendRegistryContainer::Get()->RegisterPendingNodes();
+
+		// Register Analyzers
+		// TODO: Determine if we can move this registration to Frontend where it likely belongs
+		METASOUND_REGISTER_VERTEX_ANALYZER_FACTORY(Frontend::FVertexAnalyzerAudioBuffer)
+		METASOUND_REGISTER_VERTEX_ANALYZER_FACTORY(Frontend::FVertexAnalyzerEnvelopeFollower)
+		METASOUND_REGISTER_VERTEX_ANALYZER_FACTORY(Frontend::FVertexAnalyzerForwardBool)
+		METASOUND_REGISTER_VERTEX_ANALYZER_FACTORY(Frontend::FVertexAnalyzerForwardFloat)
+		METASOUND_REGISTER_VERTEX_ANALYZER_FACTORY(Frontend::FVertexAnalyzerForwardInt)
+		METASOUND_REGISTER_VERTEX_ANALYZER_FACTORY(Frontend::FVertexAnalyzerForwardTime)
+		METASOUND_REGISTER_VERTEX_ANALYZER_FACTORY(Frontend::FVertexAnalyzerForwardString)
+		METASOUND_REGISTER_VERTEX_ANALYZER_FACTORY(Frontend::FVertexAnalyzerTriggerDensity)
+		METASOUND_REGISTER_VERTEX_ANALYZER_FACTORY(Frontend::FVertexAnalyzerTriggerToTime)
+		METASOUND_REGISTER_VERTEX_ANALYZER_FACTORY(Engine::FVertexAnalyzerAudioBusWriter)
+
+		// Register passthrough output analyzers
+		UMetasoundGeneratorHandle::RegisterPassthroughAnalyzerForType(
+			GetMetasoundDataTypeName<float>(),
+			Frontend::FVertexAnalyzerForwardFloat::GetAnalyzerName(),
+			Frontend::FVertexAnalyzerForwardFloat::FOutputs::GetValue().Name);
+		UMetasoundGeneratorHandle::RegisterPassthroughAnalyzerForType(
+			GetMetasoundDataTypeName<int32>(),
+			Frontend::FVertexAnalyzerForwardInt::GetAnalyzerName(),
+			Frontend::FVertexAnalyzerForwardInt::FOutputs::GetValue().Name);
+		UMetasoundGeneratorHandle::RegisterPassthroughAnalyzerForType(
+			GetMetasoundDataTypeName<bool>(),
+			Frontend::FVertexAnalyzerForwardBool::GetAnalyzerName(),
+			Frontend::FVertexAnalyzerForwardBool::FOutputs::GetValue().Name);
+		UMetasoundGeneratorHandle::RegisterPassthroughAnalyzerForType(
+			GetMetasoundDataTypeName<FString>(),
+			Frontend::FVertexAnalyzerForwardString::GetAnalyzerName(),
+			Frontend::FVertexAnalyzerForwardString::FOutputs::GetValue().Name);
+		UMetasoundGeneratorHandle::RegisterPassthroughAnalyzerForType(
+			GetMetasoundDataTypeName<FTime>(),
+			Frontend::FVertexAnalyzerForwardTime::GetAnalyzerName(),
+			Frontend::FVertexAnalyzerForwardTime::FOutputs::GetValue().Name);
+		UMetasoundGeneratorHandle::RegisterPassthroughAnalyzerForType(
+			GetMetasoundDataTypeName<FTrigger>(),
+			Frontend::FVertexAnalyzerTriggerToTime::GetAnalyzerName(),
+			Frontend::FVertexAnalyzerTriggerToTime::FOutputs::GetValue().Name);
+#if WITH_EDITOR
+		FAssetRegistryModule& AssetRegistryModule = FModuleManager::LoadModuleChecked<FAssetRegistryModule>("AssetRegistry");
+		if (AssetRegistryModule.Get().IsLoadingAssets())
+		{
+			AssetScanStatus = EAssetScanStatus::InProgress;
+			AssetRegistryModule.Get().OnFilesLoaded().AddRaw(this, &FModule::OnAssetScanFinished);
+		}
+		else
+		{
+			AssetScanStatus = EAssetScanStatus::Complete;
+		}
+#endif // WITH_EDITOR
+
+		UE_LOG(LogMetaSound, Log, TEXT("MetaSound Engine Initialized"));
+	}
+
+	void FModule::FModule::ShutdownModule()
 	{
-		// Supplies GC referencing in the MetaSound Frontend node registry for doing
-		// async work on UObjets
-		class FObjectReferencer 
-			: public FMetasoundFrontendRegistryContainer::IObjectReferencer
-			, public FGCObject
+#if WITH_EDITOR
+		if (NodeClassRegistryPrimeStatus != ENodeClassRegistryPrimeStatus::NotRequested && NodeClassRegistryPrimeStatus != ENodeClassRegistryPrimeStatus::Complete)
+		{
+			NodeClassRegistryPrimeStatus = ENodeClassRegistryPrimeStatus::Canceled;
+		}
+
+		ShutdownAssetClassRegistry();
+#endif // WITH_EDITOR
+		DeinitializeAssetManager();
+		Frontend::IDocumentBuilderRegistry::Deinitialize();
+	}
+
+#if WITH_EDITOR
+	void FModule::LoadAndRegisterAsset(const FAssetData& InAssetData)
+	{
+		// Ignore requests if graphs cannot be executed, as registration that results in IGraph generation is not supported.
+		if (!Metasound::CanEverExecuteGraph())
+		{
+			return;
+		}
+
+		Frontend::FMetaSoundAssetRegistrationOptions RegOptions;
+		RegOptions.bForceReregister = false;
+		if (const UMetaSoundSettings* Settings = GetDefault<UMetaSoundSettings>())
+		{
+			RegOptions.bAutoUpdateLogWarningOnDroppedConnection = Settings->bAutoUpdateLogWarningOnDroppedConnection;
+		}
+
+		if (InAssetData.IsAssetLoaded())
+		{
+			if (UObject* AssetObject = InAssetData.GetAsset())
+			{
+				OnGraphRegister.ExecuteIfBound(*AssetObject, /*bInForceViewSynchronization=*/false);
+			}
+		}
+		else
 		{
-		public:
-			virtual void AddObject(UObject* InObject) override
+			if (NodeClassRegistryPrimeStatus == ENodeClassRegistryPrimeStatus::NotRequested || NodeClassRegistryPrimeStatus == ENodeClassRegistryPrimeStatus::Canceled)
 			{
-				FScopeLock LockObjectArray(&ObjectArrayCriticalSection);
-				ObjectArray.Add(InObject);
+				return;
 			}
 
-			virtual void RemoveObject(UObject* InObject) override
+			ActiveAsyncAssetLoadRequests++;
+
+			FSoftObjectPath AssetPath = InAssetData.ToSoftObjectPath();
+			auto LoadAndRegister = [this, ObjectPath = AssetPath, RegOptions](const FName& PackageName, UPackage* Package, EAsyncLoadingResult::Type Result)
+			{
+				if (NodeClassRegistryPrimeStatus == ENodeClassRegistryPrimeStatus::Canceled)
+				{
+					return;
+				}
+
+				if (Result == EAsyncLoadingResult::Succeeded)
+				{
+					UObject* MetaSoundObj = ObjectPath.ResolveObject();
+					FMetasoundAssetBase* MetaSoundAsset = IMetasoundUObjectRegistry::Get().GetObjectAsAssetBase(MetaSoundObj);
+					check(MetaSoundAsset);
+					if (!MetaSoundAsset->IsRegistered())
+					{
+						OnGraphRegister.ExecuteIfBound(*MetaSoundObj, /*bInForceViewSynchronization=*/false);
+					}
+				}
+
+				ActiveAsyncAssetLoadRequests--;
+				if (NodeClassRegistryPrimeStatus == ENodeClassRegistryPrimeStatus::InProgress && ActiveAsyncAssetLoadRequests == 0)
+				{
+					NodeClassRegistryPrimeStatus = ENodeClassRegistryPrimeStatus::Complete;
+				}
+			};
+
+			LoadPackageAsync(AssetPath.GetLongPackageName(), FLoadPackageAsyncDelegate::CreateLambda(LoadAndRegister));
+		}
+	}
+
+	void FModule::AddClassRegistryAsset(const FAssetData& InAssetData)
+	{
+		using namespace Frontend;
+
+		// If an object's class could not be found, ignore this asset.  This can hit for non-MetaSound assets
+		// and it is up to the system in charge of interacting with that asset or the loading behavior to
+		// report the failed load of the class.
+		if (const UClass* AssetClass = InAssetData.GetClass())
+		{
+			const bool bIsRegisteredClass = IMetasoundUObjectRegistry::Get().IsRegisteredClass(*AssetClass);
+			if (bIsRegisteredClass)
+			{
+				const FNodeRegistryKey RegistryKey = IMetaSoundAssetManager::GetChecked().AddOrUpdateAsset(InAssetData);
+
+				// Can be invalid if being called for the first time on an asset before its class name is generated
+				if (RegistryKey.IsValid())
+				{
+					const bool bPrimeRequested = NodeClassRegistryPrimeStatus > ENodeClassRegistryPrimeStatus::NotRequested;
+					const bool bIsRegistered = FMetasoundFrontendRegistryContainer::Get()->IsNodeRegistered(RegistryKey);
+					if (bPrimeRequested && !bIsRegistered)
+					{
+						LoadAndRegisterAsset(InAssetData);
+					}
+				}
+			}
+		}
+	}
+	
+	void FModule::UpdateClassRegistryAsset(const FAssetData& InAssetData)
+	{
+		using namespace Frontend;
+
+		// If an object's class could not be found, ignore this asset.  This can hit for non-MetaSound assets
+		// and it is up to the system in charge of interacting with that asset or the loading behavior to
+		// report the failed load of the class.
+		if (const UClass* AssetClass = InAssetData.GetClass())
+		{
+			const bool bIsRegisteredClass = IMetasoundUObjectRegistry::Get().IsRegisteredClass(*AssetClass);
+			if (bIsRegisteredClass)
+			{
+				const FNodeRegistryKey RegistryKey = IMetaSoundAssetManager::GetChecked().AddOrUpdateAsset(InAssetData);
+				const bool bPrimeRequested = NodeClassRegistryPrimeStatus > ENodeClassRegistryPrimeStatus::NotRequested;
+				const bool bIsRegistered = FMetasoundFrontendRegistryContainer::Get()->IsNodeRegistered(RegistryKey);
+
+				// Have to re-register even if prime was not requested to avoid registry desync.
+				if (bPrimeRequested || bIsRegistered)
+				{
+					LoadAndRegisterAsset(InAssetData);
+				}
+			}
+		}
+	}
+
+	void FModule::OnPackageReloaded(const EPackageReloadPhase InPackageReloadPhase, FPackageReloadedEvent* InPackageReloadedEvent)
+	{
+		using namespace Metasound;
+		using namespace Metasound::Frontend;
+
+		if (!InPackageReloadedEvent)
+		{
+			return;
+		}
+
+		if (InPackageReloadPhase != EPackageReloadPhase::OnPackageFixup)
+		{
+			return;
+		}
+
+		auto IsAssetMetaSound = [](const UObject* Obj)
+			{
+				check(Obj);
+				if (const UClass* AssetClass = Obj->GetClass())
+				{
+					return IMetasoundUObjectRegistry::Get().IsRegisteredClass(*AssetClass);
+				}
+
+				return false;
+			};
+
+		for (const TPair<UObject*, UObject*>& Pair : InPackageReloadedEvent->GetRepointedObjects())
+		{
+			if (UObject* Obj = Pair.Key)
 			{
-				FScopeLock LockObjectArray(&ObjectArrayCriticalSection);
-				ObjectArray.Remove(InObject);
+				if (IsAssetMetaSound(Obj))
+				{
+					IMetaSoundAssetManager::GetChecked().RemoveAsset(*Pair.Key);
+
+					// Use the editor version of UnregisterWithFrontend so it refreshes any open MetaSound editors
+					OnGraphUnregister.ExecuteIfBound(*Obj);
+				}
 			}
 
-			virtual void AddReferencedObjects(FReferenceCollector& Collector) override
+			if (UObject* Obj = Pair.Value)
 			{
-				FScopeLock LockObjectArray(&ObjectArrayCriticalSection);
-				Collector.AddReferencedObjects(ObjectArray);
+				if (IsAssetMetaSound(Obj))
+				{
+					IMetaSoundAssetManager::GetChecked().AddOrUpdateAsset(*Pair.Value);
+
+					// Use the editor version of RegisterWithFrontend so it refreshes any open MetaSound editors
+					OnGraphRegister.ExecuteIfBound(*Obj, /*bInForceViewSynchronization=*/false);
+				}
 			}
+		}
+	}
+
+	void FModule::OnAssetScanFinished()
+	{
+		if (IsRunningCookCommandlet())
+		{
+			return;
+		}
+
+		AssetScanStatus = EAssetScanStatus::Complete;
 
-			virtual FString GetReferencerName() const override
+		if (NodeClassRegistryPrimeStatus == ENodeClassRegistryPrimeStatus::Requested)
+		{
+			PrimeAssetRegistryAsync();
+		}
+
+		FAssetRegistryModule& AssetRegistryModule = FModuleManager::LoadModuleChecked<FAssetRegistryModule>("AssetRegistry");
+		AssetRegistryModule.Get().OnAssetAdded().AddRaw(this, &FModule::AddClassRegistryAsset);
+		AssetRegistryModule.Get().OnAssetUpdated().AddRaw(this, &FModule::UpdateClassRegistryAsset);
+		AssetRegistryModule.Get().OnAssetRemoved().AddRaw(this, &FModule::RemoveAssetFromClassRegistry);
+		AssetRegistryModule.Get().OnAssetRenamed().AddRaw(this, &FModule::RenameAssetInClassRegistry);
+
+		AssetRegistryModule.Get().OnFilesLoaded().RemoveAll(this);
+
+		FCoreUObjectDelegates::OnPackageReloaded.AddRaw(this, &FModule::OnPackageReloaded);
+	}
+
+	void FModule::RemoveAssetFromClassRegistry(const FAssetData& InAssetData)
+	{
+		using namespace Frontend;
+		if (const UClass* AssetClass = InAssetData.GetClass())
+		{
+			const bool bIsRegisteredClass = IMetasoundUObjectRegistry::Get().IsRegisteredClass(*AssetClass);
+			if (bIsRegisteredClass)
 			{
-				return TEXT("FMetasoundEngineModule::FObjectReferencer");
+				// Use the editor version of UnregisterWithFrontend so it refreshes any open MetaSound editors
+				// Doesn't use AssetData::GetAsset() as this can result in attempting to reload the object.
+				// If this call is hit after the asset is removed, the assumption is unregistration already
+				// occurred on object destroy.
+				if (UObject* AssetObject = InAssetData.GetSoftObjectPath().ResolveObject())
+				{
+					OnGraphUnregister.ExecuteIfBound(*AssetObject);
+				}
+
+				IMetaSoundAssetManager::GetChecked().RemoveAsset(InAssetData);
 			}
+		}
+	}
 
-		private:
-			mutable FCriticalSection ObjectArrayCriticalSection;
-			TArray<TObjectPtr<UObject>> ObjectArray;
-		};
+	void FModule::RenameAssetInClassRegistry(const FAssetData& InAssetData, const FString& InOldObjectPath)
+	{
+		using namespace Frontend;
 
-	public:
-		virtual void StartupModule() override
+		if (const UClass* AssetClass = InAssetData.GetClass())
 		{
-			using namespace Frontend;
+			const bool bIsRegisteredClass = IMetasoundUObjectRegistry::Get().IsRegisteredClass(*AssetClass);
+			if (bIsRegisteredClass)
+			{
+				IMetaSoundAssetManager& AssetManager = IMetaSoundAssetManager::GetChecked();
+
+				// Unregister using the new asset data even though the old object was last to be registered
+				// as the old asset is no longer accessible by the time rename is called. The asset at this
+				// point is identical however to its prior counterpart.
+				UObject* AssetObject = InAssetData.GetAsset();
+				check(AssetObject);
 
-			METASOUND_LLM_SCOPE;
-			FModuleManager::Get().LoadModuleChecked("MetasoundGraphCore");
-			FModuleManager::Get().LoadModuleChecked("MetasoundFrontend");
-			FModuleManager::Get().LoadModuleChecked("MetasoundStandardNodes");
-			FModuleManager::Get().LoadModuleChecked("MetasoundGenerator");
-			FModuleManager::Get().LoadModuleChecked("WaveTable");
+				FMetasoundAssetBase* AssetBase = AssetManager.GetAsAsset(*AssetObject);
+				check(AssetBase);
+				bool bIsRegistered = AssetBase->IsRegistered();
+				if (bIsRegistered)
+				{
+					OnGraphUnregister.ExecuteIfBound(*AssetObject);
+				}
+
+				IMetaSoundAssetManager::GetChecked().RenameAsset(InAssetData, InOldObjectPath);
+
+				if (bIsRegistered)
+				{
+					constexpr bool bForceViewSynchronization = true;
+					OnGraphRegister.ExecuteIfBound(*AssetObject, bForceViewSynchronization);
+				}
+			}
+		}
+	}
+
+	void FModule::ShutdownAssetClassRegistry()
+	{
+		if (FAssetRegistryModule* AssetRegistryModule = static_cast<FAssetRegistryModule*>(FModuleManager::Get().GetModule("AssetRegistry")))
+		{
+			AssetRegistryModule->Get().OnAssetAdded().RemoveAll(this);
+			AssetRegistryModule->Get().OnAssetUpdated().RemoveAll(this);
+			AssetRegistryModule->Get().OnAssetRemoved().RemoveAll(this);
+			AssetRegistryModule->Get().OnAssetRenamed().RemoveAll(this);
+			AssetRegistryModule->Get().OnFilesLoaded().RemoveAll(this);
+
+			FCoreUObjectDelegates::OnPackageReloaded.RemoveAll(this);
+		}
+	}
+
+	void FModule::PrimeAssetRegistryAsync() 
+	{
+		// Ignore step if still loading assets from initial scan but set prime status as requested.
+		if (AssetScanStatus <= EAssetScanStatus::InProgress)
+		{
+			NodeClassRegistryPrimeStatus = ENodeClassRegistryPrimeStatus::Requested;
+			return;
+		}
 		
-			InitializeAssetManager();
-			IDocumentBuilderRegistry::Initialize(MakeUnique<FDocumentBuilderRegistry>());
-
-			// Set GCObject referencer for metasound frontend node registry. The MetaSound
-			// frontend does not have access to Engine GC tools and must have them 
-			// supplied externally.
-			FMetasoundFrontendRegistryContainer::Get()->SetObjectReferencer(MakeUnique<FObjectReferencer>());
-
-			// Register engine-level parameter interfaces if not done already.
-			// (Potentially not already called if plugin is loaded while cooking.)
-			UAudioSettings* AudioSettings = GetMutableDefault<UAudioSettings>();
-			check(AudioSettings);
-			AudioSettings->RegisterParameterInterfaces();
-
-			IMetasoundUObjectRegistry::Get().RegisterUClass(MakeUnique<TMetasoundUObjectRegistryEntry<UMetaSoundBuilderDocument>>());
-			IMetasoundUObjectRegistry::Get().RegisterUClass(MakeUnique<TMetasoundUObjectRegistryEntry<UMetaSoundPatch>>());
-			IMetasoundUObjectRegistry::Get().RegisterUClass(MakeUnique<TMetasoundUObjectRegistryEntry<UMetaSoundSource>>());
-
-			RegisterDeprecatedInterfaces();
-			RegisterInterfaces();
-			RegisterInternalInterfaceBindings();
-
-			// Flush node registration queue
-			FMetasoundFrontendRegistryContainer::Get()->RegisterPendingNodes();
-
-			// Register Analyzers
-			// TODO: Determine if we can move this registration to Frontend where it likely belongs
-			METASOUND_REGISTER_VERTEX_ANALYZER_FACTORY(Frontend::FVertexAnalyzerAudioBuffer)
-			METASOUND_REGISTER_VERTEX_ANALYZER_FACTORY(Frontend::FVertexAnalyzerEnvelopeFollower)
-			METASOUND_REGISTER_VERTEX_ANALYZER_FACTORY(Frontend::FVertexAnalyzerForwardBool)
-			METASOUND_REGISTER_VERTEX_ANALYZER_FACTORY(Frontend::FVertexAnalyzerForwardFloat)
-			METASOUND_REGISTER_VERTEX_ANALYZER_FACTORY(Frontend::FVertexAnalyzerForwardInt)
-			METASOUND_REGISTER_VERTEX_ANALYZER_FACTORY(Frontend::FVertexAnalyzerForwardTime)
-			METASOUND_REGISTER_VERTEX_ANALYZER_FACTORY(Frontend::FVertexAnalyzerForwardString)
-			METASOUND_REGISTER_VERTEX_ANALYZER_FACTORY(Frontend::FVertexAnalyzerTriggerDensity)
-			METASOUND_REGISTER_VERTEX_ANALYZER_FACTORY(Frontend::FVertexAnalyzerTriggerToTime)
-			METASOUND_REGISTER_VERTEX_ANALYZER_FACTORY(Engine::FVertexAnalyzerAudioBusWriter)
-
-			// Register passthrough output analyzers
-			UMetasoundGeneratorHandle::RegisterPassthroughAnalyzerForType(
-				GetMetasoundDataTypeName<float>(),
-				Frontend::FVertexAnalyzerForwardFloat::GetAnalyzerName(),
-				Frontend::FVertexAnalyzerForwardFloat::FOutputs::GetValue().Name);
-			UMetasoundGeneratorHandle::RegisterPassthroughAnalyzerForType(
-				GetMetasoundDataTypeName<int32>(),
-				Frontend::FVertexAnalyzerForwardInt::GetAnalyzerName(),
-				Frontend::FVertexAnalyzerForwardInt::FOutputs::GetValue().Name);
-			UMetasoundGeneratorHandle::RegisterPassthroughAnalyzerForType(
-				GetMetasoundDataTypeName<bool>(),
-				Frontend::FVertexAnalyzerForwardBool::GetAnalyzerName(),
-				Frontend::FVertexAnalyzerForwardBool::FOutputs::GetValue().Name);
-			UMetasoundGeneratorHandle::RegisterPassthroughAnalyzerForType(
-				GetMetasoundDataTypeName<FString>(),
-				Frontend::FVertexAnalyzerForwardString::GetAnalyzerName(),
-				Frontend::FVertexAnalyzerForwardString::FOutputs::GetValue().Name);
-			UMetasoundGeneratorHandle::RegisterPassthroughAnalyzerForType(
-				GetMetasoundDataTypeName<FTime>(),
-				Frontend::FVertexAnalyzerForwardTime::GetAnalyzerName(),
-				Frontend::FVertexAnalyzerForwardTime::FOutputs::GetValue().Name);
-			UMetasoundGeneratorHandle::RegisterPassthroughAnalyzerForType(
-				GetMetasoundDataTypeName<FTrigger>(),
-				Frontend::FVertexAnalyzerTriggerToTime::GetAnalyzerName(),
-				Frontend::FVertexAnalyzerTriggerToTime::FOutputs::GetValue().Name);
-
-
-			UE_LOG(LogMetaSound, Log, TEXT("MetaSound Engine Initialized"));
-		}
-
-		virtual void ShutdownModule() override
-		{
-			DeinitializeAssetManager();
-			Frontend::IDocumentBuilderRegistry::Deinitialize();
+		// Prime both asset manager and node class registry
+		if (NodeClassRegistryPrimeStatus != ENodeClassRegistryPrimeStatus::InProgress)
+		{
+			NodeClassRegistryPrimeStatus = ENodeClassRegistryPrimeStatus::InProgress;
+
+			TArray<FTopLevelAssetPath> ClassNames;
+			IMetasoundUObjectRegistry::Get().IterateRegisteredUClasses([&ClassNames](UClass& InClass)
+			{
+				ClassNames.Add(InClass.GetClassPathName());
+			});
+
+			FARFilter Filter;
+			Filter.ClassPaths = ClassNames;
+
+			FAssetRegistryModule& AssetRegistryModule = FModuleManager::LoadModuleChecked<FAssetRegistryModule>("AssetRegistry");
+			uint32 AssetCount = 0;
+			AssetRegistryModule.Get().EnumerateAssets(Filter, [this, &AssetCount](const FAssetData& AssetData)
+			{
+				AddClassRegistryAsset(AssetData);
+				AssetCount++;
+				return true;
+			});
+
+			// Node class registry prime is complete if there are no assets to process
+			if (AssetCount == 0 || ActiveAsyncAssetLoadRequests == 0)
+			{
+				NodeClassRegistryPrimeStatus = ENodeClassRegistryPrimeStatus::Complete;
+			}
+			// Asset manager prime also occurred as part of AddClassRegistryAsset
+			bAssetManagerPrimed = true;
 		}
-	};
+	}
+
+	void FModule::PrimeAssetManager()
+	{
+		using namespace Frontend;
+		if (AssetScanStatus <= EAssetScanStatus::InProgress)
+		{
+			UE_LOG(LogMetaSound, Display, TEXT("MetaSound Asset Manager prime requested before Asset Registry scan completed."));
+			return;
+		}
+
+		if (!IsAssetManagerPrimed())
+		{
+			TArray<FTopLevelAssetPath> ClassNames;
+			IMetasoundUObjectRegistry::Get().IterateRegisteredUClasses([&ClassNames](UClass& InClass)
+			{
+				ClassNames.Add(InClass.GetClassPathName());
+			});
+
+			FARFilter Filter;
+			Filter.ClassPaths = ClassNames;
+
+			FAssetRegistryModule& AssetRegistryModule = FModuleManager::LoadModuleChecked<FAssetRegistryModule>("AssetRegistry");
+			AssetRegistryModule.Get().EnumerateAssets(Filter, [this](const FAssetData& AssetData)
+			{
+				IMetaSoundAssetManager::GetChecked().AddOrUpdateAsset(AssetData);
+				return true;
+			});
+
+			bAssetManagerPrimed = true;
+		}
+	}
+
+	ENodeClassRegistryPrimeStatus FModule::GetNodeClassRegistryPrimeStatus() const
+	{
+		return NodeClassRegistryPrimeStatus;
+	}
+
+	EAssetScanStatus FModule::GetAssetRegistryScanStatus() const
+	{
+		return AssetScanStatus;
+	}
+
+	bool FModule::IsAssetManagerPrimed() const
+	{
+		return bAssetManagerPrimed;
+	}
+
+	FOnMetasoundGraphRegister& FModule::GetOnGraphRegisteredDelegate()
+	{
+		return OnGraphRegister;
+	}
+
+	FOnMetasoundGraphUnregister& FModule::GetOnGraphUnregisteredDelegate()
+	{
+		return OnGraphUnregister;
+	}
+#endif // WITH_EDITOR
 } // namespace Metasound::Engine
 
 IMPLEMENT_MODULE(Metasound::Engine::FModule, MetasoundEngine);
diff --git a/Engine/Plugins/Runtime/Metasound/Source/MetasoundEngine/Private/MetasoundEngineModule.h b/Engine/Plugins/Runtime/Metasound/Source/MetasoundEngine/Private/MetasoundEngineModule.h
new file mode 100644
--- /dev/null
+++ b/Engine/Plugins/Runtime/Metasound/Source/MetasoundEngine/Private/MetasoundEngineModule.h
@@ -0,0 +1,87 @@
+// Copyright Epic Games, Inc. All Rights Reserved.
+
+#pragma once
+
+#include "IMetasoundEngineModule.h"
+#include "MetasoundFrontendRegistryContainer.h"
+#include "UObject/GCObject.h"
+
+namespace Metasound::Engine
+{
+	class FModule : public IMetasoundEngineModule
+	{
+		// Supplies GC referencing in the MetaSound Frontend node registry for doing
+		// async work on UObjets
+		class FObjectReferencer
+			: public FMetasoundFrontendRegistryContainer::IObjectReferencer
+			, public FGCObject
+		{
+		public:
+			virtual void AddObject(UObject* InObject) override
+			{
+				FScopeLock LockObjectArray(&ObjectArrayCriticalSection);
+				ObjectArray.Add(InObject);
+			}
+
+			virtual void RemoveObject(UObject* InObject) override
+			{
+				FScopeLock LockObjectArray(&ObjectArrayCriticalSection);
+				ObjectArray.Remove(InObject);
+			}
+
+			virtual void AddReferencedObjects(FReferenceCollector& Collector) override
+			{
+				FScopeLock LockObjectArray(&ObjectArrayCriticalSection);
+				Collector.AddReferencedObjects(ObjectArray);
+			}
+
+			virtual FString GetReferencerName() const override
+			{
+				return TEXT("FMetasoundEngineModule::FObjectReferencer");
+			}
+
+		private:
+			mutable FCriticalSection ObjectArrayCriticalSection;
+			TArray<TObjectPtr<UObject>> ObjectArray;
+		};
+
+	public:
+		virtual void StartupModule() override;
+		virtual void ShutdownModule() override;
+
+#if WITH_EDITOR
+		virtual void PrimeAssetRegistryAsync() override;
+		virtual void PrimeAssetManager() override;
+
+		virtual ENodeClassRegistryPrimeStatus GetNodeClassRegistryPrimeStatus() const override;
+		virtual EAssetScanStatus GetAssetRegistryScanStatus() const override;
+		virtual bool IsAssetManagerPrimed() const override;
+
+		virtual FOnMetasoundGraphRegister& GetOnGraphRegisteredDelegate() override;
+		virtual FOnMetasoundGraphUnregister& GetOnGraphUnregisteredDelegate() override;
+#endif // WITH_EDITOR
+
+	private:
+#if WITH_EDITOR
+		virtual void AddClassRegistryAsset(const FAssetData& InAssetData);
+		virtual void LoadAndRegisterAsset(const FAssetData& InAssetData);
+		virtual void UpdateClassRegistryAsset(const FAssetData& InAssetData);
+		virtual void RemoveAssetFromClassRegistry(const FAssetData& InAssetData);
+		virtual void RenameAssetInClassRegistry(const FAssetData& InAssetData, const FString& InOldObjectPath);
+		virtual void ShutdownAssetClassRegistry();
+
+		virtual void OnPackageReloaded(const EPackageReloadPhase InPackageReloadPhase, FPackageReloadedEvent* InPackageReloadedEvent);
+		virtual void OnAssetScanFinished();
+
+		// Asset registry delegates for calling MetaSound editor module 
+		FOnMetasoundGraphRegister OnGraphRegister;
+		FOnMetasoundGraphUnregister OnGraphUnregister;
+
+		// Asset registry state
+		ENodeClassRegistryPrimeStatus NodeClassRegistryPrimeStatus = ENodeClassRegistryPrimeStatus::NotRequested;
+		EAssetScanStatus AssetScanStatus = EAssetScanStatus::NotRequested;
+		bool bAssetManagerPrimed = false;
+		int32 ActiveAsyncAssetLoadRequests = 0;
+#endif // WITH_EDITOR
+	};
+} // namespace Metasound::Engine
diff --git a/Engine/Plugins/Runtime/Metasound/Source/MetasoundEngine/Private/MetasoundSource.cpp b/Engine/Plugins/Runtime/Metasound/Source/MetasoundEngine/Private/MetasoundSource.cpp
--- a/Engine/Plugins/Runtime/Metasound/Source/MetasoundEngine/Private/MetasoundSource.cpp
+++ b/Engine/Plugins/Runtime/Metasound/Source/MetasoundEngine/Private/MetasoundSource.cpp
@@ -494,6 +494,12 @@ bool UMetaSoundSource::CanEditChange(const FProperty* InProperty) const
 	return true;
 }
 
+EDataValidationResult UMetaSoundSource::IsDataValid(FDataValidationContext& Context) const
+{
+	const EDataValidationResult Result = Metasound::Engine::FAssetHelper::IsDataValid(*this, Context);
+	return CombineDataValidationResults(Result, Super::IsDataValid(Context));
+}
+
 void UMetaSoundSource::PostEditChangeOutputFormat()
 {
 	using namespace Metasound::Engine;
diff --git a/Engine/Plugins/Runtime/Metasound/Source/MetasoundEngine/Public/IMetasoundEngineModule.h b/Engine/Plugins/Runtime/Metasound/Source/MetasoundEngine/Public/IMetasoundEngineModule.h
new file mode 100644
--- /dev/null
+++ b/Engine/Plugins/Runtime/Metasound/Source/MetasoundEngine/Public/IMetasoundEngineModule.h
@@ -0,0 +1,62 @@
+// Copyright Epic Games, Inc. All Rights Reserved.
+
+#pragma once
+
+#include "Delegates/Delegate.h"
+#include "Modules/ModuleInterface.h"
+#include "UObject/PackageReload.h"
+
+namespace Metasound::Engine
+{
+#if WITH_EDITOR
+	// Status of initial asset scan when editor loads up.
+	enum class EAssetScanStatus : uint8
+	{
+		NotRequested = 0,
+		InProgress = 2,
+		Complete = 3
+	};
+
+	// Node class primes status of MetaSound assets.  Priming an asset
+	// loads the asset asynchronously (if not already loaded)
+	// & registers it with the MetaSound Node Class Registry.
+	enum class ENodeClassRegistryPrimeStatus : uint8
+	{
+		NotRequested = 0,
+		Requested = 1,
+		InProgress = 2,
+		Complete = 3,
+		Canceled = 4
+	};
+
+	DECLARE_DELEGATE_TwoParams(FOnMetasoundGraphRegister, UObject&, bool)
+	DECLARE_DELEGATE_OneParam(FOnMetasoundGraphUnregister, UObject&)
+#endif // WITH_EDITOR
+
+	class IMetasoundEngineModule : public IModuleInterface
+	{
+	public:
+		virtual void StartupModule() override = 0;
+		virtual void ShutdownModule() override = 0;
+
+#if WITH_EDITOR
+		// Primes MetaSound assets, registering them with the asset manager
+		// and loads assets asynchronously (if not already loaded) 
+		// & registers them if not already registered with the MetaSound Node Class Registry.
+		virtual void PrimeAssetRegistryAsync() = 0;
+
+		// Adds assets to the MetaSound asset manager.
+		// This is a subset of PrimeAssetRegistryAsync (does not load/register assets)
+		virtual void PrimeAssetManager() = 0;
+
+		virtual ENodeClassRegistryPrimeStatus GetNodeClassRegistryPrimeStatus() const = 0;
+		virtual EAssetScanStatus GetAssetRegistryScanStatus() const = 0;
+		// Bool rather than enum because this does not require async asset loading like priming the node class registry
+		virtual bool IsAssetManagerPrimed() const = 0;
+
+		// Asset registry delegates for calling MetaSound editor module register/unregister with frontend 
+		virtual FOnMetasoundGraphRegister& GetOnGraphRegisteredDelegate() = 0;
+		virtual FOnMetasoundGraphUnregister& GetOnGraphUnregisteredDelegate() = 0;
+#endif // WITH_EDITOR
+	};
+} // namespace Metasound::Engine
diff --git a/Engine/Plugins/Runtime/Metasound/Source/MetasoundEngine/Public/Metasound.h b/Engine/Plugins/Runtime/Metasound/Source/MetasoundEngine/Public/Metasound.h
--- a/Engine/Plugins/Runtime/Metasound/Source/MetasoundEngine/Public/Metasound.h
+++ b/Engine/Plugins/Runtime/Metasound/Source/MetasoundEngine/Public/Metasound.h
@@ -153,6 +153,8 @@ public:
 	virtual void PreDuplicate(FObjectDuplicationParameters& DupParams) override;
 	virtual void PostDuplicate(EDuplicateMode::Type DuplicateMode) override;
 	virtual void PostEditUndo() override;
+	virtual EDataValidationResult IsDataValid(FDataValidationContext& Context) const override;
+
 #endif // WITH_EDITOR
 
 	virtual void BeginDestroy() override;
diff --git a/Engine/Plugins/Runtime/Metasound/Source/MetasoundEngine/Public/MetasoundSource.h b/Engine/Plugins/Runtime/Metasound/Source/MetasoundEngine/Public/MetasoundSource.h
--- a/Engine/Plugins/Runtime/Metasound/Source/MetasoundEngine/Public/MetasoundSource.h
+++ b/Engine/Plugins/Runtime/Metasound/Source/MetasoundEngine/Public/MetasoundSource.h
@@ -219,6 +219,8 @@ public:
 	virtual void PostEditChangeProperty(FPropertyChangedEvent& InEvent) override;
 	virtual bool CanEditChange(const FProperty* InProperty) const override;
 
+	virtual EDataValidationResult IsDataValid(FDataValidationContext& Context) const override;
+
 private:
 	void PostEditChangeOutputFormat();
 	void PostEditChangeQualitySettings();
diff --git a/Engine/Plugins/Runtime/Metasound/Source/MetasoundFrontend/Public/MetasoundAssetManager.h b/Engine/Plugins/Runtime/Metasound/Source/MetasoundFrontend/Public/MetasoundAssetManager.h
--- a/Engine/Plugins/Runtime/Metasound/Source/MetasoundFrontend/Public/MetasoundAssetManager.h
+++ b/Engine/Plugins/Runtime/Metasound/Source/MetasoundFrontend/Public/MetasoundAssetManager.h
@@ -164,13 +164,18 @@ namespace Metasound::Frontend
 		virtual bool ContainsKey(const Metasound::Frontend::FAssetKey& InAssetKey) const = 0;
 
 		// Returns object (if loaded) associated with the given key (null if key is not registered with the AssetManager)
+		// If multiple assets are associated with the given key, the last one is returned. 
 		virtual FMetasoundAssetBase* FindAsset(const Metasound::Frontend::FAssetKey& InAssetKey) const = 0;
 
 		// Returns object (if loaded) associated with the given key as a Document Interface (null if key is not registered with the AssetManager)
 		virtual TScriptInterface<IMetaSoundDocumentInterface> FindAssetAsDocumentInterface(const Frontend::FAssetKey& InKey) const = 0;
 
 		// Returns path associated with the given key (null if key is not registered with the AssetManager or was not loaded from asset)
+		// If multiple assets are associated with the given key, the last one is returned. 
 		virtual const FTopLevelAssetPath* FindAssetPath(const Metasound::Frontend::FAssetKey& InAssetKey) const = 0;
+		
+		// Returns all paths associated with the given key (null if key is not registered with the AssetManager or was not loaded from asset)
+		virtual const TArray<FTopLevelAssetPath>* FindAssetPaths(const Metasound::Frontend::FAssetKey& InAssetKey) const = 0;
 
 		// Converts an object to an AssetBase if its a registered asset
 		virtual FMetasoundAssetBase* GetAsAsset(UObject& InObject) const = 0;
