Fix for meters not updating to correct channel count when performing undo on Output Format. Fix for meters for Preset asset not updating to correct channel count when changing base MSS Output Format.

diff --git a/Engine/Plugins/Runtime/Metasound/Source/MetasoundEditor/Private/MetasoundEditor.cpp b/Engine/Plugins/Runtime/Metasound/Source/MetasoundEditor/Private/MetasoundEditor.cpp
--- a/Engine/Plugins/Runtime/Metasound/Source/MetasoundEditor/Private/MetasoundEditor.cpp
+++ b/Engine/Plugins/Runtime/Metasound/Source/MetasoundEditor/Private/MetasoundEditor.cpp
@@ -3943,6 +3943,12 @@ namespace Metasound
 				if (!InterfacesModified.IsEmpty() || bForceRefreshViews)
 				{
 					RefreshInterfaceView();
+
+					// Output Format may have changed, ensure analyzers are created with the correct channel count:
+					if (UMetaSoundSource* MetaSoundSource = Cast<UMetaSoundSource>(&MetaSound))
+					{
+						CreateAnalyzers(*MetaSoundSource);
+					}
 				}
 
 				HighestMessageSeverity = Graph->GetHighestMessageSeverity();
