Trivial: Deprecate old IMetaSoundDocumentBuilderRegistry class ctor

#rnx

diff --git a/Engine/Plugins/Runtime/Metasound/Source/MetasoundFrontend/Public/MetasoundDocumentInterface.h b/Engine/Plugins/Runtime/Metasound/Source/MetasoundFrontend/Public/MetasoundDocumentInterface.h
--- a/Engine/Plugins/Runtime/Metasound/Source/MetasoundFrontend/Public/MetasoundDocumentInterface.h
+++ b/Engine/Plugins/Runtime/Metasound/Source/MetasoundFrontend/Public/MetasoundDocumentInterface.h
@@ -116,6 +116,8 @@ namespace Metasound::Frontend
 	class METASOUNDFRONTEND_API IMetaSoundDocumentBuilderRegistry : public IDocumentBuilderRegistry
 	{
 	public:
+		UE_DEPRECATED(5.5, "Class Renamed: use IDocumentBuilderRegistry instead")
+		IMetaSoundDocumentBuilderRegistry() = default;
 		virtual ~IMetaSoundDocumentBuilderRegistry() = default;
 
 		UE_DEPRECATED(5.4, "Public exposition of modify delegates no longer available to discourage unsafe manipulation of builder document cache")
