Legacy Controller Page Fix-ups
- Fix-up legacy "RemoveUnusedDependencies" call to accoimdate pages until we can safely deprecate it properly (i.e. refactor AutoUpdate to use builder API)
- Move Graph Controllers to use build graph instead of default graph

#rnx

diff --git a/Engine/Plugins/Runtime/Metasound/Source/MetasoundFrontend/Private/MetasoundFrontendDocumentController.cpp b/Engine/Plugins/Runtime/Metasound/Source/MetasoundFrontend/Private/MetasoundFrontendDocumentController.cpp
--- a/Engine/Plugins/Runtime/Metasound/Source/MetasoundFrontend/Private/MetasoundFrontendDocumentController.cpp
+++ b/Engine/Plugins/Runtime/Metasound/Source/MetasoundFrontend/Private/MetasoundFrontendDocumentController.cpp
@@ -436,23 +436,25 @@ namespace Metasound
 				do
 				{
 					TSet<FGuid> ReferencedDependencyIDs;
-					auto AddNodeClassIDToSet = [&](const FMetasoundFrontendNode& Node)
+					auto AddGraphNodeClassIDsToSet = [&ReferencedDependencyIDs](const FMetasoundFrontendGraphClass& GraphClass)
 					{
-						ReferencedDependencyIDs.Add(Node.ClassID);
-					};
-
-					auto AddGraphNodeClassIDsToSet = [&](const FMetasoundFrontendGraphClass& GraphClass)
-					{
-						Algo::ForEach(FindConstBuildGraphChecked(GraphClass).Nodes, AddNodeClassIDToSet);
+						GraphClass.IterateGraphPages([&ReferencedDependencyIDs](const FMetasoundFrontendGraph& Graph)
+						{
+							auto AddNodeClassIDToSet = [&ReferencedDependencyIDs](const FMetasoundFrontendNode& Node)
+							{
+								ReferencedDependencyIDs.Add(Node.ClassID);
+							};
+							Algo::ForEach(Graph.Nodes, AddNodeClassIDToSet);
+						});
 					};
 
-					// Referenced dependencies in root class
-					Algo::ForEach(FindConstBuildGraphChecked(Document->RootGraph).Nodes, AddNodeClassIDToSet);
+					// Referenced dependencies in RootGraph
+					AddGraphNodeClassIDsToSet(Document->RootGraph);
 
-					// Referenced dependencies in subgraphs
+					// Referenced dependencies in Subgraphs
 					Algo::ForEach(Document->Subgraphs, AddGraphNodeClassIDsToSet);
 
-					auto IsDependencyUnreferenced = [&](const FMetasoundFrontendClass& ClassDependency)
+					auto IsDependencyUnreferenced = [&ReferencedDependencyIDs](const FMetasoundFrontendClass& ClassDependency)
 					{
 						return !ReferencedDependencyIDs.Contains(ClassDependency.ID);
 					};
diff --git a/Engine/Plugins/Runtime/Metasound/Source/MetasoundFrontend/Private/MetasoundFrontendGraphController.cpp b/Engine/Plugins/Runtime/Metasound/Source/MetasoundFrontend/Private/MetasoundFrontendGraphController.cpp
--- a/Engine/Plugins/Runtime/Metasound/Source/MetasoundFrontend/Private/MetasoundFrontendGraphController.cpp
+++ b/Engine/Plugins/Runtime/Metasound/Source/MetasoundFrontend/Private/MetasoundFrontendGraphController.cpp
@@ -214,7 +214,7 @@ namespace Metasound
 		{
 			if (const FMetasoundFrontendGraphClass* GraphClass = GraphClassPtr.Get())
 			{
-				return GraphClass->GetConstDefaultGraph().Style;
+				return FindConstBuildGraphChecked(*GraphClass).Style;
 			}
 
 			return Invalid::GetInvalidGraphStyle();
@@ -428,7 +428,8 @@ namespace Metasound
 
 			if (FMetasoundFrontendGraphClass* GraphClass = GraphClassPtr.Get())
 			{
-				for (const FMetasoundFrontendVariable& Variable : GraphClass->GetConstDefaultGraph().Variables)
+				const TArray<FMetasoundFrontendVariable>& Variables = FindConstBuildGraphChecked(*GraphClass).Variables;
+				for (const FMetasoundFrontendVariable& Variable : Variables)
 				{
 					VariableHandles.Add(FindVariable(Variable.ID));
 				}
@@ -443,7 +444,8 @@ namespace Metasound
 
 			if (const FMetasoundFrontendGraphClass* GraphClass = GraphClassPtr.Get())
 			{
-				for (const FMetasoundFrontendVariable& Variable : GraphClass->GetConstDefaultGraph().Variables)
+				const TArray<FMetasoundFrontendVariable>& Variables = FindConstBuildGraphChecked(*GraphClass).Variables;
+				for (const FMetasoundFrontendVariable& Variable : Variables)
 				{
 					VariableHandles.Add(FindVariable(Variable.ID));
 				}
@@ -636,7 +638,8 @@ namespace Metasound
 				{
 					return InVariable.ID == InVariableID;
 				};
-				return GraphClass->GetConstDefaultGraph().Variables.FindByPredicate(IsVariableWithID);
+				const TArray<FMetasoundFrontendVariable>& Variables = FindConstBuildGraphChecked(*GraphClass).Variables;
+				return Variables.FindByPredicate(IsVariableWithID);
 			}
 			return nullptr;
 		}
@@ -652,7 +655,8 @@ namespace Metasound
 						return (InNodeID == InVariable.VariableNodeID) || (InNodeID == InVariable.MutatorNodeID) || InVariable.AccessorNodeIDs.Contains(InNodeID) || InVariable.DeferredAccessorNodeIDs.Contains(InNodeID);
 
 					};
-					if (const FMetasoundFrontendVariable* Variable = GraphClass->GetConstDefaultGraph().Variables.FindByPredicate(ContainsNodeWithID))
+					const TArray<FMetasoundFrontendVariable>& Variables = FindConstBuildGraphChecked(*GraphClass).Variables;
+					if (const FMetasoundFrontendVariable* Variable = Variables.FindByPredicate(ContainsNodeWithID))
 					{
 						return Variable->ID;
 					}
@@ -794,7 +798,8 @@ namespace Metasound
 		{
 			if (FMetasoundFrontendGraphClass* GraphClass = GraphClassPtr.Get())
 			{
-				for (const FMetasoundFrontendNode& Node : GraphClass->GetConstDefaultGraph().Nodes)
+				const TArray<FMetasoundFrontendNode>& Nodes = FindConstBuildGraphChecked(*GraphClass).Nodes;
+				for (const FMetasoundFrontendNode& Node : Nodes)
 				{
 					FConstClassAccessPtr NodeClassPtr = OwningDocument->FindClassWithID(Node.ClassID);
 					if (const FMetasoundFrontendClass* NodeClass = NodeClassPtr.Get())
@@ -818,7 +823,8 @@ namespace Metasound
 		{
 			if (FMetasoundFrontendGraphClass* GraphClass = GraphClassPtr.Get())
 			{
-				for (const FMetasoundFrontendNode& Node : GraphClass->GetConstDefaultGraph().Nodes)
+				const TArray<FMetasoundFrontendNode>& Nodes = FindConstBuildGraphChecked(*GraphClass).Nodes;
+				for (const FMetasoundFrontendNode& Node : Nodes)
 				{
 					FConstClassAccessPtr NodeClassPtr = OwningDocument->FindClassWithID(Node.ClassID);
 					if (const FMetasoundFrontendClass* NodeClass = NodeClassPtr.Get())
@@ -1519,7 +1525,8 @@ namespace Metasound
 			{
 				FGuid NodeID = InNode.GetID();
 				auto IsNodeWithSameID = [&](const FMetasoundFrontendNode& InFrontendNode) { return InFrontendNode.GetID() == NodeID; };
-				if (const FMetasoundFrontendNode* FrontendNode = GraphClass->GetConstDefaultGraph().Nodes.FindByPredicate(IsNodeWithSameID))
+				const TArray<FMetasoundFrontendNode>& Nodes = FindConstBuildGraphChecked(*GraphClass).Nodes;
+				if (const FMetasoundFrontendNode* FrontendNode = Nodes.FindByPredicate(IsNodeWithSameID))
 				{
 					switch(NodeClassType)
 					{
@@ -1839,7 +1846,8 @@ namespace Metasound
 
 			if (const FMetasoundFrontendGraphClass* GraphClass = GraphClassPtr.Get())
 			{
-				for (const FMetasoundFrontendNode& Node : GraphClass->GetConstDefaultGraph().Nodes)
+				const TArray<FMetasoundFrontendNode>& Nodes = FindConstBuildGraphChecked(*GraphClass).Nodes;
+				for (const FMetasoundFrontendNode& Node : Nodes)
 				{
 					FConstNodeAccessPtr NodePtr = GraphClassPtr.GetNodeWithNodeID(Node.GetID());
 					FConstClassAccessPtr NodeClassPtr = OwningDocument->FindClassWithID(Node.ClassID);
@@ -1894,7 +1902,8 @@ namespace Metasound
 
 			if (const FMetasoundFrontendGraphClass* GraphClass = GraphClassPtr.Get())
 			{
-				for (const FMetasoundFrontendNode& Node : GraphClass->GetConstDefaultGraph().Nodes)
+				const TArray<FMetasoundFrontendNode>& Nodes = FindConstBuildGraphChecked(*GraphClass).Nodes;
+				for (const FMetasoundFrontendNode& Node : Nodes)
 				{
 					FConstClassAccessPtr NodeClassPtr = OwningDocument->FindClassWithID(Node.ClassID);
 					if (const FMetasoundFrontendClass* NodeClass = NodeClassPtr.Get())
