Simple polish feedback from earlier page CLs, misc clean-up:
- More const node handle fix-ups to isolate prioritization of controller callsites to refactor (and to avoid reloading active builder/resetting builder cache when not necessary)
- Make builders transactional to log when page ID is changed
- Fix mismatching EDITORONLY_DATA vs EDITOR macro chunk in MetaSoundSettings
- Protect FDocumentBuilderRegistry::FindBuilderObject calls so that callers providing an invalid asset path (ex. from controller back compat look-up calls) always returns a valid path over a stale one if both are in the registry

#rnx

diff --git a/Engine/Plugins/Runtime/Metasound/Source/MetasoundEditor/Private/MetasoundDetailCustomization.cpp b/Engine/Plugins/Runtime/Metasound/Source/MetasoundEditor/Private/MetasoundDetailCustomization.cpp
--- a/Engine/Plugins/Runtime/Metasound/Source/MetasoundEditor/Private/MetasoundDetailCustomization.cpp
+++ b/Engine/Plugins/Runtime/Metasound/Source/MetasoundEditor/Private/MetasoundDetailCustomization.cpp
@@ -420,7 +420,7 @@ namespace Metasound::Editor
 
 		if (UMetaSoundSettings* Settings = GetMutableDefault<UMetaSoundSettings>())
 		{
-			Settings->OnPageSettingsUpdated.AddSPLambda(this, [this]()
+			Settings->GetOnPageSettingsUpdatedDelegate().AddSPLambda(this, [this]()
 			{
 				UpdateItemNames();
 				if (ComboBox.IsValid())
diff --git a/Engine/Plugins/Runtime/Metasound/Source/MetasoundEditor/Private/MetasoundEditorGraph.cpp b/Engine/Plugins/Runtime/Metasound/Source/MetasoundEditor/Private/MetasoundEditorGraph.cpp
--- a/Engine/Plugins/Runtime/Metasound/Source/MetasoundEditor/Private/MetasoundEditorGraph.cpp
+++ b/Engine/Plugins/Runtime/Metasound/Source/MetasoundEditor/Private/MetasoundEditorGraph.cpp
@@ -805,7 +805,7 @@ void UMetasoundEditorGraphInput::UpdateFrontendDefaultLiteral(bool bPostTransact
 	check(MetasoundAsset);
 
 	FGraphHandle GraphHandle = MetasoundAsset->GetRootGraphHandle();
-	FNodeHandle NodeHandle = GraphHandle->GetNodeWithID(NodeID);
+	FConstNodeHandle NodeHandle = GraphHandle->GetNodeWithID(NodeID);
 
 	const Metasound::FVertexName& NodeName = NodeHandle->GetNodeName();
 	const FGuid VertexID = GraphHandle->GetVertexIDForInputVertex(NodeName);
diff --git a/Engine/Plugins/Runtime/Metasound/Source/MetasoundEditor/Private/MetasoundEditorGraphNode.cpp b/Engine/Plugins/Runtime/Metasound/Source/MetasoundEditor/Private/MetasoundEditorGraphNode.cpp
--- a/Engine/Plugins/Runtime/Metasound/Source/MetasoundEditor/Private/MetasoundEditorGraphNode.cpp
+++ b/Engine/Plugins/Runtime/Metasound/Source/MetasoundEditor/Private/MetasoundEditorGraphNode.cpp
@@ -1078,7 +1078,7 @@ void UMetasoundEditorGraphExternalNode::CacheTitle()
 	using namespace Metasound::Frontend;
 
 	constexpr bool bIncludeNamespace = false;
-	FConstNodeHandle NodeHandle = GetNodeHandle();
+	FConstNodeHandle NodeHandle = GetConstNodeHandle();
 	CachedTitle = FGraphBuilder::GetDisplayName(*NodeHandle, bIncludeNamespace);
 }
 
@@ -1106,7 +1106,7 @@ void UMetasoundEditorGraphExternalNode::Validate(Metasound::Editor::FGraphNodeVa
 #if WITH_EDITOR
 	Super::Validate(OutResult);
 
-	FConstNodeHandle NodeHandle = GetNodeHandle();
+	FConstNodeHandle NodeHandle = GetConstNodeHandle();
 	const FMetasoundFrontendClassMetadata& Metadata = NodeHandle->GetClassMetadata();
 
 	// 1. Validate external referenced graph or template node
diff --git a/Engine/Plugins/Runtime/Metasound/Source/MetasoundEditor/Private/MetasoundEditorGraphSchema.cpp b/Engine/Plugins/Runtime/Metasound/Source/MetasoundEditor/Private/MetasoundEditorGraphSchema.cpp
--- a/Engine/Plugins/Runtime/Metasound/Source/MetasoundEditor/Private/MetasoundEditorGraphSchema.cpp
+++ b/Engine/Plugins/Runtime/Metasound/Source/MetasoundEditor/Private/MetasoundEditorGraphSchema.cpp
@@ -1277,7 +1277,7 @@ void UMetasoundEditorGraphSchema::GetGraphContextActions(FGraphContextMenuBuilde
 		GetCommentAction(ContextMenuBuilder, ContextMenuBuilder.CurrentGraph);
 		if (UObject* Metasound = MetasoundEditor->GetMetasoundObject())
 		{
-			FMetasoundAssetBase* MetasoundAsset = IMetasoundUObjectRegistry::Get().GetObjectAsAssetBase(Metasound);
+			const FMetasoundAssetBase* MetasoundAsset = IMetasoundUObjectRegistry::Get().GetObjectAsAssetBase(Metasound);
 			check(MetasoundAsset);
 			GraphHandle = MetasoundAsset->GetRootGraphHandle();
 
diff --git a/Engine/Plugins/Runtime/Metasound/Source/MetasoundEditor/Private/MetasoundEditorModule.cpp b/Engine/Plugins/Runtime/Metasound/Source/MetasoundEditor/Private/MetasoundEditorModule.cpp
--- a/Engine/Plugins/Runtime/Metasound/Source/MetasoundEditor/Private/MetasoundEditorModule.cpp
+++ b/Engine/Plugins/Runtime/Metasound/Source/MetasoundEditor/Private/MetasoundEditorModule.cpp
@@ -604,7 +604,7 @@ namespace Metasound
 
 				if (UMetaSoundSettings* Settings = GetMutableDefault<UMetaSoundSettings>())
 				{
-					Settings->OnDefaultConformed.AddLambda([]()
+					Settings->GetOnDefaultConformedDelegate().AddLambda([]()
 					{
 						FNotificationInfo Info(LOCTEXT("MetaSoundSettings_CannotDeleteDefaultPage", "Cannot change name of nor delete 'Default' MetaSound Page"));
 						Info.bFireAndForget = true;
diff --git a/Engine/Plugins/Runtime/Metasound/Source/MetasoundEngine/Private/MetasoundDocumentBuilderRegistry.cpp b/Engine/Plugins/Runtime/Metasound/Source/MetasoundEngine/Private/MetasoundDocumentBuilderRegistry.cpp
--- a/Engine/Plugins/Runtime/Metasound/Source/MetasoundEngine/Private/MetasoundDocumentBuilderRegistry.cpp
+++ b/Engine/Plugins/Runtime/Metasound/Source/MetasoundEngine/Private/MetasoundDocumentBuilderRegistry.cpp
@@ -191,7 +191,7 @@ namespace Metasound::Engine
 					FTopLevelAssetPath ObjectPath;
 					if (ObjectPath.TrySetPath(&DocObject))
 					{
-						if (AssetPath == ObjectPath)
+						if (AssetPath.IsNull() || AssetPath == ObjectPath)
 						{
 							FoundEntry = Builder;
 							break;
diff --git a/Engine/Plugins/Runtime/Metasound/Source/MetasoundEngine/Private/MetasoundSettings.cpp b/Engine/Plugins/Runtime/Metasound/Source/MetasoundEngine/Private/MetasoundSettings.cpp
--- a/Engine/Plugins/Runtime/Metasound/Source/MetasoundEngine/Private/MetasoundSettings.cpp
+++ b/Engine/Plugins/Runtime/Metasound/Source/MetasoundEngine/Private/MetasoundSettings.cpp
@@ -229,6 +229,18 @@ void UMetaSoundSettings::PostInitProperties()
 
 	Super::PostInitProperties();
 }
+#endif // WITH_EDITOR
+
+#if WITH_EDITORONLY_DATA
+Metasound::Engine::FOnSettingsDefaultConformed& UMetaSoundSettings::GetOnDefaultConformedDelegate()
+{
+	return OnDefaultConformed;
+}
+
+Metasound::Engine::FOnPageSettingsUpdated& UMetaSoundSettings::GetOnPageSettingsUpdatedDelegate()
+{
+	return OnPageSettingsUpdated;
+}
 
 FName UMetaSoundSettings::GetPageSettingPropertyName()
 {
@@ -239,7 +251,7 @@ FName UMetaSoundSettings::GetQualitySettingPropertyName()
 {
 	return GET_MEMBER_NAME_CHECKED(UMetaSoundSettings, QualitySettings);
 }
-#endif // WITH_EDITOR
+#endif // WITH_EDITORONLY_DATA
 
 TArray<FName> UMetaSoundQualityHelper::GetQualityList()
 {
diff --git a/Engine/Plugins/Runtime/Metasound/Source/MetasoundEngine/Public/MetasoundSettings.h b/Engine/Plugins/Runtime/Metasound/Source/MetasoundEngine/Public/MetasoundSettings.h
--- a/Engine/Plugins/Runtime/Metasound/Source/MetasoundEngine/Public/MetasoundSettings.h
+++ b/Engine/Plugins/Runtime/Metasound/Source/MetasoundEngine/Public/MetasoundSettings.h
@@ -133,12 +133,12 @@ public:
 	UPROPERTY(Transient)
 	int32 DenyListCacheChangeID = 0;	
 
+private:
 #if WITH_EDITORONLY_DATA
 	Metasound::Engine::FOnSettingsDefaultConformed OnDefaultConformed;
 	Metasound::Engine::FOnPageSettingsUpdated OnPageSettingsUpdated;
 #endif //WITH_EDITORONLY_DATA
 
-private:
 	/** Array of possible page settings that can be added to a MetaSound object. */
 	UPROPERTY(EditAnywhere, config, Category = Pages)
 	TArray<FMetaSoundPageSettings> PageSettings;
@@ -158,6 +158,9 @@ public:
 	const TArray<FMetaSoundQualitySettings>& GetQualitySettings() const { return QualitySettings; }
 
 #if WITH_EDITORONLY_DATA
+	Metasound::Engine::FOnSettingsDefaultConformed& GetOnDefaultConformedDelegate();
+	Metasound::Engine::FOnPageSettingsUpdated& GetOnPageSettingsUpdatedDelegate();
+
 	static FName GetPageSettingPropertyName();
 	static FName GetQualitySettingPropertyName();
 #endif // WITH_EDITORONLY_DATA
