More const handle & comment fix-ups, missing transactional flag bit from last CL

#rnx

diff --git a/Engine/Plugins/Runtime/Metasound/Source/MetasoundEditor/Private/SMetasoundGraphEnumPin.cpp b/Engine/Plugins/Runtime/Metasound/Source/MetasoundEditor/Private/SMetasoundGraphEnumPin.cpp
--- a/Engine/Plugins/Runtime/Metasound/Source/MetasoundEditor/Private/SMetasoundGraphEnumPin.cpp
+++ b/Engine/Plugins/Runtime/Metasound/Source/MetasoundEditor/Private/SMetasoundGraphEnumPin.cpp
@@ -34,25 +34,25 @@ namespace Metasound
 				.OnGetTooltip(this, &SMetasoundGraphEnumPin::OnGetTooltip);
 		}
 
-		TSharedPtr<const Metasound::Frontend::IEnumDataTypeInterface>
-		SMetasoundGraphEnumPin::FindEnumInterfaceFromPin(UEdGraphPin* InPin) 
+		TSharedPtr<const Frontend::IEnumDataTypeInterface>
+		SMetasoundGraphEnumPin::FindEnumInterfaceFromPin(UEdGraphPin* InPin)
 		{
-			using namespace Metasound::Frontend;
+			using namespace Frontend;
 
 			auto MetasoundEditorNode = Cast<UMetasoundEditorGraphNode>(InPin->GetOwningNode());
-			FNodeHandle NodeHandle = MetasoundEditorNode->GetNodeHandle();
+			FConstNodeHandle NodeHandle = MetasoundEditorNode->GetConstNodeHandle();
 			FConstInputHandle Input = NodeHandle->GetConstInputWithVertexName(InPin->GetFName());
 			if (Input->IsValid())
 			{
 				FName DataType = Input->GetDataType();
-				return Metasound::Frontend::IDataTypeRegistry::Get().GetEnumInterfaceForDataType(DataType);
+				return Frontend::IDataTypeRegistry::Get().GetEnumInterfaceForDataType(DataType);
 			}
 			return nullptr;
 		}
 
 		FString SMetasoundGraphEnumPin::OnGetText() const
 		{
-			using namespace Metasound::Frontend;
+			using namespace Frontend;
 
 			TSharedPtr<const IEnumDataTypeInterface> EnumInterface = FindEnumInterfaceFromPin(GraphPinObj);
 			check(EnumInterface.IsValid());
@@ -71,7 +71,7 @@ namespace Metasound
 
 		void SMetasoundGraphEnumPin::GenerateComboBoxIndexes(TArray<TSharedPtr<int32>>& OutComboBoxIndexes)
 		{
-			using namespace Metasound::Frontend;
+			using namespace Frontend;
 			TSharedPtr<const IEnumDataTypeInterface> EnumInterface = FindEnumInterfaceFromPin(GraphPinObj);
 			check(EnumInterface.IsValid());
 
@@ -84,7 +84,7 @@ namespace Metasound
 
 		void SMetasoundGraphEnumPin::ComboBoxSelectionChanged(TSharedPtr<int32> NewSelection, ESelectInfo::Type SelectInfo)
 		{
-			using namespace Metasound::Frontend;
+			using namespace Frontend;
 			TSharedPtr<const IEnumDataTypeInterface> EnumInterface = FindEnumInterfaceFromPin(SGraphPin::GetPinObj());
 			check(EnumInterface.IsValid());
 
@@ -116,7 +116,7 @@ namespace Metasound
 
 		FText SMetasoundGraphEnumPin::OnGetFriendlyName(int32 EnumIndex)
 		{
-			using namespace Metasound::Frontend;
+			using namespace Frontend;
 
 			TSharedPtr<const IEnumDataTypeInterface> Interface = FindEnumInterfaceFromPin(SGraphPin::GetPinObj());
 			check(Interface.IsValid());
@@ -129,7 +129,7 @@ namespace Metasound
 
 		FText SMetasoundGraphEnumPin::OnGetTooltip(int32 EnumIndex)
 		{
-			using namespace Metasound::Frontend;
+			using namespace Frontend;
 
 			TSharedPtr<const IEnumDataTypeInterface> Interface = FindEnumInterfaceFromPin(GraphPinObj);
 			check(Interface.IsValid());
diff --git a/Engine/Plugins/Runtime/Metasound/Source/MetasoundEditor/Private/SMetasoundGraphNode.cpp b/Engine/Plugins/Runtime/Metasound/Source/MetasoundEditor/Private/SMetasoundGraphNode.cpp
--- a/Engine/Plugins/Runtime/Metasound/Source/MetasoundEditor/Private/SMetasoundGraphNode.cpp
+++ b/Engine/Plugins/Runtime/Metasound/Source/MetasoundEditor/Private/SMetasoundGraphNode.cpp
@@ -538,7 +538,7 @@ namespace Metasound
 		{
 			SGraphNode::SetDefaultTitleAreaWidget(DefaultTitleAreaWidget);
 
-			Metasound::Frontend::FNodeHandle NodeHandle = GetMetaSoundNode().GetNodeHandle();
+			Metasound::Frontend::FConstNodeHandle NodeHandle = GetMetaSoundNode().GetConstNodeHandle();
 			if (NodeHandle->GetClassStyle().Display.bShowName)
 			{
 				DefaultTitleAreaWidget->ClearChildren();
@@ -670,7 +670,7 @@ namespace Metasound
 
 			// Just take last type.  If more than one, all types are the same.
 			const UMetasoundEditorGraphNode& Node = GetMetaSoundNode();
-			Node.GetNodeHandle()->IterateConstOutputs([InTypeName = &TypeName](FConstOutputHandle OutputHandle)
+			Node.GetConstNodeHandle()->IterateConstOutputs([InTypeName = &TypeName](FConstOutputHandle OutputHandle)
 			{
 				*InTypeName = OutputHandle->GetDataType();
 			});
diff --git a/Engine/Plugins/Runtime/Metasound/Source/MetasoundEngine/Public/MetasoundDocumentBuilderRegistry.h b/Engine/Plugins/Runtime/Metasound/Source/MetasoundEngine/Public/MetasoundDocumentBuilderRegistry.h
--- a/Engine/Plugins/Runtime/Metasound/Source/MetasoundEngine/Public/MetasoundDocumentBuilderRegistry.h
+++ b/Engine/Plugins/Runtime/Metasound/Source/MetasoundEngine/Public/MetasoundDocumentBuilderRegistry.h
@@ -76,7 +76,7 @@ namespace Metasound::Engine
 			TObjectPtr<UMetaSoundBuilderBase> NewBuilder;
 			{
 				FScopeLock Lock(&BuildersCriticalSection);
-				NewBuilder = CastChecked<UMetaSoundBuilderBase>(NewObject<UObject>(&InMetaSoundObject, &BuilderUClass, NewName));
+				NewBuilder = CastChecked<UMetaSoundBuilderBase>(NewObject<UObject>(&InMetaSoundObject, &BuilderUClass, NewName, RF_Transactional));
 				FMetaSoundFrontendDocumentBuilder& BuilderRef = NewBuilder->GetBuilder();
 				BuilderRef = FMetaSoundFrontendDocumentBuilder(DocInterface);
 				const FMetasoundFrontendDocument& Document = DocInterface->GetConstDocument();
diff --git a/Engine/Plugins/Runtime/Metasound/Source/MetasoundFrontend/Public/MetasoundAssetBase.h b/Engine/Plugins/Runtime/Metasound/Source/MetasoundFrontend/Public/MetasoundAssetBase.h
--- a/Engine/Plugins/Runtime/Metasound/Source/MetasoundFrontend/Public/MetasoundAssetBase.h
+++ b/Engine/Plugins/Runtime/Metasound/Source/MetasoundFrontend/Public/MetasoundAssetBase.h
@@ -141,11 +141,11 @@ public:
 	// Imports the asset from a JSON file at provided path
 	bool ImportFromJSONAsset(const FString& InAbsolutePath);
 
-	// Returns handle for the root metasound graph of this asset.
+	// Soft Deprecated in favor of DocumentBuilder API. Returns handle for the root metasound graph of this asset.
 	Metasound::Frontend::FDocumentHandle GetDocumentHandle();
 	Metasound::Frontend::FConstDocumentHandle GetDocumentHandle() const;
 
-	// Returns handle for the root metasound graph of this asset.
+	// Soft Deprecated in favor of DocumentBuilder API. Returns handle for the root metasound graph of this asset.
 	Metasound::Frontend::FGraphHandle GetRootGraphHandle();
 	Metasound::Frontend::FConstGraphHandle GetRootGraphHandle() const;
 
@@ -155,7 +155,7 @@ public:
 	const FMetasoundFrontendDocument& GetConstDocumentChecked() const;
 
 	// Soft deprecated.  Document layer should not be directly mutated via asset base in anticipation
-	// of moving to all mutation via the Frontend Document & Builder Subsystem API.
+	// of moving all mutable document calls to the Frontend/Subsystem Document Builder API.
 	FMetasoundFrontendDocument& GetDocumentChecked();
 
 	UE_DEPRECATED(5.5, "Use GetConstDocumentChecked instead.")
