Make deprecated ModifyContext mutable while deprecating calls to avoid having to invalidate builder cache on legacy controller accessor calls in favor of manipulation through non-cache clearing const document accessors

#rnx

diff --git a/Engine/Plugins/Runtime/Metasound/Source/MetasoundFrontend/Private/MetasoundAssetBase.cpp b/Engine/Plugins/Runtime/Metasound/Source/MetasoundFrontend/Private/MetasoundAssetBase.cpp
--- a/Engine/Plugins/Runtime/Metasound/Source/MetasoundFrontend/Private/MetasoundAssetBase.cpp
+++ b/Engine/Plugins/Runtime/Metasound/Source/MetasoundFrontend/Private/MetasoundAssetBase.cpp
@@ -680,7 +680,9 @@ void FMetasoundAssetBase::CacheRegistryMetadata()
 
 FMetasoundFrontendDocumentModifyContext& FMetasoundAssetBase::GetModifyContext()
 {
-	return GetDocumentChecked().Metadata.ModifyContext;
+	// ModifyContext is now mutable to avoid mutations to it requiring access through
+	// the deprecated Document controller causing the builder cache to get wiped unnecessarily.
+	return GetConstDocumentChecked().Metadata.ModifyContext;
 }
 
 const FMetasoundFrontendDocumentModifyContext& FMetasoundAssetBase::GetConstModifyContext() const
diff --git a/Engine/Plugins/Runtime/Metasound/Source/MetasoundFrontend/Public/MetasoundFrontendDocument.h b/Engine/Plugins/Runtime/Metasound/Source/MetasoundFrontend/Public/MetasoundFrontendDocument.h
--- a/Engine/Plugins/Runtime/Metasound/Source/MetasoundFrontend/Public/MetasoundFrontendDocument.h
+++ b/Engine/Plugins/Runtime/Metasound/Source/MetasoundFrontend/Public/MetasoundFrontendDocument.h
@@ -1725,7 +1725,8 @@ struct METASOUNDFRONTEND_API FMetasoundFrontendDocumentMetadata
 	FMetasoundFrontendVersion Version;
 
 #if WITH_EDITORONLY_DATA
-	FMetasoundFrontendDocumentModifyContext ModifyContext;
+	// Actively being deprecated in favor of Document Builder Transaction Listener API.
+	mutable FMetasoundFrontendDocumentModifyContext ModifyContext;
 
 	// Map of MemberID to metadata used to constrain how literals can be manipulated
 	// with the editor context. This can be used to implement things like numeric ranges,
