Use full paths if valid when removing paths in the MetaSound asset manager

diff --git a/Engine/Plugins/Runtime/Metasound/Source/MetasoundEngine/Private/MetasoundAssetSubsystem.cpp b/Engine/Plugins/Runtime/Metasound/Source/MetasoundEngine/Private/MetasoundAssetSubsystem.cpp
--- a/Engine/Plugins/Runtime/Metasound/Source/MetasoundEngine/Private/MetasoundAssetSubsystem.cpp
+++ b/Engine/Plugins/Runtime/Metasound/Source/MetasoundEngine/Private/MetasoundAssetSubsystem.cpp
@@ -90,8 +90,17 @@ namespace Metasound::Engine
 			check(IsInGameThread());
 			if (TArray<FTopLevelAssetPath>* MapAssetPaths = InMap.Find(AssetKey))
 			{
-				// Package names are stripped on destruction, so only asset name is reliable
-				auto ComparePaths = [&AssetPath](const FTopLevelAssetPath& Path) { return Path.GetAssetName() == AssetPath.GetAssetName(); };
+				auto ComparePaths = [&AssetPath](const FTopLevelAssetPath& Path) 
+				{
+					// Compare full paths if valid
+					if (Path.IsValid() && AssetPath.IsValid())
+					{
+						return Path == AssetPath;
+					}
+					// Package names are stripped on destruction, so only asset name is reliable
+					return Path.GetAssetName() == AssetPath.GetAssetName(); 
+				};
+
 				if (MapAssetPaths->RemoveAllSwap(ComparePaths, EAllowShrinking::No) > 0)
 				{
 					if (MapAssetPaths->IsEmpty())
