More Page Doc Builder API updates

#rnx

diff --git a/Engine/Plugins/Runtime/Metasound/Source/MetasoundFrontend/Private/MetasoundFrontendDocumentBuilder.cpp b/Engine/Plugins/Runtime/Metasound/Source/MetasoundFrontend/Private/MetasoundFrontendDocumentBuilder.cpp
--- a/Engine/Plugins/Runtime/Metasound/Source/MetasoundFrontend/Private/MetasoundFrontendDocumentBuilder.cpp
+++ b/Engine/Plugins/Runtime/Metasound/Source/MetasoundFrontend/Private/MetasoundFrontendDocumentBuilder.cpp
@@ -1350,7 +1350,7 @@ void FMetaSoundFrontendDocumentBuilder::ClearDocument(TSharedRef<Metasound::Fron
 	GraphClass.PresetOptions.InputsInheritingDefault.Reset();
 	GraphClass.PresetOptions.bIsPreset = false;
 
-	GraphClass.RemoveAllGraphPages();
+	RemoveAllGraphPages();
 
 PRAGMA_DISABLE_DEPRECATION_WARNINGS
 	FMetasoundFrontendGraph& Graph = GraphClass.GetDefaultGraph();
@@ -2406,7 +2406,14 @@ void FMetaSoundFrontendDocumentBuilder::InitNodeLocations()
 
 bool FMetaSoundFrontendDocumentBuilder::IsDependencyReferenced(const FGuid& InClassID) const
 {
-	return DocumentCache->GetNodeCache().ContainsNodesOfClassID(InClassID);
+	bool bIsReferenced = false;
+	GetConstDocumentChecked().RootGraph.IterateGraphPages([this, &InClassID, &bIsReferenced](const FMetasoundFrontendGraph& Graph)
+	{
+		using namespace Metasound::Frontend;
+		const IDocumentGraphNodeCache& NodeCache = DocumentCache->GetNodeCache(Graph.PageID);
+		bIsReferenced |= NodeCache.ContainsNodesOfClassID(InClassID);
+	});
+	return bIsReferenced;
 }
 
 bool FMetaSoundFrontendDocumentBuilder::IsNodeInputConnected(const FGuid& InNodeID, const FGuid& InVertexID, const FGuid* InPageID) const
@@ -2637,21 +2644,19 @@ void FMetaSoundFrontendDocumentBuilder::FinishBuilding()
 
 void FMetaSoundFrontendDocumentBuilder::RemoveAllGraphPages()
 {
+	using namespace Metasound;
+
 	FMetasoundFrontendGraphClass& RootGraph = GetDocumentChecked().RootGraph;
 	RootGraph.IterateGraphPages([this](FMetasoundFrontendGraph& Graph)
 	{
-		using namespace Metasound::Frontend;
-
-		if (Graph.PageID != Metasound::Frontend::DefaultGraphPageID)
+		using namespace Metasound;
+		if (Graph.PageID != Frontend::DefaultGraphPageID)
 		{
-			DocumentDelegates->PageDelegates.OnRemovingPage.Broadcast(Metasound::Frontend::FDocumentMutatePageArgs { Graph.PageID });
+			DocumentDelegates->PageDelegates.OnRemovingPage.Broadcast(Frontend::FDocumentMutatePageArgs { Graph.PageID });
 		}
 	});
 	RootGraph.RemoveAllGraphPages();
-	if (BuildPageID == Metasound::Frontend::DefaultGraphPageID)
-	{
-		SetBuildPageID(FGuid());
-	}
+	SetBuildPageID(Frontend::DefaultGraphPageID);
 }
 
 bool FMetaSoundFrontendDocumentBuilder::RemoveDependency(const FGuid& InClassID)
@@ -2675,9 +2680,7 @@ bool FMetaSoundFrontendDocumentBuilder::RemoveDependency(const FGuid& InClassID)
 			}
 		}
 
-		const int32 LastIndex = Dependencies.Num() - 1;
-		DocumentDelegates->OnRemoveSwappingDependency.Broadcast(Index, LastIndex);
-		Dependencies.RemoveAtSwap(Index, EAllowShrinking::No);
+		RemoveSwapDependencyInternal(Index);
 	}
 
 	return true;
@@ -2705,14 +2708,21 @@ bool FMetaSoundFrontendDocumentBuilder::RemoveDependency(EMetasoundFrontendClass
 			}
 		}
 
-		const int32 LastIndex = Dependencies.Num() - 1;
-		DocumentDelegates->OnRemoveSwappingDependency.Broadcast(Index, LastIndex);
-		Dependencies.RemoveAtSwap(Index, EAllowShrinking::No);
+		RemoveSwapDependencyInternal(Index);
 	}
 
 	return true;
 }
 
+void FMetaSoundFrontendDocumentBuilder::RemoveSwapDependencyInternal(int32 Index)
+{
+	FMetasoundFrontendDocument& Document = GetDocumentChecked();
+	TArray<FMetasoundFrontendClass>& Dependencies = Document.Dependencies;
+	const int32 LastIndex = Dependencies.Num() - 1;
+	DocumentDelegates->OnRemoveSwappingDependency.Broadcast(Index, LastIndex);
+	Dependencies.RemoveAtSwap(Index, EAllowShrinking::No);
+}
+
 bool FMetaSoundFrontendDocumentBuilder::RemoveEdge(const FMetasoundFrontendEdge& EdgeToRemove, const FGuid* InPageID)
 {
 	using namespace Metasound::Frontend;
@@ -3252,14 +3262,18 @@ bool FMetaSoundFrontendDocumentBuilder::RemoveUnusedDependencies()
 {
 	bool bDidEdit = false;
 
-	TArray<FMetasoundFrontendClass>& Dependencies = GetDocumentChecked().Dependencies;
-	for (int32 i = Dependencies.Num() - 1; i >= 0; --i)
+	const FMetasoundFrontendDocument& Document = GetConstDocumentChecked();
+	const FMetasoundFrontendGraphClass& RootGraph = Document.RootGraph;
+	const TArray<FMetasoundFrontendClass>& Dependencies = Document.Dependencies;
+
+	for (int32 Index = Dependencies.Num() - 1; Index >= 0; --Index)
 	{
-		const FGuid& ClassID = Dependencies[i].ID;
-		if (!DocumentCache->GetNodeCache().ContainsNodesOfClassID(ClassID))
+		const FGuid& ClassID = Dependencies[Index].ID;
+		const bool bIsReferenced = IsDependencyReferenced(ClassID);
+		if (!bIsReferenced)
 		{
-			const bool bRemoved = RemoveDependency(ClassID);
-			bDidEdit |= ensureAlwaysMsgf(bRemoved, TEXT("Failed to remove dependency that was found on document and was not referenced by nodes"));
+			RemoveSwapDependencyInternal(Index);
+			bDidEdit = true;
 		}
 	}
 
@@ -3299,6 +3313,7 @@ bool FMetaSoundFrontendDocumentBuilder::SetBuildPageID(const FGuid& InBuildPageI
 
 			constexpr bool bPrimeCache = false;
 			DocumentCache->SetBuildPageID(BuildPageID);
+			DocumentDelegates->PageDelegates.OnPageSet.Broadcast({ BuildPageID });
 		}
 		return true;
 	}
@@ -3730,9 +3745,9 @@ bool FMetaSoundFrontendDocumentBuilder::SetNodeLocation(const FGuid& InNodeID, c
 	return false;
 }
 
-bool FMetaSoundFrontendDocumentBuilder::SetNodeUnconnectedPinsHidden(const FGuid& InNodeID, const bool bUnconnectedPinsHidden)
+bool FMetaSoundFrontendDocumentBuilder::SetNodeUnconnectedPinsHidden(const FGuid& InNodeID, const bool bUnconnectedPinsHidden, const FGuid* InPageID)
 {
-	if (FMetasoundFrontendNode* Node = FindNodeInternal(InNodeID))
+	if (FMetasoundFrontendNode* Node = FindNodeInternal(InNodeID, InPageID))
 	{
 		Node->Style.bUnconnectedPinsHidden = bUnconnectedPinsHidden;
 		return true;
@@ -3741,9 +3756,9 @@ bool FMetaSoundFrontendDocumentBuilder::SetNodeUnconnectedPinsHidden(const FGuid
 	return false;
 }
 
-const FMetasoundFrontendNodeStyle* FMetaSoundFrontendDocumentBuilder::GetNodeStyle(const FGuid& InNodeID)
+const FMetasoundFrontendNodeStyle* FMetaSoundFrontendDocumentBuilder::GetNodeStyle(const FGuid& InNodeID, const FGuid* InPageID)
 {
-	if (const FMetasoundFrontendNode* Node = FindNodeInternal(InNodeID))
+	if (const FMetasoundFrontendNode* Node = FindNodeInternal(InNodeID, InPageID))
 	{
 		return &Node->Style;
 	}
diff --git a/Engine/Plugins/Runtime/Metasound/Source/MetasoundFrontend/Public/MetasoundFrontendDocumentBuilder.h b/Engine/Plugins/Runtime/Metasound/Source/MetasoundFrontend/Public/MetasoundFrontendDocumentBuilder.h
--- a/Engine/Plugins/Runtime/Metasound/Source/MetasoundFrontend/Public/MetasoundFrontendDocumentBuilder.h
+++ b/Engine/Plugins/Runtime/Metasound/Source/MetasoundFrontend/Public/MetasoundFrontendDocumentBuilder.h
@@ -244,6 +244,11 @@ public:
 	// The graph ID used when requests are made to mutate specific paged graph topology (ex. adding or removing nodes or edges)
 	const FGuid& GetBuildPageID() const;
 
+#if WITH_EDITOR
+	// Gets the editor-only style of a node with the given ID.
+	const FMetasoundFrontendNodeStyle* GetNodeStyle(const FGuid& InNodeID, const FGuid* InPageID = nullptr);
+#endif // WITH_EDITOR
+
 	template<typename TObjectType>
 	TObjectType& CastDocumentObjectChecked() const
 	{
@@ -394,10 +399,7 @@ public:
 	bool SetNodeLocation(const FGuid& InNodeID, const FVector2D& InLocation, const FGuid* InLocationGuid = nullptr, const FGuid* InPageID = nullptr);
 
 	// Sets the editor-only Unconnected Pins Hidden for a node with the given ID.
-	bool SetNodeUnconnectedPinsHidden(const FGuid& InNodeID, const bool bUnconnectedPinsHidden);
-
-	// Gets the editor-only style of a node with the given ID.
-	const FMetasoundFrontendNodeStyle* GetNodeStyle(const FGuid& InNodeID);
+	bool SetNodeUnconnectedPinsHidden(const FGuid& InNodeID, const bool bUnconnectedPinsHidden, const FGuid* InPageID = nullptr);
 #endif // WITH_EDITOR
 
 	bool SetNodeInputDefault(const FGuid& InNodeID, const FGuid& InVertexID, const FMetasoundFrontendLiteral& InLiteral, const FGuid* InPageID = nullptr);
@@ -445,6 +447,8 @@ private:
 	FMetasoundFrontendDocument& GetDocumentChecked() const;
 	IMetaSoundDocumentInterface& GetDocumentInterfaceChecked() const;
 
+	void RemoveSwapDependencyInternal(int32 Index);
+
 	bool SetGraphInputInheritsDefault(FName InName, bool bInputInheritsDefault);
 
 	bool SpliceVariableNodeFromStack(const FGuid& InNodeID, const FGuid& InPageID);
@@ -453,7 +457,8 @@ private:
 	UPROPERTY(Transient)
 	TScriptInterface<IMetaSoundDocumentInterface> DocumentInterface;
 
-	// Default page ID to apply build transactions to if no page ID is optionally set directly.
+	// Page ID to apply build transaction to if no optional PageID is provided in explicit function call.
+	// (Also used to support back compat for Controller API until mutable controllers are adequately deprecated).
 	UPROPERTY(Transient)
 	FGuid BuildPageID;
 
