Missing file
#rnx

diff --git a/Engine/Plugins/Runtime/Metasound/Source/MetasoundFrontend/Public/MetasoundFrontendDocumentModifyDelegates.h b/Engine/Plugins/Runtime/Metasound/Source/MetasoundFrontend/Public/MetasoundFrontendDocumentModifyDelegates.h
--- a/Engine/Plugins/Runtime/Metasound/Source/MetasoundFrontend/Public/MetasoundFrontendDocumentModifyDelegates.h
+++ b/Engine/Plugins/Runtime/Metasound/Source/MetasoundFrontend/Public/MetasoundFrontendDocumentModifyDelegates.h
@@ -23,11 +23,13 @@ namespace Metasound::Frontend
 
 	DECLARE_MULTICAST_DELEGATE_OneParam(FOnDocumentPageAdded, const FDocumentMutatePageArgs& /* Args */);
 	DECLARE_MULTICAST_DELEGATE_OneParam(FOnDocumentRemovingPage, const FDocumentMutatePageArgs& /* Args */);
+	DECLARE_MULTICAST_DELEGATE_OneParam(FOnDocumentPageSet, const FDocumentMutatePageArgs& /* Args */);
 
 	struct METASOUNDFRONTEND_API FPageModifyDelegates
 	{
 		FOnDocumentPageAdded OnPageAdded;
 		FOnDocumentRemovingPage OnRemovingPage;
+		FOnDocumentPageSet OnPageSet;
 	};
 
 	struct METASOUNDFRONTEND_API FInterfaceModifyDelegates
