Fix MetaSound non interface members being set to read only if their name contains an interface namespace

diff --git a/Engine/Plugins/Runtime/Metasound/Source/MetasoundEditor/Private/MetasoundEditorGraph.cpp b/Engine/Plugins/Runtime/Metasound/Source/MetasoundEditor/Private/MetasoundEditorGraph.cpp
--- a/Engine/Plugins/Runtime/Metasound/Source/MetasoundEditor/Private/MetasoundEditorGraph.cpp
+++ b/Engine/Plugins/Runtime/Metasound/Source/MetasoundEditor/Private/MetasoundEditorGraph.cpp
@@ -410,6 +410,11 @@ const FMetasoundFrontendVersion& UMetasoundEditorGraphVertex::GetInterfaceVersio
 }
 
 bool UMetasoundEditorGraphVertex::IsInterfaceMember(FMetasoundFrontendInterface* OutInterface) const
+{
+	return false;
+}
+
+bool UMetasoundEditorGraphVertex::NameContainsInterfaceNamespace(FMetasoundFrontendInterface* OutInterface) const
 {
 	using namespace Metasound::Frontend;
 
@@ -418,14 +423,13 @@ bool UMetasoundEditorGraphVertex::IsInterfaceMember(FMetasoundFrontendInterface*
 	FName ParamName;
 	Audio::FParameterPath::SplitName(MemberName, InterfaceNamespace, ParamName);
 
-	FMetasoundFrontendInterface InterfaceToValidate;
-	if (!InterfaceNamespace.IsNone() && ISearchEngine::Get().FindInterfaceWithHighestVersion(InterfaceNamespace, InterfaceToValidate))
+	FMetasoundFrontendInterface FoundInterface;
+	if (!InterfaceNamespace.IsNone() && ISearchEngine::Get().FindInterfaceWithHighestVersion(InterfaceNamespace, FoundInterface))
 	{
 		if (OutInterface)
 		{
-			*OutInterface = MoveTemp(InterfaceToValidate);
+			*OutInterface = MoveTemp(FoundInterface);
 		}
-
 		return true;
 	}
 
@@ -663,6 +667,38 @@ TArray<UMetasoundEditorGraphMemberNode*> UMetasoundEditorGraphInput::GetNodes()
 	return Nodes;
 }
 
+bool UMetasoundEditorGraphInput::IsInterfaceMember(FMetasoundFrontendInterface* OutInterface) const
+{
+	FMetasoundFrontendInterface Interface;
+	if (NameContainsInterfaceNamespace(&Interface))
+	{
+		// Check if Input is a member of the found interface
+		if (const FMetasoundFrontendNode* InputNode = GetFrontendNode())
+		{
+			const FMetasoundFrontendVertex& Input = InputNode->Interface.Inputs.Last();
+			auto IsInput = [&Input](const FMetasoundFrontendClassInput& InterfaceInput)
+			{
+				return FMetasoundFrontendVertex::IsFunctionalEquivalent(Input, InterfaceInput);
+			};
+
+			if (Interface.Inputs.ContainsByPredicate(IsInput))
+			{
+				if (OutInterface)
+				{
+					*OutInterface = MoveTemp(Interface);
+				}
+				return true;
+			}
+		}
+	}
+
+	if (OutInterface)
+	{
+		*OutInterface = { };
+	}
+	return false;
+}
+
 void UMetasoundEditorGraphInput::SetSortOrderIndex(int32 InSortOrderIndex)
 {
 	using namespace Metasound::Editor;
@@ -890,6 +926,40 @@ int32 UMetasoundEditorGraphOutput::GetSortOrderIndex() const
 	return GraphHandle->GetSortOrderIndexForOutput(NodeName);
 }
 
+bool UMetasoundEditorGraphOutput::IsInterfaceMember(FMetasoundFrontendInterface* OutInterface) const
+{
+	using namespace Metasound::Frontend;
+	
+	FMetasoundFrontendInterface Interface;
+	if (NameContainsInterfaceNamespace(&Interface))
+	{
+		// Check if Output is a member of the found interface
+		if (const FMetasoundFrontendNode* OutputNode = GetFrontendNode())
+		{
+			const FMetasoundFrontendVertex& Output = OutputNode->Interface.Outputs.Last();
+			auto IsOutput = [&Output](const FMetasoundFrontendClassOutput& InterfaceOutput)
+			{
+				return FMetasoundFrontendVertex::IsFunctionalEquivalent(Output, InterfaceOutput);
+			};
+
+			if (Interface.Outputs.ContainsByPredicate(IsOutput))
+			{
+				if (OutInterface)
+				{
+					*OutInterface = MoveTemp(Interface);
+				}
+				return true;
+			}
+		}
+	}
+	
+	if (OutInterface)
+	{
+		*OutInterface = { };
+	}
+	return false;
+}
+
 void UMetasoundEditorGraphOutput::SetSortOrderIndex(int32 InSortOrderIndex)
 {
 	using namespace Metasound::Editor;
diff --git a/Engine/Plugins/Runtime/Metasound/Source/MetasoundEditor/Public/MetasoundEditorGraph.h b/Engine/Plugins/Runtime/Metasound/Source/MetasoundEditor/Public/MetasoundEditorGraph.h
--- a/Engine/Plugins/Runtime/Metasound/Source/MetasoundEditor/Public/MetasoundEditorGraph.h
+++ b/Engine/Plugins/Runtime/Metasound/Source/MetasoundEditor/Public/MetasoundEditorGraph.h
@@ -192,6 +192,9 @@ protected:
 
 	const FMetasoundFrontendNode* GetFrontendNode() const;
 
+	/** If the vertex's name contains an interface namespace, return true and set OutInterface to that interface if provided. */
+	bool NameContainsInterfaceNamespace(FMetasoundFrontendInterface* OutInterface) const;
+
 public:
 	/** Initializes all properties with the given parameters required to identify the frontend member from this editor graph member. */
 	void InitMember(FName InDataType, const FMetasoundFrontendLiteral& InDefaultLiteral, FGuid InNodeID, FMetasoundFrontendClassName&& InClassName);
@@ -262,6 +265,7 @@ public:
 	virtual TArray<UMetasoundEditorGraphMemberNode*> GetNodes() const override;
 
 	virtual const FText& GetGraphMemberLabel() const override;
+	virtual bool IsInterfaceMember(FMetasoundFrontendInterface* OutInterface = nullptr) const override;
 	virtual void ResetToClassDefault() override;
 	virtual void SetDataType(FName InNewType, bool bPostTransaction) override;
 	virtual void SetVertexAccessType(EMetasoundFrontendVertexAccessType InNewAccessType, bool bPostTransaction) override;
@@ -286,6 +290,7 @@ class METASOUNDEDITOR_API UMetasoundEditorGraphOutput : public UMetasoundEditorG
 public:
 	virtual FText GetDescription() const override;
 	virtual int32 GetSortOrderIndex() const override;
+	virtual bool IsInterfaceMember(FMetasoundFrontendInterface* OutInterface = nullptr) const override;
 	virtual void SetSortOrderIndex(int32 InSortOrderIndex) override;
 	virtual const FText& GetGraphMemberLabel() const override;
 	virtual void ResetToClassDefault() override;
