Temporarily disable MetaSound class name asset validation, add cvar

#rnx

diff --git a/Engine/Plugins/Runtime/Metasound/Source/MetasoundEngine/Private/MetasoundEngineAsset.h b/Engine/Plugins/Runtime/Metasound/Source/MetasoundEngine/Private/MetasoundEngineAsset.h
--- a/Engine/Plugins/Runtime/Metasound/Source/MetasoundEngine/Private/MetasoundEngineAsset.h
+++ b/Engine/Plugins/Runtime/Metasound/Source/MetasoundEngine/Private/MetasoundEngineAsset.h
@@ -150,7 +150,11 @@ namespace Metasound::Engine
 		template <typename TMetaSoundObject>
 		static EDataValidationResult IsDataValid(const TMetaSoundObject& InMetaSound, FDataValidationContext& InOutContext)
 		{
-			const EDataValidationResult Result = IsClassNameUnique(InMetaSound, InOutContext);
+			EDataValidationResult Result = EDataValidationResult::Valid;
+			if (MetasoundEngineModulePrivate::EnableMetaSoundEditorAssetValidation)
+			{
+				Result = IsClassNameUnique(InMetaSound, InOutContext);
+			}
 			return Result;
 		}
 
diff --git a/Engine/Plugins/Runtime/Metasound/Source/MetasoundEngine/Private/MetasoundEngineModule.cpp b/Engine/Plugins/Runtime/Metasound/Source/MetasoundEngine/Private/MetasoundEngineModule.cpp
--- a/Engine/Plugins/Runtime/Metasound/Source/MetasoundEngine/Private/MetasoundEngineModule.cpp
+++ b/Engine/Plugins/Runtime/Metasound/Source/MetasoundEngine/Private/MetasoundEngineModule.cpp
@@ -53,6 +53,21 @@ REGISTER_METASOUND_DATATYPE(Metasound::FWaveTableBankAsset, "WaveTableBankAsset"
 
 namespace Metasound::Engine
 {
+#if WITH_EDITOR
+	namespace MetasoundEngineModulePrivate
+	{
+		// Temporarily disabled while fixing up assets
+		int32 EnableMetaSoundEditorAssetValidation = 0;
+
+		FAutoConsoleVariableRef CVarEnableMetaSoundEditorAssetValidation(
+			TEXT("au.MetaSound.Editor.EnableAssetValidation"),
+			EnableMetaSoundEditorAssetValidation,
+			TEXT("Enables MetaSound specific asset validation.\n")
+			TEXT("Default: 1 (Enabled)"),
+			ECVF_Default);
+	} // namespace MetasoundEngineModulePrivate
+#endif // WITH_EDITOR
+
 	void FModule::FModule::StartupModule() 
 	{
 		using namespace Frontend;
diff --git a/Engine/Plugins/Runtime/Metasound/Source/MetasoundEngine/Private/MetasoundEngineModule.h b/Engine/Plugins/Runtime/Metasound/Source/MetasoundEngine/Private/MetasoundEngineModule.h
--- a/Engine/Plugins/Runtime/Metasound/Source/MetasoundEngine/Private/MetasoundEngineModule.h
+++ b/Engine/Plugins/Runtime/Metasound/Source/MetasoundEngine/Private/MetasoundEngineModule.h
@@ -6,8 +6,20 @@
 #include "MetasoundFrontendRegistryContainer.h"
 #include "UObject/GCObject.h"
 
+#if WITH_EDITOR
+#include "HAL/IConsoleManager.h"
+#endif // WITH_EDITOR
+
 namespace Metasound::Engine
 {
+#if WITH_EDITOR
+	namespace MetasoundEngineModulePrivate
+	{
+		extern int32 EnableMetaSoundEditorAssetValidation;
+		extern FAutoConsoleVariableRef CVarEnableMetaSoundEditorAssetValidation;
+	} // namespace MetasoundEngineModulePrivate
+#endif // WITH_EDITOR
+
 	class FModule : public IMetasoundEngineModule
 	{
 		// Supplies GC referencing in the MetaSound Frontend node registry for doing
