UE-131681 - Improving Convert From Preset Workflow

diff --git a/Engine/Plugins/Runtime/Metasound/Source/MetasoundEditor/Private/MetasoundEditor.cpp b/Engine/Plugins/Runtime/Metasound/Source/MetasoundEditor/Private/MetasoundEditor.cpp
--- a/Engine/Plugins/Runtime/Metasound/Source/MetasoundEditor/Private/MetasoundEditor.cpp
+++ b/Engine/Plugins/Runtime/Metasound/Source/MetasoundEditor/Private/MetasoundEditor.cpp
@@ -1342,6 +1342,12 @@ namespace Metasound
 			}
 
 			FSlateApplication::Get().DismissAllMenus();
+
+			// In case of undoing 'convert from preset' refresh toolbar to include ConvertFromPreset button
+			if (UToolMenus* ToolMenus = UToolMenus::Get())
+			{
+				ToolMenus->RefreshAllWidgets();
+			}
 		}
 
 		void FEditor::NotifyAssetPrimeInProgress()
@@ -2268,15 +2274,68 @@ namespace Metasound
 
 			if (Builder.IsValid())
 			{
-				EMetaSoundBuilderResult Result = EMetaSoundBuilderResult::Failed;
-				Builder->ConvertFromPreset(Result);
-				ensure(Result == EMetaSoundBuilderResult::Succeeded);
+				TSharedPtr<SWindow> DialogWindow =
+					SNew(SWindow)
+					.Title(LOCTEXT("MetasoundPresetDialogTitle", "Convert From Preset?"))
+					.SupportsMinimize(false)
+					.SupportsMaximize(false)
+					.SizingRule(ESizingRule::Autosized)
+					.AutoCenter(EAutoCenter::PreferredWorkArea);
+				
+				TSharedPtr<SBox> DialogContent =
+					SNew(SBox)
+					.HAlign(HAlign_Center)
+					.VAlign(VAlign_Center)
+					[
+						SNew(SHorizontalBox)
+						+ SHorizontalBox::Slot()
+						.AutoWidth()
+						.HAlign(HAlign_Left)
+						.VAlign(VAlign_Bottom)
+						[
+							SNew(SButton)
+							.Text(LOCTEXT("MetasoundPresetDialogAccept", "Accept"))
+							.OnClicked_Lambda([this, &DialogWindow]()
+							{
+								const FScopedTransaction Transaction(LOCTEXT("ConvertFromPresetText", "Convert From Preset"));
+								GetMetasoundObject()->Modify();
+								
+								EMetaSoundBuilderResult Result = EMetaSoundBuilderResult::Failed;
+								Builder->ConvertFromPreset(Result);
+								ensure(Result == EMetaSoundBuilderResult::Succeeded);
+								
+								if (UToolMenus* ToolMenus = UToolMenus::Get())
+								{
+									ToolMenus->RefreshAllWidgets();
+								}
+								
+								RefreshGraphMemberMenu();
+								RefreshDetails();
 
-				if (UToolMenus* ToolMenus = UToolMenus::Get())
-				{
-					ToolMenus->RefreshAllWidgets();
-				}
-				RefreshDetails();
+								DialogWindow->RequestDestroyWindow();
+								
+								return FReply::Handled();
+							})
+						]
+						+ SHorizontalBox::Slot()
+						.AutoWidth()
+						.HAlign(HAlign_Right)
+						.VAlign(VAlign_Bottom)
+						[
+							SNew(SButton)
+							.Text(LOCTEXT("MetasoundPresetDialogCancel", "Cancel"))
+							.OnClicked_Lambda([&DialogWindow]()
+							{
+								DialogWindow->RequestDestroyWindow();
+								
+								return FReply::Handled();
+							})
+						]
+				];
+				
+				DialogWindow->SetContent(DialogContent.ToSharedRef());
+				
+				FSlateApplication::Get().AddModalWindow(DialogWindow.ToSharedRef(), GetGraphEditor());
 			}
 		}
 
