UE-216214 - [Workflow Improvement] - Add autoconversion nodes when making connections for advanced MetaSound types

diff --git a/Engine/Plugins/Runtime/Metasound/Source/MetasoundStandardNodes/Private/MetasoundConversionNodes.cpp b/Engine/Plugins/Runtime/Metasound/Source/MetasoundStandardNodes/Private/MetasoundConversionNodes.cpp
--- a/Engine/Plugins/Runtime/Metasound/Source/MetasoundStandardNodes/Private/MetasoundConversionNodes.cpp
+++ b/Engine/Plugins/Runtime/Metasound/Source/MetasoundStandardNodes/Private/MetasoundConversionNodes.cpp
@@ -15,6 +15,8 @@
 
 #define LOCTEXT_NAMESPACE "MetasoundStandardNodes_Conversion"
 
+#define METASOUND_REGISTER_CONVERSION(NodeClass, FromDataType, ToDataType) static bool bSuccessfullyRegisteredConversion##NodeClass = FMetasoundFrontendRegistryContainer::Get()->EnqueueInitCommand([]() { Metasound::RegisterConversionOperator<FromDataType, ToDataType>(); });
+
 namespace Metasound
 {
 	namespace ConversionNodeVertexNames
@@ -321,6 +323,38 @@ namespace Metasound
 
 	using FConversionAudioToFloat = TConversionNode<FAudioBuffer, float>;
 	METASOUND_REGISTER_NODE(FConversionAudioToFloat)
+
+
+	template<typename TFromType, typename TToType>
+	void RegisterConversionOperator()
+	{
+		using FConverterNodeRegistryKey = ::Metasound::Frontend::FConverterNodeRegistryKey;
+		using FConverterNodeInfo = ::Metasound::Frontend::FConverterNodeInfo;
+
+		FName FromType = ::Metasound::TDataReferenceTypeInfo<TFromType>::TypeName;
+		FName ToType = ::Metasound::TDataReferenceTypeInfo<TToType>::TypeName;
+
+		FConverterNodeRegistryKey RegistryKey = { FromType, ToType };
+
+		const FNodeClassMetadata& Metadata = TConversionOperator<TFromType, TToType>::GetNodeInfo();
+
+		FConverterNodeInfo ConverterNodeInfo =
+		{
+			FromType,
+			ToType,
+			Metasound::Frontend::FNodeRegistryKey(Metadata)
+		};
+
+		FMetasoundFrontendRegistryContainer::Get()->RegisterConversionNode(RegistryKey, ConverterNodeInfo);
+	}
+
+	METASOUND_REGISTER_CONVERSION(FConversionFloatToTime, float, FTime)
+	METASOUND_REGISTER_CONVERSION(FConversionTimeToFloat, FTime, float)
+	METASOUND_REGISTER_CONVERSION(FConversionInt32ToTime, int32, FTime)
+	METASOUND_REGISTER_CONVERSION(FConversionTimeToInt32, FTime, int32)
+	METASOUND_REGISTER_CONVERSION(FConversionFloatToAudio, float, FAudioBuffer)
+	METASOUND_REGISTER_CONVERSION(FConversionAudioToFloat, FAudioBuffer, float)
 }
 
+#undef METASOUND_REGISTER_CONVERSION
 #undef LOCTEXT_NAMESPACE
