Reenable MetaSound asset validation

#rnx

diff --git a/Engine/Plugins/Runtime/Metasound/Source/MetasoundEngine/Private/MetasoundEngineModule.cpp b/Engine/Plugins/Runtime/Metasound/Source/MetasoundEngine/Private/MetasoundEngineModule.cpp
--- a/Engine/Plugins/Runtime/Metasound/Source/MetasoundEngine/Private/MetasoundEngineModule.cpp
+++ b/Engine/Plugins/Runtime/Metasound/Source/MetasoundEngine/Private/MetasoundEngineModule.cpp
@@ -56,8 +56,7 @@ namespace Metasound::Engine
 #if WITH_EDITOR
 	namespace MetasoundEngineModulePrivate
 	{
-		// Temporarily disabled while fixing up assets
-		int32 EnableMetaSoundEditorAssetValidation = 0;
+		int32 EnableMetaSoundEditorAssetValidation = 1;
 
 		FAutoConsoleVariableRef CVarEnableMetaSoundEditorAssetValidation(
 			TEXT("au.MetaSound.Editor.EnableAssetValidation"),
