UE-218301 - Text conflict from LOCTEXT macro for namespace for MetasoundEditor.cpp

diff --git a/Engine/Plugins/Runtime/Metasound/Source/MetasoundEditor/Private/MetasoundEditor.cpp b/Engine/Plugins/Runtime/Metasound/Source/MetasoundEditor/Private/MetasoundEditor.cpp
--- a/Engine/Plugins/Runtime/Metasound/Source/MetasoundEditor/Private/MetasoundEditor.cpp
+++ b/Engine/Plugins/Runtime/Metasound/Source/MetasoundEditor/Private/MetasoundEditor.cpp
@@ -499,6 +499,9 @@ namespace Metasound
 				FSlateColor SecondaryColor;
 				GetDefaultStatusSymbol(PrimarySymbol, PrimaryColor, SecondarySymbol, SecondaryColor);
 
+				const FText IncompatibleText = LOCTEXT("MetasoundHoverNotCompatibleText", "'{0}' is not compatible with '{1}'");
+				const FText CompatibleText = LOCTEXT("MetasoundHoverCompatibleText", "Convert {0} to {1}.");
+
 				FText Message;
 				if (GraphMember.IsValid())
 				{
@@ -538,14 +541,12 @@ namespace Metasound
 									if (Connectability.Connectable == Frontend::FConnectability::EConnectable::No)
 									{
 										PrimarySymbol = FAppStyle::GetBrush(TEXT("Graph.ConnectorFeedback.Error"));
-										Message = FText::Format(LOCTEXT("NotCompatible_Error", "'{0}' is not compatible with '{1}'"),
-											FText::FromName(DataType), FText::FromName(OtherDataType));
+										Message = FText::Format(IncompatibleText, FText::FromName(DataType), FText::FromName(OtherDataType));
 									}
 									else if (Connectability.Connectable == Frontend::FConnectability::EConnectable::YesWithConverterNode)
 									{
 										PrimarySymbol = FAppStyle::GetBrush(TEXT("Graph.ConnectorFeedback.ViaCast"));
-										Message = FText::Format(LOCTEXT("NotCompatible_Error", "Convert {0} to {1}."),
-											FText::FromName(DataType), FText::FromName(OtherDataType));
+										Message = FText::Format(CompatibleText, FText::FromName(DataType), FText::FromName(OtherDataType));
 									}
 								}
 							}
@@ -582,14 +583,12 @@ namespace Metasound
 										if (Connectability.Connectable == Frontend::FConnectability::EConnectable::No)
 										{
 											PrimarySymbol = FAppStyle::GetBrush(TEXT("Graph.ConnectorFeedback.Error"));
-											Message = FText::Format(LOCTEXT("NotCompatible_Error", "'{0}' is not compatible with '{1}'"),
-												FText::FromName(DataType), FText::FromName(OtherDataType));
+											Message = FText::Format(IncompatibleText, FText::FromName(DataType), FText::FromName(OtherDataType));
 										}
 										else if (Connectability.Connectable == Frontend::FConnectability::EConnectable::YesWithConverterNode)
 										{
 											PrimarySymbol = FAppStyle::GetBrush(TEXT("Graph.ConnectorFeedback.ViaCast"));
-											Message = FText::Format(LOCTEXT("NotCompatible_Error", "Convert {0} to {1}."),
-												FText::FromName(DataType), FText::FromName(OtherDataType));
+											Message = FText::Format(CompatibleText, FText::FromName(DataType), FText::FromName(OtherDataType));
 										}
 									}
 								}
@@ -644,14 +643,12 @@ namespace Metasound
 									if (Connectability.Connectable == Frontend::FConnectability::EConnectable::No)
 									{
 										PrimarySymbol = FAppStyle::GetBrush(TEXT("Graph.ConnectorFeedback.Error"));
-										Message = FText::Format(LOCTEXT("NotCompatible_Error", "'{0}' is not compatible with '{1}'"),
-											FText::FromName(DataType), FText::FromName(OtherDataType));
+										Message = FText::Format(IncompatibleText, FText::FromName(DataType), FText::FromName(OtherDataType));
 									}
 									else if (Connectability.Connectable == Frontend::FConnectability::EConnectable::YesWithConverterNode)
 									{
 										PrimarySymbol = FAppStyle::GetBrush(TEXT("Graph.ConnectorFeedback.ViaCast"));
-										Message = FText::Format(LOCTEXT("NotCompatible_Error", "Convert {0} to {1}."),
-											FText::FromName(DataType), FText::FromName(OtherDataType));
+										Message = FText::Format(CompatibleText, FText::FromName(DataType), FText::FromName(OtherDataType));
 									}
 								}
 							}
