Remove editor setting for "Show Oscilloscope On Audio Pin Mouse Over" and make the feature always enabled.

diff --git a/Engine/Plugins/Runtime/Metasound/Source/MetasoundEditor/Private/MetasoundEditorGraphBuilder.cpp b/Engine/Plugins/Runtime/Metasound/Source/MetasoundEditor/Private/MetasoundEditorGraphBuilder.cpp
--- a/Engine/Plugins/Runtime/Metasound/Source/MetasoundEditor/Private/MetasoundEditorGraphBuilder.cpp
+++ b/Engine/Plugins/Runtime/Metasound/Source/MetasoundEditor/Private/MetasoundEditorGraphBuilder.cpp
@@ -232,14 +232,11 @@ namespace Metasound
 				DataType = OutputHandle->GetDataType();
 			}
 
-			const UMetasoundEditorSettings* MetaSoundSettings = GetDefault<UMetasoundEditorSettings>();
-			const bool bIsDataTypeFAudioBufferSupported = DataType == GetMetasoundDataTypeName<FAudioBuffer>() && MetaSoundSettings && MetaSoundSettings->bShowOscilloscopeOnAudioPinMouseOver;
-
 			const bool bIsSupportedType = DataType == GetMetasoundDataTypeName<float>()
 				|| DataType == GetMetasoundDataTypeName<int32>()
 				|| DataType == GetMetasoundDataTypeName<FString>()
 				|| DataType == GetMetasoundDataTypeName<bool>()
-				|| bIsDataTypeFAudioBufferSupported;
+				|| DataType == GetMetasoundDataTypeName<FAudioBuffer>();
 
 			if (!bIsSupportedType)
 			{
diff --git a/Engine/Plugins/Runtime/Metasound/Source/MetasoundEditor/Private/MetasoundEditorSettings.h b/Engine/Plugins/Runtime/Metasound/Source/MetasoundEditor/Private/MetasoundEditorSettings.h
--- a/Engine/Plugins/Runtime/Metasound/Source/MetasoundEditor/Private/MetasoundEditorSettings.h
+++ b/Engine/Plugins/Runtime/Metasound/Source/MetasoundEditor/Private/MetasoundEditorSettings.h
@@ -173,10 +173,6 @@ public:
 	UPROPERTY(EditAnywhere, config, Category = GraphAnimation, meta = (ShowOnlyInnerProperties))
 	FMetasoundAnalyzerAnimationSettings AnalyzerAnimationSettings;
 
-	/** Whether to show a small oscilloscope when hovering audio pins. */
-	UPROPERTY(EditAnywhere, config, DisplayName = "Show Oscilloscope On Audio Pin Mouse Over", Category = Visualization)
-	bool bShowOscilloscopeOnAudioPinMouseOver = false;
-
 	/** Determines which details view to show in Metasounds Editor */
 	UPROPERTY(Transient)
 	EMetasoundActiveDetailView DetailView = EMetasoundActiveDetailView::General;
diff --git a/Engine/Plugins/Runtime/Metasound/Source/MetasoundEditor/Private/SMetasoundGraphPin.h b/Engine/Plugins/Runtime/Metasound/Source/MetasoundEditor/Private/SMetasoundGraphPin.h
--- a/Engine/Plugins/Runtime/Metasound/Source/MetasoundEditor/Private/SMetasoundGraphPin.h
+++ b/Engine/Plugins/Runtime/Metasound/Source/MetasoundEditor/Private/SMetasoundGraphPin.h
@@ -496,29 +496,25 @@ namespace Metasound
 						});
 
 					// Audio Pin Inspector update
-					if (const UMetasoundEditorSettings* MetasoundSettings = GetDefault<UMetasoundEditorSettings>();
-					    MetasoundSettings && MetasoundSettings->bShowOscilloscopeOnAudioPinMouseOver)
+					FName DataType;
+					if (GraphPin->Direction == EGPD_Input)
 					{
-						FName DataType;
-						if (GraphPin->Direction == EGPD_Input)
-						{
-							const Metasound::Frontend::FConstInputHandle InputHandle = FGraphBuilder::GetConstInputHandleFromPin(GraphPin);
-							DataType = InputHandle->GetDataType();
-						}
-						else
-						{
-							const Metasound::Frontend::FConstOutputHandle OutputHandle = FGraphBuilder::GetConstOutputHandleFromPin(GraphPin);
-							DataType = OutputHandle->GetDataType();
-						}
+						const Metasound::Frontend::FConstInputHandle InputHandle = FGraphBuilder::GetConstInputHandleFromPin(GraphPin);
+						DataType = InputHandle->GetDataType();
+					}
+					else
+					{
+						const Metasound::Frontend::FConstOutputHandle OutputHandle = FGraphBuilder::GetConstOutputHandleFromPin(GraphPin);
+						DataType = OutputHandle->GetDataType();
+					}
 
-						if (DataType == GetMetasoundDataTypeName<FAudioBuffer>())
-						{
-							UpdatePinAudioInspector(*GraphPin, bIsHoveringPin, PinAudioInspector, ParentPinType::ValueInspectorTooltip,
-								[this](FVector2D& OutTooltipLocation)
-								{
-									ParentPinType::GetInteractiveTooltipLocation(OutTooltipLocation);
-								});
-						}
+					if (DataType == GetMetasoundDataTypeName<FAudioBuffer>())
+					{
+						UpdatePinAudioInspector(*GraphPin, bIsHoveringPin, PinAudioInspector, ParentPinType::ValueInspectorTooltip,
+							[this](FVector2D& OutTooltipLocation)
+							{
+								ParentPinType::GetInteractiveTooltipLocation(OutTooltipLocation);
+							});
 					}
 				}
 			}
