Fix mismatched loctext key for MetaSounds graph tab label

#rnx

diff --git a/Engine/Plugins/Runtime/Metasound/Source/MetasoundEditor/Private/MetasoundEditorTabFactory.h b/Engine/Plugins/Runtime/Metasound/Source/MetasoundEditor/Private/MetasoundEditorTabFactory.h
--- a/Engine/Plugins/Runtime/Metasound/Source/MetasoundEditor/Private/MetasoundEditorTabFactory.h
+++ b/Engine/Plugins/Runtime/Metasound/Source/MetasoundEditor/Private/MetasoundEditorTabFactory.h
@@ -51,7 +51,7 @@ namespace Metasound
 				check(Args.GetTabId() == Names::GraphCanvas);
 
 				TSharedRef<SDockTab> SpawnedTab = SNew(SDockTab)
-					.Label(LOCTEXT("MetasoundGraphCanvasTitle", "Viewport"));
+					.Label(LOCTEXT("MetasoundGraphCanvasTitle", "MetaSound Graph"));
 
 				if (GraphEditor.IsValid())
 				{
