ClearDocument clean/fix-ups & bind page delegate to doc cache to ensure node & edge caches don't stick around post page removal
#rnx

diff --git a/Engine/Plugins/Runtime/Metasound/Source/MetasoundFrontend/Private/MetasoundFrontendDocumentBuilder.cpp b/Engine/Plugins/Runtime/Metasound/Source/MetasoundFrontend/Private/MetasoundFrontendDocumentBuilder.cpp
--- a/Engine/Plugins/Runtime/Metasound/Source/MetasoundFrontend/Private/MetasoundFrontendDocumentBuilder.cpp
+++ b/Engine/Plugins/Runtime/Metasound/Source/MetasoundFrontend/Private/MetasoundFrontendDocumentBuilder.cpp
@@ -1336,17 +1336,16 @@ bool FMetaSoundFrontendDocumentBuilder::CanAddEdge(const FMetasoundFrontendEdge&
 void FMetaSoundFrontendDocumentBuilder::ClearDocument(TSharedRef<Metasound::Frontend::FDocumentModifyDelegates> ModifyDelegates)
 {
 	FMetasoundFrontendDocument& Doc = GetDocumentChecked();
+	FMetasoundFrontendGraphClass& GraphClass = Doc.RootGraph;
 
-	Doc.Interfaces.Reset();
-	Doc.Dependencies.Reset();
+	GraphClass.Interface.Inputs.Empty();
+	GraphClass.Interface.Outputs.Empty();
 
-#if WITH_EDITORONLY_DATA
-	Doc.Metadata.MemberMetadata.Reset();
-#endif // WITH_EDITORONLY_DATA
+#if WITH_EDITOR
+	GraphClass.Interface.SetInputStyle({ });
+	GraphClass.Interface.SetOutputStyle({ });
+#endif // WITH_EDITOR
 
-	FMetasoundFrontendGraphClass& GraphClass = Doc.RootGraph;
-	GraphClass.Interface.Inputs.Reset();
-	GraphClass.Interface.Outputs.Reset();
 	GraphClass.PresetOptions.InputsInheritingDefault.Reset();
 	GraphClass.PresetOptions.bIsPreset = false;
 
@@ -1354,10 +1353,18 @@ void FMetaSoundFrontendDocumentBuilder::ClearDocument(TSharedRef<Metasound::Fron
 
 PRAGMA_DISABLE_DEPRECATION_WARNINGS
 	FMetasoundFrontendGraph& Graph = GraphClass.GetDefaultGraph();
-	Graph.Nodes.Reset();
-	Graph.Edges.Reset();
+	Graph.Variables.Empty();
+	Graph.Nodes.Empty();
+	Graph.Edges.Empty();
 PRAGMA_ENABLE_DEPRECATION_WARNINGS
 
+	Doc.Interfaces.Empty();
+	Doc.Dependencies.Empty();
+
+#if WITH_EDITORONLY_DATA
+	Doc.Metadata.MemberMetadata.Empty();
+#endif // WITH_EDITORONLY_DATA
+
 	Reload(ModifyDelegates);
 }
 
diff --git a/Engine/Plugins/Runtime/Metasound/Source/MetasoundFrontend/Private/MetasoundFrontendDocumentCache.cpp b/Engine/Plugins/Runtime/Metasound/Source/MetasoundFrontend/Private/MetasoundFrontendDocumentCache.cpp
--- a/Engine/Plugins/Runtime/Metasound/Source/MetasoundFrontend/Private/MetasoundFrontendDocumentCache.cpp
+++ b/Engine/Plugins/Runtime/Metasound/Source/MetasoundFrontend/Private/MetasoundFrontendDocumentCache.cpp
@@ -114,12 +114,19 @@ namespace Metasound::Frontend
 			EdgeCacheMap.Add(InBuildPageID, EdgeCache);
 			TSharedRef<FDocumentGraphNodeCache> NodeCache = FDocumentGraphNodeCache::Create(ThisShared, InBuildPageID, ModifyDelegates.Get());
 			NodeCacheMap.Add(InBuildPageID, NodeCache);
+
 			InterfaceCache = FDocumentGraphInterfaceCache::Create(ThisShared, ModifyDelegates->InterfaceDelegates);
 		}
 
 		ModifyDelegates->OnDependencyAdded.AddSP(this, &FDocumentCache::OnDependencyAdded);
 		ModifyDelegates->OnRemoveSwappingDependency.AddSP(this, &FDocumentCache::OnRemoveSwappingDependency);
 		ModifyDelegates->OnRenamingDependencyClass.AddSP(this, &FDocumentCache::OnRenamingDependencyClass);
+
+		ModifyDelegates->PageDelegates.OnRemovingPage.AddSPLambda(this, [this](const FDocumentMutatePageArgs& Args)
+		{
+			EdgeCacheMap.Remove(Args.PageID);
+			NodeCacheMap.Remove(Args.PageID);
+		});
 	}
 
 	const FMetasoundFrontendDocument& FDocumentCache::GetDocument() const
