MetasoundAssetProxy

Create FMetasoundAssetProxy class for FMetasoundAssetBase
Make FMetasoundAssetBase inherit from IAudioProxyDataFactory to return the proxy type

#rnx

diff --git a/Engine/Plugins/Runtime/Metasound/Source/MetasoundFrontend/Private/MetasoundAssetBase.cpp b/Engine/Plugins/Runtime/Metasound/Source/MetasoundFrontend/Private/MetasoundAssetBase.cpp
--- a/Engine/Plugins/Runtime/Metasound/Source/MetasoundFrontend/Private/MetasoundAssetBase.cpp
+++ b/Engine/Plugins/Runtime/Metasound/Source/MetasoundFrontend/Private/MetasoundAssetBase.cpp
@@ -261,6 +261,21 @@ bool FMetasoundAssetBase::ConformObjectDataToInterfaces()
 	return false;
 }
 
+TSharedPtr<Audio::IProxyData> FMetasoundAssetBase::CreateProxyData(const Audio::FProxyDataInitParams& InitParams)
+{
+	using namespace Metasound::Frontend;
+	TScriptInterface<const IMetaSoundDocumentInterface> DocInterface = GetOwningAsset();
+	const FGraphRegistryKey& Key = GetGraphRegistryKey();
+	FMetasoundAssetProxy::FParameters Args;
+	Args.Interfaces = DocInterface->GetConstDocument().Interfaces;
+	Args.Graph = FMetasoundFrontendRegistryContainer::Get()->GetGraph(Key);
+	if (Args.Graph.IsValid())
+	{
+		return MakeShared<FMetasoundAssetProxy>(Args);
+	}
+	return nullptr;
+}
+
 void FMetasoundAssetBase::RegisterGraphWithFrontend(Metasound::Frontend::FMetaSoundAssetRegistrationOptions InRegistrationOptions)
 {
 	UpdateAndRegisterForExecution(MoveTemp(InRegistrationOptions));
@@ -1059,4 +1074,35 @@ const FMetasoundAssetBase::FRuntimeData& FMetasoundAssetBase::GetRuntimeData() c
 }
 PRAGMA_ENABLE_DEPRECATION_WARNINGS
 
+FMetasoundAssetProxy::FMetasoundAssetProxy(const FParameters& InParams)
+{
+	Interfaces = InParams.Interfaces;
+	Graph = InParams.Graph;
+}
+
+FMetasoundAssetProxy::FMetasoundAssetProxy(const FMetasoundAssetProxy& Other)
+{
+	Interfaces = Other.Interfaces;
+	Graph = Other.Graph;
+}
+
+TUniquePtr<Metasound::IOperator> FMetasoundAssetProxy::CreateOperator(const Metasound::FBuildOperatorParams& BuildOperatorParams, Metasound::FInputVertexInterfaceData& InputData) const
+{
+	using namespace Metasound;
+	if (ensure(Graph.IsValid()))
+	{
+		FBuildResults Results;
+		FBuildOperatorParams Params
+		{
+			*Graph.Get(),
+			BuildOperatorParams.OperatorSettings,
+			InputData,
+			BuildOperatorParams.Environment,
+			nullptr
+		};
+		return Graph->GetDefaultOperatorFactory()->CreateOperator(Params, Results);
+	}
+	return nullptr;
+}
+
 #undef LOCTEXT_NAMESPACE // "MetaSound"
diff --git a/Engine/Plugins/Runtime/Metasound/Source/MetasoundFrontend/Public/MetasoundAssetBase.h b/Engine/Plugins/Runtime/Metasound/Source/MetasoundFrontend/Public/MetasoundAssetBase.h
--- a/Engine/Plugins/Runtime/Metasound/Source/MetasoundFrontend/Public/MetasoundAssetBase.h
+++ b/Engine/Plugins/Runtime/Metasound/Source/MetasoundFrontend/Public/MetasoundAssetBase.h
@@ -25,6 +25,7 @@ typedef TMulticastDelegate<void(IConsoleVariable*), FDefaultDelegateUserPolicy>
 
 namespace Metasound
 {
+	class IGraph;
 	class FGraph;
 	namespace Frontend
 	{
@@ -43,15 +44,16 @@ namespace Metasound
 	} // namespace Frontend
 } // namespace Metasound
 
-
 /** FMetasoundAssetBase is intended to be a mix-in subclass for UObjects which utilize
  * Metasound assets.  It provides consistent access to FMetasoundFrontendDocuments, control
  * over the FMetasoundFrontendClassInterface of the FMetasoundFrontendDocument.  It also enables the UObject
  * to be utilized by a host of other engine tools built to support MetaSounds.
  */
-class METASOUNDFRONTEND_API FMetasoundAssetBase
+class METASOUNDFRONTEND_API FMetasoundAssetBase : public IAudioProxyDataFactory
 {
 public:
+	virtual TSharedPtr<Audio::IProxyData> CreateProxyData(const Audio::FProxyDataInitParams& InitParams) override;
+
 	static const FString FileExtension;
 
 	FMetasoundAssetBase() = default;
@@ -295,3 +297,38 @@ private:
 
 	Metasound::Frontend::FGraphRegistryKey GraphRegistryKey;
 };
+
+class METASOUNDFRONTEND_API FMetasoundAssetProxy final : public Audio::TProxyData<FMetasoundAssetProxy>
+{
+public:
+	IMPL_AUDIOPROXY_CLASS(FMetasoundAssetProxy);
+
+	struct FParameters
+	{
+		TSet<FMetasoundFrontendVersion> Interfaces;
+		TSharedPtr<const Metasound::IGraph> Graph;
+
+	};
+
+	explicit FMetasoundAssetProxy(const FParameters& InParams);
+	
+	FMetasoundAssetProxy(const FMetasoundAssetProxy& Other);
+
+	const Metasound::IGraph* GetGraph() const
+	{
+		return Graph.Get();
+	}
+
+	const TSet<FMetasoundFrontendVersion>& GetInterfaces() const
+	{
+		return Interfaces;
+	}
+
+	TUniquePtr<Metasound::IOperator> CreateOperator(const Metasound::FBuildOperatorParams& BuildOperatorParams, Metasound::FInputVertexInterfaceData& InputData) const;
+
+private:
+	
+	TSet<FMetasoundFrontendVersion> Interfaces;
+	TSharedPtr<const Metasound::IGraph> Graph;
+};
+using FMetasoundAssetProxyPtr = TSharedPtr<FMetasoundAssetProxy, ESPMode::ThreadSafe>;
