Minor comment add per earlier feedback
#rnx

diff --git a/Engine/Plugins/Runtime/Metasound/Source/MetasoundFrontend/Public/MetasoundFrontendDocumentIdGenerator.h b/Engine/Plugins/Runtime/Metasound/Source/MetasoundFrontend/Public/MetasoundFrontendDocumentIdGenerator.h
--- a/Engine/Plugins/Runtime/Metasound/Source/MetasoundFrontend/Public/MetasoundFrontendDocumentIdGenerator.h
+++ b/Engine/Plugins/Runtime/Metasound/Source/MetasoundFrontend/Public/MetasoundFrontendDocumentIdGenerator.h
@@ -75,7 +75,10 @@ namespace Metasound
 			bool bIsDeterministic = false;
 		};
 
-		// For generating IDs that are not derived from a given document. 
+		// For generating IDs that are derived from a given class vertex.  Unlike
+		// document ID generation, the class ID generator's results are not unique
+		// upon each request and therefor can deterministically generate the same
+		// ID for the same provided vertex.
 		class FClassIDGenerator
 		{
 		public: 
