Remove "CreateOperator" functionality from MetasoundAssetProxy

diff --git a/Engine/Plugins/Runtime/Metasound/Source/MetasoundFrontend/Private/MetasoundAssetBase.cpp b/Engine/Plugins/Runtime/Metasound/Source/MetasoundFrontend/Private/MetasoundAssetBase.cpp
--- a/Engine/Plugins/Runtime/Metasound/Source/MetasoundFrontend/Private/MetasoundAssetBase.cpp
+++ b/Engine/Plugins/Runtime/Metasound/Source/MetasoundFrontend/Private/MetasoundAssetBase.cpp
@@ -1086,23 +1086,4 @@ FMetasoundAssetProxy::FMetasoundAssetProxy(const FMetasoundAssetProxy& Other)
 	Graph = Other.Graph;
 }
 
-TUniquePtr<Metasound::IOperator> FMetasoundAssetProxy::CreateOperator(const Metasound::FBuildOperatorParams& BuildOperatorParams, Metasound::FInputVertexInterfaceData& InputData) const
-{
-	using namespace Metasound;
-	if (ensure(Graph.IsValid()))
-	{
-		FBuildResults Results;
-		FBuildOperatorParams Params
-		{
-			*Graph.Get(),
-			BuildOperatorParams.OperatorSettings,
-			InputData,
-			BuildOperatorParams.Environment,
-			nullptr
-		};
-		return Graph->GetDefaultOperatorFactory()->CreateOperator(Params, Results);
-	}
-	return nullptr;
-}
-
 #undef LOCTEXT_NAMESPACE // "MetaSound"
diff --git a/Engine/Plugins/Runtime/Metasound/Source/MetasoundFrontend/Public/MetasoundAssetBase.h b/Engine/Plugins/Runtime/Metasound/Source/MetasoundFrontend/Public/MetasoundAssetBase.h
--- a/Engine/Plugins/Runtime/Metasound/Source/MetasoundFrontend/Public/MetasoundAssetBase.h
+++ b/Engine/Plugins/Runtime/Metasound/Source/MetasoundFrontend/Public/MetasoundAssetBase.h
@@ -324,8 +324,6 @@ public:
 		return Interfaces;
 	}
 
-	TUniquePtr<Metasound::IOperator> CreateOperator(const Metasound::FBuildOperatorParams& BuildOperatorParams, Metasound::FInputVertexInterfaceData& InputData) const;
-
 private:
 	
 	TSet<FMetasoundFrontendVersion> Interfaces;
