MetaSound Template API fix-ups to support pages

#rnx

diff --git a/Engine/Plugins/Runtime/Metasound/Source/MetasoundEditor/Private/MetasoundEditorGraphNode.cpp b/Engine/Plugins/Runtime/Metasound/Source/MetasoundEditor/Private/MetasoundEditorGraphNode.cpp
--- a/Engine/Plugins/Runtime/Metasound/Source/MetasoundEditor/Private/MetasoundEditorGraphNode.cpp
+++ b/Engine/Plugins/Runtime/Metasound/Source/MetasoundEditor/Private/MetasoundEditorGraphNode.cpp
@@ -1161,7 +1161,8 @@ void UMetasoundEditorGraphExternalNode::Validate(Metasound::Editor::FGraphNodeVa
 				{
 #if WITH_EDITOR
 					FString Message;
-					if (!Template->HasRequiredConnections(GetBuilderChecked().GetConstBuilder(), GetNodeID(), &Message))
+					const FMetaSoundFrontendDocumentBuilder& DocBuilder = GetBuilderChecked().GetConstBuilder();
+					if (!Template->HasRequiredConnections(DocBuilder, DocBuilder.GetBuildPageID(), GetNodeID(), &Message))
 					{
 						OutResult.SetMessage(EMessageSeverity::Warning, Message);
 					}
diff --git a/Engine/Plugins/Runtime/Metasound/Source/MetasoundEditor/Private/MetasoundEditorGraphSchema.cpp b/Engine/Plugins/Runtime/Metasound/Source/MetasoundEditor/Private/MetasoundEditorGraphSchema.cpp
--- a/Engine/Plugins/Runtime/Metasound/Source/MetasoundEditor/Private/MetasoundEditorGraphSchema.cpp
+++ b/Engine/Plugins/Runtime/Metasound/Source/MetasoundEditor/Private/MetasoundEditorGraphSchema.cpp
@@ -1796,11 +1796,11 @@ FText UMetasoundEditorGraphSchema::GetPinDisplayName(const UEdGraphPin* Pin) con
 			{
 				if (Pin->Direction == EGPD_Input)
 				{
-					return Template->GetInputVertexDisplayName(Builder, Node->GetID(), Pin->GetFName());
+					return Template->GetInputVertexDisplayName(Builder, Builder.GetBuildPageID(), Node->GetID(), Pin->GetFName());
 				}
 				else
 				{
-					return Template->GetOutputVertexDisplayName(Builder, Node->GetID(), Pin->GetFName());
+					return Template->GetOutputVertexDisplayName(Builder, Builder.GetBuildPageID(), Node->GetID(), Pin->GetFName());
 				}
 			}
 		}
diff --git a/Engine/Plugins/Runtime/Metasound/Source/MetasoundEditor/Private/SMetasoundGraphPin.cpp b/Engine/Plugins/Runtime/Metasound/Source/MetasoundEditor/Private/SMetasoundGraphPin.cpp
--- a/Engine/Plugins/Runtime/Metasound/Source/MetasoundEditor/Private/SMetasoundGraphPin.cpp
+++ b/Engine/Plugins/Runtime/Metasound/Source/MetasoundEditor/Private/SMetasoundGraphPin.cpp
@@ -43,7 +43,8 @@ namespace Metasound
 				{
 					if (const INodeTemplate* Template = INodeTemplateRegistry::Get().FindTemplate(FRerouteNodeTemplate::GetRegistryKey()))
 					{
-						bHasRequiredConnections = Template->HasRequiredConnections(GetBuilderChecked().GetConstBuilder(), OwningNode->GetNodeID());
+						const FMetaSoundFrontendDocumentBuilder& DocBuilder = GetBuilderChecked().GetConstBuilder();
+						bHasRequiredConnections = Template->HasRequiredConnections(DocBuilder, DocBuilder.GetBuildPageID(), OwningNode->GetNodeID());
 					}
 				}
 			}
diff --git a/Engine/Plugins/Runtime/Metasound/Source/MetasoundFrontend/Private/MetasoundFrontendDocumentBuilder.cpp b/Engine/Plugins/Runtime/Metasound/Source/MetasoundFrontend/Private/MetasoundFrontendDocumentBuilder.cpp
--- a/Engine/Plugins/Runtime/Metasound/Source/MetasoundFrontend/Private/MetasoundFrontendDocumentBuilder.cpp
+++ b/Engine/Plugins/Runtime/Metasound/Source/MetasoundFrontend/Private/MetasoundFrontendDocumentBuilder.cpp
@@ -2117,7 +2117,7 @@ EMetasoundFrontendVertexAccessType FMetaSoundFrontendDocumentBuilder::GetNodeInp
 					{
 						if (Template->IsInputAccessTypeDynamic())
 						{
-							return Template->GetNodeInputAccessType(*this, InNodeID, InVertexID);
+							return Template->GetNodeInputAccessType(*this, PageID, InNodeID, InVertexID);
 						}
 					}
 				}
@@ -2246,7 +2246,7 @@ EMetasoundFrontendVertexAccessType FMetaSoundFrontendDocumentBuilder::GetNodeOut
 					{
 						if (Template->IsOutputAccessTypeDynamic())
 						{
-							return Template->GetNodeOutputAccessType(*this, InNodeID, InVertexID);
+							return Template->GetNodeOutputAccessType(*this, PageID, InNodeID, InVertexID);
 						}
 					}
 				}
@@ -2533,7 +2533,7 @@ bool FMetaSoundFrontendDocumentBuilder::TransformTemplateNodes()
 
 	struct FTemplateTransformParams
 	{
-		const Metasound::Frontend::INodeTemplate* Template;
+		const Metasound::Frontend::INodeTemplate* Template = nullptr;
 		TArray<FGuid> NodeIDs;
 	};
 	using FTemplateTransformParamsMap = TSortedMap<FGuid, FTemplateTransformParams>;
@@ -2562,34 +2562,36 @@ bool FMetaSoundFrontendDocumentBuilder::TransformTemplateNodes()
 	// which allows for addition/removal of nodes to/from original array container
 	// without template transform having to worry about mutation while iterating
 	TArray<FGuid> TemplateNodeIDs;
-	FMetasoundFrontendGraph& Graph = FindBuildGraphChecked();
-	for (const FMetasoundFrontendNode& Node : Graph.Nodes)
+	bool bModified = false;
+	Document.RootGraph.IterateGraphPages([this, &Dependencies, &TemplateParams, &bModified](FMetasoundFrontendGraph& Graph)
 	{
-		if (FTemplateTransformParams* Params = TemplateParams.Find(Node.ClassID))
+		for (const FMetasoundFrontendNode& Node : Graph.Nodes)
 		{
-			Params->NodeIDs.Add(Node.GetID());
+			if (FTemplateTransformParams* Params = TemplateParams.Find(Node.ClassID))
+			{
+				Params->NodeIDs.Add(Node.GetID());
+			}
 		}
-	}
 
-	// 2. Transform nodes
-	bool bModified = false;
-	for (const TPair<FGuid, FTemplateTransformParams>& Pair : TemplateParams)
-	{
-		const FTemplateTransformParams& Params = Pair.Value;
-		if (Params.Template)
+		for (TPair<FGuid, FTemplateTransformParams>& Pair : TemplateParams)
 		{
-			TUniquePtr<INodeTemplateTransform> NodeTransform = Params.Template->GenerateNodeTransform();
-			check(NodeTransform.IsValid());
-
-			for (const FGuid& NodeID : Params.NodeIDs)
+			FTemplateTransformParams& Params = Pair.Value;
+			if (Params.Template)
 			{
-				bModified = true;
-				NodeTransform->Transform(NodeID, *this);
+				TUniquePtr<INodeTemplateTransform> NodeTransform = Params.Template->GenerateNodeTransform();
+				check(NodeTransform.IsValid());
+
+				for (const FGuid& NodeID : Params.NodeIDs)
+				{
+					bModified = true;
+					NodeTransform->Transform(Graph.PageID, NodeID, *this);
+				}
 			}
+			Params.NodeIDs.Reset();
 		}
-	}
+	});
 
-	// 3. Remove template classes from dependency list
+	// 2. Remove template classes from dependency list
 	for (int32 i = Dependencies.Num() - 1; i >= 0; --i)
 	{
 		const FMetasoundFrontendClass& Class = Dependencies[i];
diff --git a/Engine/Plugins/Runtime/Metasound/Source/MetasoundFrontend/Private/MetasoundFrontendNodeTemplateRegistry.cpp b/Engine/Plugins/Runtime/Metasound/Source/MetasoundFrontend/Private/MetasoundFrontendNodeTemplateRegistry.cpp
--- a/Engine/Plugins/Runtime/Metasound/Source/MetasoundFrontend/Private/MetasoundFrontendNodeTemplateRegistry.cpp
+++ b/Engine/Plugins/Runtime/Metasound/Source/MetasoundFrontend/Private/MetasoundFrontendNodeTemplateRegistry.cpp
@@ -108,22 +108,22 @@ namespace Metasound::Frontend
 	}
 
 #if WITH_EDITOR
-	FText FNodeTemplateBase::GetNodeDisplayName(const IMetaSoundDocumentInterface& Interface, const FGuid& InNodeID) const
+	FText FNodeTemplateBase::GetNodeDisplayName(const IMetaSoundDocumentInterface& Interface, const FGuid& InPageID, const FGuid& InNodeID) const
 	{
 		return { };
 	}
 
-	FText FNodeTemplateBase::GetInputVertexDisplayName(const FMetaSoundFrontendDocumentBuilder& InBuilder, const FGuid& InNodeID, FName InputName) const
+	FText FNodeTemplateBase::GetInputVertexDisplayName(const FMetaSoundFrontendDocumentBuilder& InBuilder, const FGuid& InPageID, const FGuid& InNodeID, FName InputName) const
 	{
 		return FText::FromName(InputName);
 	}
 
-	FText FNodeTemplateBase::GetOutputVertexDisplayName(const FMetaSoundFrontendDocumentBuilder& InBuilder, const FGuid& InNodeID, FName OutputName) const
+	FText FNodeTemplateBase::GetOutputVertexDisplayName(const FMetaSoundFrontendDocumentBuilder& InBuilder, const FGuid& InPageID, const FGuid& InNodeID, FName OutputName) const
 	{
 		return FText::FromName(OutputName);
 	}
 
-	bool FNodeTemplateBase::HasRequiredConnections(const FMetaSoundFrontendDocumentBuilder& InBuilder, const FGuid& InNodeID, FString* OutMessage) const
+	bool FNodeTemplateBase::HasRequiredConnections(const FMetaSoundFrontendDocumentBuilder& InBuilder, const FGuid& InPageID, const FGuid& InNodeID, FString* OutMessage) const
 	{
 		return true;
 	}
diff --git a/Engine/Plugins/Runtime/Metasound/Source/MetasoundFrontend/Private/NodeTemplates/MetasoundFrontendNodeTemplateAudioAnalyzer.cpp b/Engine/Plugins/Runtime/Metasound/Source/MetasoundFrontend/Private/NodeTemplates/MetasoundFrontendNodeTemplateAudioAnalyzer.cpp
--- a/Engine/Plugins/Runtime/Metasound/Source/MetasoundFrontend/Private/NodeTemplates/MetasoundFrontendNodeTemplateAudioAnalyzer.cpp
+++ b/Engine/Plugins/Runtime/Metasound/Source/MetasoundFrontend/Private/NodeTemplates/MetasoundFrontendNodeTemplateAudioAnalyzer.cpp
@@ -14,13 +14,13 @@ namespace Metasound::Frontend
 			FAudioAnalyzerNodeTemplateTransform() = default;
 			virtual ~FAudioAnalyzerNodeTemplateTransform() = default;
 
-			virtual bool Transform(const FGuid& InNodeID, FMetaSoundFrontendDocumentBuilder& OutBuilder) const override;
+			virtual bool Transform(const FGuid& InPageID, const FGuid& InNodeID, FMetaSoundFrontendDocumentBuilder& OutBuilder) const override;
 		};
 
-		bool FAudioAnalyzerNodeTemplateTransform::Transform(const FGuid& InNodeID, FMetaSoundFrontendDocumentBuilder& OutBuilder) const
+		bool FAudioAnalyzerNodeTemplateTransform::Transform(const FGuid& InPageID, const FGuid& InNodeID, FMetaSoundFrontendDocumentBuilder& OutBuilder) const
 		{
 			// Strip it out:
-			return OutBuilder.RemoveNode(InNodeID);
+			return OutBuilder.RemoveNode(InNodeID, &InPageID);
 		}
 	} // namespace AudioAnalyzerNodeTemplatePrivate
 
@@ -68,12 +68,12 @@ namespace Metasound::Frontend
 		return FrontendClass;
 	}
 	
-	EMetasoundFrontendVertexAccessType FAudioAnalyzerNodeTemplate::GetNodeInputAccessType(const FMetaSoundFrontendDocumentBuilder& InBuilder, const FGuid& InNodeID, const FGuid& InVertexID) const
+	EMetasoundFrontendVertexAccessType FAudioAnalyzerNodeTemplate::GetNodeInputAccessType(const FMetaSoundFrontendDocumentBuilder& InBuilder, const FGuid& InPageID, const FGuid& InNodeID, const FGuid& InVertexID) const
 	{
 		return EMetasoundFrontendVertexAccessType::Unset;
 	}
 
-	EMetasoundFrontendVertexAccessType FAudioAnalyzerNodeTemplate::GetNodeOutputAccessType(const FMetaSoundFrontendDocumentBuilder& InBuilder, const FGuid& InNodeID, const FGuid& InVertexID) const
+	EMetasoundFrontendVertexAccessType FAudioAnalyzerNodeTemplate::GetNodeOutputAccessType(const FMetaSoundFrontendDocumentBuilder& InBuilder, const FGuid& InPageID, const FGuid& InNodeID, const FGuid& InVertexID) const
 	{
 		return EMetasoundFrontendVertexAccessType::Unset;
 	}
diff --git a/Engine/Plugins/Runtime/Metasound/Source/MetasoundFrontend/Private/NodeTemplates/MetasoundFrontendNodeTemplateInput.cpp b/Engine/Plugins/Runtime/Metasound/Source/MetasoundFrontend/Private/NodeTemplates/MetasoundFrontendNodeTemplateInput.cpp
--- a/Engine/Plugins/Runtime/Metasound/Source/MetasoundFrontend/Private/NodeTemplates/MetasoundFrontendNodeTemplateInput.cpp
+++ b/Engine/Plugins/Runtime/Metasound/Source/MetasoundFrontend/Private/NodeTemplates/MetasoundFrontendNodeTemplateInput.cpp
@@ -74,7 +74,7 @@ namespace Metasound::Frontend
 	}
 
 #if WITH_EDITOR
-	FText FInputNodeTemplate::GetNodeDisplayName(const IMetaSoundDocumentInterface& Interface, const FGuid& InNodeID) const
+	FText FInputNodeTemplate::GetNodeDisplayName(const IMetaSoundDocumentInterface& Interface, const FGuid& InPageID, const FGuid& InNodeID) const
 	{
 		return { };
 	}
@@ -124,10 +124,10 @@ namespace Metasound::Frontend
 		return RegistryKey;
 	}
 
-	EMetasoundFrontendVertexAccessType FInputNodeTemplate::GetNodeInputAccessType(const FMetaSoundFrontendDocumentBuilder& InBuilder, const FGuid& InNodeID, const FGuid& InVertexID) const
+	EMetasoundFrontendVertexAccessType FInputNodeTemplate::GetNodeInputAccessType(const FMetaSoundFrontendDocumentBuilder& InBuilder, const FGuid& InPageID, const FGuid& InNodeID, const FGuid& InVertexID) const
 	{
 		const FMetasoundFrontendNode* ConnectedInputNode = nullptr;
-		if (const FMetasoundFrontendVertex* ConnectedInputOutput = InBuilder.FindNodeOutputConnectedToNodeInput(InNodeID, InVertexID, &ConnectedInputNode))
+		if (const FMetasoundFrontendVertex* ConnectedInputOutput = InBuilder.FindNodeOutputConnectedToNodeInput(InNodeID, InVertexID, &ConnectedInputNode, &InPageID))
 		{
 			const FMetasoundFrontendClass* InputClass = InBuilder.FindDependency(ConnectedInputNode->ClassID);
 			check(InputClass);
@@ -137,28 +137,28 @@ namespace Metasound::Frontend
 		return EMetasoundFrontendVertexAccessType::Unset;
 	}
 
-	EMetasoundFrontendVertexAccessType FInputNodeTemplate::GetNodeOutputAccessType(const FMetaSoundFrontendDocumentBuilder& InBuilder, const FGuid& InNodeID, const FGuid& InVertexID) const
+	EMetasoundFrontendVertexAccessType FInputNodeTemplate::GetNodeOutputAccessType(const FMetaSoundFrontendDocumentBuilder& InBuilder, const FGuid& InPageID, const FGuid& InNodeID, const FGuid& InVertexID) const
 	{
-		if (const FMetasoundFrontendNode* Node = InBuilder.FindNode(InNodeID))
+		if (const FMetasoundFrontendNode* Node = InBuilder.FindNode(InNodeID, &InPageID))
 		{
 			const FMetasoundFrontendVertex& Input = Node->Interface.Inputs.Last();
-			return GetNodeInputAccessType(InBuilder, InNodeID, Input.VertexID);
+			return GetNodeInputAccessType(InBuilder, InPageID, InNodeID, Input.VertexID);
 		}
 
 		return EMetasoundFrontendVertexAccessType::Unset;
 	}
 
 #if WITH_EDITOR
-	FText FInputNodeTemplate::GetOutputVertexDisplayName(const FMetaSoundFrontendDocumentBuilder& InBuilder, const FGuid& InNodeID, FName OutputName) const
+	FText FInputNodeTemplate::GetOutputVertexDisplayName(const FMetaSoundFrontendDocumentBuilder& InBuilder, const FGuid& InPageID, const FGuid& InNodeID, FName OutputName) const
 	{
-		const FMetasoundFrontendNode* OwningNode = InBuilder.FindNode(InNodeID);
+		const FMetasoundFrontendNode* OwningNode = InBuilder.FindNode(InNodeID, &InPageID);
 		if (!OwningNode)
 		{
 			return FText::FromName(OutputName);
 		}
 
 		const FMetasoundFrontendNode* ConnectedInputNode = nullptr;
-		const FMetasoundFrontendVertex* ConnectedOutput = InBuilder.FindNodeOutputConnectedToNodeInput(InNodeID, OwningNode->Interface.Inputs.Last().VertexID, &ConnectedInputNode);
+		const FMetasoundFrontendVertex* ConnectedOutput = InBuilder.FindNodeOutputConnectedToNodeInput(InNodeID, OwningNode->Interface.Inputs.Last().VertexID, &ConnectedInputNode, &InPageID);
 		if (ensureMsgf(ConnectedInputNode, TEXT("Input template node should always be connected to associated input node's only output")))
 		{
 			FName NodeName = ConnectedInputNode->Name;
@@ -172,10 +172,10 @@ namespace Metasound::Frontend
 			return INodeTemplate::ResolveMemberDisplayName(NodeName, DisplayName, bIncludeNamespace);
 		}
 
-		return FRerouteNodeTemplate::GetOutputVertexDisplayName(InBuilder, InNodeID, OutputName);
+		return FRerouteNodeTemplate::GetOutputVertexDisplayName(InBuilder, InPageID, InNodeID, OutputName);
 	}
 
-	bool FInputNodeTemplate::HasRequiredConnections(const FMetaSoundFrontendDocumentBuilder& InBuilder, const FGuid& InNodeID, FString* OutMessage) const
+	bool FInputNodeTemplate::HasRequiredConnections(const FMetaSoundFrontendDocumentBuilder& InBuilder, const FGuid& InPageID, const FGuid& InNodeID, FString* OutMessage) const
 	{
 		return true;
 	}
diff --git a/Engine/Plugins/Runtime/Metasound/Source/MetasoundFrontend/Private/NodeTemplates/MetasoundFrontendNodeTemplateReroute.cpp b/Engine/Plugins/Runtime/Metasound/Source/MetasoundFrontend/Private/NodeTemplates/MetasoundFrontendNodeTemplateReroute.cpp
--- a/Engine/Plugins/Runtime/Metasound/Source/MetasoundFrontend/Private/NodeTemplates/MetasoundFrontendNodeTemplateReroute.cpp
+++ b/Engine/Plugins/Runtime/Metasound/Source/MetasoundFrontend/Private/NodeTemplates/MetasoundFrontendNodeTemplateReroute.cpp
@@ -22,15 +22,15 @@ namespace Metasound::Frontend
 			FRerouteNodeTemplateTransform() = default;
 			virtual ~FRerouteNodeTemplateTransform() = default;
 
-			virtual bool Transform(const FGuid& InNodeID, FMetaSoundFrontendDocumentBuilder& OutBuilder) const override;
+			virtual bool Transform(const FGuid& InPageID, const FGuid& InNodeID, FMetaSoundFrontendDocumentBuilder& OutBuilder) const override;
 		};
 
-		bool FRerouteNodeTemplateTransform::Transform(const FGuid& InNodeID, FMetaSoundFrontendDocumentBuilder& OutBuilder) const
+		bool FRerouteNodeTemplateTransform::Transform(const FGuid& InPageID, const FGuid& InNodeID, FMetaSoundFrontendDocumentBuilder& OutBuilder) const
 		{
 			FMetasoundFrontendEdge InputEdge;
 			TArray<FMetasoundFrontendEdge> OutputEdges;
 
-			const FMetasoundFrontendNode* Node = OutBuilder.FindNode(InNodeID);
+			const FMetasoundFrontendNode* Node = OutBuilder.FindNode(InNodeID, &InPageID);
 			if (ensureMsgf(Node, TEXT("Failed to find node with ID '%s' when reroute template node transform was given a valid ID for builder '%s'."),
 				*InNodeID.ToString(),
 				*OutBuilder.GetDebugName()))
@@ -48,7 +48,7 @@ namespace Metasound::Frontend
 				// Copy input edge to mutate from fields and avoid pointer going out of scope when template node is removed below
 				{
 					const FMetasoundFrontendVertex& InputVertex = Node->Interface.Inputs.Last();
-					TArray<const FMetasoundFrontendEdge*> InputEdges = OutBuilder.FindEdges(Node->GetID(), InputVertex.VertexID);
+					TArray<const FMetasoundFrontendEdge*> InputEdges = OutBuilder.FindEdges(Node->GetID(), InputVertex.VertexID, &InPageID);
 					if (!InputEdges.IsEmpty())
 					{
 						InputEdge = *InputEdges.Last();
@@ -58,7 +58,7 @@ namespace Metasound::Frontend
 				// Copy output edges to mutate from fields and avoid pointer going out of scope when swapping below
 				{
 					const FMetasoundFrontendVertex& OutputVertex = Node->Interface.Outputs.Last();
-					TArray<const FMetasoundFrontendEdge*> CurrentOutputEdges = OutBuilder.FindEdges(Node->GetID(), OutputVertex.VertexID);
+					TArray<const FMetasoundFrontendEdge*> CurrentOutputEdges = OutBuilder.FindEdges(Node->GetID(), OutputVertex.VertexID, &InPageID);
 					Algo::Transform(CurrentOutputEdges, OutputEdges, [](const FMetasoundFrontendEdge* CurrentEdge)
 					{
 						check(CurrentEdge);
@@ -67,7 +67,7 @@ namespace Metasound::Frontend
 				}
 
 				// Remove the template node
-				OutBuilder.RemoveNode(Node->GetID());
+				OutBuilder.RemoveNode(Node->GetID(), &InPageID);
 
 				// Add new connections from reroute source node to reroute destination node. Either could be another reroute,
 				// which is valid because said node will subsequently get processed.
@@ -78,7 +78,7 @@ namespace Metasound::Frontend
 					{
 						OutputEdge.FromNodeID = InputEdge.FromNodeID;
 						OutputEdge.FromVertexID = InputEdge.FromVertexID;
-						OutBuilder.AddEdge(MoveTemp(OutputEdge));
+						OutBuilder.AddEdge(MoveTemp(OutputEdge), &InPageID);
 					}
 
 					return bModified;
@@ -88,7 +88,7 @@ namespace Metasound::Frontend
 			return false;
 		}
 
-		const FMetasoundFrontendVertex* FindReroutedOutputVertex(const FMetaSoundFrontendDocumentBuilder& InBuilder, const FMetasoundFrontendNode& InOutputOwningNode, const FMetasoundFrontendVertex& InOutputVertex)
+		const FMetasoundFrontendVertex* FindReroutedOutputVertex(const FMetaSoundFrontendDocumentBuilder& InBuilder, const FGuid& InPageID, const FMetasoundFrontendNode& InOutputOwningNode, const FMetasoundFrontendVertex& InOutputVertex)
 		{
 			const FMetasoundFrontendVertex* ReroutedVertex = &InOutputVertex;
 			if (const FMetasoundFrontendClass* Class = InBuilder.FindDependency(InOutputOwningNode.ClassID))
@@ -96,16 +96,16 @@ namespace Metasound::Frontend
 				if (Class->Metadata.GetClassName() == FRerouteNodeTemplate::ClassName)
 				{
 					const FGuid& OutputOwningNodeID = InOutputOwningNode.GetID();
-					TArray<const FMetasoundFrontendVertex*> Inputs = InBuilder.FindNodeInputs(OutputOwningNodeID);
+					TArray<const FMetasoundFrontendVertex*> Inputs = InBuilder.FindNodeInputs(OutputOwningNodeID, { }, &InPageID);
 					if (!Inputs.IsEmpty())
 					{
 						if (const FMetasoundFrontendVertex* RerouteInput = Inputs.Last())
 						{
 							const FMetasoundFrontendNode* ConnectedNode = nullptr;
-							if (const FMetasoundFrontendVertex* OutputVertex = InBuilder.FindNodeOutputConnectedToNodeInput(OutputOwningNodeID, RerouteInput->VertexID, &ConnectedNode))
+							if (const FMetasoundFrontendVertex* OutputVertex = InBuilder.FindNodeOutputConnectedToNodeInput(OutputOwningNodeID, RerouteInput->VertexID, &ConnectedNode, &InPageID))
 							{
 								check(ConnectedNode);
-								return FindReroutedOutputVertex(InBuilder, *ConnectedNode, *OutputVertex);
+								return FindReroutedOutputVertex(InBuilder, InPageID, *ConnectedNode, *OutputVertex);
 							}
 						}
 					}
@@ -119,6 +119,7 @@ namespace Metasound::Frontend
 
 		void FindReroutedInputVertices(
 			const FMetaSoundFrontendDocumentBuilder& InBuilder,
+			const FGuid& InPageID,
 			const FMetasoundFrontendNode& InInputOwningNode,
 			const FMetasoundFrontendVertex& InInputVertex,
 			TArray<const FMetasoundFrontendNode*>& InOutReroutedInputOwningNodes,
@@ -137,7 +138,7 @@ namespace Metasound::Frontend
 						check(Output);
 
 						TArray<const FMetasoundFrontendNode*> ConnectedInputNodes;
-						TArray<const FMetasoundFrontendVertex*> ConnectedInputVertices = InBuilder.FindNodeInputsConnectedToNodeOutput(InputOwningNodeID, Output->VertexID, &ConnectedInputNodes);
+						TArray<const FMetasoundFrontendVertex*> ConnectedInputVertices = InBuilder.FindNodeInputsConnectedToNodeOutput(InputOwningNodeID, Output->VertexID, &ConnectedInputNodes, &InPageID);
 						check(ConnectedInputNodes.Num() == ConnectedInputVertices.Num());
 
 						for (int32 Index = 0; Index < ConnectedInputNodes.Num(); ++Index)
@@ -147,7 +148,7 @@ namespace Metasound::Frontend
 							const FMetasoundFrontendVertex* ConnectedInputVertex = ConnectedInputVertices[Index];
 							check(ConnectedInputVertex);
 
-							FindReroutedInputVertices(InBuilder, *ConnectedInputOwningNode, *ConnectedInputVertex, InOutReroutedInputOwningNodes, InOutReroutedInputVertices);
+							FindReroutedInputVertices(InBuilder, InPageID, *ConnectedInputOwningNode, *ConnectedInputVertex, InOutReroutedInputOwningNodes, InOutReroutedInputVertices);
 						}
 					}
 
@@ -170,7 +171,7 @@ namespace Metasound::Frontend
 	}
 
 #if WITH_EDITOR
-	FText FRerouteNodeTemplate::GetNodeDisplayName(const IMetaSoundDocumentInterface& DocumentInterface, const FGuid& InNodeID) const
+	FText FRerouteNodeTemplate::GetNodeDisplayName(const IMetaSoundDocumentInterface& DocumentInterface, const FGuid& InPageID, const FGuid& InNodeID) const
 	{
 		return { };
 	}
@@ -244,16 +245,16 @@ namespace Metasound::Frontend
 		return FrontendClass;
 	}
 
-	EMetasoundFrontendVertexAccessType FRerouteNodeTemplate::GetNodeInputAccessType(const FMetaSoundFrontendDocumentBuilder& InBuilder, const FGuid& InNodeID, const FGuid& InVertexID) const
+	EMetasoundFrontendVertexAccessType FRerouteNodeTemplate::GetNodeInputAccessType(const FMetaSoundFrontendDocumentBuilder& InBuilder, const FGuid& InPageID, const FGuid& InNodeID, const FGuid& InVertexID) const
 	{
 		// Recursive search up DAG for first connected non-reroute node's input access type
-		if (const FMetasoundFrontendNode* Node = InBuilder.FindNode(InNodeID))
+		if (const FMetasoundFrontendNode* Node = InBuilder.FindNode(InNodeID, &InPageID))
 		{
 			// Should only ever be one
 			const FMetasoundFrontendVertex& RerouteOutput = Node->Interface.Outputs.Last();
 
 			TArray<const FMetasoundFrontendNode*> ConnectedNodes;
-			TArray<const FMetasoundFrontendVertex*> ConnectedInputs = InBuilder.FindNodeInputsConnectedToNodeOutput(InNodeID, RerouteOutput.VertexID, &ConnectedNodes);
+			TArray<const FMetasoundFrontendVertex*> ConnectedInputs = InBuilder.FindNodeInputsConnectedToNodeOutput(InNodeID, RerouteOutput.VertexID, &ConnectedNodes, &InPageID);
 			for (int32 Index = 0; Index < ConnectedNodes.Num(); ++Index)
 			{
 				const FMetasoundFrontendNode* ConnectedNode = ConnectedNodes[Index];
@@ -262,10 +263,10 @@ namespace Metasound::Frontend
 					const FMetasoundFrontendVertex* ConnectedInput = ConnectedInputs[Index];
 					if (ConnectedNodeClass->Metadata.GetClassName() == GetClassName())
 					{
-						return this->GetNodeInputAccessType(InBuilder, ConnectedNode->GetID(), ConnectedInput->VertexID);
+						return this->GetNodeInputAccessType(InBuilder, InPageID, ConnectedNode->GetID(), ConnectedInput->VertexID);
 					}
 
-					return InBuilder.GetNodeInputAccessType(ConnectedNode->GetID(), ConnectedInput->VertexID);
+					return InBuilder.GetNodeInputAccessType(ConnectedNode->GetID(), ConnectedInput->VertexID, &InPageID);
 				}
 			}
 		}
@@ -273,7 +274,7 @@ namespace Metasound::Frontend
 		return EMetasoundFrontendVertexAccessType::Unset;
 	}
 
-	EMetasoundFrontendVertexAccessType FRerouteNodeTemplate::GetNodeOutputAccessType(const FMetaSoundFrontendDocumentBuilder& InBuilder, const FGuid& InNodeID, const FGuid& InVertexID) const
+	EMetasoundFrontendVertexAccessType FRerouteNodeTemplate::GetNodeOutputAccessType(const FMetaSoundFrontendDocumentBuilder& InBuilder, const FGuid& InPageID, const FGuid& InNodeID, const FGuid& InVertexID) const
 	{
 		// Depth-first recursive search for first connected non-reroute node's output access type
 		if (const FMetasoundFrontendNode* Node = InBuilder.FindNode(InNodeID))
@@ -282,16 +283,16 @@ namespace Metasound::Frontend
 			const FMetasoundFrontendVertex& RerouteInput = Node->Interface.Inputs.Last();
 
 			const FMetasoundFrontendNode* ConnectedNode = nullptr;
-			if (const FMetasoundFrontendVertex* ConnectedOutput = InBuilder.FindNodeOutputConnectedToNodeInput(InNodeID, RerouteInput.VertexID, &ConnectedNode))
+			if (const FMetasoundFrontendVertex* ConnectedOutput = InBuilder.FindNodeOutputConnectedToNodeInput(InNodeID, RerouteInput.VertexID, &ConnectedNode, &InPageID))
 			{
 				if (const FMetasoundFrontendClass* ConnectedNodeClass = InBuilder.FindDependency(ConnectedNode->ClassID))
 				{
 					if (ConnectedNodeClass->Metadata.GetClassName() == ClassName)
 					{
-						return this->GetNodeOutputAccessType(InBuilder, ConnectedNode->GetID(), ConnectedOutput->VertexID);
+						return this->GetNodeOutputAccessType(InBuilder, InPageID, ConnectedNode->GetID(), ConnectedOutput->VertexID);
 					}
 
-					return InBuilder.GetNodeOutputAccessType(ConnectedNode->GetID(), ConnectedOutput->VertexID);
+					return InBuilder.GetNodeOutputAccessType(ConnectedNode->GetID(), ConnectedOutput->VertexID, &InPageID);
 				}
 			}
 		}
@@ -311,7 +312,7 @@ namespace Metasound::Frontend
 	}
 
 #if WITH_EDITOR
-	bool FRerouteNodeTemplate::HasRequiredConnections(const FMetaSoundFrontendDocumentBuilder& InBuilder, const FGuid& InNodeID, FString* OutMessage) const
+	bool FRerouteNodeTemplate::HasRequiredConnections(const FMetaSoundFrontendDocumentBuilder& InBuilder, const FGuid& InPageID, const FGuid& InNodeID, FString* OutMessage) const
 	{
 		const FMetasoundFrontendNode* Node = InBuilder.FindNode(InNodeID);
 		if (!Node)
@@ -319,20 +320,20 @@ namespace Metasound::Frontend
 			return false;
 		}
 
-		TArray<const FMetasoundFrontendVertex*> Outputs = InBuilder.FindNodeOutputs(InNodeID);
-		const bool bConnectedToNonRerouteOutputs = Algo::AnyOf(Outputs, [&InBuilder, &Node](const FMetasoundFrontendVertex* OutputVertex)
+		TArray<const FMetasoundFrontendVertex*> Outputs = InBuilder.FindNodeOutputs(InNodeID, { }, &InPageID);
+		const bool bConnectedToNonRerouteOutputs = Algo::AnyOf(Outputs, [&InBuilder, &InPageID, &Node](const FMetasoundFrontendVertex* OutputVertex)
 		{
 			using namespace ReroutePrivate;
-			return FindReroutedOutputVertex(InBuilder, *Node, *OutputVertex) != nullptr;
+			return FindReroutedOutputVertex(InBuilder, InPageID, *Node, *OutputVertex) != nullptr;
 		});
 
-		TArray<const FMetasoundFrontendVertex*> Inputs = InBuilder.FindNodeInputs(InNodeID);
-		const bool bConnectedToNonRerouteInputs = Algo::AnyOf(Inputs, [&InBuilder, &Node](const FMetasoundFrontendVertex* InputVertex)
+		TArray<const FMetasoundFrontendVertex*> Inputs = InBuilder.FindNodeInputs(InNodeID, { }, &InPageID);
+		const bool bConnectedToNonRerouteInputs = Algo::AnyOf(Inputs, [&InBuilder, &InPageID, &Node](const FMetasoundFrontendVertex* InputVertex)
 		{
 			using namespace ReroutePrivate;
 			TArray<const FMetasoundFrontendVertex*> InputVertices;
 			TArray<const FMetasoundFrontendNode*> InputVerticesOwningNodes;
-			FindReroutedInputVertices(InBuilder, *Node, *InputVertex, InputVerticesOwningNodes, InputVertices);
+			FindReroutedInputVertices(InBuilder, InPageID, *Node, *InputVertex, InputVerticesOwningNodes, InputVertices);
 			return !InputVertices.IsEmpty();
 		});
 
diff --git a/Engine/Plugins/Runtime/Metasound/Source/MetasoundFrontend/Public/MetasoundFrontendDocumentIdGenerator.h b/Engine/Plugins/Runtime/Metasound/Source/MetasoundFrontend/Public/MetasoundFrontendDocumentIdGenerator.h
--- a/Engine/Plugins/Runtime/Metasound/Source/MetasoundFrontend/Public/MetasoundFrontendDocumentIdGenerator.h
+++ b/Engine/Plugins/Runtime/Metasound/Source/MetasoundFrontend/Public/MetasoundFrontendDocumentIdGenerator.h
@@ -77,7 +77,7 @@ namespace Metasound
 
 		// For generating IDs that are derived from a given class vertex.  Unlike
 		// document ID generation, the class ID generator's results are not unique
-		// upon each request and therefor can deterministically generate the same
+		// upon each request and therefore can deterministically generate the same
 		// ID for the same provided vertex.
 		class FClassIDGenerator
 		{
diff --git a/Engine/Plugins/Runtime/Metasound/Source/MetasoundFrontend/Public/MetasoundFrontendNodeTemplateRegistry.h b/Engine/Plugins/Runtime/Metasound/Source/MetasoundFrontend/Public/MetasoundFrontendNodeTemplateRegistry.h
--- a/Engine/Plugins/Runtime/Metasound/Source/MetasoundFrontend/Public/MetasoundFrontendNodeTemplateRegistry.h
+++ b/Engine/Plugins/Runtime/Metasound/Source/MetasoundFrontend/Public/MetasoundFrontendNodeTemplateRegistry.h
@@ -42,7 +42,7 @@ namespace Metasound::Frontend
 		virtual ~INodeTemplateTransform() = default;
 
 		/** Return true if the builder was modified, false otherwise. */
-		virtual bool Transform(const FGuid& InNodeID, FMetaSoundFrontendDocumentBuilder& OutBuilder) const = 0;
+		virtual bool Transform(const FGuid& InPageID, const FGuid& InNodeID, FMetaSoundFrontendDocumentBuilder& OutBuilder) const = 0;
 	};
 
 	/**
@@ -64,17 +64,17 @@ namespace Metasound::Frontend
 		virtual TUniquePtr<INodeTransform> GenerateNodeTransform(FMetasoundFrontendDocument& InDocument) const;
 
 #if WITH_EDITOR
-		virtual FText GetNodeDisplayName(const IMetaSoundDocumentInterface& Interface, const FGuid& InNodeID) const = 0;
+		virtual FText GetNodeDisplayName(const IMetaSoundDocumentInterface& Interface, const FGuid& InPageID, const FGuid& InNodeID) const = 0;
 
-		UE_DEPRECATED(5.5, "Use overload GetOutputVertexDisplayName with supplied builder instead")
+		UE_DEPRECATED(5.5, "Use overload GetOutputVertexDisplayName with supplied builder & pageID instead")
 		virtual FText GetInputPinDisplayName(const Frontend::IInputController& InInput) const { return FText(); }
 
-		virtual FText GetInputVertexDisplayName(const FMetaSoundFrontendDocumentBuilder& InBuilder, const FGuid& InNodeID, FName InputName) const = 0;
+		virtual FText GetInputVertexDisplayName(const FMetaSoundFrontendDocumentBuilder& InBuilder, const FGuid& InPageID, const FGuid& InNodeID, FName InputName) const = 0;
 
-		UE_DEPRECATED(5.5, "Use overload GetOutputVertexDisplayName with supplied builder instead")
+		UE_DEPRECATED(5.5, "Use overload GetOutputVertexDisplayName with supplied builder & pageID instead")
 		virtual FText GetOutputPinDisplayName(const Frontend::IOutputController& InOutput) const { return FText(); };
 
-		virtual FText GetOutputVertexDisplayName(const FMetaSoundFrontendDocumentBuilder& InBuilder, const FGuid& InNodeID, FName OutputName) const = 0;
+		virtual FText GetOutputVertexDisplayName(const FMetaSoundFrontendDocumentBuilder& InBuilder, const FGuid& InPageID, const FGuid& InNodeID, FName OutputName) const = 0;
 
 		static FText ResolveMemberDisplayName(FName VertexName, FText DisplayName, bool bIncludeNamespace);
 #endif // WITH_EDITOR
@@ -86,10 +86,10 @@ namespace Metasound::Frontend
 		virtual const FMetasoundFrontendClass& GetFrontendClass() const = 0;
 
 		// Returns access type of the given input within the provided builder's document
-		virtual EMetasoundFrontendVertexAccessType GetNodeInputAccessType(const FMetaSoundFrontendDocumentBuilder& InBuilder, const FGuid& InNodeID, const FGuid& InVertexID) const = 0;
+		virtual EMetasoundFrontendVertexAccessType GetNodeInputAccessType(const FMetaSoundFrontendDocumentBuilder& InBuilder, const FGuid& InPageID, const FGuid& InNodeID, const FGuid& InVertexID) const = 0;
 
 		// Returns access type of the given output within the provided builder's document
-		virtual EMetasoundFrontendVertexAccessType GetNodeOutputAccessType(const FMetaSoundFrontendDocumentBuilder& InBuilder, const FGuid& InNodeID, const FGuid& InVertexID) const = 0;
+		virtual EMetasoundFrontendVertexAccessType GetNodeOutputAccessType(const FMetaSoundFrontendDocumentBuilder& InBuilder, const FGuid& InPageID, const FGuid& InNodeID, const FGuid& InVertexID) const = 0;
 
 		UE_DEPRECATED(5.4, "Use version number or classname instead")
 		virtual const FMetasoundFrontendVersion& GetVersion() const { const static FMetasoundFrontendVersion NullVersion; return NullVersion; }
@@ -98,12 +98,12 @@ namespace Metasound::Frontend
 		virtual const FMetasoundFrontendVersionNumber& GetVersionNumber() const = 0;
 
 #if WITH_EDITOR
-		UE_DEPRECATED(5.5, "Use version that takes a Frontend NodeID and builder instead")
+		UE_DEPRECATED(5.5, "Use version that takes a Frontend PageID, NodeID and builder instead")
 		virtual bool HasRequiredConnections(FConstNodeHandle InNodeHandle, FString* OutMessage = nullptr) const { return false; }
 
 		// Returns whether or not the given node template has the necessary
 		// required connections to be preprocessed (editor only).
-		virtual bool HasRequiredConnections(const FMetaSoundFrontendDocumentBuilder& InBuilder, const FGuid& InNodeID, FString* OutMessage = nullptr) const = 0;
+		virtual bool HasRequiredConnections(const FMetaSoundFrontendDocumentBuilder& InBuilder, const FGuid& InPageID, const FGuid& InNodeID, FString* OutMessage = nullptr) const = 0;
 #endif // WITH_EDITOR
 
 		// Returns whether template can dynamically assign a node's input access type (as opposed to it being assigned on the class input definition)
@@ -142,10 +142,10 @@ namespace Metasound::Frontend
 	{
 	public:
 #if WITH_EDITOR
-		virtual FText GetNodeDisplayName(const IMetaSoundDocumentInterface& Interface, const FGuid& InNodeID) const override;
-		virtual FText GetInputVertexDisplayName(const FMetaSoundFrontendDocumentBuilder& InBuilder, const FGuid& InNodeID, FName InputName) const override;
-		virtual FText GetOutputVertexDisplayName(const FMetaSoundFrontendDocumentBuilder& InBuilder, const FGuid& InNodeID, FName OutputName) const override;
-		virtual bool HasRequiredConnections(const FMetaSoundFrontendDocumentBuilder& InBuilder, const FGuid& InNodeID, FString* OutMessage = nullptr) const override;
+		virtual FText GetNodeDisplayName(const IMetaSoundDocumentInterface& Interface, const FGuid& InPageID, const FGuid& InNodeID) const override;
+		virtual FText GetInputVertexDisplayName(const FMetaSoundFrontendDocumentBuilder& InBuilder, const FGuid& InPageID, const FGuid& InNodeID, FName InputName) const override;
+		virtual FText GetOutputVertexDisplayName(const FMetaSoundFrontendDocumentBuilder& InBuilder, const FGuid& InPageID, const FGuid& InNodeID, FName OutputName) const override;
+		virtual bool HasRequiredConnections(const FMetaSoundFrontendDocumentBuilder& InBuilder, const FGuid& InPageID, const FGuid& InNodeID, FString* OutMessage = nullptr) const override;
 #endif // WITH_EDITOR
 	};
 
diff --git a/Engine/Plugins/Runtime/Metasound/Source/MetasoundFrontend/Public/NodeTemplates/MetasoundFrontendNodeTemplateAudioAnalyzer.h b/Engine/Plugins/Runtime/Metasound/Source/MetasoundFrontend/Public/NodeTemplates/MetasoundFrontendNodeTemplateAudioAnalyzer.h
--- a/Engine/Plugins/Runtime/Metasound/Source/MetasoundFrontend/Public/NodeTemplates/MetasoundFrontendNodeTemplateAudioAnalyzer.h
+++ b/Engine/Plugins/Runtime/Metasound/Source/MetasoundFrontend/Public/NodeTemplates/MetasoundFrontendNodeTemplateAudioAnalyzer.h
@@ -15,8 +15,8 @@ namespace Metasound::Frontend
 		virtual const FMetasoundFrontendClassName& GetClassName() const override;
 		virtual TUniquePtr<INodeTemplateTransform> GenerateNodeTransform() const override;
 		virtual const FMetasoundFrontendClass& GetFrontendClass() const override;
-		virtual EMetasoundFrontendVertexAccessType GetNodeInputAccessType(const FMetaSoundFrontendDocumentBuilder& InBuilder, const FGuid& InNodeID, const FGuid& InVertexID) const override;
-		virtual EMetasoundFrontendVertexAccessType GetNodeOutputAccessType(const FMetaSoundFrontendDocumentBuilder& InBuilder, const FGuid& InNodeID, const FGuid& InVertexID) const override;
+		virtual EMetasoundFrontendVertexAccessType GetNodeInputAccessType(const FMetaSoundFrontendDocumentBuilder& InBuilder, const FGuid& InPageID, const FGuid& InNodeID, const FGuid& InVertexID) const override;
+		virtual EMetasoundFrontendVertexAccessType GetNodeOutputAccessType(const FMetaSoundFrontendDocumentBuilder& InBuilder, const FGuid& InPageID, const FGuid& InNodeID, const FGuid& InVertexID) const override;
 		virtual const FMetasoundFrontendVersionNumber& GetVersionNumber() const override;
 		virtual bool IsInputAccessTypeDynamic() const override;
 		virtual bool IsInputConnectionUserModifiable() const override;
diff --git a/Engine/Plugins/Runtime/Metasound/Source/MetasoundFrontend/Public/NodeTemplates/MetasoundFrontendNodeTemplateInput.h b/Engine/Plugins/Runtime/Metasound/Source/MetasoundFrontend/Public/NodeTemplates/MetasoundFrontendNodeTemplateInput.h
--- a/Engine/Plugins/Runtime/Metasound/Source/MetasoundFrontend/Public/NodeTemplates/MetasoundFrontendNodeTemplateInput.h
+++ b/Engine/Plugins/Runtime/Metasound/Source/MetasoundFrontend/Public/NodeTemplates/MetasoundFrontendNodeTemplateInput.h
@@ -34,16 +34,16 @@ namespace Metasound::Frontend
 		virtual const FMetasoundFrontendClassName& GetClassName() const override;
 
 #if WITH_EDITOR
-		virtual FText GetNodeDisplayName(const IMetaSoundDocumentInterface& Interface, const FGuid& InNodeID) const override;
+		virtual FText GetNodeDisplayName(const IMetaSoundDocumentInterface& Interface, const FGuid& InPageID, const FGuid& InNodeID) const override;
 #endif // WITH_EDITOR
 
 		virtual const FMetasoundFrontendClass& GetFrontendClass() const override;
-		virtual EMetasoundFrontendVertexAccessType GetNodeInputAccessType(const FMetaSoundFrontendDocumentBuilder& InBuilder, const FGuid& InNodeID, const FGuid& InVertexID) const override;
-		virtual EMetasoundFrontendVertexAccessType GetNodeOutputAccessType(const FMetaSoundFrontendDocumentBuilder& InBuilder, const FGuid& InNodeID, const FGuid& InVertexID) const override;
+		virtual EMetasoundFrontendVertexAccessType GetNodeInputAccessType(const FMetaSoundFrontendDocumentBuilder& InBuilder, const FGuid& InPageID, const FGuid& InNodeID, const FGuid& InVertexID) const override;
+		virtual EMetasoundFrontendVertexAccessType GetNodeOutputAccessType(const FMetaSoundFrontendDocumentBuilder& InBuilder, const FGuid& InPageID, const FGuid& InNodeID, const FGuid& InVertexID) const override;
 
 #if WITH_EDITOR
-		virtual FText GetOutputVertexDisplayName(const FMetaSoundFrontendDocumentBuilder& InBuilder, const FGuid& InNodeID, FName OutputName) const override;
-		virtual bool HasRequiredConnections(const FMetaSoundFrontendDocumentBuilder& InBuilder, const FGuid& InNodeID, FString* OutMessage = nullptr) const override;
+		virtual FText GetOutputVertexDisplayName(const FMetaSoundFrontendDocumentBuilder& InBuilder, const FGuid& InPageID, const FGuid& InNodeID, FName OutputName) const override;
+		virtual bool HasRequiredConnections(const FMetaSoundFrontendDocumentBuilder& InBuilder, const FGuid& InPageID, const FGuid& InNodeID, FString* OutMessage = nullptr) const override;
 
 		// Injects template nodes between builder's document inputs not connected
 		// to existing template inputs, copying locational data from the represented
diff --git a/Engine/Plugins/Runtime/Metasound/Source/MetasoundFrontend/Public/NodeTemplates/MetasoundFrontendNodeTemplateReroute.h b/Engine/Plugins/Runtime/Metasound/Source/MetasoundFrontend/Public/NodeTemplates/MetasoundFrontendNodeTemplateReroute.h
--- a/Engine/Plugins/Runtime/Metasound/Source/MetasoundFrontend/Public/NodeTemplates/MetasoundFrontendNodeTemplateReroute.h
+++ b/Engine/Plugins/Runtime/Metasound/Source/MetasoundFrontend/Public/NodeTemplates/MetasoundFrontendNodeTemplateReroute.h
@@ -29,7 +29,7 @@ namespace Metasound::Frontend
 		virtual const FMetasoundFrontendClassName& GetClassName() const override;
 
 #if WITH_EDITOR
-		virtual FText GetNodeDisplayName(const IMetaSoundDocumentInterface& DocumentInterface, const FGuid& InNodeID) const override;
+		virtual FText GetNodeDisplayName(const IMetaSoundDocumentInterface& DocumentInterface, const FGuid& InPageID, const FGuid& InNodeID) const override;
 #endif // WITH_EDITOR
 
 		UE_DEPRECATED(5.4, "Use version that does not require mutating a provided PreprocessedDocument")
@@ -38,8 +38,8 @@ namespace Metasound::Frontend
 		virtual FMetasoundFrontendNodeInterface GenerateNodeInterface(FNodeTemplateGenerateInterfaceParams InParams) const override;
 		virtual TUniquePtr<INodeTemplateTransform> GenerateNodeTransform() const override;
 		virtual const FMetasoundFrontendClass& GetFrontendClass() const override;
-		virtual EMetasoundFrontendVertexAccessType GetNodeInputAccessType(const FMetaSoundFrontendDocumentBuilder& InBuilder, const FGuid& InNodeID, const FGuid& InVertexID) const override;
-		virtual EMetasoundFrontendVertexAccessType GetNodeOutputAccessType(const FMetaSoundFrontendDocumentBuilder& InBuilder, const FGuid& InNodeID, const FGuid& InVertexID) const override;
+		virtual EMetasoundFrontendVertexAccessType GetNodeInputAccessType(const FMetaSoundFrontendDocumentBuilder& InBuilder, const FGuid& InPageID, const FGuid& InNodeID, const FGuid& InVertexID) const override;
+		virtual EMetasoundFrontendVertexAccessType GetNodeOutputAccessType(const FMetaSoundFrontendDocumentBuilder& InBuilder, const FGuid& InPageID, const FGuid& InNodeID, const FGuid& InVertexID) const override;
 		virtual const FMetasoundFrontendVersionNumber& GetVersionNumber() const override;
 		virtual bool IsInputAccessTypeDynamic() const override;
 		virtual bool IsInputConnectionUserModifiable() const override;
@@ -48,7 +48,7 @@ namespace Metasound::Frontend
 		virtual bool IsValidNodeInterface(const FMetasoundFrontendNodeInterface& InNodeInterface) const override;
 
 #if WITH_EDITOR
-		virtual bool HasRequiredConnections(const FMetaSoundFrontendDocumentBuilder& InBuilder, const FGuid& InNodeID, FString* OutMessage = nullptr) const override;
+		virtual bool HasRequiredConnections(const FMetaSoundFrontendDocumentBuilder& InBuilder, const FGuid& InPageID, const FGuid& InNodeID, FString* OutMessage = nullptr) const override;
 #endif // WITH_EDITOR
 	};
 } // namespace Metasound::Frontend
