Initial checkpoint for runtime graph to support Metasound document pages
- Adds ability for graph generation via target page ID
- Deprecates the requirement of a GraphController to build graph to IGraph

#tests MetaSound automated tests & UnitTests
#rnx

diff --git a/Engine/Plugins/Runtime/Metasound/Source/MetasoundEngine/Private/MetasoundBuilderSubsystem.cpp b/Engine/Plugins/Runtime/Metasound/Source/MetasoundEngine/Private/MetasoundBuilderSubsystem.cpp
--- a/Engine/Plugins/Runtime/Metasound/Source/MetasoundEngine/Private/MetasoundBuilderSubsystem.cpp
+++ b/Engine/Plugins/Runtime/Metasound/Source/MetasoundEngine/Private/MetasoundBuilderSubsystem.cpp
@@ -266,10 +266,10 @@ void UMetaSoundSourceBuilder::InitDelegates(Metasound::Frontend::FDocumentModify
 	OutDocumentDelegates.InterfaceDelegates.OnRemovingInput.AddUObject(this, &UMetaSoundSourceBuilder::OnRemovingInput);
 	OutDocumentDelegates.InterfaceDelegates.OnRemovingOutput.AddUObject(this, &UMetaSoundSourceBuilder::OnRemovingOutput);
 
-	InitExecutablePageDelegates(OutDocumentDelegates);
+	InitTargetPageDelegates(OutDocumentDelegates);
 }
 
-void UMetaSoundSourceBuilder::InitExecutablePageDelegates(Metasound::Frontend::FDocumentModifyDelegates& OutDocumentDelegates)
+void UMetaSoundSourceBuilder::InitTargetPageDelegates(Metasound::Frontend::FDocumentModifyDelegates& OutDocumentDelegates)
 {
 	using namespace Metasound::DynamicGraph;
 	using namespace Metasound::Engine;
@@ -305,13 +305,14 @@ void UMetaSoundSourceBuilder::InitExecutablePageDelegates(Metasound::Frontend::F
 		EdgeDelegates.OnRemoveSwappingEdge.RemoveAll(this);
 	});
 
-	const FGuid ExecutablePageID = FDocumentBuilderRegistry::GetChecked().ResolveExecutablePageID(&GetMetaSoundSource());
+	const IMetaSoundDocumentInterface& DocInterface = GetConstBuilder().GetConstDocumentInterfaceChecked();
+	const FGuid PageID = FDocumentBuilderRegistry::GetChecked().ResolveTargetPageID(DocInterface.GetConstDocument(), DocInterface.GetAssetPathChecked());
 
-	FEdgeModifyDelegates& EdgeDelegates = OutDocumentDelegates.FindEdgeDelegatesChecked(ExecutablePageID);
+	FEdgeModifyDelegates& EdgeDelegates = OutDocumentDelegates.FindEdgeDelegatesChecked(PageID);
 	EdgeDelegates.OnEdgeAdded.AddUObject(this, &UMetaSoundSourceBuilder::OnEdgeAdded);
 	EdgeDelegates.OnRemoveSwappingEdge.AddUObject(this, &UMetaSoundSourceBuilder::OnRemoveSwappingEdge);
 
-	FNodeModifyDelegates& NodeDelegates = OutDocumentDelegates.FindNodeDelegatesChecked(ExecutablePageID);
+	FNodeModifyDelegates& NodeDelegates = OutDocumentDelegates.FindNodeDelegatesChecked(PageID);
 	NodeDelegates.OnNodeAdded.AddUObject(this, &UMetaSoundSourceBuilder::OnNodeAdded);
 	NodeDelegates.OnNodeInputLiteralSet.AddUObject(this, &UMetaSoundSourceBuilder::OnNodeInputLiteralSet);
 	NodeDelegates.OnRemoveSwappingNode.AddUObject(this, &UMetaSoundSourceBuilder::OnRemoveSwappingNode);
@@ -557,7 +558,7 @@ void UMetaSoundSourceBuilder::OnPageAdded(const Metasound::Frontend::FDocumentMu
 	using namespace Metasound::Frontend;
 
 	FDocumentModifyDelegates& DocDelegates = Builder.GetDocumentDelegates();
-	InitExecutablePageDelegates(DocDelegates);
+	InitTargetPageDelegates(DocDelegates);
 }
 
 void UMetaSoundSourceBuilder::OnRemovingPage(const Metasound::Frontend::FDocumentMutatePageArgs& Args)
@@ -565,7 +566,7 @@ void UMetaSoundSourceBuilder::OnRemovingPage(const Metasound::Frontend::FDocumen
 	using namespace Metasound::Frontend;
 
 	FDocumentModifyDelegates& DocDelegates = Builder.GetDocumentDelegates();
-	InitExecutablePageDelegates(DocDelegates);
+	InitTargetPageDelegates(DocDelegates);
 
 }
 
@@ -1136,6 +1137,16 @@ void UMetaSoundBuilderSubsystem::RegisterSourceBuilder(FName BuilderName, UMetaS
 	}
 }
 
+bool UMetaSoundBuilderSubsystem::SetTargetPage(FName PageName)
+{
+	if (UMetaSoundSettings* Settings = GetMutableDefault<UMetaSoundSettings>())
+	{
+		return Settings->SetTargetPage(PageName);
+	}
+
+	return false;
+}
+
 bool UMetaSoundBuilderSubsystem::UnregisterBuilder(FName BuilderName)
 {
 	return NamedBuilders.Remove(BuilderName) > 0;
diff --git a/Engine/Plugins/Runtime/Metasound/Source/MetasoundEngine/Private/MetasoundDocumentBuilderRegistry.cpp b/Engine/Plugins/Runtime/Metasound/Source/MetasoundEngine/Private/MetasoundDocumentBuilderRegistry.cpp
--- a/Engine/Plugins/Runtime/Metasound/Source/MetasoundEngine/Private/MetasoundDocumentBuilderRegistry.cpp
+++ b/Engine/Plugins/Runtime/Metasound/Source/MetasoundEngine/Private/MetasoundDocumentBuilderRegistry.cpp
@@ -2,6 +2,7 @@
 
 #include "MetasoundDocumentBuilderRegistry.h"
 
+#include "HAL/PlatformProperties.h"
 #include "MetasoundAssetManager.h"
 #include "MetasoundGlobals.h"
 #include "MetasoundSettings.h"
@@ -233,12 +234,6 @@ namespace Metasound::Engine
 		return FoundBuilders;
 	}
 
-	FGuid FDocumentBuilderRegistry::ResolveExecutablePageID(const TScriptInterface<IMetaSoundDocumentInterface> DocumentInterface) const
-	{
-		// TODO: implement selection logic.  For now default to 0 guid, which is always valid on an asset.
-		return FGuid();
-	}
-
 	FMetaSoundFrontendDocumentBuilder* FDocumentBuilderRegistry::FindOutermostBuilder(const UObject& InSubObject) const
 	{
 		using namespace Metasound::Frontend;
@@ -309,6 +304,13 @@ namespace Metasound::Engine
 		}
 	}
 
+#if WITH_EDITOR
+	FOnResolvePreviewPageInfo& FDocumentBuilderRegistry::GetOnResolvePreviewPageInfoDelegate()
+	{
+		return OnResolvePreviewPageInfo;
+	}
+#endif // WITH_EDITOR
+
 	bool FDocumentBuilderRegistry::ReloadBuilder(const FMetasoundFrontendClassName& InClassName) const
 	{
 		bool bReloaded = false;
@@ -322,6 +324,90 @@ namespace Metasound::Engine
 		return bReloaded;
 	}
 
+	FGuid FDocumentBuilderRegistry::ResolveTargetPageID(const FMetasoundFrontendDocument& Document, const FTopLevelAssetPath& AssetPath) const
+	{
+		FName PlatformName = FPlatformProperties::IniPlatformName();
+
+#if WITH_EDITOR
+		if (OnResolvePreviewPageInfo.IsBound())
+		{
+			FPreviewPageInfo PreviewInfo = OnResolvePreviewPageInfo.Execute(Document);
+			if (PreviewInfo.PageID.IsSet())
+			{
+				return PreviewInfo.PageID.GetValue();
+			}
+
+			PlatformName = PreviewInfo.PlatformName;
+		}
+#endif // WITH_EDITOR
+
+		TSet<FGuid> DocPageIds;
+		Document.RootGraph.IterateGraphPages([&DocPageIds](const FMetasoundFrontendGraph& PageGraph)
+		{
+			DocPageIds.Add(PageGraph.PageID);
+		});
+
+		bool bImplementsPages = false;
+		const UMetaSoundSettings* Settings = GetDefault<UMetaSoundSettings>();
+		if (Settings)
+		{
+			const FGuid& TargetPageID = Settings->GetTargetPageID();
+			const TArray<FMetaSoundPageSettings>& PageSettingsArray = Settings->GetPageSettings();
+			bImplementsPages = Settings->GetPageSettings().IsEmpty();
+
+			bool bFoundMatch = false;
+			for (int32 Index = PageSettingsArray.Num() - 1; Index >= 0; --Index)
+			{
+				const FMetaSoundPageSettings& PageSettings = PageSettingsArray[Index];
+				bFoundMatch |= PageSettings.UniqueId == TargetPageID;
+				const bool bAssetImplementsPage = DocPageIds.Contains(PageSettings.UniqueId);
+				if (bFoundMatch && bAssetImplementsPage)
+				{
+#if WITH_EDITOR
+					const bool bIsCooked = PageSettings.IsCooked.GetValueForPlatform(PlatformName);
+					if (bIsCooked)
+					{
+						return PageSettings.UniqueId;
+					}
+#else // !WITH_EDITOR
+					return PageSettings.UniqueId;
+#endif // !WITH_EDITOR
+				}
+			}
+		}
+
+		// Set to arbitrary ID so the default is prioritized when iterating implemented
+		// asset pages. All documents are guaranteed to implement at least one page.
+		FGuid PageID = FGuid::NewGuid();
+		Document.RootGraph.IterateGraphPages([&PageID](const FMetasoundFrontendGraph& Graph)
+		{
+			if (PageID.IsValid())
+			{
+				PageID = Graph.PageID;
+			}
+		});
+
+#if !NO_LOGGING
+		// Log error if page settings exist, but could not resolve ID with the given document
+		if (bImplementsPages)
+		{
+			FString PageIdentifier = PageID.ToString();
+			if (const FMetaSoundPageSettings* PageSettings = Settings->FindPageSettings(PageID))
+			{
+				PageIdentifier = PageSettings->Name.ToString();
+			}
+
+			UE_LOG(LogMetaSound, Error,
+				TEXT("'%s' failed to resolve executable page ID:  \nMetaSound 'Page Settings' does not provide a valid fallback page for execution"
+					"on the desired platform, which can result in undefined behavior. Registering asset's page with ID '%s'."),
+				*AssetPath.ToString(),
+				*PageIdentifier);
+		}
+#endif // !NO_LOGGING
+
+		return PageID;
+	}
+
 	void FDocumentBuilderRegistry::SetEventLogVerbosity(ELogEvent Event, ELogVerbosity::Type Verbosity)
 	{
 		EventLogVerbosity.FindOrAdd(Event) = Verbosity;
diff --git a/Engine/Plugins/Runtime/Metasound/Source/MetasoundEngine/Private/MetasoundSettings.cpp b/Engine/Plugins/Runtime/Metasound/Source/MetasoundEngine/Private/MetasoundSettings.cpp
--- a/Engine/Plugins/Runtime/Metasound/Source/MetasoundEngine/Private/MetasoundSettings.cpp
+++ b/Engine/Plugins/Runtime/Metasound/Source/MetasoundEngine/Private/MetasoundSettings.cpp
@@ -196,6 +196,21 @@ const FMetaSoundQualitySettings* UMetaSoundSettings::FindQualitySettings(const F
 	return Metasound::SettingsPrivate::FindSettingsStruct(QualitySettings, InQualityID);
 }
 
+const FGuid& UMetaSoundSettings::GetTargetPageID() const
+{
+	if (const FMetaSoundPageSettings* TargetSettings = FindPageSettings(TargetPageName))
+	{
+		return TargetSettings->UniqueId;
+	}
+
+	if (PageSettings.IsEmpty())
+	{
+		return Metasound::Frontend::DefaultGraphPageID;
+	}
+
+	return PageSettings.Last().UniqueId;
+}
+
 #if WITH_EDITOR
 void UMetaSoundSettings::PostEditChangeChainProperty(FPropertyChangedChainEvent& PostEditChangeChainProperty)
 {
@@ -231,6 +246,20 @@ void UMetaSoundSettings::PostInitProperties()
 }
 #endif // WITH_EDITOR
 
+bool UMetaSoundSettings::SetTargetPage(FName PageName)
+{
+	if (const FMetaSoundPageSettings* PageSetting = FindPageSettings(PageName))
+	{
+		if (TargetPageName != PageSetting->Name)
+		{
+			TargetPageName = PageSetting->Name;
+			return true;
+		}
+	}
+
+	return false;
+}
+
 #if WITH_EDITORONLY_DATA
 Metasound::Engine::FOnSettingsDefaultConformed& UMetaSoundSettings::GetOnDefaultConformedDelegate()
 {
diff --git a/Engine/Plugins/Runtime/Metasound/Source/MetasoundEngine/Public/MetasoundBuilderBase.cpp b/Engine/Plugins/Runtime/Metasound/Source/MetasoundEngine/Public/MetasoundBuilderBase.cpp
--- a/Engine/Plugins/Runtime/Metasound/Source/MetasoundEngine/Public/MetasoundBuilderBase.cpp
+++ b/Engine/Plugins/Runtime/Metasound/Source/MetasoundEngine/Public/MetasoundBuilderBase.cpp
@@ -152,6 +152,7 @@ FMetaSoundBuilderNodeInputHandle UMetaSoundBuilderBase::AddGraphOutputNode(FName
 	return NewHandle;
 }
 
+#if WITH_EDITORONLY_DATA
 void UMetaSoundBuilderBase::AddGraphPage(FName PageName, bool bDuplicateLastGraph, bool bSetAsBuildGraph, EMetaSoundBuilderResult& OutResult)
 {
 	if (const UMetaSoundSettings* Settings = GetDefault<UMetaSoundSettings>())
@@ -166,6 +167,7 @@ void UMetaSoundBuilderBase::AddGraphPage(FName PageName, bool bDuplicateLastGrap
 
 	OutResult = EMetaSoundBuilderResult::Failed;
 }
+#endif // WITH_EDITORONLY_DATA
 
 void UMetaSoundBuilderBase::AddInterface(FName InterfaceName, EMetaSoundBuilderResult& OutResult)
 {
@@ -941,6 +943,7 @@ void UMetaSoundBuilderBase::ReloadCache(bool bPrimeCache)
 void UMetaSoundBuilderBase::RemoveAllGraphPages()
 {
 	Builder.RemoveAllGraphPages();
+	Builder.RemoveUnusedDependencies();
 }
 
 #if WITH_EDITOR
@@ -962,6 +965,7 @@ void UMetaSoundBuilderBase::RemoveGraphOutput(FName Name, EMetaSoundBuilderResul
 	OutResult = bRemoved ? EMetaSoundBuilderResult::Succeeded : EMetaSoundBuilderResult::Failed;
 }
 
+#if WITH_EDITORONLY_DATA
 void UMetaSoundBuilderBase::RemoveGraphPage(FName Name, EMetaSoundBuilderResult& OutResult)
 {
 	const UMetaSoundSettings* Settings = GetDefault<UMetaSoundSettings>();
@@ -975,6 +979,7 @@ void UMetaSoundBuilderBase::RemoveGraphPage(FName Name, EMetaSoundBuilderResult&
 
 	OutResult = EMetaSoundBuilderResult::Failed;
 }
+#endif // WITH_EDITORONLY_DATA
 
 void UMetaSoundBuilderBase::RemoveInterface(FName InterfaceName, EMetaSoundBuilderResult& OutResult)
 {
diff --git a/Engine/Plugins/Runtime/Metasound/Source/MetasoundEngine/Public/MetasoundBuilderBase.h b/Engine/Plugins/Runtime/Metasound/Source/MetasoundEngine/Public/MetasoundBuilderBase.h
--- a/Engine/Plugins/Runtime/Metasound/Source/MetasoundEngine/Public/MetasoundBuilderBase.h
+++ b/Engine/Plugins/Runtime/Metasound/Source/MetasoundEngine/Public/MetasoundBuilderBase.h
@@ -135,10 +135,12 @@ public:
 	UFUNCTION(BlueprintCallable, Category = "Audio|MetaSound|Builder", meta = (ExpandEnumAsExecs = "OutResult", AdvancedDisplay = "3"))
 	UPARAM(DisplayName = "Input Handle") FMetaSoundBuilderNodeInputHandle AddGraphOutputNode(FName Name, FName DataType, FMetasoundFrontendLiteral DefaultValue, EMetaSoundBuilderResult& OutResult, bool bIsConstructorOutput = false);
 
+#if WITH_EDITORONLY_DATA
 	// Adds a graph page to the given builder's document. Fails if the page is not a valid page registered with MetaSoundSettings
 	// or if the document already contains a page with the given name. No check is done here to determine cook eligibility (i.e.
 	// pages can be added even if set to be stripped for the active platform).
 	void AddGraphPage(FName PageName, bool bDuplicateLastGraph, bool bSetAsBuildGraph, EMetaSoundBuilderResult& OutResult);
+#endif // WITH_EDITORONLY_DATA
 
 	// Adds an interface registered with the given name to the graph, adding associated input and output nodes.
 	UFUNCTION(BlueprintCallable, Category = "Audio|MetaSound|Builder", meta = (ExpandEnumAsExecs = "OutResult"))
@@ -347,9 +349,11 @@ public:
 	UFUNCTION(BlueprintCallable, Category = "Audio|MetaSound|Builder", meta = (ExpandEnumAsExecs = "OutResult"))
 	void RemoveGraphOutput(FName Name, EMetaSoundBuilderResult& OutResult);
 
+#if WITH_EDITORONLY_DATA
 	// Removes a graph page with the given name, setting result to failed if the name was not found, was invalid,
 	// or was the default (which cannot be removed).
 	void RemoveGraphPage(FName Name, EMetaSoundBuilderResult& OutResult);
+#endif // WITH_EDITORONLY_DATA
 
 	// Removes the interface with the given name from the builder's MetaSound. Removes any graph inputs
 	// and outputs associated with the given interface and their respective connections (if any).
diff --git a/Engine/Plugins/Runtime/Metasound/Source/MetasoundEngine/Public/MetasoundBuilderSubsystem.h b/Engine/Plugins/Runtime/Metasound/Source/MetasoundEngine/Public/MetasoundBuilderSubsystem.h
--- a/Engine/Plugins/Runtime/Metasound/Source/MetasoundEngine/Public/MetasoundBuilderSubsystem.h
+++ b/Engine/Plugins/Runtime/Metasound/Source/MetasoundEngine/Public/MetasoundBuilderSubsystem.h
@@ -121,7 +121,7 @@ private:
 	const UMetaSoundSource& GetMetaSoundSource() const;
 	UMetaSoundSource& GetMetaSoundSource();
 
-	void InitExecutablePageDelegates(Metasound::Frontend::FDocumentModifyDelegates& OutDocumentDelegates);
+	void InitTargetPageDelegates(Metasound::Frontend::FDocumentModifyDelegates& OutDocumentDelegates);
 
 	void OnEdgeAdded(int32 EdgeIndex) const;
 	void OnInputAdded(int32 InputIndex);
@@ -267,6 +267,9 @@ public:
 	UFUNCTION(BlueprintCallable, Category = "Audio|MetaSound|Builder")
 	void RegisterSourceBuilder(FName BuilderName, UMetaSoundSourceBuilder* Builder);
 
+	UFUNCTION(BlueprintCallable, Category = "MetaSounds|Pages")
+	UPARAM(DisplayName = "TargetPageChanged") bool SetTargetPage(FName PageName);
+
 	UFUNCTION(BlueprintCallable, Category = "Audio|MetaSound|Builder")
 	UPARAM(DisplayName = "Unregistered") bool UnregisterBuilder(FName BuilderName);
 
diff --git a/Engine/Plugins/Runtime/Metasound/Source/MetasoundEngine/Public/MetasoundDocumentBuilderRegistry.h b/Engine/Plugins/Runtime/Metasound/Source/MetasoundEngine/Public/MetasoundDocumentBuilderRegistry.h
--- a/Engine/Plugins/Runtime/Metasound/Source/MetasoundEngine/Public/MetasoundDocumentBuilderRegistry.h
+++ b/Engine/Plugins/Runtime/Metasound/Source/MetasoundEngine/Public/MetasoundDocumentBuilderRegistry.h
@@ -8,11 +8,22 @@
 #include "MetasoundFrontendDocument.h"
 #include "MetasoundFrontendRegistries.h"
 #include "MetasoundLog.h"
+#include "Misc/Optional.h"
 #include "Misc/ScopeLock.h"
 
 
 namespace Metasound::Engine
 {
+#if WITH_EDITOR
+	struct FPreviewPageInfo
+	{
+		FName PlatformName;
+		TOptional<FGuid> PageID;
+	};
+
+	DECLARE_DELEGATE_RetVal_OneParam(FPreviewPageInfo, FOnResolvePreviewPageInfo, const FMetasoundFrontendDocument& /* Document */);
+#endif // WITH_EDITOR
+
 	class METASOUNDENGINE_API FDocumentBuilderRegistry : public Frontend::IDocumentBuilderRegistry
 	{
 		mutable TMultiMap<FMetasoundFrontendClassName, TWeakObjectPtr<UMetaSoundBuilderBase>> Builders;
@@ -119,11 +130,15 @@ namespace Metasound::Engine
 		// Returns all builder objects registered and active associated with the given ClassName.
 		TArray<UMetaSoundBuilderBase*> FindBuilderObjects(const FMetasoundFrontendClassName& InClassName) const;
 
-		// Given the provided document, returns the valid page ID to be executed.
-		virtual FGuid ResolveExecutablePageID(const TScriptInterface<IMetaSoundDocumentInterface> DocumentInterface) const;
+#if WITH_EDITOR
+		FOnResolvePreviewPageInfo& GetOnResolvePreviewPageInfoDelegate();
+#endif // WITH_EDITOR
 
 		bool ReloadBuilder(const FMetasoundFrontendClassName& InClassName) const override;
 
+		// Given the provided document and its respective pages, returns the PageID to be used for runtime IGraph and proxy generation.
+		virtual FGuid ResolveTargetPageID(const FMetasoundFrontendDocument& Document, const FTopLevelAssetPath& AssetPath) const override;
+
 		void SetEventLogVerbosity(ELogEvent Event, ELogVerbosity::Type Verbosity);
 
 	private:
@@ -131,6 +146,10 @@ namespace Metasound::Engine
 		bool CanPostEventLog(ELogEvent Event, ELogVerbosity::Type Verbosity) const;
 		void FinishBuildingInternal(UMetaSoundBuilderBase& Builder, bool bForceUnregisterNodeClass) const;
 
+#if WITH_EDITOR
+		FOnResolvePreviewPageInfo OnResolvePreviewPageInfo;
+#endif // WITH_EDITOR
+
 		TSortedMap<ELogEvent, ELogVerbosity::Type> EventLogVerbosity;
 	};
 } // namespace Metasound::Engine
\ No newline at end of file
diff --git a/Engine/Plugins/Runtime/Metasound/Source/MetasoundEngine/Public/MetasoundSettings.h b/Engine/Plugins/Runtime/Metasound/Source/MetasoundEngine/Public/MetasoundSettings.h
--- a/Engine/Plugins/Runtime/Metasound/Source/MetasoundEngine/Public/MetasoundSettings.h
+++ b/Engine/Plugins/Runtime/Metasound/Source/MetasoundEngine/Public/MetasoundSettings.h
@@ -131,7 +131,7 @@ public:
 	TArray<FDirectoryPath> DirectoriesToRegister;
 		
 	UPROPERTY(Transient)
-	int32 DenyListCacheChangeID = 0;	
+	int32 DenyListCacheChangeID = 0;
 
 private:
 #if WITH_EDITORONLY_DATA
@@ -139,7 +139,16 @@ private:
 	Metasound::Engine::FOnPageSettingsUpdated OnPageSettingsUpdated;
 #endif //WITH_EDITORONLY_DATA
 
-	/** Array of possible page settings that can be added to a MetaSound object. */
+	/** PageID to target when attempting to execute MetaSound. If target page is not implemented for the set
+	  *  platform, uses order of cooked pages (see 'Page Settings' for order) falling back to lower index-ordered page
+	  * implemented in MetaSound asset. */
+	UPROPERTY(EditAnywhere, config, Category = Pages)
+	FName TargetPageName = Metasound::Frontend::DefaultGraphPageName;
+
+	/** Array of possible page settings that can be added to a MetaSound object. Order
+	  * defines default fallback logic whereby a higher index-ordered page
+	  * implemented in a MetaSound asset is higher priority (see 'Target Page').
+	  */
 	UPROPERTY(EditAnywhere, config, Category = Pages)
 	TArray<FMetaSoundPageSettings> PageSettings;
 
@@ -148,12 +157,26 @@ private:
 	TArray<FMetaSoundQualitySettings> QualitySettings;
 
 public:
+	// Returns the page settings with the provided name. If there are multiple settings
+	// with the same name, selection within the duplicates is undefined.
 	const FMetaSoundPageSettings* FindPageSettings(FName Name) const;
+
+	// Returns the page settings with the unique ID given.
 	const FMetaSoundPageSettings* FindPageSettings(const FGuid& InPageID) const;
 
+	// Returns the quality settings with the provided name. If there are multiple settings
+	// with the same name, selection within the duplicates is undefined.
 	const FMetaSoundQualitySettings* FindQualitySettings(FName Name) const;
+
+	// Returns the quality settings with the unique ID given.
 	const FMetaSoundQualitySettings* FindQualitySettings(const FGuid& InQualityID) const;
 
+	// Returns the target page name.
+	const FName& GetTargetPageName() const { return TargetPageName; }
+
+	// Returns the target page ID.
+	const FGuid& GetTargetPageID() const;
+
 	const TArray<FMetaSoundPageSettings>& GetPageSettings() const { return PageSettings; }
 	const TArray<FMetaSoundQualitySettings>& GetQualitySettings() const { return QualitySettings; }
 
@@ -165,6 +188,10 @@ public:
 	static FName GetQualitySettingPropertyName();
 #endif // WITH_EDITORONLY_DATA
 
+	// Sets the target page to the given name. Returns true if associated page settings were found
+	// and target set, false if not found and not set.
+	bool SetTargetPage(FName PageName);
+
 #if WITH_EDITOR
 private:
 	void ConformPageSettingsDefault(bool bNotifyDefaultConformed);
diff --git a/Engine/Plugins/Runtime/Metasound/Source/MetasoundFrontend/Private/MetasoundDocumentInterface.cpp b/Engine/Plugins/Runtime/Metasound/Source/MetasoundFrontend/Private/MetasoundDocumentInterface.cpp
--- a/Engine/Plugins/Runtime/Metasound/Source/MetasoundFrontend/Private/MetasoundDocumentInterface.cpp
+++ b/Engine/Plugins/Runtime/Metasound/Source/MetasoundFrontend/Private/MetasoundDocumentInterface.cpp
@@ -1,6 +1,11 @@
 // Copyright Epic Games, Inc. All Rights Reserved.
 
 #include "MetasoundDocumentInterface.h"
+#include "MetasoundFrontendDocumentIdGenerator.h"
+#include "MetasoundFrontendGraph.h"
+#include "MetasoundFrontendProxyDataCache.h"
+#include "MetasoundOperatorBuilder.h"
+#include "MetasoundOperatorBuilderSettings.h"
 
 
 namespace Metasound::Frontend
diff --git a/Engine/Plugins/Runtime/Metasound/Source/MetasoundFrontend/Private/MetasoundFrontendDocument.cpp b/Engine/Plugins/Runtime/Metasound/Source/MetasoundFrontend/Private/MetasoundFrontendDocument.cpp
--- a/Engine/Plugins/Runtime/Metasound/Source/MetasoundFrontend/Private/MetasoundFrontendDocument.cpp
+++ b/Engine/Plugins/Runtime/Metasound/Source/MetasoundFrontend/Private/MetasoundFrontendDocument.cpp
@@ -1127,10 +1127,27 @@ namespace Metasound::Frontend
 	void ForEachLiteral(const FMetasoundFrontendDocument& InDoc, FForEachLiteralFunctionRef OnLiteral)
 	{
 		ForEachLiteral(InDoc.RootGraph, OnLiteral);
+
 		for (const FMetasoundFrontendGraphClass& GraphClass : InDoc.Subgraphs)
 		{
 			ForEachLiteral(GraphClass, OnLiteral);
 		}
+
+		for (const FMetasoundFrontendClass& Dependency : InDoc.Dependencies)
+		{
+			ForEachLiteral(Dependency, OnLiteral);
+		}
+	}
+
+	void ForEachLiteral(const FMetasoundFrontendDocument& InDoc, FForEachLiteralFunctionRef OnLiteral, const FGuid& InPageID)
+	{
+		ForEachLiteral(InDoc.RootGraph, OnLiteral, InPageID);
+
+		for (const FMetasoundFrontendGraphClass& GraphClass : InDoc.Subgraphs)
+		{
+			ForEachLiteral(GraphClass, OnLiteral, InPageID);
+		}
+
 		for (const FMetasoundFrontendClass& Dependency : InDoc.Dependencies)
 		{
 			ForEachLiteral(Dependency, OnLiteral);
@@ -1141,12 +1158,31 @@ namespace Metasound::Frontend
 	{
 		ForEachLiteral(static_cast<const FMetasoundFrontendClass&>(InGraphClass), OnLiteral);
 
-		for (const FMetasoundFrontendNode& Node : InGraphClass.GetConstDefaultGraph().Nodes)
+		InGraphClass.IterateGraphPages([&OnLiteral](const FMetasoundFrontendGraph& Graph)
+		{
+			for (const FMetasoundFrontendNode& Node : Graph.Nodes)
+			{
+				ForEachLiteral(Node, OnLiteral);
+			}
+
+			for (const FMetasoundFrontendVariable& Variable : Graph.Variables)
+			{
+				OnLiteral(Variable.TypeName, Variable.Literal);
+			}
+		});
+	}
+
+	void ForEachLiteral(const FMetasoundFrontendGraphClass& InGraphClass, FForEachLiteralFunctionRef OnLiteral, const FGuid& InPageID)
+	{
+		ForEachLiteral(static_cast<const FMetasoundFrontendClass&>(InGraphClass), OnLiteral);
+
+		const FMetasoundFrontendGraph& Graph = InGraphClass.FindConstGraphChecked(InPageID);
+		for (const FMetasoundFrontendNode& Node : Graph.Nodes)
 		{
 			ForEachLiteral(Node, OnLiteral);
 		}
 
-		for (const FMetasoundFrontendVariable& Variable : InGraphClass.GetConstDefaultGraph().Variables)
+		for (const FMetasoundFrontendVariable& Variable : Graph.Variables)
 		{
 			OnLiteral(Variable.TypeName, Variable.Literal);
 		}
diff --git a/Engine/Plugins/Runtime/Metasound/Source/MetasoundFrontend/Private/MetasoundFrontendGraph.cpp b/Engine/Plugins/Runtime/Metasound/Source/MetasoundFrontend/Private/MetasoundFrontendGraph.cpp
--- a/Engine/Plugins/Runtime/Metasound/Source/MetasoundFrontend/Private/MetasoundFrontendGraph.cpp
+++ b/Engine/Plugins/Runtime/Metasound/Source/MetasoundFrontend/Private/MetasoundFrontendGraph.cpp
@@ -354,13 +354,13 @@ namespace Metasound
 	bool FFrontendGraphBuilder::AddNodesToGraph(FBuildGraphContext& InGraphContext)
 	{
 		TSet<FNodeIDVertexID> GraphEdgeDestinations;
-		const TArray<FMetasoundFrontendEdge>& GraphEdges = InGraphContext.GraphClass.GetConstDefaultGraph().Edges;
+		const TArray<FMetasoundFrontendEdge>& GraphEdges = InGraphContext.PagedGraph.Edges;
 		Algo::Transform(GraphEdges, GraphEdgeDestinations, [](const FMetasoundFrontendEdge& Edge) 
 		{
 			return FNodeIDVertexID{Edge.ToNodeID, Edge.ToVertexID};
 		});
 
-		for (const FMetasoundFrontendNode& Node : InGraphContext.GraphClass.GetConstDefaultGraph().Nodes)
+		for (const FMetasoundFrontendNode& Node : InGraphContext.PagedGraph.Nodes)
 		{
 			const FMetasoundFrontendClass* NodeClass = InGraphContext.BuildContext.FrontendClasses.FindRef(Node.ClassID);
 
@@ -429,7 +429,7 @@ namespace Metasound
 
 					case EMetasoundFrontendClassType::Variable:
 					{
-						TSharedPtr<const INode> VariableNode(CreateVariableNode(InGraphContext.BuildContext, Node, InGraphContext.GraphClass.GetConstDefaultGraph()).Release());
+						TSharedPtr<const INode> VariableNode(CreateVariableNode(InGraphContext.BuildContext, Node, InGraphContext.PagedGraph).Release());
 						InGraphContext.Graph->AddNode(Node.GetID(), VariableNode);
 					}
 					break;
@@ -468,7 +468,7 @@ namespace Metasound
 		TMap<FNodeIDVertexID, FCoreNodeAndFrontendVertex> NodeDestinationsByID;
 
 		// Add nodes to NodeID/VertexID map
-		for (const FMetasoundFrontendNode& Node : InGraphContext.GraphClass.GetConstDefaultGraph().Nodes)
+		for (const FMetasoundFrontendNode& Node : InGraphContext.PagedGraph.Nodes)
 		{
 			const INode* CoreNode = InGraphContext.Graph->FindNode(Node.GetID());
 			if (nullptr == CoreNode)
@@ -488,7 +488,7 @@ namespace Metasound
 			}
 		};
 
-		for (const FMetasoundFrontendEdge& Edge : InGraphContext.GraphClass.GetConstDefaultGraph().Edges)
+		for (const FMetasoundFrontendEdge& Edge : InGraphContext.PagedGraph.Edges)
 		{
 			const FNodeIDVertexID DestinationKey(Edge.ToNodeID, Edge.ToVertexID);
 			const FCoreNodeAndFrontendVertex* DestinationNodeAndVertex = NodeDestinationsByID.Find(DestinationKey);
@@ -741,14 +741,19 @@ PRAGMA_ENABLE_DEPRECATION_WARNINGS
 		return bSuccess;
 	}
 
-	TUniquePtr<FFrontendGraph> FFrontendGraphBuilder::CreateGraph(FBuildContext& InContext, const FMetasoundFrontendGraphClass& InGraphClass)
+	TUniquePtr<FFrontendGraph> FFrontendGraphBuilder::CreateGraph(FBuildContext& InContext, const FMetasoundFrontendGraphClass& InGraphClass, const FGuid* InPageId)
 	{
 		const FString GraphName = InContext.DebugAssetName;
 
+		const FMetasoundFrontendGraph& PageGraph = InPageId
+			? InGraphClass.FindConstGraphChecked(*InPageId)
+			: InGraphClass.GetConstDefaultGraph();
+
 		FBuildGraphContext BuildGraphContext
 		{
 			MakeUnique<FFrontendGraph>(GraphName, InContext.GraphId),
 			InGraphClass,
+			PageGraph,
 			InContext
 		};
 
@@ -779,7 +784,7 @@ PRAGMA_ENABLE_DEPRECATION_WARNINGS
 		return CreateGraph(InGraph, InSubgraphs, InDependencies, InDebugAssetName);
 	}
 
-	TUniquePtr<FFrontendGraph> FFrontendGraphBuilder::CreateGraph(const FMetasoundFrontendGraphClass& InGraph, const TArray<FMetasoundFrontendGraphClass>& InSubgraphs, const TArray<FMetasoundFrontendClass>& InDependencies, const Frontend::FProxyDataCache& InProxyDataCache, const FString& InDebugAssetName, const FGuid InGraphId)
+	TUniquePtr<FFrontendGraph> FFrontendGraphBuilder::CreateGraph(const FMetasoundFrontendGraphClass& InGraph, const TArray<FMetasoundFrontendGraphClass>& InSubgraphs, const TArray<FMetasoundFrontendClass>& InDependencies, const Frontend::FProxyDataCache& InProxyDataCache, const FString& InDebugAssetName, const FGuid InGraphId, const FGuid* InPageId)
 	{
 		FBuildContext Context
 		{
@@ -813,9 +818,10 @@ PRAGMA_ENABLE_DEPRECATION_WARNINGS
 		}
 
 		// Create each subgraph.
+		const FGuid PageId = InPageId ? *InPageId : FGuid();
 		for (const FMetasoundFrontendGraphClass* FrontendSubgraphPtr : FrontendSubgraphPtrs)
 		{
-			TSharedPtr<const INode> Subgraph(CreateGraph(Context, *FrontendSubgraphPtr).Release());
+			TSharedPtr<const INode> Subgraph(CreateGraph(Context, *FrontendSubgraphPtr, &PageId).Release());
 			if (!Subgraph.IsValid())
 			{
 				UE_LOG(LogMetaSound, Warning, TEXT("Failed to create subgraph [SubgraphName: %s] in asset '%s'"), *FrontendSubgraphPtr->Metadata.GetClassName().ToString(), *InDebugAssetName);
@@ -828,26 +834,28 @@ PRAGMA_ENABLE_DEPRECATION_WARNINGS
 		}
 
 		// Create parent graph.
-		return CreateGraph(Context, InGraph);
+		return CreateGraph(Context, InGraph, &PageId);
 	}
 
-	TUniquePtr<FFrontendGraph> FFrontendGraphBuilder::CreateGraph(const FMetasoundFrontendDocument& InDocument, const FString& InDebugAssetName)
+	TUniquePtr<FFrontendGraph> FFrontendGraphBuilder::CreateGraph(const FMetasoundFrontendDocument& InDocument, const FString& InDebugAssetName, const FGuid* InPageId)
 	{
 		// Create proxies before creating graph.
+		const FGuid PageId = InPageId ? *InPageId : FGuid();
 		Frontend::FProxyDataCache ProxyDataCache;
-		ProxyDataCache.CreateAndCacheProxies(InDocument);
+		ProxyDataCache.CreateAndCacheProxies(InDocument, PageId);
 		
-		return CreateGraph(InDocument, ProxyDataCache, InDebugAssetName);
+		return CreateGraph(InDocument, ProxyDataCache, InDebugAssetName, Frontend::CreateLocallyUniqueId(), &PageId);
 	}
 
-	TUniquePtr<FFrontendGraph> FFrontendGraphBuilder::CreateGraph(const FMetasoundFrontendGraphClass& InGraph, const TArray<FMetasoundFrontendGraphClass>& InSubgraphs, const TArray<FMetasoundFrontendClass>& InDependencies, const FString& InDebugAssetName)
+	TUniquePtr<FFrontendGraph> FFrontendGraphBuilder::CreateGraph(const FMetasoundFrontendGraphClass& InGraph, const TArray<FMetasoundFrontendGraphClass>& InSubgraphs, const TArray<FMetasoundFrontendClass>& InDependencies, const FString& InDebugAssetName, const FGuid* InPageId)
 	{
 		// Create proxies before building graph
+		const FGuid PageId = InPageId ? *InPageId : FGuid();
 		Frontend::FProxyDataCache ProxyDataCache;
-		ProxyDataCache.CreateAndCacheProxies(InGraph);
+		ProxyDataCache.CreateAndCacheProxies(InGraph, PageId);
 		for (const FMetasoundFrontendGraphClass& SubgraphClass : InSubgraphs)
 		{
-			ProxyDataCache.CreateAndCacheProxies(SubgraphClass);
+			ProxyDataCache.CreateAndCacheProxies(SubgraphClass, PageId);
 		}
 		for (const FMetasoundFrontendClass& DependencyClass : InDependencies)
 		{
@@ -855,7 +863,7 @@ PRAGMA_ENABLE_DEPRECATION_WARNINGS
 		}
 
 
-		return CreateGraph(InGraph, InSubgraphs, InDependencies, ProxyDataCache, InDebugAssetName);
+		return CreateGraph(InGraph, InSubgraphs, InDependencies, ProxyDataCache, InDebugAssetName, Frontend::CreateLocallyUniqueId(), &PageId);
 	}
 	
 	TUniquePtr<FFrontendGraph> FFrontendGraphBuilder::CreateGraph(const FMetasoundFrontendDocument& InDocument, const TSet<FName>& InTransmittableInputNames, const FString& InDebugAssetName)
@@ -863,8 +871,8 @@ PRAGMA_ENABLE_DEPRECATION_WARNINGS
 		return CreateGraph(InDocument, InDebugAssetName);
 	}
 
-	TUniquePtr<FFrontendGraph> FFrontendGraphBuilder::CreateGraph(const FMetasoundFrontendDocument& InDocument, const Frontend::FProxyDataCache& InProxyDataCache, const FString& InDebugAssetName, const FGuid InGraphId)
+	TUniquePtr<FFrontendGraph> FFrontendGraphBuilder::CreateGraph(const FMetasoundFrontendDocument& InDocument, const Frontend::FProxyDataCache& InProxyDataCache, const FString& InDebugAssetName, const FGuid InGraphId, const FGuid* InPageId)
 	{
-		return CreateGraph(InDocument.RootGraph, InDocument.Subgraphs, InDocument.Dependencies, InProxyDataCache, InDebugAssetName, InGraphId);
+		return CreateGraph(InDocument.RootGraph, InDocument.Subgraphs, InDocument.Dependencies, InProxyDataCache, InDebugAssetName, InGraphId, InPageId);
 	}
 }
diff --git a/Engine/Plugins/Runtime/Metasound/Source/MetasoundFrontend/Private/MetasoundFrontendGraphController.cpp b/Engine/Plugins/Runtime/Metasound/Source/MetasoundFrontend/Private/MetasoundFrontendGraphController.cpp
--- a/Engine/Plugins/Runtime/Metasound/Source/MetasoundFrontend/Private/MetasoundFrontendGraphController.cpp
+++ b/Engine/Plugins/Runtime/Metasound/Source/MetasoundFrontend/Private/MetasoundFrontendGraphController.cpp
@@ -1630,8 +1630,10 @@ namespace Metasound
 
 				FString UnknownAsset = TEXT("UnknownAsset");
 				FProxyDataCache ProxyCache;
-				ProxyCache.CreateAndCacheProxies(*(OwningDocument->GetDocumentPtr().Get()));
-				TUniquePtr<FFrontendGraph> Graph = FFrontendGraphBuilder::CreateGraph(*GraphClass, Subgraphs, Dependencies, ProxyCache, UnknownAsset);
+				const FMetasoundFrontendDocument& Doc = *(OwningDocument->GetDocumentPtr().Get());
+				ProxyCache.CreateAndCacheProxies(Doc, Frontend::DefaultGraphPageID);
+
+				TUniquePtr<FFrontendGraph> Graph = FFrontendGraphBuilder::CreateGraph(*GraphClass, Subgraphs, Dependencies, ProxyCache, UnknownAsset, Frontend::CreateLocallyUniqueId(), { });
 
 				if (!Graph.IsValid())
 				{
diff --git a/Engine/Plugins/Runtime/Metasound/Source/MetasoundFrontend/Private/MetasoundFrontendInvalidController.h b/Engine/Plugins/Runtime/Metasound/Source/MetasoundFrontend/Private/MetasoundFrontendInvalidController.h
--- a/Engine/Plugins/Runtime/Metasound/Source/MetasoundFrontend/Private/MetasoundFrontendInvalidController.h
+++ b/Engine/Plugins/Runtime/Metasound/Source/MetasoundFrontend/Private/MetasoundFrontendInvalidController.h
@@ -430,8 +430,6 @@ namespace Metasound
 
 			virtual FNodeHandle CreateEmptySubgraph(const FMetasoundFrontendClassMetadata& InInfo) override { return INodeController::GetInvalidHandle(); }
 
-			virtual TUniquePtr<IOperator> BuildOperator(const FOperatorSettings& InSettings, const FMetasoundEnvironment& InEnvironment, FBuildResults& OutBuildErrors) const override { return { }; }
-
 			virtual FDocumentHandle GetOwningDocument() override;
 			virtual FConstDocumentHandle GetOwningDocument() const override;
 
diff --git a/Engine/Plugins/Runtime/Metasound/Source/MetasoundFrontend/Private/MetasoundFrontendProxyDataCache.cpp b/Engine/Plugins/Runtime/Metasound/Source/MetasoundFrontend/Private/MetasoundFrontendProxyDataCache.cpp
--- a/Engine/Plugins/Runtime/Metasound/Source/MetasoundFrontend/Private/MetasoundFrontendProxyDataCache.cpp
+++ b/Engine/Plugins/Runtime/Metasound/Source/MetasoundFrontend/Private/MetasoundFrontendProxyDataCache.cpp
@@ -7,6 +7,7 @@
 #include "Containers/Map.h"
 #include "CoreGlobals.h"
 #include "IAudioProxyInitializer.h"
+#include "MetasoundDocumentInterface.h"
 #include "MetasoundFrontendDataTypeRegistry.h"
 #include "MetasoundFrontendDocument.h"
 #include "MetasoundFrontendLiteral.h"
@@ -21,59 +22,59 @@ namespace Metasound
 	{
 		namespace MetasoundFrontendProxyDataCachePrivate
 		{
-			template<typename DocElementType>
-			void CreateAndCacheProxies(FProxyDataCache& InCache, const DocElementType& InDocElement)
+			template<typename DocElementType, typename ...Args>
+			void CreateAndCacheProxies(FProxyDataCache& InCache, const DocElementType& InDocElement, Args&&... InArgs)
 			{
 				checkf(IsInGameThread() || IsInAudioThread(), TEXT("Proxies should only be created in the game thread or audio thread"));
 				IDataTypeRegistry& DataRegistry = IDataTypeRegistry::Get();
 				TArray<UObject*> UObjectArray;
 
-				ForEachLiteral(InDocElement, [&InCache, &DataRegistry, &UObjectArray](const FName& InDataType, const FMetasoundFrontendLiteral& InLiteral)
+				auto CreateProxies = [&InCache, &DataRegistry, &UObjectArray](const FName& InDataType, const FMetasoundFrontendLiteral& InLiteral)
+				{
+					EMetasoundFrontendLiteralType LiteralType = InLiteral.GetType();
+					if (LiteralType == EMetasoundFrontendLiteralType::UObject)
 					{
-						EMetasoundFrontendLiteralType LiteralType = InLiteral.GetType();
-						if (LiteralType == EMetasoundFrontendLiteralType::UObject)
+						UObject* Object = nullptr;
+						InLiteral.TryGet(Object);
+						if (Object)
 						{
-							UObject* Object = nullptr;
-							InLiteral.TryGet(Object);
-							if (Object)
+							if (!InCache.Contains(Object))
 							{
-								if (!InCache.Contains(Object))
-								{
-									InCache.CacheProxy(Object, DataRegistry.CreateProxyFromUObject(InDataType, Object));
-								}
+								InCache.CacheProxy(Object, DataRegistry.CreateProxyFromUObject(InDataType, Object));
 							}
 						}
-						else if (LiteralType == EMetasoundFrontendLiteralType::UObjectArray)
+					}
+					else if (LiteralType == EMetasoundFrontendLiteralType::UObjectArray)
+					{
+						FName ElementDataTypeName = CreateElementTypeNameFromArrayTypeName(InDataType);
+						UObjectArray.Reset();
+						InLiteral.TryGet(UObjectArray);
+						for (UObject* Object : UObjectArray)
 						{
-							FName ElementDataTypeName = CreateElementTypeNameFromArrayTypeName(InDataType);
-							UObjectArray.Reset();
-							InLiteral.TryGet(UObjectArray);
-							for (UObject* Object : UObjectArray)
+							if (Object)
 							{
-								if (Object)
+								if (!InCache.Contains(Object))
 								{
-									if (!InCache.Contains(Object))
-									{
-										InCache.CacheProxy(Object, DataRegistry.CreateProxyFromUObject(ElementDataTypeName, Object));
-									}
+									InCache.CacheProxy(Object, DataRegistry.CreateProxyFromUObject(ElementDataTypeName, Object));
 								}
 							}
 						}
 					}
-				);
+				};
+				ForEachLiteral(InDocElement, CreateProxies, Forward<Args>(InArgs)...);
 			}
 		}
 
-		void FProxyDataCache::CreateAndCacheProxies(const FMetasoundFrontendDocument& InDocument)
+		void FProxyDataCache::CreateAndCacheProxies(const FMetasoundFrontendDocument& InDocument, const FGuid& InPageID)
 		{
-			METASOUND_TRACE_CPUPROFILER_EVENT_SCOPE(Metasound::FProxyDataCache::CreateAndCacheProxies_Document)	
-			MetasoundFrontendProxyDataCachePrivate::CreateAndCacheProxies(*this, InDocument);
+			METASOUND_TRACE_CPUPROFILER_EVENT_SCOPE(Metasound::FProxyDataCache::CreateAndCacheProxies_Document)
+			MetasoundFrontendProxyDataCachePrivate::CreateAndCacheProxies(*this, InDocument, InPageID);
 		}
 
-		void FProxyDataCache::CreateAndCacheProxies(const FMetasoundFrontendGraphClass& InGraphClass)
+		void FProxyDataCache::CreateAndCacheProxies(const FMetasoundFrontendGraphClass& InGraphClass, const FGuid& InPageID)
 		{
-			METASOUND_TRACE_CPUPROFILER_EVENT_SCOPE(Metasound::FProxyDataCache::CreateAndCacheProxies_GraphClass)	
-			MetasoundFrontendProxyDataCachePrivate::CreateAndCacheProxies(*this, InGraphClass);
+			METASOUND_TRACE_CPUPROFILER_EVENT_SCOPE(Metasound::FProxyDataCache::CreateAndCacheProxies_GraphClass)
+			MetasoundFrontendProxyDataCachePrivate::CreateAndCacheProxies(*this, InGraphClass, InPageID);
 		}
 
 		void FProxyDataCache::CreateAndCacheProxies(const FMetasoundFrontendClass& InClass)
diff --git a/Engine/Plugins/Runtime/Metasound/Source/MetasoundFrontend/Private/MetasoundFrontendProxyDataCache.h b/Engine/Plugins/Runtime/Metasound/Source/MetasoundFrontend/Private/MetasoundFrontendProxyDataCache.h
--- a/Engine/Plugins/Runtime/Metasound/Source/MetasoundFrontend/Private/MetasoundFrontendProxyDataCache.h
+++ b/Engine/Plugins/Runtime/Metasound/Source/MetasoundFrontend/Private/MetasoundFrontendProxyDataCache.h
@@ -32,10 +32,10 @@ namespace Metasound
 		{
 		public:
 			/** Create and cache all proxies found in a document. */
-			void CreateAndCacheProxies(const FMetasoundFrontendDocument& InDocument);
+			void CreateAndCacheProxies(const FMetasoundFrontendDocument& InDocument, const FGuid& InPageID);
 
 			/** Create and cache all proxies found in a graph class. */
-			void CreateAndCacheProxies(const FMetasoundFrontendGraphClass& InGraphClass);
+			void CreateAndCacheProxies(const FMetasoundFrontendGraphClass& InGraphClass, const FGuid& InPageID);
 
 			/** Create and cache all proxies found in a frontend class class. */
 			void CreateAndCacheProxies(const FMetasoundFrontendClass& InClass);
diff --git a/Engine/Plugins/Runtime/Metasound/Source/MetasoundFrontend/Private/MetasoundFrontendRegistryContainerImpl.cpp b/Engine/Plugins/Runtime/Metasound/Source/MetasoundFrontend/Private/MetasoundFrontendRegistryContainerImpl.cpp
--- a/Engine/Plugins/Runtime/Metasound/Source/MetasoundFrontend/Private/MetasoundFrontendRegistryContainerImpl.cpp
+++ b/Engine/Plugins/Runtime/Metasound/Source/MetasoundFrontend/Private/MetasoundFrontendRegistryContainerImpl.cpp
@@ -48,15 +48,14 @@ namespace Metasound::Frontend
 
 	namespace RegistryPrivate
 	{
-		TScriptInterface<IMetaSoundDocumentInterface> BuildRegistryDocument(TScriptInterface<IMetaSoundDocumentInterface> DocumentInterface, bool bAsync)
+		TScriptInterface<IMetaSoundDocumentInterface> BuildRegistryDocument(TScriptInterface<IMetaSoundDocumentInterface> DocumentInterface, bool bAsync, FGuid& OutPageID)
 		{
-			using namespace Metasound::Frontend;
-
 			METASOUND_TRACE_CPUPROFILER_EVENT_SCOPE(Metasound::Frontend::BuildRegistryDocument);
 
 			UObject* DocObject = DocumentInterface.GetObject();
 			check(DocObject);
 			const FMetasoundFrontendDocument& Document = DocumentInterface->GetConstDocument();
+			OutPageID = IDocumentBuilderRegistry::GetChecked().ResolveTargetPageID(Document, DocumentInterface->GetAssetPathChecked());
 
 #if WITH_EDITOR
 			// Only assets require template node processing
@@ -275,16 +274,15 @@ namespace Metasound::Frontend
 		};
 	} // namespace RegistryPrivate
 
-	void FRegistryContainerImpl::BuildAndRegisterGraphFromDocument(TScriptInterface<IMetaSoundDocumentInterface> DocumentInterface, const FProxyDataCache& InProxyDataCache, FNodeClassInfo&& InNodeClassInfo)
+	void FRegistryContainerImpl::BuildAndRegisterGraphFromDocument(const FMetasoundFrontendDocument& InDocument, const FProxyDataCache& InProxyDataCache, FNodeClassInfo&& InNodeClassInfo, const FGuid& InPageID)
 	{
 		using namespace RegistryPrivate;
 
 		METASOUND_TRACE_CPUPROFILER_EVENT_SCOPE(Metasound::FRegistryContainerImpl::BuildAndRegisterGraphFromDocument);
 		METASOUND_TRACE_CPUPROFILER_EVENT_SCOPE_TEXT(*FString::Printf(TEXT("Metasound::FRegistryContainerImpl::BuildAndRegisterGraphFromDocument asset %s"), *InNodeClassInfo.AssetPath.ToString()));
 
-		const FMetasoundFrontendDocument& Document = DocumentInterface->GetConstDocument();
 		// Use the asset class id for the graph id because it should be locally unique. Unlike other ids, it is regenerated on asset duplicate. 
-		TUniquePtr<FFrontendGraph> FrontendGraph = FFrontendGraphBuilder::CreateGraph(Document, InProxyDataCache, InNodeClassInfo.AssetPath.ToString(), /*GraphId=*/InNodeClassInfo.AssetClassID);
+		TUniquePtr<FFrontendGraph> FrontendGraph = FFrontendGraphBuilder::CreateGraph(InDocument, InProxyDataCache, InNodeClassInfo.AssetPath.ToString(), /*GraphId=*/InNodeClassInfo.AssetClassID, &InPageID);
 		if (!FrontendGraph.IsValid())
 		{
 			UE_LOG(LogMetaSound, Error, TEXT("Failed to build MetaSound graph in asset '%s'"), *InNodeClassInfo.AssetPath.ToString());
@@ -292,12 +290,12 @@ namespace Metasound::Frontend
 
 		TSharedPtr<const FGraph> GraphToRegister = MakeShareable<const FGraph>(FrontendGraph.Release());
 
-		TUniquePtr<INodeRegistryEntry> RegistryEntry = MakeUnique<FDocumentNodeRegistryEntry>(Document.RootGraph, Document.Interfaces, MoveTemp(InNodeClassInfo), GraphToRegister);
+		TUniquePtr<INodeRegistryEntry> RegistryEntry = MakeUnique<FDocumentNodeRegistryEntry>(InDocument.RootGraph, InDocument.Interfaces, MoveTemp(InNodeClassInfo), GraphToRegister);
 
 		const FNodeRegistryKey RegistryKey = RegisterNodeInternal(MoveTemp(RegistryEntry));
 
 		// Key must use the asset path provided to class info and *NOT* that of the
-		// provided DocumentInterface object, as that may be a built transient object
+		// document's owning DocumentInterface object, as that may be a built transient object
 		// with a different transient asset path.
 		const FGraphRegistryKey GraphKey { RegistryKey, InNodeClassInfo.AssetPath };
 		RegisterGraphInternal(GraphKey, GraphToRegister);
@@ -482,7 +480,8 @@ namespace Metasound::Frontend
 
 		// Use the asset path of the provided document interface object for identification, *NOT* the
 		// built version as the build process may in fact create a new object with a transient path.
-		const TScriptInterface<IMetaSoundDocumentInterface> RegistryDocInterface = RegistryPrivate::BuildRegistryDocument(InDocumentInterface, bAsync);
+		FGuid PageID = Frontend::DefaultGraphPageID;
+		const TScriptInterface<IMetaSoundDocumentInterface> RegistryDocInterface = RegistryPrivate::BuildRegistryDocument(InDocumentInterface, bAsync, PageID);
 
 		UObject* OwningObject = RegistryDocInterface.GetObject();
 		check(OwningObject);
@@ -490,7 +489,7 @@ namespace Metasound::Frontend
 		// Proxies are created synchronously to avoid creating proxies in async tasks. Proxies
 		// are created from UObjects which need to be protected from GC and non-GT access.
 		FProxyDataCache ProxyDataCache;
-		ProxyDataCache.CreateAndCacheProxies(Document);
+		ProxyDataCache.CreateAndCacheProxies(Document, PageID);
 
 		// Store update to newly registered node in history so nodes
 		// can be queried by transaction ID
@@ -504,12 +503,12 @@ namespace Metasound::Frontend
 		{
 			Tasks::FTask BuildAndRegisterTask = AsyncRegistrationPipe.Launch(
 				UE_SOURCE_LOCATION,
-				[RegistryKey, ClassInfo = MoveTemp(NodeClassInfo), RegistryDocInterface, ProxyDataCache = MoveTemp(ProxyDataCache)]() mutable
+				[RegistryKey, PageID, ClassInfo = MoveTemp(NodeClassInfo), RegistryDocInterface, ProxyDataCache = MoveTemp(ProxyDataCache)]() mutable
 				{
 					FRegistryContainerImpl& Registry = FRegistryContainerImpl::Get();
 					// Unregister the graph before re-registering
 					Registry.UnregisterGraphInternal(RegistryKey);
-					Registry.BuildAndRegisterGraphFromDocument(RegistryDocInterface, ProxyDataCache, MoveTemp(ClassInfo));
+					Registry.BuildAndRegisterGraphFromDocument(RegistryDocInterface->GetConstDocument(), ProxyDataCache, MoveTemp(ClassInfo), PageID);
 					Registry.RemoveRegistrationTask(RegistryKey, FNodeRegistryTransaction::ETransactionType::NodeRegistration);
 					Registry.RemoveDocumentReference(RegistryDocInterface);
 				}
@@ -528,7 +527,7 @@ namespace Metasound::Frontend
 			UnregisterGraphInternal(RegistryKey);
 
 			// Build and register graph synchronously
-			BuildAndRegisterGraphFromDocument(RegistryDocInterface, ProxyDataCache, MoveTemp(NodeClassInfo));
+			BuildAndRegisterGraphFromDocument(RegistryDocInterface->GetConstDocument(), ProxyDataCache, MoveTemp(NodeClassInfo), PageID);
 		}
 
 		return RegistryKey;
diff --git a/Engine/Plugins/Runtime/Metasound/Source/MetasoundFrontend/Private/MetasoundFrontendRegistryContainerImpl.h b/Engine/Plugins/Runtime/Metasound/Source/MetasoundFrontend/Private/MetasoundFrontendRegistryContainerImpl.h
--- a/Engine/Plugins/Runtime/Metasound/Source/MetasoundFrontend/Private/MetasoundFrontendRegistryContainerImpl.h
+++ b/Engine/Plugins/Runtime/Metasound/Source/MetasoundFrontend/Private/MetasoundFrontendRegistryContainerImpl.h
@@ -148,7 +148,7 @@ namespace Metasound
 				FTopLevelAssetPath AssetPath;
 			};
 
-			void BuildAndRegisterGraphFromDocument(TScriptInterface<IMetaSoundDocumentInterface> DocumentInterface, const FProxyDataCache& InProxyDataCache, FNodeClassInfo&& InNodeClassInfo);
+			void BuildAndRegisterGraphFromDocument(const FMetasoundFrontendDocument& InDocument, const FProxyDataCache& InProxyDataCache, FNodeClassInfo&& InNodeClassInfo, const FGuid& InPageID);
 
 			void AddRegistrationTask(const FGraphRegistryKey& InKey, FActiveRegistrationTaskInfo&& TaskInfo);
 			void RemoveRegistrationTask(const FGraphRegistryKey& InKey, FNodeRegistryTransaction::ETransactionType TransactionType);
diff --git a/Engine/Plugins/Runtime/Metasound/Source/MetasoundFrontend/Public/MetasoundAssetBase.h b/Engine/Plugins/Runtime/Metasound/Source/MetasoundFrontend/Public/MetasoundAssetBase.h
--- a/Engine/Plugins/Runtime/Metasound/Source/MetasoundFrontend/Public/MetasoundAssetBase.h
+++ b/Engine/Plugins/Runtime/Metasound/Source/MetasoundFrontend/Public/MetasoundAssetBase.h
@@ -154,7 +154,7 @@ public:
 	UE_DEPRECATED(5.5, "Direct mutation of the document is no longer supported via AssetBase.")
 	void SetDocument(FMetasoundFrontendDocument InDocument, bool bMarkDirty = true);
 
-	const FMetasoundFrontendDocument& GetConstDocumentChecked() const;
+	virtual const FMetasoundFrontendDocument& GetConstDocumentChecked() const;
 
 	// Soft deprecated.  Document layer should not be directly mutated via asset base in anticipation
 	// of moving all mutable document calls to the Frontend/Subsystem Document Builder API.
diff --git a/Engine/Plugins/Runtime/Metasound/Source/MetasoundFrontend/Public/MetasoundDocumentInterface.h b/Engine/Plugins/Runtime/Metasound/Source/MetasoundFrontend/Public/MetasoundDocumentInterface.h
--- a/Engine/Plugins/Runtime/Metasound/Source/MetasoundFrontend/Public/MetasoundDocumentInterface.h
+++ b/Engine/Plugins/Runtime/Metasound/Source/MetasoundFrontend/Public/MetasoundDocumentInterface.h
@@ -2,8 +2,11 @@
 #pragma once
 
 #include "UObject/Interface.h"
+#include "MetasoundBuilderInterface.h"
 #include "MetasoundFrontendDocument.h"
 #include "MetasoundFrontendDocumentModifyDelegates.h"
+#include "MetasoundOperatorInterface.h"
+#include "MetasoundOperatorSettings.h"
 #include "Templates/Function.h"
 #include "UObject/TopLevelAssetPath.h"
 
@@ -103,6 +106,9 @@ namespace Metasound::Frontend
 		// false if not found.
 		virtual bool ReloadBuilder(const FMetasoundFrontendClassName& InClassName) const = 0;
 
+		// Given the provided document and its respective pages, returns the PageID to be used for runtime IGraph and proxy generation.
+		virtual FGuid ResolveTargetPageID(const FMetasoundFrontendDocument& Document, const FTopLevelAssetPath& AssetPath) const = 0;
+
 		static IDocumentBuilderRegistry* Get();
 		static IDocumentBuilderRegistry& GetChecked();
 		static void Deinitialize();
diff --git a/Engine/Plugins/Runtime/Metasound/Source/MetasoundFrontend/Public/MetasoundFrontendController.h b/Engine/Plugins/Runtime/Metasound/Source/MetasoundFrontend/Public/MetasoundFrontendController.h
--- a/Engine/Plugins/Runtime/Metasound/Source/MetasoundFrontend/Public/MetasoundFrontendController.h
+++ b/Engine/Plugins/Runtime/Metasound/Source/MetasoundFrontend/Public/MetasoundFrontendController.h
@@ -987,15 +987,8 @@ namespace Metasound
 			 */
 			virtual FNodeHandle CreateEmptySubgraph(const FMetasoundFrontendClassMetadata& InInfo) = 0;
 
-			/** Creates a runtime operator for the given graph. Does not support input value manipulation via transmission.
-			 *
-			 * @param InSettings - Settings to use when creating operators.
-			 * @param InEnvironment - Environment variables available during creation.
-			 * @param OutResults - Results pertaining to operator build process & resulting IOperator instance.
-			 *
-			 * @return On success, a valid pointer to a Metasound operator. An invalid pointer on failure.
-			 */
-			virtual TUniquePtr<IOperator> BuildOperator(const FOperatorSettings& InSettings, const FMetasoundEnvironment& InEnvironment, FBuildResults& OutResults) const = 0;
+			UE_DEPRECATED(5.5, "Building graph operator is implemented internally and now supports pages")
+			virtual TUniquePtr<IOperator> BuildOperator(const FOperatorSettings& InSettings, const FMetasoundEnvironment& InEnvironment, FBuildResults& OutResults) const { return { }; }
 
 			/** Returns a handle to the document owning this graph. */
 			virtual FDocumentHandle GetOwningDocument() = 0;
diff --git a/Engine/Plugins/Runtime/Metasound/Source/MetasoundFrontend/Public/MetasoundFrontendDocument.h b/Engine/Plugins/Runtime/Metasound/Source/MetasoundFrontend/Public/MetasoundFrontendDocument.h
--- a/Engine/Plugins/Runtime/Metasound/Source/MetasoundFrontend/Public/MetasoundFrontendDocument.h
+++ b/Engine/Plugins/Runtime/Metasound/Source/MetasoundFrontend/Public/MetasoundFrontendDocument.h
@@ -1814,21 +1814,27 @@ METASOUNDFRONTEND_API const TCHAR* LexToString(EMetasoundFrontendVertexAccessTyp
 
 namespace Metasound::Frontend
 {
-		METASOUNDFRONTEND_API bool StringToClassType(const FString& InString, EMetasoundFrontendClassType& OutClassType);
+	METASOUNDFRONTEND_API bool StringToClassType(const FString& InString, EMetasoundFrontendClassType& OutClassType);
 
-		/** Signature of function called for each found literal. */
-		using FForEachLiteralFunctionRef = TFunctionRef<void(const FName& InDataTypeName, const FMetasoundFrontendLiteral&)>; 
+	/** Signature of function called for each found literal. */
+	using FForEachLiteralFunctionRef = TFunctionRef<void(const FName& InDataTypeName, const FMetasoundFrontendLiteral&)>;
 
-		/** Execute the provided function for each literal on a FMetasoundFrontendDocument.*/
-		METASOUNDFRONTEND_API void ForEachLiteral(const FMetasoundFrontendDocument& InDoc, FForEachLiteralFunctionRef OnLiteral);
+	/** Execute the provided function for each literal on a FMetasoundFrontendDocument.*/
+	METASOUNDFRONTEND_API void ForEachLiteral(const FMetasoundFrontendDocument& InDoc, FForEachLiteralFunctionRef OnLiteral);
 
-		/** Execute the provided function for each literal on a FMetasoundFrontendGraphClass.*/
-		METASOUNDFRONTEND_API void ForEachLiteral(const FMetasoundFrontendGraphClass& InGraphClass, FForEachLiteralFunctionRef OnLiteral);
+	/** Execute the provided function for each literal on a FMetasoundFrontendDocument filtered by the given PageID.*/
+	METASOUNDFRONTEND_API void ForEachLiteral(const FMetasoundFrontendDocument& InDoc, FForEachLiteralFunctionRef OnLiteral, const FGuid& InPageID);
 
-		/** Execute the provided function for each literal on a FMetasoundFrontendClass.*/
-		METASOUNDFRONTEND_API void ForEachLiteral(const FMetasoundFrontendClass& InClass, FForEachLiteralFunctionRef OnLiteral);
+	/** Execute the provided function for each literal on a FMetasoundFrontendGraphClass.*/
+	METASOUNDFRONTEND_API void ForEachLiteral(const FMetasoundFrontendGraphClass& InGraphClass, FForEachLiteralFunctionRef OnLiteral);
 
-		/** Execute the provided function for each literal on a FMetasoundFrontendNode.*/
-		METASOUNDFRONTEND_API void ForEachLiteral(const FMetasoundFrontendNode& InNode, FForEachLiteralFunctionRef OnLiteral);
+	/** Execute the provided function for each literal on a FMetasoundFrontendGraphClass filtered by the given PageID.*/
+	METASOUNDFRONTEND_API void ForEachLiteral(const FMetasoundFrontendGraphClass& InGraphClass, FForEachLiteralFunctionRef OnLiteral, const FGuid& InPageID);
+
+	/** Execute the provided function for each literal on a FMetasoundFrontendClass.*/
+	METASOUNDFRONTEND_API void ForEachLiteral(const FMetasoundFrontendClass& InClass, FForEachLiteralFunctionRef OnLiteral);
+
+	/** Execute the provided function for each literal on a FMetasoundFrontendNode.*/
+	METASOUNDFRONTEND_API void ForEachLiteral(const FMetasoundFrontendNode& InNode, FForEachLiteralFunctionRef OnLiteral);
 } // namespace Metasound::Frontend
 
diff --git a/Engine/Plugins/Runtime/Metasound/Source/MetasoundFrontend/Public/MetasoundFrontendGraph.h b/Engine/Plugins/Runtime/Metasound/Source/MetasoundFrontend/Public/MetasoundFrontendGraph.h
--- a/Engine/Plugins/Runtime/Metasound/Source/MetasoundFrontend/Public/MetasoundFrontendGraph.h
+++ b/Engine/Plugins/Runtime/Metasound/Source/MetasoundFrontend/Public/MetasoundFrontendGraph.h
@@ -104,16 +104,23 @@ namespace Metasound
 		static TUniquePtr<FFrontendGraph> CreateGraph(const FMetasoundFrontendGraphClass& InGraphClass, const TArray<FMetasoundFrontendGraphClass>& InSubgraphs, const TArray<FMetasoundFrontendClass>& InDependencies, const TSet<FName>& TransmittableInputNames, const FString& InDebugAssetName);
 
 		/* Create a FFrontendGraph from a FMetasoundFrontendDocument.*/
-		static TUniquePtr<FFrontendGraph> CreateGraph(const FMetasoundFrontendDocument& InDocument, const FString& InDebugAssetName);
+		static TUniquePtr<FFrontendGraph> CreateGraph(const FMetasoundFrontendDocument& InDocument, const FString& InDebugAssetName, const FGuid* InPageId = nullptr);
 
 		/* Create a FFrontendGraph from a FMetasoundFrontendDocument retrieving proxies from a FProxyDataCache.*/
-		static TUniquePtr<FFrontendGraph> CreateGraph(const FMetasoundFrontendDocument& InDocument, const Frontend::FProxyDataCache& InProxies, const FString& InDebugAssetName, const FGuid InGraphId = Frontend::CreateLocallyUniqueId());
+		static TUniquePtr<FFrontendGraph> CreateGraph(const FMetasoundFrontendDocument& InDocument, const Frontend::FProxyDataCache& InProxies, const FString& InDebugAssetName, const FGuid InGraphId = Frontend::CreateLocallyUniqueId(), const FGuid* InPageId = nullptr);
 
 		/* Create a FFrontendGraph from a FMetasoundFrontendDocument and subobjects.*/
-		static TUniquePtr<FFrontendGraph> CreateGraph(const FMetasoundFrontendGraphClass& InGraph, const TArray<FMetasoundFrontendGraphClass>& InSubgraphs, const TArray<FMetasoundFrontendClass>& InDependencies, const FString& InDebugAssetName);
+		static TUniquePtr<FFrontendGraph> CreateGraph(const FMetasoundFrontendGraphClass& InGraph, const TArray<FMetasoundFrontendGraphClass>& InSubgraphs, const TArray<FMetasoundFrontendClass>& InDependencies, const FString& InDebugAssetName, const FGuid* InPageId = nullptr);
 
 		/* Create a FFrontendGraph from a FMetasoundFrontendDocument and subobjects, retrieving proxies from a FProxyDataCache.*/
-		static TUniquePtr<FFrontendGraph> CreateGraph(const FMetasoundFrontendGraphClass& InGraph, const TArray<FMetasoundFrontendGraphClass>& InSubgraphs, const TArray<FMetasoundFrontendClass>& InDependencies, const Frontend::FProxyDataCache& InProxyDataCache, const FString& InDebugAssetName, const FGuid InGraphId = Frontend::CreateLocallyUniqueId());
+		static TUniquePtr<FFrontendGraph> CreateGraph(
+			const FMetasoundFrontendGraphClass& InGraph,
+			const TArray<FMetasoundFrontendGraphClass>& InSubgraphs,
+			const TArray<FMetasoundFrontendClass>& InDependencies,
+			const Frontend::FProxyDataCache& InProxyDataCache,
+			const FString& InDebugAssetName,
+			const FGuid InGraphId = Frontend::CreateLocallyUniqueId(),
+			const FGuid* InPageId = nullptr);
 
 	private:
 		struct FDefaultLiteralData
@@ -148,6 +155,7 @@ namespace Metasound
 		{
 			TUniquePtr<FFrontendGraph> Graph;
 			const FMetasoundFrontendGraphClass& GraphClass;
+			const FMetasoundFrontendGraph& PagedGraph;
 			FBuildContext& BuildContext;
 			FDefaultInputByIDMap DefaultInputs;
 		};
@@ -157,7 +165,7 @@ namespace Metasound
 
 		static bool SortSubgraphDependencies(TArray<const FMetasoundFrontendGraphClass*>& Subgraphs);
 
-		static TUniquePtr<FFrontendGraph> CreateGraph(FBuildContext& InContext, const FMetasoundFrontendGraphClass& InSubgraph);
+		static TUniquePtr<FFrontendGraph> CreateGraph(FBuildContext& InContext, const FMetasoundFrontendGraphClass& InSubgraph, const FGuid* InPageId = nullptr);
 
 		static const FMetasoundFrontendClassInput* FindClassInputForInputNode(const FMetasoundFrontendGraphClass& InOwningGraph, const FMetasoundFrontendNode& InInputNode, int32& OutClassInputIndex);
 		static const FMetasoundFrontendClassOutput* FindClassOutputForOutputNode(const FMetasoundFrontendGraphClass& InOwningGraph, const FMetasoundFrontendNode& InOutputNode, int32& OutClassOutputIndex);
