UE-218541 - Conversion nodes are not removed when using the undo function within a MetaSound

diff --git a/Engine/Plugins/Runtime/Metasound/Source/MetasoundEditor/Private/MetasoundEditorGraphSchema.cpp b/Engine/Plugins/Runtime/Metasound/Source/MetasoundEditor/Private/MetasoundEditorGraphSchema.cpp
--- a/Engine/Plugins/Runtime/Metasound/Source/MetasoundEditor/Private/MetasoundEditorGraphSchema.cpp
+++ b/Engine/Plugins/Runtime/Metasound/Source/MetasoundEditor/Private/MetasoundEditorGraphSchema.cpp
@@ -1567,6 +1567,8 @@ bool UMetasoundEditorGraphSchema::TryCreateConnection(UEdGraphPin* PinA, UEdGrap
 		return false;
 	}
 
+	const FScopedTransaction Transaction(LOCTEXT("MetasoundConnect", "Connect Pins"));
+
 	FConstInputHandle InputHandle = FGraphBuilder::GetConstInputHandleFromPin(InputPin);
 	FConstOutputHandle OutputHandle = FGraphBuilder::GetConstOutputHandleFromPin(OutputPin);
 	FConnectability Connectability = InputHandle->CanConnectTo(*OutputHandle);
@@ -1575,6 +1577,8 @@ bool UMetasoundEditorGraphSchema::TryCreateConnection(UEdGraphPin* PinA, UEdGrap
 		UMetasoundEditorGraph* MetaSoundGraph = CastChecked<UMetasoundEditorGraph>(InputPin->GetOwningNode()->GetGraph());
 		UObject& ParentMetaSound = MetaSoundGraph->GetMetasoundChecked();
 
+		ParentMetaSound.Modify();
+
 		if (Connectability.PossibleConverterNodeClasses.Num() == 0)
 		{
 			return false;
