UE-216286 - [MetaSounds] - When autoconverting nodes, the node placement is not centred

diff --git a/Engine/Plugins/Runtime/Metasound/Source/MetasoundEditor/Private/MetasoundEditorGraphSchema.cpp b/Engine/Plugins/Runtime/Metasound/Source/MetasoundEditor/Private/MetasoundEditorGraphSchema.cpp
--- a/Engine/Plugins/Runtime/Metasound/Source/MetasoundEditor/Private/MetasoundEditorGraphSchema.cpp
+++ b/Engine/Plugins/Runtime/Metasound/Source/MetasoundEditor/Private/MetasoundEditorGraphSchema.cpp
@@ -1595,8 +1595,9 @@ bool UMetasoundEditorGraphSchema::TryCreateConnection(UEdGraphPin* PinA, UEdGrap
 			UEdGraphNode* OutputNode = OutputPin->GetOwningNode();
 			
 			FVector2D Location = FVector2D();
+			Location += FVector2D(InputNode->NodePosX, InputNode->NodePosY);
 			Location += FVector2D(OutputNode->NodePosX, OutputNode->NodePosY);
-			Location.Y += 40.f;
+			Location *= 0.5f;
 
 			NewGraphNode->Modify();
 			NewGraphNode->UpdateFrontendNodeLocation(Location);
