Moving page mutable calls to builds loading editor-only data (pages initially will be a non-runtime mutable feature on serialized MetaSound assets only, so locking down accordingly)

#rnx

diff --git a/Engine/Plugins/Runtime/Metasound/Source/MetasoundEngine/Public/MetasoundBuilderBase.cpp b/Engine/Plugins/Runtime/Metasound/Source/MetasoundEngine/Public/MetasoundBuilderBase.cpp
--- a/Engine/Plugins/Runtime/Metasound/Source/MetasoundEngine/Public/MetasoundBuilderBase.cpp
+++ b/Engine/Plugins/Runtime/Metasound/Source/MetasoundEngine/Public/MetasoundBuilderBase.cpp
@@ -940,11 +940,13 @@ void UMetaSoundBuilderBase::ReloadCache(bool bPrimeCache)
 	Reload({ }, bPrimeCache);
 }
 
+#if WITH_EDITORONLY_DATA
 void UMetaSoundBuilderBase::RemoveAllGraphPages()
 {
 	Builder.RemoveAllGraphPages();
 	Builder.RemoveUnusedDependencies();
 }
+#endif // WITH_EDITORONLY_DATA
 
 #if WITH_EDITOR
 bool UMetaSoundBuilderBase::RemoveGraphComment(const FGuid& InCommentID)
diff --git a/Engine/Plugins/Runtime/Metasound/Source/MetasoundEngine/Public/MetasoundBuilderBase.h b/Engine/Plugins/Runtime/Metasound/Source/MetasoundEngine/Public/MetasoundBuilderBase.h
--- a/Engine/Plugins/Runtime/Metasound/Source/MetasoundEngine/Public/MetasoundBuilderBase.h
+++ b/Engine/Plugins/Runtime/Metasound/Source/MetasoundEngine/Public/MetasoundBuilderBase.h
@@ -338,8 +338,11 @@ public:
 	UFUNCTION(BlueprintCallable, Category = "Audio|MetaSound|Builder", meta = (ExpandEnumAsExecs = "OutResult"))
 	void ConvertToPreset(const TScriptInterface<IMetaSoundDocumentInterface>& ReferencedNodeClass, EMetaSoundBuilderResult& OutResult);
 
+
+#if WITH_EDITORONLY_DATA
 	// Removes all graph pages, leaving just the default.
 	void RemoveAllGraphPages();
+#endif // WITH_EDITORONLY_DATA
 
 	// Removes graph input if it exists; sets result to succeeded if it was removed and failed if it was not.
 	UFUNCTION(BlueprintCallable, Category = "Audio|MetaSound|Builder", meta = (ExpandEnumAsExecs = "OutResult"))
diff --git a/Engine/Plugins/Runtime/Metasound/Source/MetasoundFrontend/Private/MetasoundFrontendDocument.cpp b/Engine/Plugins/Runtime/Metasound/Source/MetasoundFrontend/Private/MetasoundFrontendDocument.cpp
--- a/Engine/Plugins/Runtime/Metasound/Source/MetasoundFrontend/Private/MetasoundFrontendDocument.cpp
+++ b/Engine/Plugins/Runtime/Metasound/Source/MetasoundFrontend/Private/MetasoundFrontendDocument.cpp
@@ -885,6 +885,7 @@ FMetasoundFrontendGraphClass::FMetasoundFrontendGraphClass()
 	Metadata.SetType(EMetasoundFrontendClassType::Graph);
 }
 
+#if WITH_EDITORONLY_DATA
 const FMetasoundFrontendGraph& FMetasoundFrontendGraphClass::AddGraphPage(const FGuid& InPageID, bool bDuplicateLastGraph, bool bSetAsBuildGraph)
 {
 	checkf(InPageID != Metasound::Frontend::DefaultGraphPageID, TEXT("Cannot add graph with DefaultPageID"))
@@ -906,6 +907,7 @@ PRAGMA_ENABLE_DEPRECATION_WARNINGS
 	NewGraph->PageID = InPageID;
 	return *NewGraph;
 }
+#endif // WITH_EDITORONLY_DATA
 
 bool FMetasoundFrontendGraphClass::ContainsGraphPage(const FGuid& InPageID) const
 {
@@ -920,6 +922,7 @@ PRAGMA_ENABLE_DEPRECATION_WARNINGS
 	return PagedGraphs.ContainsByPredicate(MatchesPageID);
 }
 
+#if WITH_EDITORONLY_DATA
 void FMetasoundFrontendGraphClass::RemoveAllGraphPages()
 {
 	PagedGraphs.Empty();
@@ -956,6 +959,7 @@ bool FMetasoundFrontendGraphClass::RemoveGraphPage(const FGuid& InPageID, FGuid*
 
 	return false;
 }
+#endif // WITH_EDITORONLY_DATA
 
 FMetasoundFrontendGraph* FMetasoundFrontendGraphClass::FindGraph(const FGuid& InPageID)
 {
diff --git a/Engine/Plugins/Runtime/Metasound/Source/MetasoundFrontend/Private/MetasoundFrontendDocumentBuilder.cpp b/Engine/Plugins/Runtime/Metasound/Source/MetasoundFrontend/Private/MetasoundFrontendDocumentBuilder.cpp
--- a/Engine/Plugins/Runtime/Metasound/Source/MetasoundFrontend/Private/MetasoundFrontendDocumentBuilder.cpp
+++ b/Engine/Plugins/Runtime/Metasound/Source/MetasoundFrontend/Private/MetasoundFrontendDocumentBuilder.cpp
@@ -1304,6 +1304,7 @@ FMetasoundFrontendNode* FMetaSoundFrontendDocumentBuilder::AddNodeInternal(const
 	return nullptr;
 }
 
+#if WITH_EDITORONLY_DATA
 const FMetasoundFrontendGraph& FMetaSoundFrontendDocumentBuilder::AddGraphPage(const FGuid& InPageID, bool bDuplicateLastGraph, bool bSetAsBuildGraph)
 {
 	using namespace Metasound::Frontend;
@@ -1316,6 +1317,7 @@ const FMetasoundFrontendGraph& FMetaSoundFrontendDocumentBuilder::AddGraphPage(c
 	}
 	return ToReturn;
 }
+#endif // WITH_EDITORONLY_DATA
 
 bool FMetaSoundFrontendDocumentBuilder::CanAddEdge(const FMetasoundFrontendEdge& InEdge, const FGuid* InPageID) const
 {
@@ -1349,14 +1351,23 @@ void FMetaSoundFrontendDocumentBuilder::ClearDocument(TSharedRef<Metasound::Fron
 	GraphClass.PresetOptions.InputsInheritingDefault.Reset();
 	GraphClass.PresetOptions.bIsPreset = false;
 
+	// Removing graph pages is not necessary when editor only data is not available as graph mutation
+	// is only supported in builds with editor data loaded. Otherwise, anything calling ClearDocument
+	// should only be a transient, non serialized asset graph which does not support page mutation.
+#if WITH_EDITORONLY_DATA
 	RemoveAllGraphPages();
+#else // !WITH_EDITORONLY_DATA
+	UObject& DocObject = CastDocumentObjectChecked<UObject>();
+	checkf(!DocObject.IsAsset(), TEXT("Cannot call clear document on asset '%s': builder API does not support document mutation on serialized objects without editor data loaded"), *GetDebugName());
+#endif // !WITH_EDITORONLY_DATA
 
-PRAGMA_DISABLE_DEPRECATION_WARNINGS
-	FMetasoundFrontendGraph& Graph = GraphClass.GetDefaultGraph();
-	Graph.Variables.Empty();
-	Graph.Nodes.Empty();
-	Graph.Edges.Empty();
-PRAGMA_ENABLE_DEPRECATION_WARNINGS
+	// Empty implementation for single remaining graph.
+	GraphClass.IterateGraphPages([](FMetasoundFrontendGraph& Graph)
+	{
+		Graph.Variables.Empty();
+		Graph.Nodes.Empty();
+		Graph.Edges.Empty();
+	});
 
 	Doc.Interfaces.Empty();
 	Doc.Dependencies.Empty();
@@ -2651,6 +2662,7 @@ void FMetaSoundFrontendDocumentBuilder::FinishBuilding()
 	DocumentCache.Reset();
 }
 
+#if WITH_EDITORONLY_DATA
 void FMetaSoundFrontendDocumentBuilder::RemoveAllGraphPages()
 {
 	using namespace Metasound;
@@ -2667,6 +2679,7 @@ void FMetaSoundFrontendDocumentBuilder::RemoveAllGraphPages()
 	RootGraph.RemoveAllGraphPages();
 	SetBuildPageID(Frontend::DefaultGraphPageID);
 }
+#endif // WITH_EDITORONLY_DATA
 
 bool FMetaSoundFrontendDocumentBuilder::RemoveDependency(const FGuid& InClassID)
 {
@@ -3128,6 +3141,7 @@ bool FMetaSoundFrontendDocumentBuilder::RemoveGraphOutput(FName InOutputName)
 	return false;
 }
 
+#if WITH_EDITORONLY_DATA
 bool FMetaSoundFrontendDocumentBuilder::RemoveGraphPage(const FGuid& InPageID)
 {
 	using namespace Metasound::Frontend;
@@ -3152,6 +3166,7 @@ bool FMetaSoundFrontendDocumentBuilder::RemoveGraphPage(const FGuid& InPageID)
 
 	return bPageRemoved;
 }
+#endif // WITH_EDITORONLY_DATA
 
 bool FMetaSoundFrontendDocumentBuilder::RemoveInterface(FName InterfaceName)
 {
@@ -3309,6 +3324,7 @@ void FMetaSoundFrontendDocumentBuilder::SetAuthor(const FString& InAuthor)
 }
 #endif // WITH_EDITOR
 
+#if WITH_EDITORONLY_DATA
 bool FMetaSoundFrontendDocumentBuilder::SetBuildPageID(const FGuid& InBuildPageID)
 {
 	using namespace Metasound::Frontend;
@@ -3329,6 +3345,7 @@ bool FMetaSoundFrontendDocumentBuilder::SetBuildPageID(const FGuid& InBuildPageI
 
 	return false;
 }
+#endif // WITH_EDITORONLY_DATA
 
 bool FMetaSoundFrontendDocumentBuilder::SetGraphInputAccessType(FName InputName, EMetasoundFrontendVertexAccessType AccessType)
 {
diff --git a/Engine/Plugins/Runtime/Metasound/Source/MetasoundFrontend/Private/MetasoundFrontendDocumentCache.cpp b/Engine/Plugins/Runtime/Metasound/Source/MetasoundFrontend/Private/MetasoundFrontendDocumentCache.cpp
--- a/Engine/Plugins/Runtime/Metasound/Source/MetasoundFrontend/Private/MetasoundFrontendDocumentCache.cpp
+++ b/Engine/Plugins/Runtime/Metasound/Source/MetasoundFrontend/Private/MetasoundFrontendDocumentCache.cpp
@@ -335,6 +335,7 @@ namespace Metasound::Frontend
 		++TransactionCount;
 	}
 
+#if WITH_EDITORONLY_DATA
 	void FDocumentCache::SetBuildPageID(const FGuid& InPageID)
 	{
 		if (InPageID != BuildPageID)
@@ -343,6 +344,7 @@ namespace Metasound::Frontend
 			BuildPageID = InPageID;
 		}
 	}
+#endif // WITH_EDITORONLY_DATA
 
 	FDocumentGraphInterfaceCache::FDocumentGraphInterfaceCache(TSharedRef<const FDocumentCache> ParentCache)
 		: Parent(ParentCache)
diff --git a/Engine/Plugins/Runtime/Metasound/Source/MetasoundFrontend/Private/MetasoundFrontendDocumentCache.h b/Engine/Plugins/Runtime/Metasound/Source/MetasoundFrontend/Private/MetasoundFrontendDocumentCache.h
--- a/Engine/Plugins/Runtime/Metasound/Source/MetasoundFrontend/Private/MetasoundFrontendDocumentCache.h
+++ b/Engine/Plugins/Runtime/Metasound/Source/MetasoundFrontend/Private/MetasoundFrontendDocumentCache.h
@@ -180,7 +180,9 @@ namespace Metasound::Frontend
 		virtual const IDocumentGraphNodeCache& GetNodeCache(const FGuid& InPageID) const override;
 		virtual const IDocumentGraphInterfaceCache& GetInterfaceCache() const override;
 
+#if WITH_EDITORONLY_DATA
 		virtual void SetBuildPageID(const FGuid& InPageID) override;
+#endif // WITH_EDITORONLY_DATA
 
 		int32 GetTransactionCount() const;
 		int32 GetTransactionCount(const FGuid& InPageID) const;
diff --git a/Engine/Plugins/Runtime/Metasound/Source/MetasoundFrontend/Public/MetasoundFrontendDocument.h b/Engine/Plugins/Runtime/Metasound/Source/MetasoundFrontend/Public/MetasoundFrontendDocument.h
--- a/Engine/Plugins/Runtime/Metasound/Source/MetasoundFrontend/Public/MetasoundFrontendDocument.h
+++ b/Engine/Plugins/Runtime/Metasound/Source/MetasoundFrontend/Public/MetasoundFrontendDocument.h
@@ -1676,6 +1676,7 @@ public:
 	FMetasoundFrontendGraphClassPresetOptions PresetOptions;
 
 public:
+#if WITH_EDITORONLY_DATA
 	const FMetasoundFrontendGraph& AddGraphPage(const FGuid& InPageID, bool bDuplicateLastGraph = true, bool bSetAsBuildGraph = true);
 	void RemoveAllGraphPages();
 
@@ -1683,6 +1684,7 @@ public:
 	// If provided an "AdjacentPageID," sets the value at the given pointer to a page ID adjacent to
 	// the removed page. If last page was removed, returns the default graph ID.
 	bool RemoveGraphPage(const FGuid& InPageID, FGuid* OutAdjacentPageID = nullptr);
+#endif // WITH_EDITORONLY_DATA
 
 	bool ContainsGraphPage(const FGuid& InPageID) const;
 	void IterateGraphPages(TFunctionRef<void(FMetasoundFrontendGraph&)> IterFunc);
diff --git a/Engine/Plugins/Runtime/Metasound/Source/MetasoundFrontend/Public/MetasoundFrontendDocumentBuilder.h b/Engine/Plugins/Runtime/Metasound/Source/MetasoundFrontend/Public/MetasoundFrontendDocumentBuilder.h
--- a/Engine/Plugins/Runtime/Metasound/Source/MetasoundFrontend/Public/MetasoundFrontendDocumentBuilder.h
+++ b/Engine/Plugins/Runtime/Metasound/Source/MetasoundFrontend/Public/MetasoundFrontendDocumentBuilder.h
@@ -159,8 +159,17 @@ public:
 	bool AddEdgesByNodeClassInterfaceBindings(const FGuid& InFromNodeID, const FGuid& InToNodeID, bool bReplaceExistingConnections = true, const FGuid* InPageID = nullptr);
 	bool AddEdgesFromMatchingInterfaceNodeOutputsToGraphOutputs(const FGuid& InNodeID, TArray<const FMetasoundFrontendEdge*>& OutEdgesCreated, bool bReplaceExistingConnections = true, const FGuid* InPageID = nullptr);
 	bool AddEdgesFromMatchingInterfaceNodeInputsToGraphInputs(const FGuid& InNodeID, TArray<const FMetasoundFrontendEdge*>& OutEdgesCreated, bool bReplaceExistingConnections = true, const FGuid* InPageID = nullptr);
+
+	// Adds Graph Input to document, which in turn adds a referencing input node to ALL pages.  If valid PageID is provided, returns associated page's node pointer.
+	// If none provided, returns node pointer to node for the builder's currently set build page ID (see 'GetBuildPageID').
 	const FMetasoundFrontendNode* AddGraphInput(const FMetasoundFrontendClassInput& InClassInput, const FGuid* InPageID = nullptr);
+
+	// Adds node to document to the page associated with the given PageID.  If no valid PageID is provided, adds and returns node pointer to node for the builder's
+	// currently set build page ID (see 'GetBuildPageID').
 	const FMetasoundFrontendNode* AddGraphNode(const FMetasoundFrontendGraphClass& InClass, FGuid InNodeID = FGuid::NewGuid(), const FGuid* InPageID = nullptr);
+
+	// Adds Graph Output to document, which in turn adds a referencing output node to ALL pages.  If valid PageID is provided, returns associated page's node pointer.
+	// If none provided, returns node pointer to node for the builder's currently set build page ID (see 'GetBuildPageID').
 	const FMetasoundFrontendNode* AddGraphOutput(const FMetasoundFrontendClassOutput& InClassOutput, const FGuid* InPageID = nullptr);
 
 	bool AddInterface(FName InterfaceName);
@@ -168,8 +177,10 @@ public:
 	const FMetasoundFrontendNode* AddNodeByClassName(const FMetasoundFrontendClassName& InClassName, int32 InMajorVersion = 1, FGuid InNodeID = FGuid::NewGuid(), const FGuid* InPageID = nullptr);
 	const FMetasoundFrontendNode* AddNodeByTemplate(const Metasound::Frontend::INodeTemplate& InTemplate, Metasound::Frontend::FNodeTemplateGenerateInterfaceParams Params, FGuid InNodeID = FGuid::NewGuid(), const FGuid* InPageID = nullptr);
 
+#if WITH_EDITORONLY_DATA
 	// Adds a graph page to the given builder's document
 	const FMetasoundFrontendGraph& AddGraphPage(const FGuid& InPageID, bool bDuplicateLastGraph = true, bool bSetAsBuildGraph = true);
+#endif // WITH_EDITORONLY_DATA
 
 	// Returns whether or not the given edge can be added, which requires that its input
 	// is not already connected and the edge is valid (see function 'IsValidEdge').
@@ -178,7 +189,7 @@ public:
 	// Clears document completely of all graph page data (nodes, edges, & member metadata), dependencies,
 	// interfaces, member metadata, preset state, etc. Leaves ClassMetadata intact. Reloads the builder state,
 	// so external delegates must be relinked if desired.
-	void ClearDocument(TSharedRef<Metasound::Frontend::FDocumentModifyDelegates> ModifyDelegates = { });
+	void ClearDocument(TSharedRef<Metasound::Frontend::FDocumentModifyDelegates> ModifyDelegates);
 
 	UE_DEPRECATED(5.5, "Use ClearDocument instead")
 	void ClearGraph() {  }
@@ -306,7 +317,9 @@ public:
 
 	bool ModifyInterfaces(Metasound::Frontend::FModifyInterfaceOptions&& InOptions);
 
+#if WITH_EDITORONLY_DATA
 	void RemoveAllGraphPages();
+#endif // WITH_EDITORONLY_DATA
 
 	UE_DEPRECATED(5.5,
 		"Cache invalidation may require new copy of delegates. In addition, re-priming is discouraged. "
@@ -330,7 +343,11 @@ public:
 
 	bool RemoveGraphInput(FName InInputName);
 	bool RemoveGraphOutput(FName InOutputName);
+
+#if WITH_EDITORONLY_DATA
 	bool RemoveGraphPage(const FGuid& InPageID);
+#endif // WITH_EDITORONLY_DATA
+
 	bool RemoveInterface(FName InName);
 	bool RemoveNamedEdges(const TSet<Metasound::Frontend::FNamedEdge>& InNamedEdgesToRemove, TArray<FMetasoundFrontendEdge>* OutRemovedEdges = nullptr, const FGuid* InPageID = nullptr);
 	bool RemoveNode(const FGuid& InNodeID, const FGuid* InPageID = nullptr);
@@ -352,12 +369,14 @@ public:
 
 #endif // WITH_EDITOR
 
+#if WITH_EDITORONLY_DATA
 	// Sets the builder's targeted paged graph ID to the given ID if it exists.
 	// Returns true if the builder is already targeting the given ID or if it successfully
 	// found a page implementation with the given ID and was able to switch to it, false if not.
 	// Swapping the targeted build graph ID clears the local cache, so swapping frequently can
 	// induce cash thrashing.
 	bool SetBuildPageID(const FGuid& InBuildPageID);
+#endif // WITH_EDITORONLY_DATA
 
 	// Sets the given graph input's access type. If connected to other nodes and access type is not compatible,
 	// associated edges/connections are removed.  Returns true if either DataType was successfully set to new
diff --git a/Engine/Plugins/Runtime/Metasound/Source/MetasoundFrontend/Public/MetasoundFrontendDocumentCacheInterface.h b/Engine/Plugins/Runtime/Metasound/Source/MetasoundFrontend/Public/MetasoundFrontendDocumentCacheInterface.h
--- a/Engine/Plugins/Runtime/Metasound/Source/MetasoundFrontend/Public/MetasoundFrontendDocumentCacheInterface.h
+++ b/Engine/Plugins/Runtime/Metasound/Source/MetasoundFrontend/Public/MetasoundFrontendDocumentCacheInterface.h
@@ -123,8 +123,10 @@ namespace Metasound::Frontend
 		// Returns InterfaceCache for the given document
 		virtual const IDocumentGraphInterfaceCache& GetInterfaceCache() const = 0;
 
+#if WITH_EDITORONLY_DATA
 		// Sets the default BuildPageID for the given document
 		virtual void SetBuildPageID(const FGuid& InPageID) = 0;
+#endif // WITH_EDITORONLY_DATA
 	};
 
 } // namespace Metasound::Frontend
