First Rough Pass at Pages Editor Audition Support
- Fix-up new transaction listener for updating page details on page transactions
- Misc clean-up

#rnx

diff --git a/Engine/Plugins/Runtime/Metasound/Source/MetasoundEditor/Private/MetasoundDetailCustomization.cpp b/Engine/Plugins/Runtime/Metasound/Source/MetasoundEditor/Private/MetasoundDetailCustomization.cpp
--- a/Engine/Plugins/Runtime/Metasound/Source/MetasoundEditor/Private/MetasoundDetailCustomization.cpp
+++ b/Engine/Plugins/Runtime/Metasound/Source/MetasoundEditor/Private/MetasoundDetailCustomization.cpp
@@ -18,6 +18,7 @@
 #include "MetasoundEditor.h"
 #include "MetasoundEditorGraphBuilder.h"
 #include "MetasoundEditorSettings.h"
+#include "MetasoundEditorSubsystem.h"
 #include "MetasoundFrontend.h"
 #include "MetasoundFrontendController.h"
 #include "MetasoundFrontendSearchEngine.h"
@@ -344,56 +345,6 @@
 	{
 	}
 
-	void FMetasoundPagesDetailCustomization::FPageListener::OnBuilderReloaded(Frontend::FDocumentModifyDelegates& OutDelegates)
-	{
-		OutDelegates.PageDelegates.OnPageAdded.AddSP(this, &FPageListener::OnPageAdded);
-		OutDelegates.PageDelegates.OnRemovingPage.AddSP(this, &FPageListener::OnRemovingPage);
-	}
-
-	void FMetasoundPagesDetailCustomization::FPageListener::OnPageAdded(const Frontend::FDocumentMutatePageArgs& Args)
-	{
-		if (TSharedPtr<FMetasoundPagesDetailCustomization> ParentPtr = Parent.Pin())
-		{
-			const UMetaSoundSettings* Settings = GetDefault<UMetaSoundSettings>();
-			check(Settings);
-			if (const FMetaSoundPageSettings* PageSettings = Settings->FindPageSettings(Args.PageID))
-			{
-				if (PageSettings->Name != ParentPtr->BuildPageName)
-				{
-					ParentPtr->BuildPageName = PageSettings->Name;
-					FGraphBuilder::RegisterGraphWithFrontend(ParentPtr->GetMetaSound());
-				}
-
-				ParentPtr->AddableItems.RemoveAll([&PageSettings](const TSharedPtr<FString>& Item) { return Item->Compare(PageSettings->Name.ToString()) == 0; });
-				ParentPtr->ImplementedNames.Add(PageSettings->Name);
-				ParentPtr->ComboBox->RefreshOptions();
-				ParentPtr->RebuildImplemented();
-			}
-		}
-	}
-
-	void FMetasoundPagesDetailCustomization::FPageListener::OnRemovingPage(const Frontend::FDocumentMutatePageArgs& Args)
-	{
-		if (TSharedPtr<FMetasoundPagesDetailCustomization> ParentPtr = Parent.Pin())
-		{
-			const UMetaSoundSettings* Settings = GetDefault<UMetaSoundSettings>();
-			check(Settings);
-			if (const FMetaSoundPageSettings* PageSettings = Settings->FindPageSettings(Args.PageID))
-			{
-				if (PageSettings->Name != ParentPtr->BuildPageName)
-				{
-					ParentPtr->BuildPageName = PageSettings->Name;
-					FGraphBuilder::RegisterGraphWithFrontend(ParentPtr->GetMetaSound());
-				}
-
-				ParentPtr->AddableItems.Add(MakeShared<FString>(PageSettings->Name.ToString()));
-				ParentPtr->ImplementedNames.Remove(PageSettings->Name);
-				ParentPtr->ComboBox->RefreshOptions();
-				ParentPtr->RebuildImplemented();
-			}
-		}
-	}
-
 	void FMetasoundPagesDetailCustomization::CustomizeDetails(IDetailLayoutBuilder& DetailLayout)
 	{
 		using namespace Engine;
@@ -408,15 +359,6 @@
 		{
 			return;
 		}
-		if (UMetasoundEditorViewBase* View = CastChecked<UMetasoundEditorViewBase>(Objects.Last()))
-		{
-			if (UObject* MetaSound = View->GetMetasound())
-			{
-				InitBuilder(*MetaSound);
-				PageListener = MakeShared<FPageListener>(StaticCastSharedRef<FMetasoundPagesDetailCustomization>(AsShared()));
-				Builder->AddTransactionListener(PageListener->AsShared());
-			}
-		}
 
 		if (UMetaSoundSettings* Settings = GetMutableDefault<UMetaSoundSettings>())
 		{
@@ -430,8 +372,6 @@
 			});
 		}
 
-		UpdateItemNames();
-
 		SAssignNew(ComboBox, SSearchableComboBox)
 			.OptionsSource(&AddableItems)
 			.OnGenerateWidget_Lambda([](TSharedPtr<FString> InItem)
@@ -506,23 +446,22 @@
 		IDetailCategoryBuilder& Category = DetailLayout.EditCategory(FName(ItemName), HeaderName);
 		Category.AddCustomRow(HeaderName) [ Utilities ];
 
-		const UMetaSoundSettings* Settings = GetDefault<UMetaSoundSettings>();
-		check(Settings);
-		if (const FMetaSoundPageSettings* PageSettings = Settings->FindPageSettings(Builder->GetBuilder().GetBuildPageID()))
-		{
-			BuildPageName = PageSettings->Name;
-		}
-		else
-		{
-			BuildPageName = Metasound::Frontend::DefaultGraphPageName;
-		}
-
 		EntryWidgets = SNew(SVerticalBox);
-		RebuildImplemented();
 		Category.AddCustomRow(LOCTEXT("ImplementedPagesLabel", "Implemented Pages"))
 		[
 			EntryWidgets->AsShared()
 		];
+
+		// Registration of page listener instance calls OnReload which in turn causes RefreshView, so no need to call directly
+		if (UMetasoundEditorViewBase* View = CastChecked<UMetasoundEditorViewBase>(Objects.Last()))
+		{
+			if (UObject* MetaSound = View->GetMetasound())
+			{
+				InitBuilder(*MetaSound);
+				PageListener = MakeShared<FPageListener>(StaticCastSharedRef<FMetasoundPagesDetailCustomization>(AsShared()));
+				Builder->AddTransactionListener(PageListener->AsShared());
+			}
+		}
 	}
 
 	UObject& FMetasoundPagesDetailCustomization::GetMetaSound() const
@@ -540,19 +479,20 @@
 
 			TSharedRef<SWidget> SelectButtonWidget = PropertyCustomizationHelpers::MakeUseSelectedButton(FSimpleDelegate::CreateLambda([this, InName]()
 			{
-				using namespace Frontend;
-				const UMetaSoundSettings* Settings = GetDefault<UMetaSoundSettings>();
-				check(Settings);
-				if (const FMetaSoundPageSettings* PageSettings = Settings->FindPageSettings(InName))
+				EMetaSoundBuilderResult Result = EMetaSoundBuilderResult::Failed;
+				constexpr bool bFocusEditor = false; // Already focused by user action
+				UMetaSoundEditorSubsystem::GetConstChecked().SetFocusedPage(Builder.Get(), InName, bFocusEditor, Result);
+				if (Result == EMetaSoundBuilderResult::Succeeded)
 				{
-					if (Builder->GetBuilder().SetBuildPageID(PageSettings->UniqueId))
-					{
-						BuildPageName = InName;
-						FGraphBuilder::RegisterGraphWithFrontend(GetMetaSound());
-					}
+					BuildPageName = InName;
 				}
 			}),
-			LOCTEXT("SetPageTooltip", "Sets the actively displayed graph page of the MetaSound."),
+			TAttribute<FText>::Create([this, InName]()
+			{
+				return BuildPageName == InName
+					? LOCTEXT("FocusedPageTooltip", "Currently focused page.")
+					: LOCTEXT("SetFocusedPageTooltip", "Sets the actively focused graph page of the MetaSound.");
+			}),
 			TAttribute<bool>::Create([this, InName]()
 			{
 				return BuildPageName != InName;
@@ -587,10 +527,6 @@
 					{
 						UpdateItemNames();
 						ComboBox->RefreshOptions();
-						if (InName == BuildPageName)
-						{
-							FGraphBuilder::RegisterGraphWithFrontend(MetaSound);
-						}
 					}
 				}), LOCTEXT("RemovePageTooltip2", "Removes the associated page from the MetaSound."));
 				EntryWidget->AddSlot()
@@ -627,6 +563,34 @@
 		}
 	}
 
+	void FMetasoundPagesDetailCustomization::RefreshView()
+	{
+		if (Builder.IsValid())
+		{
+			const UMetaSoundSettings* Settings = GetDefault<UMetaSoundSettings>();
+			check(Settings);
+			FMetaSoundFrontendDocumentBuilder& DocBuilder = Builder->GetBuilder();
+			const FGuid& PageID = DocBuilder.GetBuildPageID();
+			if (const FMetaSoundPageSettings* PageSettings = Settings->FindPageSettings(PageID))
+			{
+				BuildPageName = PageSettings->Name;
+			}
+			else
+			{
+				BuildPageName = Frontend::DefaultGraphPageName;
+				DocBuilder.SetBuildPageID(Frontend::DefaultGraphPageID);
+			}
+		}
+		else
+		{
+			BuildPageName = Frontend::DefaultGraphPageName;
+		}
+
+		UpdateItemNames();
+		ComboBox->RefreshOptions();
+		RebuildImplemented();
+	}
+
 	void FMetasoundPagesDetailCustomization::UpdateItemNames()
 	{
 		using namespace Frontend;
@@ -663,6 +627,77 @@
 		Algo::Transform(ImplementedGuids, ImplementedNames, GetPageName);
 	}
 
+	void FMetasoundPagesDetailCustomization::FPageListener::OnBuilderReloaded(Frontend::FDocumentModifyDelegates& OutDelegates)
+	{
+		if (TSharedPtr<FMetasoundPagesDetailCustomization> ParentPtr = Parent.Pin())
+		{
+			ParentPtr->RefreshView();
+		}
+
+		OutDelegates.PageDelegates.OnPageAdded.AddSP(this, &FPageListener::OnPageAdded);
+		OutDelegates.PageDelegates.OnPageSet.AddSP(this, &FPageListener::OnPageSet);
+		OutDelegates.PageDelegates.OnRemovingPage.AddSP(this, &FPageListener::OnRemovingPage);
+	}
+
+	void FMetasoundPagesDetailCustomization::FPageListener::OnPageAdded(const Frontend::FDocumentMutatePageArgs& Args)
+	{
+		if (TSharedPtr<FMetasoundPagesDetailCustomization> ParentPtr = Parent.Pin())
+		{
+			const UMetaSoundSettings* Settings = GetDefault<UMetaSoundSettings>();
+			check(Settings);
+			if (const FMetaSoundPageSettings* PageSettings = Settings->FindPageSettings(Args.PageID))
+			{
+				if (PageSettings->Name != ParentPtr->BuildPageName)
+				{
+					ParentPtr->BuildPageName = PageSettings->Name;
+					FGraphBuilder::RegisterGraphWithFrontend(ParentPtr->GetMetaSound());
+				}
+
+				ParentPtr->AddableItems.RemoveAll([&PageSettings](const TSharedPtr<FString>& Item) { return Item->Compare(PageSettings->Name.ToString()) == 0; });
+				ParentPtr->ImplementedNames.Add(PageSettings->Name);
+				ParentPtr->ComboBox->RefreshOptions();
+				ParentPtr->RebuildImplemented();
+			}
+		}
+	}
+
+	void FMetasoundPagesDetailCustomization::FPageListener::OnPageSet(const Frontend::FDocumentMutatePageArgs& Args)
+	{
+		if (TSharedPtr<FMetasoundPagesDetailCustomization> ParentPtr = Parent.Pin())
+		{
+			const UMetaSoundSettings* Settings = GetDefault<UMetaSoundSettings>();
+			check(Settings);
+			if (const FMetaSoundPageSettings* PageSettings = Settings->FindPageSettings(Args.PageID))
+			{
+				ParentPtr->BuildPageName = PageSettings->Name;
+				ParentPtr->ComboBox->RefreshOptions();
+				ParentPtr->RebuildImplemented();
+			}
+		}
+	}
+
+	void FMetasoundPagesDetailCustomization::FPageListener::OnRemovingPage(const Frontend::FDocumentMutatePageArgs& Args)
+	{
+		if (TSharedPtr<FMetasoundPagesDetailCustomization> ParentPtr = Parent.Pin())
+		{
+			const UMetaSoundSettings* Settings = GetDefault<UMetaSoundSettings>();
+			check(Settings);
+			if (const FMetaSoundPageSettings* PageSettings = Settings->FindPageSettings(Args.PageID))
+			{
+				if (PageSettings->Name != ParentPtr->BuildPageName)
+				{
+					ParentPtr->BuildPageName = PageSettings->Name;
+					FGraphBuilder::RegisterGraphWithFrontend(ParentPtr->GetMetaSound());
+				}
+
+				ParentPtr->AddableItems.Add(MakeShared<FString>(PageSettings->Name.ToString()));
+				ParentPtr->ImplementedNames.Remove(PageSettings->Name);
+				ParentPtr->ComboBox->RefreshOptions();
+				ParentPtr->RebuildImplemented();
+			}
+		}
+	}
+
 	void FMetasoundInterfacesDetailCustomization::CustomizeDetails(IDetailLayoutBuilder& DetailLayout)
 	{
 		TArray<TWeakObjectPtr<UObject>> Objects;
diff --git a/Engine/Plugins/Runtime/Metasound/Source/MetasoundEditor/Private/MetasoundDetailCustomization.h b/Engine/Plugins/Runtime/Metasound/Source/MetasoundEditor/Private/MetasoundDetailCustomization.h
--- a/Engine/Plugins/Runtime/Metasound/Source/MetasoundEditor/Private/MetasoundDetailCustomization.h
+++ b/Engine/Plugins/Runtime/Metasound/Source/MetasoundEditor/Private/MetasoundDetailCustomization.h
@@ -64,6 +64,7 @@ namespace Metasound
 		private:
 			UObject& GetMetaSound() const;
 			void RebuildImplemented();
+			void RefreshView();
 			void UpdateItemNames();
 
 			TArray<TSharedPtr<FString>> AddableItems;
@@ -90,6 +91,7 @@ namespace Metasound
 			private:
 				virtual void OnBuilderReloaded(Frontend::FDocumentModifyDelegates& OutDelegates) override;
 				void OnPageAdded(const Frontend::FDocumentMutatePageArgs& Args);
+				void OnPageSet(const Frontend::FDocumentMutatePageArgs& Args);
 				void OnRemovingPage(const Frontend::FDocumentMutatePageArgs& Args);
 			};
 			TSharedPtr<FPageListener> PageListener;
diff --git a/Engine/Plugins/Runtime/Metasound/Source/MetasoundEditor/Private/MetasoundEditor.cpp b/Engine/Plugins/Runtime/Metasound/Source/MetasoundEditor/Private/MetasoundEditor.cpp
--- a/Engine/Plugins/Runtime/Metasound/Source/MetasoundEditor/Private/MetasoundEditor.cpp
+++ b/Engine/Plugins/Runtime/Metasound/Source/MetasoundEditor/Private/MetasoundEditor.cpp
@@ -22,7 +22,9 @@
 #include "Framework/Commands/GenericCommands.h"
 #include "Framework/Docking/TabManager.h"
 #include "Framework/MultiBox/MultiBoxBuilder.h"
+#include "Framework/MultiBox/SToolBarButtonBlock.h"
 #include "Framework/Notifications/NotificationManager.h"
+#include "Framework/SlateDelegates.h"
 #include "GenericPlatform/GenericApplication.h"
 #include "GraphEditor.h"
 #include "GraphEditorActions.h"
@@ -62,6 +64,7 @@
 #include "MetasoundGenerator.h"
 #include "MetasoundLog.h"
 #include "MetasoundNodeDetailCustomization.h"
+#include "MetasoundSettings.h"
 #include "MetasoundSource.h"
 #include "MetasoundUObjectRegistry.h"
 #include "Misc/Attribute.h"
@@ -78,6 +81,7 @@
 #include "Styling/AppStyle.h"
 #include "Templates/Function.h"
 #include "Templates/SharedPointer.h"
+#include "ToolMenuEntry.h"
 #include "ToolMenus.h"
 #include "UObject/ScriptInterface.h"
 #include "Widgets/Docking/SDockTab.h"
@@ -1058,6 +1062,10 @@ namespace Metasound
 			// This forces versioning the document on load prior to the editor synchronizing and building the editor graph if an asset is
 			// reloaded while the asset editor was open.
 			Builder.Reset(&Engine::FDocumentBuilderRegistry::GetChecked().FindOrBeginBuilding(*ObjectToEdit));
+			DocListener = MakeShared<FDocumentListener>(StaticCastSharedRef<FEditor>(AsShared()));
+			Builder->AddTransactionListener(DocListener->AsShared());
+			SyncFocusedPage();
+
 			if (FMetasoundAssetBase* MetaSoundAsset = IMetasoundUObjectRegistry::Get().GetObjectAsAssetBase(ObjectToEdit))
 			{
 				if (MetaSoundAsset->VersionAsset(Builder->GetBuilder()))
@@ -1626,6 +1634,171 @@ namespace Metasound
 			DestroyAnalyzers();
 		}
 
+		TSharedRef<SWidget> FEditor::CreateAuditionMenuOptions() const
+		{
+			TSharedPtr<FUICommandList> Commands = MakeShared<FUICommandList>();
+			FMenuBuilder MenuBuilder(true, Commands);
+			
+			constexpr bool bOpenOnClick = false;
+			constexpr bool bCloseAfterSelection = false;
+			const FText Label = LOCTEXT("PageAuditionSettingsSubMenu", "Pages");
+			const FText ToolTip = LOCTEXT("AuditionPageSettingsTooltip", "Settings related to auditioning pages");
+			MenuBuilder.AddSubMenu(Label, ToolTip,
+				FNewMenuDelegate::CreateLambda([this](FMenuBuilder& SubMenuBuilder)
+				{
+					CreateAuditionPageSubMenuOptions(SubMenuBuilder);
+				}),
+				bOpenOnClick,
+				FSlateIcon(),
+				bCloseAfterSelection
+			);
+			return MenuBuilder.MakeWidget();
+		}
+
+		void FEditor::CreateAuditionPageSubMenuOptions(FMenuBuilder& MenuBuilder) const
+		{
+			TWeakObjectPtr<const UMetaSoundSettings> SettingsPtr = GetDefault<UMetaSoundSettings>();
+			if (!SettingsPtr.IsValid())
+			{
+				return;
+			}
+
+			const TArray<FMetaSoundPageSettings>& PageSettingsArray = SettingsPtr->GetPageSettings();
+			MenuBuilder.BeginSection("SetAuditionPlatformSectionHeader", LOCTEXT("SetAuditionPlatformDescription", "Platform"));
+			{
+				TSet<FName> ImplementedPlatforms;
+				for (const FMetaSoundPageSettings& PageSettings : PageSettingsArray)
+				{
+					auto PairToKey = [](const TPair<FName, int32>& PerPlatformPair) { return PerPlatformPair.Key; };
+					Algo::Transform(PageSettings.IsCooked.PerPlatform, ImplementedPlatforms, PairToKey);
+				}
+
+				auto CreatePlatformEntry = [this, &MenuBuilder](FName PlatformName, const FText& PlatformText)
+				{
+					FUIAction SetPlatformAction;
+					SetPlatformAction.ExecuteAction = FExecuteAction::CreateLambda([this, PlatformName]()
+					{
+						if (UMetasoundEditorSettings* EditorSettings = GetMutableDefault<UMetasoundEditorSettings>())
+						{
+							EditorSettings->AuditionPlatform = PlatformName;
+						}
+					});
+
+					SetPlatformAction.CanExecuteAction = FCanExecuteAction::CreateLambda([this, PlatformName]()
+					{
+						if (const UMetasoundEditorSettings* EditorSettings = GetDefault<UMetasoundEditorSettings>())
+						{
+							return EditorSettings->AuditionPlatform != PlatformName;
+						}
+
+						return false;
+					});
+
+					MenuBuilder.AddMenuEntry(PlatformText,
+						FText::Format(LOCTEXT("SetAuditionPlatformToolTip", "Sets the audition platform to '{0}'."), PlatformText),
+						FSlateIcon(),
+						SetPlatformAction);
+				};
+
+				CreatePlatformEntry(FName(), LOCTEXT("DefaultPlatformDisplayName", "Default"));
+				for (const FName& PlatformName : ImplementedPlatforms)
+				{
+					const FText PlatformText = FText::FromName(PlatformName);
+					CreatePlatformEntry(PlatformName, PlatformText);
+				}
+			}
+			MenuBuilder.EndSection();
+
+			MenuBuilder.BeginSection("SetAuditionTargetPageSectionHeader", LOCTEXT("SetAuditionTargetPageDescription", "Target Page"));
+			{
+				const FText FocusPageTooltip = LOCTEXT("EnableFocusTargetPageSwapTooltip", "Dynamically swap which page is targeted based on which page is focused");
+				MenuBuilder.AddWidget(
+					SNew(SCheckBox)
+					.OnCheckStateChanged_Lambda([this](ECheckBoxState State)
+					{
+						if (UMetasoundEditorSettings* EdSettings = GetMutableDefault<UMetasoundEditorSettings>())
+						{
+							switch (State)
+							{
+								case ECheckBoxState::Checked:
+								{
+									EdSettings->AuditionPageMode = EAuditionPageMode::Focused;
+									SyncFocusedPage();
+									break;
+								}
+
+								case ECheckBoxState::Unchecked:
+								case ECheckBoxState::Undetermined:
+								default:
+								{
+									EdSettings->AuditionPageMode = EAuditionPageMode::User;
+								}
+								break;
+							}
+						}
+					})
+					.IsChecked_Lambda([this]()
+					{
+						if (const UMetasoundEditorSettings* EdSettings = GetDefault<UMetasoundEditorSettings>())
+						{
+							if (EdSettings->AuditionPageMode == EAuditionPageMode::Focused)
+							{
+								return ECheckBoxState::Checked;
+							}
+						}
+
+						return ECheckBoxState::Unchecked;
+					})
+					.ToolTipText(FocusPageTooltip),
+					LOCTEXT("EnableFocusTargetPageSwap", "Set To Focused"),
+					true,
+					true,
+					FocusPageTooltip
+				);
+
+				TSet<FName> PageNames;
+				Algo::Transform(PageSettingsArray, PageNames, [](const FMetaSoundPageSettings& PageSettings) { return PageSettings.Name; });
+
+				auto CreateTargetPageEntry = [this, &MenuBuilder](FName AuditionTargetPage, const FText& PageText)
+				{
+					FUIAction SetPlatformAction;
+					SetPlatformAction.ExecuteAction = FExecuteAction::CreateLambda([this, AuditionTargetPage]()
+					{
+						if (UMetasoundEditorSettings* EditorSettings = GetMutableDefault<UMetasoundEditorSettings>())
+						{
+							EditorSettings->AuditionTargetPage = AuditionTargetPage;
+						}
+					});
+
+					SetPlatformAction.CanExecuteAction = FCanExecuteAction::CreateLambda([this, AuditionTargetPage]()
+					{
+						if (const UMetasoundEditorSettings* EditorSettings = GetDefault<UMetasoundEditorSettings>())
+						{
+							const bool bUserAuditionMode = EditorSettings->AuditionPageMode == EAuditionPageMode::User;
+							if (bUserAuditionMode)
+							{
+								return EditorSettings->AuditionTargetPage != AuditionTargetPage;
+							}
+						}
+
+						return false;
+					});
+
+					MenuBuilder.AddMenuEntry(PageText,
+						FText::Format(LOCTEXT("SetAuditionPageToolTip", "Sets the audition target page to '{0}'."), PageText),
+						FSlateIcon(),
+						SetPlatformAction);
+				};
+
+				for (const FName& PageName : PageNames)
+				{
+					const FText PageText = FText::FromName(PageName);
+					CreateTargetPageEntry(PageName, PageText);
+				}
+			}
+			MenuBuilder.EndSection();
+		}
+
 		void FEditor::DestroyAnalyzers()
 		{
 			OutputMeter.Reset();
@@ -1691,7 +1864,7 @@ namespace Metasound
 								NAME_None,
 								TAttribute<FText>(),
 								TAttribute<FText>(),
-								TAttribute<FSlateIcon>::Create([this]() { return GetSettingsImage(); }),
+								Style::CreateSlateIcon("MetasoundEditor.Settings"),
 								"EditSourceSettings"
 							);
 						}
@@ -1701,7 +1874,7 @@ namespace Metasound
 							NAME_None,
 							TAttribute<FText>(),
 							TAttribute<FText>(),
-							TAttribute<FSlateIcon>::Create([this]() { return GetSettingsImage(); }),
+							Style::CreateSlateIcon("MetasoundEditor.MetasoundSource.Thumbnail"),
 							"EditMetasoundSettings"
 						);
 					}
@@ -1740,6 +1913,25 @@ namespace Metasound
 				})
 			);
 
+			if (AssetEditorPrivate::EnablePageEditor)
+			{
+				if (UToolMenu* AssetToolbar = UToolMenus::Get()->ExtendMenu(GetToolMenuToolbarName()))
+				{
+					TSharedPtr<FUICommandList> CommandList = MakeShared<FUICommandList>();
+					FToolMenuSection& Section = AssetToolbar->FindOrAddSection("Asset.Utilities");
+					FToolMenuEntry Entry = FToolMenuEntry::InitComboButton(
+						"AuditionMenu",
+						FUIAction(),
+						FOnGetContent::CreateRaw(this, &FEditor::CreateAuditionMenuOptions),
+						LOCTEXT("AuditionSettingsMenu", "Audition"),
+						LOCTEXT("AuditionSettingsMenu_Tooltip", "Settings related to auditioning MetaSound (Target page, platform etc.)"),
+						Style::CreateSlateIcon("MetasoundEditor.Audition"),
+						false);
+					Entry.StyleNameOverride = "CalloutToolbar";
+					Section.AddEntry(Entry);
+				}
+			}
+
 			AddToolbarExtender(ToolbarExtender);
 
 			if (GEditor)
@@ -1760,12 +1952,6 @@ namespace Metasound
 			return FSlateIcon("MetaSoundStyle", IconName);
 		}
 
-		FSlateIcon FEditor::GetSettingsImage() const
-		{
-			const FName IconName = "MetasoundEditor.Settings";
-			return FSlateIcon("MetaSoundStyle", IconName);
-		}
-
 		FSlateIcon FEditor::GetExportStatusImage() const
 		{
 			FName IconName = "MetasoundEditor.Export";
@@ -2196,6 +2382,20 @@ namespace Metasound
 			SetPreviewID(INDEX_NONE);
 		}
 
+		void FEditor::SyncFocusedPage() const
+		{
+			if (const UMetaSoundSettings* Settings = GetDefault<UMetaSoundSettings>())
+			{
+				const FMetaSoundFrontendDocumentBuilder& DocBuilder = Builder->GetConstBuilder();
+				if (const FMetaSoundPageSettings* PageSettings = Settings->FindPageSettings(DocBuilder.GetBuildPageID()))
+				{
+					constexpr bool bFocusPageEditor = false; // Already Focused
+					EMetaSoundBuilderResult Result = EMetaSoundBuilderResult::Failed;
+					UMetaSoundEditorSubsystem::GetChecked().SetFocusedPage(Builder.Get(), PageSettings->Name, bFocusPageEditor, Result);
+				}
+			}
+		}
+
 		void FEditor::TogglePlayback()
 		{
 			check(GEditor);
@@ -3910,7 +4110,8 @@ namespace Metasound
 			FGraphBuilder::BindEditorGraph(DocBuilder, &Graph);
 			check(Graph);
 
-			const bool bSynchronizedGraph = FGraphBuilder::SynchronizeGraph(DocBuilder, *Graph);
+			const bool bSynchronizedGraph = FGraphBuilder::SynchronizeGraph(DocBuilder, *Graph, !bRefreshGraph);
+			bRefreshGraph = false;
 
 			FMetasoundAssetBase* MetaSoundAsset = IMetasoundUObjectRegistry::Get().GetObjectAsAssetBase(&MetaSound);
 			check(MetaSoundAsset);
@@ -4289,6 +4490,19 @@ namespace Metasound
 		{
 			bMemberRenameRequested = true;
 		}
+
+		void FEditor::FDocumentListener::OnBuilderReloaded(Frontend::FDocumentModifyDelegates& OutDelegates)
+		{
+			OutDelegates.PageDelegates.OnPageSet.AddSP(this, &FDocumentListener::OnPageSet);
+		}
+
+		void FEditor::FDocumentListener::OnPageSet(const Frontend::FDocumentMutatePageArgs& Args)
+		{
+			if (TSharedPtr<FEditor> ParentPtr = Parent.Pin())
+			{
+				ParentPtr->bRefreshGraph = true;
+			}
+		}
 	}
 }
 #undef LOCTEXT_NAMESPACE
diff --git a/Engine/Plugins/Runtime/Metasound/Source/MetasoundEditor/Private/MetasoundEditor.h b/Engine/Plugins/Runtime/Metasound/Source/MetasoundEditor/Private/MetasoundEditor.h
--- a/Engine/Plugins/Runtime/Metasound/Source/MetasoundEditor/Private/MetasoundEditor.h
+++ b/Engine/Plugins/Runtime/Metasound/Source/MetasoundEditor/Private/MetasoundEditor.h
@@ -382,6 +382,9 @@ namespace Metasound
 
 			void ExportNodesToText(FString& OutText) const;
 
+			// Sets the globally targeted audition page for execution to this editor's currently focused page
+			void SyncFocusedPage() const;
+
 			/** FNotifyHook interface */
 			virtual void NotifyPostChange(const FPropertyChangedEvent& PropertyChangedEvent, FProperty* PropertyThatChanged) override;
 
@@ -398,8 +401,6 @@ namespace Metasound
 
 			FSlateIcon GetExportStatusImage() const;
 
-			FSlateIcon GetSettingsImage() const;
-
 			// TODO: Move import/export out of editor and into import/export asset actions
 			void Import();
 			void Export();
@@ -422,6 +423,12 @@ namespace Metasound
 			/** Converts the MetaSound from a preset to a fully modifiable MetaSound. */
 			void ConvertFromPreset();
 
+			/** Creates audition menu options */
+			TSharedRef<SWidget> CreateAuditionMenuOptions() const;
+
+			/** Creates page menu options */
+			void CreateAuditionPageSubMenuOptions(FMenuBuilder& MenuBuilder) const;
+
 			/** Show the Metasound object's Source settings in the Details panel */
 			void EditSourceSettings();
 
@@ -521,6 +528,28 @@ namespace Metasound
 			TSharedPtr<SNotificationItem> NotificationPtr;
 
 			bool bMemberRenameRequested = false;
+
+			bool bRefreshGraph = false;
+
+			class FDocumentListener : public Frontend::IDocumentBuilderTransactionListener
+			{
+				TWeakPtr<FEditor> Parent;
+
+			public:
+				FDocumentListener() = default;
+				FDocumentListener(TSharedRef<FEditor> InParent)
+					: Parent(InParent)
+				{
+				}
+
+				virtual ~FDocumentListener() = default;
+
+			private:
+				virtual void OnBuilderReloaded(Frontend::FDocumentModifyDelegates& OutDelegates) override;
+				void OnPageSet(const Frontend::FDocumentMutatePageArgs& Args);
+			};
+
+			TSharedPtr<FDocumentListener> DocListener;
 		};
 	} // namespace Editor
 } // namespace Metasound
diff --git a/Engine/Plugins/Runtime/Metasound/Source/MetasoundEditor/Private/MetasoundEditorGraphBuilder.cpp b/Engine/Plugins/Runtime/Metasound/Source/MetasoundEditor/Private/MetasoundEditorGraphBuilder.cpp
--- a/Engine/Plugins/Runtime/Metasound/Source/MetasoundEditor/Private/MetasoundEditorGraphBuilder.cpp
+++ b/Engine/Plugins/Runtime/Metasound/Source/MetasoundEditor/Private/MetasoundEditorGraphBuilder.cpp
@@ -121,7 +121,7 @@ namespace Metasound
 				}
 			};
 
-			void SynchronizeGraphRecursively(const FMetaSoundFrontendDocumentBuilder& InBuilder, UMetasoundEditorGraph& OutGraph)
+			void SynchronizeGraphRecursively(const FMetaSoundFrontendDocumentBuilder& InBuilder, UMetasoundEditorGraph& OutGraph, bool bSkipIfModifyContextUnchanged)
 			{
 				using namespace Engine;
 				using namespace Frontend;
@@ -143,10 +143,10 @@ namespace Metasound
 
 					UMetasoundEditorGraph* Graph = nullptr;
 					FGraphBuilder::BindEditorGraph(RefBuilder, &Graph);
-					SynchronizeGraphRecursively(RefBuilder, *Graph);
+					SynchronizeGraphRecursively(RefBuilder, *Graph, bSkipIfModifyContextUnchanged);
 				}
 
-				if (!MetaSoundAsset->GetConstModifyContext().GetDocumentModified())
+				if (bSkipIfModifyContextUnchanged && !MetaSoundAsset->GetConstModifyContext().GetDocumentModified())
 				{
 					return;
 				}
@@ -492,12 +492,12 @@ namespace Metasound
 			FMetaSoundFrontendDocumentBuilder& Builder = IDocumentBuilderRegistry::GetChecked().FindOrBeginBuilding(&InMetaSound);
 			if (const FMetasoundFrontendNode* NewNode = Builder.AddNodeByClassName(InMetadata.GetClassName(), InMetadata.GetVersion().Major))
 			{
-				FMetasoundAssetBase* MetaSoundAsset = IMetasoundUObjectRegistry::Get().GetObjectAsAssetBase(&InMetaSound);
-				check(MetaSoundAsset);
-				MetaSoundAsset->GetDocumentChecked().Metadata.ModifyContext.AddNodeIDModified(NewNode->GetID());
 				const FMetasoundFrontendClass* Dependency = Builder.FindDependency(NewNode->ClassID);
 				if (ensure(Dependency))
 				{
+					FMetasoundAssetBase* MetaSoundAsset = IMetasoundUObjectRegistry::Get().GetObjectAsAssetBase(&InMetaSound);
+					check(MetaSoundAsset);
+					MetaSoundAsset->GetModifyContext().AddNodeIDModified(NewNode->GetID());
 					return AddExternalNode(InMetaSound, NewNode->GetID(), Dependency->Metadata, bInSelectNewNode);
 				}
 			}
@@ -2174,7 +2174,7 @@ namespace Metasound
 			return bIsGraphDirty;
 		}
 
-		bool FGraphBuilder::SynchronizeGraph(const FMetaSoundFrontendDocumentBuilder& InBuilder, UMetasoundEditorGraph& OutGraph)
+		bool FGraphBuilder::SynchronizeGraph(const FMetaSoundFrontendDocumentBuilder& InBuilder, UMetasoundEditorGraph& OutGraph, bool bSkipIfModifyContextUnchanged)
 		{
 			TRACE_CPUPROFILER_EVENT_SCOPE(Metasound::Editor::FGraphBuilder::SynchronizeGraph);
 
@@ -2184,7 +2184,7 @@ namespace Metasound
 			FMetasoundAssetBase* MetaSoundAsset = IMetasoundUObjectRegistry::Get().GetObjectAsAssetBase(&MetaSound);
 			check(MetaSoundAsset);
 
-			if (RecurseGetDocumentModified(*MetaSoundAsset))
+			if (!bSkipIfModifyContextUnchanged || RecurseGetDocumentModified(*MetaSoundAsset))
 			{
 				TSet<FMetasoundAssetBase*> EditedReferencingMetaSounds;
 				if (GEditor)
@@ -2208,7 +2208,7 @@ namespace Metasound
 				if (EditedReferencingMetaSounds.IsEmpty())
 				{
 					MetaSoundAsset->CacheRegistryMetadata();
-					GraphBuilderPrivate::SynchronizeGraphRecursively(InBuilder, OutGraph);
+					GraphBuilderPrivate::SynchronizeGraphRecursively(InBuilder, OutGraph, bSkipIfModifyContextUnchanged);
 					GraphBuilderPrivate::RecurseClearDocumentModified(*MetaSoundAsset);
 				}
 				else
@@ -2219,7 +2219,7 @@ namespace Metasound
 						UObject* OwningMetaSound = EditedMetaSound->GetOwningAsset();
 						check(OwningMetaSound);
 						const FMetaSoundFrontendDocumentBuilder& EditedBuilder = IDocumentBuilderRegistry::GetChecked().FindOrBeginBuilding(OwningMetaSound);
-						SynchronizeGraph(EditedBuilder, *CastChecked<UMetasoundEditorGraph>(&EditedMetaSound->GetGraphChecked()));
+						SynchronizeGraph(EditedBuilder, *CastChecked<UMetasoundEditorGraph>(&EditedMetaSound->GetGraphChecked()), bSkipIfModifyContextUnchanged);
 					}
 				}
 
diff --git a/Engine/Plugins/Runtime/Metasound/Source/MetasoundEditor/Private/MetasoundEditorGraphBuilder.h b/Engine/Plugins/Runtime/Metasound/Source/MetasoundEditor/Private/MetasoundEditorGraphBuilder.h
--- a/Engine/Plugins/Runtime/Metasound/Source/MetasoundEditor/Private/MetasoundEditorGraphBuilder.h
+++ b/Engine/Plugins/Runtime/Metasound/Source/MetasoundEditor/Private/MetasoundEditorGraphBuilder.h
@@ -260,9 +260,10 @@ namespace Metasound
 			// FMetasoundFrontendDocument model. Validates the graph (and those referenced recursively).
 			//
 			// @param InBuilder - Builder to synchronize ed graph with.
-			// @param OutGraph = Graph to mutate to conform to provided builder's selected build graph
+			// @param OutGraph - Graph to mutate to conform to provided builder's selected build graph
+			// @param bSkipIfModifyContextUnchanged - If true, recursively checks deprecated document modification for builder's document and all referenced documents, and skips EdGraph synchronization if no modifications are logged.
 			// @return whether or not EditorGraph synchronization was performed.
-			static bool SynchronizeGraph(const FMetaSoundFrontendDocumentBuilder& InBuilder, UMetasoundEditorGraph& OutGraph);
+			static bool SynchronizeGraph(const FMetaSoundFrontendDocumentBuilder& InBuilder, UMetasoundEditorGraph& OutGraph, bool bSkipIfModifyContextUnchanged = true);
 
 			// Synchronizes editor nodes with frontend nodes, removing editor nodes that are not represented in the frontend, and adding editor nodes to represent missing frontend nodes.
 			static void SynchronizeNodes(const FMetaSoundFrontendDocumentBuilder& InBuilder, UMetasoundEditorGraph& OutGraph);
diff --git a/Engine/Plugins/Runtime/Metasound/Source/MetasoundEditor/Private/MetasoundEditorModule.cpp b/Engine/Plugins/Runtime/Metasound/Source/MetasoundEditor/Private/MetasoundEditorModule.cpp
--- a/Engine/Plugins/Runtime/Metasound/Source/MetasoundEditor/Private/MetasoundEditorModule.cpp
+++ b/Engine/Plugins/Runtime/Metasound/Source/MetasoundEditor/Private/MetasoundEditorModule.cpp
@@ -6,6 +6,7 @@
 #include "EdGraph/EdGraphNode.h"
 #include "EdGraph/EdGraphPin.h"
 #include "EdGraphUtilities.h"
+#include "Editor.h"
 #include "Framework/Notifications/NotificationManager.h"
 #include "HAL/IConsoleManager.h"
 #include "IDetailCustomization.h"
@@ -159,6 +160,7 @@ namespace Metasound
 					Set("MetasoundEditor.Analyzers.BackgroundColor", FLinearColor(0.0075f, 0.0075f, 0.0075, 1.0f));
 
 					// Misc
+					Set("MetasoundEditor.Audition", new IMAGE_BRUSH_SVG(TEXT("Icons/metasound_page"), Icon16));
 					Set("MetasoundEditor.Speaker", new FSlateImageBrush(RootToContentDir(TEXT("/Icons/speaker_144x.png")), FVector2D(144.0f, 144.0f)));
 					Set("MetasoundEditor.Metasound.Icon", new IMAGE_BRUSH_SVG(TEXT("Icons/metasound_icon"), Icon16));
 
@@ -397,6 +399,50 @@ namespace Metasound
 
 			}
 
+			void RegisterSettingsDelegates()
+			{
+				using namespace Engine;
+
+				if (UMetaSoundSettings* Settings = GetMutableDefault<UMetaSoundSettings>())
+				{
+					Settings->GetOnDefaultConformedDelegate().AddLambda([]()
+					{
+						FNotificationInfo Info(LOCTEXT("MetaSoundSettings_CannotDeleteDefaultPage", "Cannot change name of nor delete 'Default' MetaSound Page"));
+						Info.bFireAndForget = true;
+						Info.ExpireDuration = 2.0f;
+						Info.bUseThrobber = true;
+						FSlateNotificationManager::Get().AddNotification(Info);
+					});
+				}
+
+				const UMetasoundEditorSettings* EditorSettings = GetDefault<UMetasoundEditorSettings>();
+				FDocumentBuilderRegistry::GetChecked().GetOnResolveAuditionPageInfoDelegate().BindUObject(EditorSettings, &UMetasoundEditorSettings::GetAuditionPageInfo);
+
+				FEditorDelegates::PreBeginPIE.AddWeakLambda(EditorSettings, [](const bool /* bSimulating */)
+				{
+					if (const UMetasoundEditorSettings* EdSettings = GetDefault<UMetasoundEditorSettings>())
+					{
+						FOnResolveAuditionPageInfo& OnResolveAuditionPageInfo = FDocumentBuilderRegistry::GetChecked().GetOnResolveAuditionPageInfoDelegate();
+						if (OnResolveAuditionPageInfo.IsBoundToObject(EdSettings))
+						{
+							OnResolveAuditionPageInfo.Unbind();
+						}
+					}
+				});
+				FEditorDelegates::EndPIE.AddWeakLambda(EditorSettings, [](const bool /* bSimulating */)
+				{
+					if (const UMetasoundEditorSettings* EdSettings = GetDefault<UMetasoundEditorSettings>())
+					{
+						FOnResolveAuditionPageInfo& OnResolveAuditionPageInfo = FDocumentBuilderRegistry::GetChecked().GetOnResolveAuditionPageInfoDelegate();
+						if (!OnResolveAuditionPageInfo.IsBoundToObject(EdSettings))
+						{
+							OnResolveAuditionPageInfo.BindUObject(EdSettings, &UMetasoundEditorSettings::GetAuditionPageInfo);
+						}
+					}
+				});
+
+			}
+
 			virtual void PrimeAssetRegistryAsync() override {}
 
 			virtual EAssetPrimeStatus GetAssetRegistryPrimeStatus() const override
@@ -407,7 +453,7 @@ namespace Metasound
 			virtual EAssetScanStatus GetAssetRegistryScanStatus() const override
 			{
 				return EAssetScanStatus::NotRequested;
-			}			
+			}
 
 			virtual void RegisterExplicitProxyClass(const UClass& InClass) override
 			{
@@ -601,18 +647,7 @@ namespace Metasound
 				});
 
 				AssetTools.GetOnPackageMigration().AddRaw(this, &FModule::OnPackageMigration);
-
-				if (UMetaSoundSettings* Settings = GetMutableDefault<UMetaSoundSettings>())
-				{
-					Settings->GetOnDefaultConformedDelegate().AddLambda([]()
-					{
-						FNotificationInfo Info(LOCTEXT("MetaSoundSettings_CannotDeleteDefaultPage", "Cannot change name of nor delete 'Default' MetaSound Page"));
-						Info.bFireAndForget = true;
-						Info.ExpireDuration = 2.0f;
-						Info.bUseThrobber = true;
-						FSlateNotificationManager::Get().AddNotification(Info);
-					});
-				}
+				RegisterSettingsDelegates();
 			}
 
 			virtual void ShutdownModule() override
diff --git a/Engine/Plugins/Runtime/Metasound/Source/MetasoundEditor/Private/MetasoundEditorSettings.cpp b/Engine/Plugins/Runtime/Metasound/Source/MetasoundEditor/Private/MetasoundEditorSettings.cpp
--- a/Engine/Plugins/Runtime/Metasound/Source/MetasoundEditor/Private/MetasoundEditorSettings.cpp
+++ b/Engine/Plugins/Runtime/Metasound/Source/MetasoundEditor/Private/MetasoundEditorSettings.cpp
@@ -5,11 +5,15 @@
 #include "AudioWidgetsStyle.h"
 #include "Framework/Notifications/NotificationManager.h"
 #include "Internationalization/Text.h"
+#include "MetasoundDocumentBuilderRegistry.h"
+#include "MetasoundFrontendDocument.h"
+#include "MetasoundSettings.h"
 #include "Styling/AppStyle.h"
 #include "Styling/SlateWidgetStyleAsset.h"
 #include "Widgets/Notifications/SNotificationList.h"
 #include "UObject/UnrealType.h"
 
+
 #include UE_INLINE_GENERATED_CPP_BY_NAME(MetasoundEditorSettings)
 
 #define LOCTEXT_NAMESPACE "MetasoundEditorSettings"
@@ -104,4 +108,53 @@ const FAudioMaterialMeterStyle* UMetasoundEditorSettings::GetMeterStyle() const
 
 	return &FAudioWidgetsStyle::Get().GetWidgetStyle<FAudioMaterialMeterStyle>("AudioMaterialMeter.Style");
 }
+
+Metasound::Engine::FAuditionPageInfo UMetasoundEditorSettings::GetAuditionPageInfo(const FMetasoundFrontendDocument& InDocument) const
+{
+	using namespace Metasound::Engine;
+
+	FAuditionPageInfo PreviewInfo { .PlatformName = AuditionPlatform };
+
+	TSet<FGuid> DocPageIds;
+	InDocument.RootGraph.IterateGraphPages([&DocPageIds](const FMetasoundFrontendGraph& PageGraph)
+	{
+		DocPageIds.Add(PageGraph.PageID);
+	});
+
+	auto PageIsCooked = [&DocPageIds, &PreviewInfo](const FMetaSoundPageSettings& PageSettings)
+	{
+		if (DocPageIds.Contains(PageSettings.UniqueId))
+		{
+			return PageSettings.IsCooked.GetValueForPlatform(PreviewInfo.PlatformName);
+		}
+
+		return false;
+	};
+
+	if (const UMetaSoundSettings* Settings = GetDefault<UMetaSoundSettings>())
+	{
+		if (const FMetaSoundPageSettings* TargetPageSettings = Settings->FindPageSettings(AuditionTargetPage))
+		{
+			const FGuid& TargetPageID = TargetPageSettings->UniqueId;
+			const TArray<FMetaSoundPageSettings>& PageSettingsArray = Settings->GetPageSettings();
+			bool bFoundMatch = false;
+			for (int32 Index = PageSettingsArray.Num() - 1; Index >= 0; --Index)
+			{
+				const FMetaSoundPageSettings& PageSettings = PageSettingsArray[Index];
+				bFoundMatch |= PageSettings.UniqueId == TargetPageID;
+				if (bFoundMatch)
+				{
+					if (PageIsCooked(PageSettings))
+					{
+						PreviewInfo.PageID = PageSettings.UniqueId;
+						return PreviewInfo;
+					}
+				}
+			}
+		}
+	}
+
+	return PreviewInfo;
+
+}
 #undef LOCTEXT_NAMESPACE // "MetaSoundEditor"
diff --git a/Engine/Plugins/Runtime/Metasound/Source/MetasoundEditor/Private/MetasoundEditorSettings.h b/Engine/Plugins/Runtime/Metasound/Source/MetasoundEditor/Private/MetasoundEditorSettings.h
--- a/Engine/Plugins/Runtime/Metasound/Source/MetasoundEditor/Private/MetasoundEditorSettings.h
+++ b/Engine/Plugins/Runtime/Metasound/Source/MetasoundEditor/Private/MetasoundEditorSettings.h
@@ -9,11 +9,19 @@
 
 #include "MetasoundEditorSettings.generated.h"
 
+// Forward Declarations
 class USlateWidgetStyleAsset;
 struct FAudioMaterialKnobStyle;
 struct FAudioMaterialButtonStyle;
 struct FAudioMaterialSliderStyle;
 struct FAudioMaterialMeterStyle;
+struct FMetasoundFrontendDocument;
+
+namespace Metasound::Engine
+{
+	struct FAuditionPageInfo;
+} // namespace Metasound::Engine
+
 
 UENUM()
 enum class EMetasoundActiveAnalyzerEnvelopeDirection : uint8
@@ -44,6 +52,13 @@ enum class EMetasoundActiveDetailView : uint8
 	General
 };
 
+UENUM()
+enum class EAuditionPageMode : uint8
+{
+	Focused,
+	User
+};
+
 USTRUCT()
 struct FMetasoundAnalyzerAnimationSettings
 {
@@ -102,6 +117,23 @@ public:
 	UPROPERTY(EditAnywhere, config, Category=General)
 	FString DefaultAuthor;
 
+	/* Currently set page audition mode. Set by the MetaSound Asset Editor. */
+	UPROPERTY(VisibleAnywhere, config, Category = Pages)
+	EAuditionPageMode AuditionPageMode = EAuditionPageMode::Focused;
+
+	/** Name of platform to mock when previewing playback. This will limit playback to only pages that are cooked for the given platform.
+	  * Set in the MetaSound Asset Editor.
+	  */
+	UPROPERTY(VisibleAnywhere, config, Category = Pages)
+	FName AuditionPlatform;
+
+	/** Name of the page to target when previewing playback in editor. If target page is not implemented for the set audition platform,
+	  *  uses order of cooked pages (see 'project MetaSound Settings --> Page Settings' for order) falling back to lower index-ordered page
+	  * implemented in MetaSound asset. Set in the MetaSound Asset Editor.
+	  */
+	UPROPERTY(VisibleAnywhere, config, Category = Pages)
+	FName AuditionTargetPage;
+
 	/** Maps Pin Category To Pin Color */
 	TMap<FName, FLinearColor> CustomPinTypeColors;
 	
@@ -183,7 +215,7 @@ public:
 	
 	/**Override the Knob Style used in the Metasound Editor.*/
 	UPROPERTY(EditAnywhere, config, Category = WidgetStyling, meta = (AllowedClasses = "/Script/SlateCore.SlateWidgetStyleAsset", EditCondition = "bUseAudioMaterialWidgets", DisplayName = "Knob Style"))
-	FSoftObjectPath KnobStyleOverride;	
+	FSoftObjectPath KnobStyleOverride;
 	
 	/**Override the Slider Style used in the Metasound Editor.*/
 	UPROPERTY(EditAnywhere, config, Category = WidgetStyling, meta = (AllowedClasses = "/Script/SlateCore.SlateWidgetStyleAsset", EditCondition = "bUseAudioMaterialWidgets", DisplayName = "Slider Style"))
@@ -191,11 +223,11 @@ public:
 
 	/**Override the Button Style used in the Metasound Editor.*/
 	UPROPERTY(EditAnywhere, config, Category = WidgetStyling, meta = (AllowedClasses = "/Script/SlateCore.SlateWidgetStyleAsset", EditCondition = "bUseAudioMaterialWidgets", DisplayName = "Button Style"))
-	FSoftObjectPath ButtonStyleOverride;	
+	FSoftObjectPath ButtonStyleOverride;
 	
 	/**Override the Meter Style used in the Metasound Editor.*/
 	UPROPERTY(EditAnywhere, config, Category = WidgetStyling, meta = (AllowedClasses = "/Script/SlateCore.SlateWidgetStyleAsset", EditCondition = "bUseAudioMaterialWidgets", DisplayName = "Meter Style"))
-	FSoftObjectPath MeterStyleOverride;	
+	FSoftObjectPath MeterStyleOverride;
 
 	//UObject
 #if WITH_EDITOR
@@ -204,15 +236,16 @@ public:
 	//~UObject
 		
 	/** Get the AudioMaterialKnob Style. If KnobStyleOverride is not set, returns default style.*/
-	const FAudioMaterialKnobStyle* GetKnobStyle() const;	
+	const FAudioMaterialKnobStyle* GetKnobStyle() const;
 
 	/** Get the AudioMaterialSlider Style. If SliderStyleOverride is not set, returns default style.*/
-	const FAudioMaterialSliderStyle* GetSliderStyle() const;	
+	const FAudioMaterialSliderStyle* GetSliderStyle() const;
 	
 	/** Get the AudioMaterialButton Style. If ButtonStyleOverride is not set, returns default style.*/
-	const FAudioMaterialButtonStyle* GetButtonStyle() const;	
+	const FAudioMaterialButtonStyle* GetButtonStyle() const;
 	
 	/** Get the AudioMaterialMeter Style. If MeterStyleOverride is not set, returns default style.*/
 	const FAudioMaterialMeterStyle* GetMeterStyle() const;
 
+	Metasound::Engine::FAuditionPageInfo GetAuditionPageInfo(const FMetasoundFrontendDocument& InDocument) const;
 };
diff --git a/Engine/Plugins/Runtime/Metasound/Source/MetasoundEditor/Private/MetasoundEditorSubsystem.cpp b/Engine/Plugins/Runtime/Metasound/Source/MetasoundEditor/Private/MetasoundEditorSubsystem.cpp
--- a/Engine/Plugins/Runtime/Metasound/Source/MetasoundEditor/Private/MetasoundEditorSubsystem.cpp
+++ b/Engine/Plugins/Runtime/Metasound/Source/MetasoundEditor/Private/MetasoundEditorSubsystem.cpp
@@ -10,10 +10,13 @@
 #include "MetasoundEditorGraphSchema.h"
 #include "MetasoundEditorSettings.h"
 #include "MetasoundFactory.h"
+#include "MetasoundSettings.h"
 #include "MetasoundUObjectRegistry.h"
 #include "NodeTemplates/MetasoundFrontendNodeTemplateInput.h"
+#include "ScopedTransaction.h"
 #include "Sound/SoundSourceBusSend.h"
 #include "Sound/SoundSubmixSend.h"
+#include "Subsystems/AssetEditorSubsystem.h"
 
 #include UE_INLINE_GENERATED_CPP_BY_NAME(MetasoundEditorSubsystem)
 
@@ -246,6 +249,44 @@ void UMetaSoundEditorSubsystem::RegisterToolbarExtender(TSharedRef<FExtender> In
 	EditorToolbarExtenders.AddUnique(InExtender);
 }
 
+void UMetaSoundEditorSubsystem::SetFocusedPage(UMetaSoundBuilderBase* InBuilder, FName InPageName, bool bFocusPageEditor, EMetaSoundBuilderResult& OutResult) const
+{
+	using namespace Metasound::Frontend;
+	if (!InBuilder)
+	{
+		OutResult = EMetaSoundBuilderResult::Failed;
+		return;
+	}
+
+	const UMetaSoundSettings* Settings = GetDefault<UMetaSoundSettings>();
+	check(Settings);
+	if (const FMetaSoundPageSettings* PageSettings = Settings->FindPageSettings(InPageName))
+	{
+		const FScopedTransaction Transaction(FText::Format(LOCTEXT("SetFocusedBuildPageTransactionFormat", "Set Focused Builder Page '{0}'"), FText::FromName(InPageName)));
+		InBuilder->Modify();
+		if (InBuilder->GetBuilder().SetBuildPageID(PageSettings->UniqueId))
+		{
+			if (UMetasoundEditorSettings* EditorSettings = GetMutableDefault<UMetasoundEditorSettings>())
+			{
+				EditorSettings->AuditionTargetPage = InPageName;
+			}
+
+			if (GEditor && bFocusPageEditor)
+			{
+				if (UAssetEditorSubsystem* AssetEditorSubsystem = GEditor->GetEditorSubsystem<UAssetEditorSubsystem>())
+				{
+					AssetEditorSubsystem->OpenEditorForAsset(&InBuilder->GetConstBuilder().CastDocumentObjectChecked<UObject>());
+				}
+			}
+
+			OutResult = EMetaSoundBuilderResult::Succeeded;
+			return;
+		}
+	};
+
+	OutResult = EMetaSoundBuilderResult::Failed;
+}
+
 bool UMetaSoundEditorSubsystem::UnregisterToolbarExtender(TSharedRef<FExtender> InExtender)
 {
 	const int32 NumRemoved = EditorToolbarExtenders.RemoveAllSwap([&InExtender](const TSharedRef<FExtender>& Extender) { return Extender == InExtender; });
diff --git a/Engine/Plugins/Runtime/Metasound/Source/MetasoundEditor/Public/MetasoundEditorSubsystem.h b/Engine/Plugins/Runtime/Metasound/Source/MetasoundEditor/Public/MetasoundEditorSubsystem.h
--- a/Engine/Plugins/Runtime/Metasound/Source/MetasoundEditor/Public/MetasoundEditorSubsystem.h
+++ b/Engine/Plugins/Runtime/Metasound/Source/MetasoundEditor/Public/MetasoundEditorSubsystem.h
@@ -75,6 +75,12 @@ public:
 	// Register toolbar extender that will be displayed in the MetaSound Asset Editor.
 	void RegisterToolbarExtender(TSharedRef<FExtender> InExtender);
 
+	// If the given page name is implemented on the provided builder, sets the focused page of
+	// the provided builder to the given page name if and sets the audition target page to
+	// the provided name. If the given builder has an asset editor open, optionally opens or brings
+	// that editor's associated page into user focus.
+	void SetFocusedPage(UMetaSoundBuilderBase* InBuilder, FName InPageName, bool bFocusPageEditor, EMetaSoundBuilderResult& OutResult) const;
+
 	// Unregisters toolbar extender that is displayed in the MetaSound Asset Editor.
 	bool UnregisterToolbarExtender(TSharedRef<FExtender> InExtender);
 
diff --git a/Engine/Plugins/Runtime/Metasound/Source/MetasoundEngine/Private/MetasoundDocumentBuilderRegistry.cpp b/Engine/Plugins/Runtime/Metasound/Source/MetasoundEngine/Private/MetasoundDocumentBuilderRegistry.cpp
--- a/Engine/Plugins/Runtime/Metasound/Source/MetasoundEngine/Private/MetasoundDocumentBuilderRegistry.cpp
+++ b/Engine/Plugins/Runtime/Metasound/Source/MetasoundEngine/Private/MetasoundDocumentBuilderRegistry.cpp
@@ -305,9 +305,9 @@ namespace Metasound::Engine
 	}
 
 #if WITH_EDITOR
-	FOnResolvePreviewPageInfo& FDocumentBuilderRegistry::GetOnResolvePreviewPageInfoDelegate()
+	FOnResolveAuditionPageInfo& FDocumentBuilderRegistry::GetOnResolveAuditionPageInfoDelegate()
 	{
-		return OnResolvePreviewPageInfo;
+		return OnResolveAuditionPageInfo;
 	}
 #endif // WITH_EDITOR
 
@@ -329,9 +329,9 @@ namespace Metasound::Engine
 		FName PlatformName = FPlatformProperties::IniPlatformName();
 
 #if WITH_EDITOR
-		if (OnResolvePreviewPageInfo.IsBound())
+		if (OnResolveAuditionPageInfo.IsBound())
 		{
-			FPreviewPageInfo PreviewInfo = OnResolvePreviewPageInfo.Execute(Document);
+			FAuditionPageInfo PreviewInfo = OnResolveAuditionPageInfo.Execute(Document);
 			if (PreviewInfo.PageID.IsSet())
 			{
 				return PreviewInfo.PageID.GetValue();
diff --git a/Engine/Plugins/Runtime/Metasound/Source/MetasoundEngine/Public/MetasoundDocumentBuilderRegistry.h b/Engine/Plugins/Runtime/Metasound/Source/MetasoundEngine/Public/MetasoundDocumentBuilderRegistry.h
--- a/Engine/Plugins/Runtime/Metasound/Source/MetasoundEngine/Public/MetasoundDocumentBuilderRegistry.h
+++ b/Engine/Plugins/Runtime/Metasound/Source/MetasoundEngine/Public/MetasoundDocumentBuilderRegistry.h
@@ -15,13 +15,13 @@
 namespace Metasound::Engine
 {
 #if WITH_EDITOR
-	struct FPreviewPageInfo
+	struct FAuditionPageInfo
 	{
 		FName PlatformName;
 		TOptional<FGuid> PageID;
 	};
 
-	DECLARE_DELEGATE_RetVal_OneParam(FPreviewPageInfo, FOnResolvePreviewPageInfo, const FMetasoundFrontendDocument& /* Document */);
+	DECLARE_DELEGATE_RetVal_OneParam(FAuditionPageInfo, FOnResolveAuditionPageInfo, const FMetasoundFrontendDocument& /* Document */);
 #endif // WITH_EDITOR
 
 	class METASOUNDENGINE_API FDocumentBuilderRegistry : public Frontend::IDocumentBuilderRegistry
@@ -131,7 +131,7 @@ namespace Metasound::Engine
 		TArray<UMetaSoundBuilderBase*> FindBuilderObjects(const FMetasoundFrontendClassName& InClassName) const;
 
 #if WITH_EDITOR
-		FOnResolvePreviewPageInfo& GetOnResolvePreviewPageInfoDelegate();
+		FOnResolveAuditionPageInfo& GetOnResolveAuditionPageInfoDelegate();
 #endif // WITH_EDITOR
 
 		bool ReloadBuilder(const FMetasoundFrontendClassName& InClassName) const override;
@@ -147,7 +147,7 @@ namespace Metasound::Engine
 		void FinishBuildingInternal(UMetaSoundBuilderBase& Builder, bool bForceUnregisterNodeClass) const;
 
 #if WITH_EDITOR
-		FOnResolvePreviewPageInfo OnResolvePreviewPageInfo;
+		FOnResolveAuditionPageInfo OnResolveAuditionPageInfo;
 #endif // WITH_EDITOR
 
 		TSortedMap<ELogEvent, ELogVerbosity::Type> EventLogVerbosity;
