Improve logging for FMetasoundGeneratorHandle

#rnx

diff --git a/Engine/Plugins/Runtime/Metasound/Source/MetasoundEngine/Private/MetasoundGeneratorHandle.cpp b/Engine/Plugins/Runtime/Metasound/Source/MetasoundEngine/Private/MetasoundGeneratorHandle.cpp
--- a/Engine/Plugins/Runtime/Metasound/Source/MetasoundEngine/Private/MetasoundGeneratorHandle.cpp
+++ b/Engine/Plugins/Runtime/Metasound/Source/MetasoundEngine/Private/MetasoundGeneratorHandle.cpp
@@ -53,7 +53,7 @@ namespace Metasound
 
 			if (!Source.IsValid())
 			{
-				UE_LOG(LogMetaSound, Error, TEXT("Couldn't get the source for the audio component."));
+				UE_LOG(LogMetaSound, Error, TEXT("FMetaSoundGeneratorHandle missing source: %s."), *Handle->ToString());
 				return nullptr;
 			}
 			
@@ -206,6 +206,17 @@ namespace Metasound
 		return 0;
 	}
 
+	FString FMetasoundGeneratorHandle::ToString() const
+	{
+		if (!IsValid())
+		{
+			return FString::Printf(TEXT("Invalid Handle"));
+		}
+
+		check(AudioComponent.IsValid());
+		return FString::Printf(TEXT("%s [Id:%d] with owner %s"), *GetNameSafe(AudioComponent.Get()), AudioComponentId, *GetNameSafe(AudioComponent->GetOwner()));
+	}
+
 	void FMetasoundGeneratorHandle::SetGenerator(TWeakPtr<FMetasoundGenerator>&& InGenerator)
 	{
 		METASOUND_LLM_SCOPE;
diff --git a/Engine/Plugins/Runtime/Metasound/Source/MetasoundEngine/Public/MetasoundGeneratorHandle.h b/Engine/Plugins/Runtime/Metasound/Source/MetasoundEngine/Public/MetasoundGeneratorHandle.h
--- a/Engine/Plugins/Runtime/Metasound/Source/MetasoundEngine/Public/MetasoundGeneratorHandle.h
+++ b/Engine/Plugins/Runtime/Metasound/Source/MetasoundEngine/Public/MetasoundGeneratorHandle.h
@@ -193,6 +193,11 @@ namespace Metasound
 		*/
 		double GetCPUCoreUtilization() const;
 
+		/**
+		* Returns a text representation of this generator handle, useful for logging
+		*/
+		FString ToString() const;
+
 	private:
 		void SetGenerator(TWeakPtr<FMetasoundGenerator>&& InGenerator);
 
