Migrate MetaSound editor input/output creation to builder API in preparation of supporting paged input literals

#tests automated in-editor, asset editor addition/removal of inputs/outputs via misc means
#rnx

diff --git a/Engine/Plugins/Runtime/Metasound/Source/MetasoundEditor/Private/MetasoundEditor.cpp b/Engine/Plugins/Runtime/Metasound/Source/MetasoundEditor/Private/MetasoundEditor.cpp
--- a/Engine/Plugins/Runtime/Metasound/Source/MetasoundEditor/Private/MetasoundEditor.cpp
+++ b/Engine/Plugins/Runtime/Metasound/Source/MetasoundEditor/Private/MetasoundEditor.cpp
@@ -69,6 +69,7 @@
 #include "MetasoundUObjectRegistry.h"
 #include "Misc/Attribute.h"
 #include "Modules/ModuleManager.h"
+#include "NodeTemplates/MetasoundFrontendNodeTemplateInput.h"
 #include "NodeTemplates/MetasoundFrontendNodeTemplateReroute.h"
 #include "PropertyCustomizationHelpers.h"
 #include "PropertyEditorModule.h"
@@ -266,6 +267,7 @@ namespace Metasound
 			virtual FReply DroppedOnNode(FVector2D ScreenPosition, FVector2D GraphPosition) override { return FReply::Unhandled(); }
 			virtual FReply DroppedOnPin(FVector2D InScreenPosition, FVector2D InGraphPosition) override
 			{
+				using namespace Engine;
 				using namespace Frontend;
 
 				if (!GraphMember.IsValid())
@@ -284,9 +286,9 @@ namespace Metasound
 					MetasoundGraph->Modify();
 					Input->Modify();
 
-					FNodeHandle InputHandle = Input->GetNodeHandle();
-					FConstNodeHandle TemplateNodeHandle = FGraphBuilder::AddInputTemplateNodeHandle(ParentMetasound, InputHandle);
-					if (UMetasoundEditorGraphNode* NewGraphNode = FGraphBuilder::AddInputNode(ParentMetasound, TemplateNodeHandle))
+					FMetaSoundFrontendDocumentBuilder& Builder = FDocumentBuilderRegistry::GetChecked().FindOrBeginBuilding(&ParentMetasound);
+					const FMetasoundFrontendNode* TemplateNode = FInputNodeTemplate::CreateNode(Builder, Input->GetMemberName());
+					if (UMetasoundEditorGraphNode* NewGraphNode = FGraphBuilder::AddInputNode(ParentMetasound, TemplateNode->GetID()))
 					{
 						NewGraphNode->Modify();
 						NewGraphNode->UpdateFrontendNodeLocation(InGraphPosition);
@@ -314,8 +316,7 @@ namespace Metasound
 						MetasoundGraph->Modify();
 						Output->Modify();
 
-						FConstNodeHandle OutputHandle = Output->GetConstNodeHandle();
-						if (UMetasoundEditorGraphOutputNode* NewGraphNode = FGraphBuilder::AddOutputNode(ParentMetasound, OutputHandle))
+						if (UMetasoundEditorGraphOutputNode* NewGraphNode = FGraphBuilder::AddOutputNode(ParentMetasound, Output->NodeID))
 						{
 							NewGraphNode->Modify();
 							NewGraphNode->UpdateFrontendNodeLocation(InGraphPosition);
@@ -4351,6 +4352,8 @@ namespace Metasound
 				return FReply::Unhandled();
 			}
 
+			const FName DataTypeName = GetMetasoundDataTypeName<float>();
+
 			UMetasoundEditorGraph& Graph = GetMetaSoundGraphChecked();
 
 			TArray<TObjectPtr<UObject>> SelectedObjects;
@@ -4363,16 +4366,15 @@ namespace Metasound
 					const FScopedTransaction Transaction(LOCTEXT("AddInputNode", "Add MetaSound Input"));
 					MetaSound->Modify();
 
-					const FName DataTypeName = GetMetasoundDataTypeName<float>();
 					FCreateNodeVertexParams VertexParams;
 					VertexParams.DataType = DataTypeName;
 
-					Frontend::FNodeHandle NodeHandle = FGraphBuilder::AddInputNodeHandle(*MetaSound, VertexParams);
-					if (ensure(NodeHandle->IsValid()))
+					FMetasoundFrontendClassInput ClassInput = FGraphBuilder::CreateUniqueClassInput(*MetaSound, VertexParams);
+					if (const FMetasoundFrontendNode* NewNode = Builder->GetBuilder().AddGraphInput(ClassInput))
 					{
-						NameToSelect = NodeHandle->GetNodeName();
+						NameToSelect = NewNode->Name;
 
-						TObjectPtr<UMetasoundEditorGraphInput> Input = Graph.FindOrAddInput(NodeHandle);
+						TObjectPtr<UMetasoundEditorGraphInput> Input = Graph.FindOrAddInput(NewNode->GetID());
 						if (ensure(Input))
 						{
 							SelectedObjects.Add(Input);
@@ -4383,19 +4385,18 @@ namespace Metasound
 
 				case ENodeSection::Outputs:
 				{
-					const FScopedTransaction Transaction(TEXT(""), LOCTEXT("AddOutputNode", "Add MetaSound Output"), MetaSound);
+					const FScopedTransaction Transaction(LOCTEXT("AddOutputNode", "Add MetaSound Output"));
 					MetaSound->Modify();
 
-					const FName DataTypeName = GetMetasoundDataTypeName<float>();
 					FCreateNodeVertexParams VertexParams;
 					VertexParams.DataType = DataTypeName;
 
-					Frontend::FNodeHandle NodeHandle = FGraphBuilder::AddOutputNodeHandle(*MetaSound, VertexParams);
-					if (ensure(NodeHandle->IsValid()))
+					FMetasoundFrontendClassOutput ClassOutput = FGraphBuilder::CreateUniqueClassOutput(*MetaSound, VertexParams);
+					if (const FMetasoundFrontendNode* NewNode = Builder->GetBuilder().AddGraphOutput(ClassOutput))
 					{
-						NameToSelect = NodeHandle->GetNodeName();
+						NameToSelect = NewNode->Name;
 
-						TObjectPtr<UMetasoundEditorGraphOutput> Output = Graph.FindOrAddOutput(NodeHandle);
+						TObjectPtr<UMetasoundEditorGraphOutput> Output = Graph.FindOrAddOutput(NewNode->GetID());
 						if (ensure(Output))
 						{
 							SelectedObjects.Add(Output);
@@ -4409,8 +4410,6 @@ namespace Metasound
 					const FScopedTransaction Transaction(TEXT(""), LOCTEXT("AddVariableNode", "Add MetaSound Variable"), MetaSound);
 					MetaSound->Modify();
 
-					const FName DataTypeName = GetMetasoundDataTypeName<float>();
-					
 					Frontend::FVariableHandle FrontendVariable = FGraphBuilder::AddVariableHandle(*MetaSound, DataTypeName);
 					if (ensure(FrontendVariable->IsValid()))
 					{
diff --git a/Engine/Plugins/Runtime/Metasound/Source/MetasoundEditor/Private/MetasoundEditorDocumentClipboardUtils.cpp b/Engine/Plugins/Runtime/Metasound/Source/MetasoundEditor/Private/MetasoundEditorDocumentClipboardUtils.cpp
--- a/Engine/Plugins/Runtime/Metasound/Source/MetasoundEditor/Private/MetasoundEditorDocumentClipboardUtils.cpp
+++ b/Engine/Plugins/Runtime/Metasound/Source/MetasoundEditor/Private/MetasoundEditorDocumentClipboardUtils.cpp
@@ -7,6 +7,7 @@
 #include "Logging/TokenizedMessage.h"
 #include "MetasoundDocumentBuilderRegistry.h"
 #include "MetasoundFrontendSearchEngine.h"
+#include "NodeTemplates/MetasoundFrontendNodeTemplateInput.h"
 #include "ScopedTransaction.h"
 
 
@@ -14,8 +15,11 @@ namespace Metasound::Editor
 {
 	void FDocumentClipboardUtils::ProcessPastedInputNodes(FMetasoundAssetBase& OutAsset, TArray<UMetasoundEditorGraphNode*>& OutPastedNodes)
 	{
+		using namespace Engine;
 		using namespace Frontend;
 
+		FMetaSoundFrontendDocumentBuilder& Builder = FDocumentBuilderRegistry::GetChecked().FindOrBeginBuilding(OutAsset.GetOwningAsset());
+
 		TMap<FName, TObjectPtr<UMetasoundEditorGraphInput>> MappedGeneratedInputNames;
 		UMetasoundEditorGraph& Graph = *CastChecked<UMetasoundEditorGraph>(&OutAsset.GetGraphChecked());
 		for (int32 Index = OutPastedNodes.Num() - 1; Index >= 0; --Index)
@@ -62,25 +66,31 @@ namespace Metasound::Editor
 						{
 							FCreateNodeVertexParams VertexParams;
 							VertexParams.DataType = Breadcrumb.DataType;
-							FNodeHandle NewNodeHandle = FGraphBuilder::AddInputNodeHandle(*OutAsset.GetOwningAsset(), VertexParams, &Breadcrumb.DefaultLiteral, &Breadcrumb.MemberName);
-							MappedGeneratedInputNames.Add(Breadcrumb.MemberName, Input);
-							Input = Graph.FindOrAddInput(NewNodeHandle);
+
+							FMetasoundFrontendClassInput ClassInput = FGraphBuilder::CreateUniqueClassInput(*OutAsset.GetOwningAsset(), VertexParams);
+							if (const FMetasoundFrontendNode* NewNode = Builder.AddGraphInput(ClassInput))
+							{
+								Input = Graph.FindOrAddInput(NewNode->GetID());
+								MappedGeneratedInputNames.Add(Breadcrumb.MemberName, Input);
+							}
 						}
 					}
 				}
 			}
 
-			if (!Input)
+			if (Input)
 			{
-				constexpr bool bAllowShrinking = false;
-				Graph.RemoveNode(InputNode);
-				OutPastedNodes.RemoveAtSwap(Index, 1, EAllowShrinking::No);
+				const FMetasoundFrontendNode* InputTemplateNode = FInputNodeTemplate::CreateNode(Builder, Input->GetMemberName());
+				if (ensure(InputTemplateNode))
+				{
+					InputNode->NodeID = InputTemplateNode->GetID();
+				}
 			}
 			else
 			{
-				FNodeHandle InputNodeHandle = Input->GetNodeHandle();
-				FConstNodeHandle InputTemplateNode = FGraphBuilder::AddInputTemplateNodeHandle(*OutAsset.GetOwningAsset(), InputNodeHandle);
-				InputNode->NodeID = InputTemplateNode->GetID();
+				constexpr bool bAllowShrinking = false;
+				Graph.RemoveNode(InputNode);
+				OutPastedNodes.RemoveAtSwap(Index, 1, EAllowShrinking::No);
 			}
 		}
 	}
diff --git a/Engine/Plugins/Runtime/Metasound/Source/MetasoundEditor/Private/MetasoundEditorGraph.cpp b/Engine/Plugins/Runtime/Metasound/Source/MetasoundEditor/Private/MetasoundEditorGraph.cpp
--- a/Engine/Plugins/Runtime/Metasound/Source/MetasoundEditor/Private/MetasoundEditorGraph.cpp
+++ b/Engine/Plugins/Runtime/Metasound/Source/MetasoundEditor/Private/MetasoundEditorGraph.cpp
@@ -598,21 +598,29 @@ Metasound::Frontend::FNodeHandle UMetasoundEditorGraphInput::AddNodeHandle(const
 		return Metasound::Frontend::INodeController::GetInvalidHandle();
 	}
 
-	UObject& Metasound = Graph->GetMetasoundChecked();
-	Metasound::Frontend::FNodeHandle NewNodeHandle = FGraphBuilder::AddInputNodeHandle(Metasound, InParams, nullptr /* DefaultValue */, &InName);
-	return NewNodeHandle;
+	UObject& MetaSound = Graph->GetMetasoundChecked();
+	FMetasoundAssetBase* MetaSoundAsset = Metasound::IMetasoundUObjectRegistry::Get().GetObjectAsAssetBase(&MetaSound);
+	check(MetaSoundAsset);
+
+	FMetasoundFrontendClassInput ClassInput = FGraphBuilder::CreateUniqueClassInput(MetaSound, InParams, nullptr, &InName);
+	return MetaSoundAsset->GetRootGraphHandle()->AddInputVertex(ClassInput);
 }
 
 UMetasoundEditorGraphNode* UMetasoundEditorGraphInput::AddNode(Metasound::Frontend::FNodeHandle InNodeHandle, bool bInSelectNewNode)
 {
 	using namespace Metasound::Editor;
+	using namespace Metasound::Engine;
 	using namespace Metasound::Frontend;
 
 	UMetasoundEditorGraph* Graph = GetOwningGraph();
 	check(Graph);
 
-	FConstNodeHandle TemplateNodeHandle = FGraphBuilder::AddInputTemplateNodeHandle(Graph->GetMetasoundChecked(), InNodeHandle);
-	return FGraphBuilder::AddInputNode(Graph->GetMetasoundChecked(), TemplateNodeHandle, bInSelectNewNode);
+	if (const FMetasoundFrontendNode* TemplateNode = FInputNodeTemplate::CreateNode(GetFrontendBuilderChecked(), GetMemberName()))
+	{
+		return FGraphBuilder::AddInputNode(Graph->GetMetasoundChecked(), TemplateNode->GetID(), bInSelectNewNode);
+	}
+
+	return nullptr;
 }
 
 const FText& UMetasoundEditorGraphInput::GetGraphMemberLabel() const
@@ -890,9 +898,13 @@ Metasound::Frontend::FNodeHandle UMetasoundEditorGraphOutput::AddNodeHandle(cons
 		return Metasound::Frontend::INodeController::GetInvalidHandle();
 	} 
 
-	UObject& Metasound = Graph->GetMetasoundChecked();
-	FNodeHandle NewNodeHandle = FGraphBuilder::AddOutputNodeHandle(Metasound, InParams, &InName);
-	return NewNodeHandle;
+	UObject& MetaSound = Graph->GetMetasoundChecked();
+
+	FMetasoundAssetBase* MetaSoundAsset = Metasound::IMetasoundUObjectRegistry::Get().GetObjectAsAssetBase(&MetaSound);
+	check(MetaSoundAsset);
+
+	FMetasoundFrontendClassOutput ClassOutput = FGraphBuilder::CreateUniqueClassOutput(MetaSound, InParams, &InName);
+	return MetaSoundAsset->GetRootGraphHandle()->AddOutputVertex(ClassOutput);
 }
 
 UMetasoundEditorGraphNode* UMetasoundEditorGraphOutput::AddNode(Metasound::Frontend::FNodeHandle InNodeHandle, bool bInSelectNewNode)
@@ -900,7 +912,7 @@ UMetasoundEditorGraphNode* UMetasoundEditorGraphOutput::AddNode(Metasound::Front
 	using namespace Metasound::Editor;
 	UMetasoundEditorGraph* Graph = GetOwningGraph();
 	check(Graph);
-	return FGraphBuilder::AddOutputNode(Graph->GetMetasoundChecked(), InNodeHandle, bInSelectNewNode);
+	return FGraphBuilder::AddOutputNode(Graph->GetMetasoundChecked(), InNodeHandle->GetID(), bInSelectNewNode);
 }
 
 FText UMetasoundEditorGraphOutput::GetDescription() const
@@ -1698,6 +1710,41 @@ UMetasoundEditorGraphInput* UMetasoundEditorGraph::FindInput(FName InName) const
 	return Input ? Input->Get() : nullptr;
 }
 
+UMetasoundEditorGraphInput* UMetasoundEditorGraph::FindOrAddInput(const FGuid& InNodeID)
+{
+	using namespace Metasound::Editor;
+	using namespace Metasound::Frontend;
+
+	if (TObjectPtr<UMetasoundEditorGraphInput> Input = FindInput(InNodeID))
+	{
+		return Input;
+	}
+
+	const FMetaSoundFrontendDocumentBuilder& Builder = GetBuilderChecked().GetConstBuilder();
+	if (const FMetasoundFrontendNode* Node = Builder.FindNode(InNodeID))
+	{
+		if (const FMetasoundFrontendClassInput* ClassInput = Builder.FindGraphInput(Node->Name))
+		{
+			const FMetasoundFrontendLiteral& DefaultLiteral = ClassInput->DefaultLiteral;
+			if (const FMetasoundFrontendClass* Class = Builder.FindDependency(Node->ClassID))
+			{
+				FMetasoundFrontendClassName ClassName = Class->Metadata.GetClassName();
+
+				UMetasoundEditorGraphInput* NewInput = NewObject<UMetasoundEditorGraphInput>(this, GraphPrivate::GetUniqueTransientMemberName(), RF_Transactional);
+				if (ensure(NewInput))
+				{
+					NewInput->InitMember(ClassInput->TypeName, DefaultLiteral, InNodeID, MoveTemp(ClassName));
+					Inputs.Add(NewInput);
+				}
+
+				return NewInput;
+			}
+		}
+	}
+
+	return nullptr;
+}
+
 UMetasoundEditorGraphInput* UMetasoundEditorGraph::FindOrAddInput(Metasound::Frontend::FConstNodeHandle InNodeHandle)
 {
 	using namespace Metasound::Editor;
@@ -1762,6 +1809,42 @@ UMetasoundEditorGraphOutput* UMetasoundEditorGraph::FindOutput(FName InName) con
 	return Output ? Output->Get() : nullptr;
 }
 
+UMetasoundEditorGraphOutput* UMetasoundEditorGraph::FindOrAddOutput(const FGuid& InNodeID)
+{
+	using namespace Metasound::Editor;
+	using namespace Metasound::Frontend;
+
+	if (TObjectPtr<UMetasoundEditorGraphOutput> Output = FindOutput(InNodeID))
+	{
+		return Output;
+	}
+
+	const FMetaSoundFrontendDocumentBuilder& Builder = GetBuilderChecked().GetConstBuilder();
+	if (const FMetasoundFrontendNode* Node = Builder.FindNode(InNodeID))
+	{
+		if (const FMetasoundFrontendClassOutput* ClassOutput = Builder.FindGraphOutput(Node->Name))
+		{
+			if (const FMetasoundFrontendClass* Class = Builder.FindDependency(Node->ClassID))
+			{
+				FMetasoundFrontendClassName ClassName = Class->Metadata.GetClassName();
+
+				UMetasoundEditorGraphOutput* NewOutput = NewObject<UMetasoundEditorGraphOutput>(this, GraphPrivate::GetUniqueTransientMemberName(), RF_Transactional);
+				if (ensure(NewOutput))
+				{
+					FMetasoundFrontendLiteral DefaultLiteral;
+					DefaultLiteral.SetFromLiteral(IDataTypeRegistry::Get().CreateDefaultLiteral(ClassOutput->TypeName));
+					NewOutput->InitMember(ClassOutput->TypeName, DefaultLiteral, InNodeID, MoveTemp(ClassName));
+					Outputs.Add(NewOutput);
+				}
+
+				return NewOutput;
+			}
+		}
+	}
+
+	return nullptr;
+}
+
 UMetasoundEditorGraphOutput* UMetasoundEditorGraph::FindOrAddOutput(Metasound::Frontend::FConstNodeHandle InNodeHandle)
 {
 	using namespace Metasound::Editor;
diff --git a/Engine/Plugins/Runtime/Metasound/Source/MetasoundEditor/Private/MetasoundEditorGraphBuilder.cpp b/Engine/Plugins/Runtime/Metasound/Source/MetasoundEditor/Private/MetasoundEditorGraphBuilder.cpp
--- a/Engine/Plugins/Runtime/Metasound/Source/MetasoundEditor/Private/MetasoundEditorGraphBuilder.cpp
+++ b/Engine/Plugins/Runtime/Metasound/Source/MetasoundEditor/Private/MetasoundEditorGraphBuilder.cpp
@@ -91,7 +91,7 @@ namespace Metasound
 				return nullptr;
 			}
 
-			FName GenerateUniqueName(const TArray<FName>& InExistingNames, const FString& InBaseName)
+			FName GenerateUniqueName(const TSet<FName>& InExistingNames, const FString& InBaseName)
 			{
 				int32 PostFixInt = 0;
 				FString NewName = InBaseName;
@@ -555,51 +555,73 @@ namespace Metasound
 			return nullptr;
 		}
 
-		UMetasoundEditorGraphInputNode* FGraphBuilder::AddInputNode(UObject& InMetaSound, const Frontend::FConstNodeHandle& InNodeHandle, bool bInSelectNewNode)
+		UMetasoundEditorGraphInputNode* FGraphBuilder::AddInputNode(UObject& InMetaSound, const FGuid& InTemplateNodeID, bool bInSelectNewNode)
 		{
 			using namespace Frontend;
 
-			const FMetasoundFrontendClassMetadata& InputMetadata = InNodeHandle->GetClassMetadata();
-			if (ensureMsgf(InputMetadata.GetType() == EMetasoundFrontendClassType::Template, TEXT("Cannot call 'AddInputNode' with node of class type '%s': Must be input template."), LexToString(InputMetadata.GetType())))
+			using namespace Engine;
+			using namespace Frontend;
+
+			FMetaSoundFrontendDocumentBuilder& Builder = FDocumentBuilderRegistry::GetChecked().FindOrBeginBuilding(&InMetaSound);
+			if (const FMetasoundFrontendNode* TemplateNode = Builder.FindNode(InTemplateNodeID))
 			{
-				if (ensureMsgf(InputMetadata.GetClassName() == FInputNodeTemplate::ClassName, TEXT("Cannot call 'AddInputNode with node handle that is not of class '%s'"), *InputMetadata.GetClassName().ToString()))
+				if (const FMetasoundFrontendClass* Class = Builder.FindDependency(TemplateNode->ClassID))
 				{
-					auto InitNodeFunc = [&InNodeHandle](UMetasoundEditorGraph& MetasoundGraph, UMetasoundEditorGraphInputNode& NewGraphNode)
+					const EMetasoundFrontendClassType& ClassType = Class->Metadata.GetType();
+					if (ensureMsgf(ClassType == EMetasoundFrontendClassType::Template, TEXT("Cannot call 'AddInputNode' with node of class type '%s': Must be input template."), LexToString(ClassType)))
 					{
-						FConstInputHandle TemplateInput = InNodeHandle->GetConstInputs().Last();
-						FConstOutputHandle InputNodeOutput = TemplateInput->GetConnectedOutput();
-						UMetasoundEditorGraphInput* Input = MetasoundGraph.FindOrAddInput(InputNodeOutput->GetOwningNode());
-						if (ensure(Input))
+						if (ensureMsgf(Class->Metadata.GetClassName() == FInputNodeTemplate::ClassName, TEXT("Cannot call 'AddInputNode with node handle that is not of class '%s'"), *FInputNodeTemplate::ClassName.ToString()))
 						{
-							NewGraphNode.Input = Input;
-							NewGraphNode.NodeID = InNodeHandle->GetID();
-						}
-					};
+							const FGuid& TemplateNodeInputVertexID = TemplateNode->Interface.Inputs.Last().VertexID;
+							auto InitNodeFunc = [&Builder, &TemplateNodeInputVertexID, &InTemplateNodeID](UMetasoundEditorGraph& MetasoundGraph, UMetasoundEditorGraphInputNode& NewGraphNode)
+							{
+								const FMetasoundFrontendNode* ConnectedInputNode = nullptr;
+								Builder.FindNodeOutputConnectedToNodeInput(InTemplateNodeID, TemplateNodeInputVertexID, &ConnectedInputNode);
+								if (ensureMsgf(ConnectedInputNode, TEXT("Failed to find required input connected to template node")))
+								{
+									UMetasoundEditorGraphInput* Input = MetasoundGraph.FindOrAddInput(ConnectedInputNode->GetID());
+									if (ensure(Input))
+									{
+										NewGraphNode.Input = Input;
+										NewGraphNode.NodeID = InTemplateNodeID;
+									}
+								}
+							};
 
-					return GraphBuilderPrivate::AddNode<UMetasoundEditorGraphInputNode>(InMetaSound, InitNodeFunc, bInSelectNewNode);
+							return GraphBuilderPrivate::AddNode<UMetasoundEditorGraphInputNode>(InMetaSound, InitNodeFunc, bInSelectNewNode);
+						}
+					}
 				}
 			}
 
 			return nullptr;
 		}
 
-		UMetasoundEditorGraphOutputNode* FGraphBuilder::AddOutputNode(UObject& InMetaSound, const Frontend::FConstNodeHandle& InNodeHandle, bool bInSelectNewNode)
+		UMetasoundEditorGraphOutputNode* FGraphBuilder::AddOutputNode(UObject& InMetaSound, const FGuid& InNodeID, bool bInSelectNewNode)
 		{
+			using namespace Engine;
 			using namespace Frontend;
 
-			const EMetasoundFrontendClassType ClassType = InNodeHandle->GetClassMetadata().GetType();
-			if (ensureMsgf(ClassType == EMetasoundFrontendClassType::Output, TEXT("Cannot call 'AddOutputNode' with node of class type '%s'"), LexToString(ClassType)))
+			FMetaSoundFrontendDocumentBuilder& Builder = FDocumentBuilderRegistry::GetChecked().FindOrBeginBuilding(&InMetaSound);
+			if (const FMetasoundFrontendNode* Node = Builder.FindNode(InNodeID))
 			{
-				auto InitNodeFunc = [&InNodeHandle](UMetasoundEditorGraph& MetasoundGraph, UMetasoundEditorGraphOutputNode& NewGraphNode)
+				if (const FMetasoundFrontendClass* Class = Builder.FindDependency(Node->ClassID))
 				{
-					UMetasoundEditorGraphOutput* Output = MetasoundGraph.FindOrAddOutput(InNodeHandle);
-					if (ensure(Output))
+					const EMetasoundFrontendClassType& ClassType = Class->Metadata.GetType();
+					if (ensureMsgf(ClassType == EMetasoundFrontendClassType::Output, TEXT("Cannot call 'AddOutputNode' with node of class type '%s'"), LexToString(ClassType)))
 					{
-						NewGraphNode.Output = Output;
-					}
-				};
+						auto InitNodeFunc = [&InNodeID](UMetasoundEditorGraph& MetasoundGraph, UMetasoundEditorGraphOutputNode& NewGraphNode)
+						{
+							UMetasoundEditorGraphOutput* Output = MetasoundGraph.FindOrAddOutput(InNodeID);
+							if (ensure(Output))
+							{
+								NewGraphNode.Output = Output;
+							}
+						};
 
-				return GraphBuilderPrivate::AddNode<UMetasoundEditorGraphOutputNode>(InMetaSound, InitNodeFunc, bInSelectNewNode);
+						return GraphBuilderPrivate::AddNode<UMetasoundEditorGraphOutputNode>(InMetaSound, InitNodeFunc, bInSelectNewNode);
+					}
+				}
 			}
 
 			return nullptr;
@@ -647,17 +669,14 @@ namespace Metasound
 			return Categories;
 		}
 
-		FName FGraphBuilder::GenerateUniqueNameByClassType(const UObject& InMetaSound, EMetasoundFrontendClassType InClassType, const FString& InBaseName)
+		FName FGraphBuilder::GenerateUniqueNameByClassType(UObject& InMetaSound, EMetasoundFrontendClassType InClassType, const FString& InBaseName)
 		{
-			const FMetasoundAssetBase* MetaSoundAsset = IMetasoundUObjectRegistry::Get().GetObjectAsAssetBase(&InMetaSound);
-			check(MetaSoundAsset);
+			using namespace Engine;
 
-			// Get existing names.
-			TArray<FName> ExistingNames;
-			MetaSoundAsset->GetRootGraphHandle()->IterateConstNodes([&](const Frontend::FConstNodeHandle& Node)
-			{
-				ExistingNames.Add(Node->GetNodeName());
-			}, InClassType);
+			TSet<FName> ExistingNames;
+			FMetaSoundFrontendDocumentBuilder& Builder = FDocumentBuilderRegistry::GetChecked().FindOrBeginBuilding(&InMetaSound);
+			auto GetName = [&](const FMetasoundFrontendClass&, const FMetasoundFrontendNode& Node) { ExistingNames.Add(Node.Name); };
+			Builder.IterateNodesByClassType(GetName, InClassType);
 
 			return GraphBuilderPrivate::GenerateUniqueName(ExistingNames, InBaseName);
 		}
@@ -1177,45 +1196,6 @@ namespace Metasound
 			}
 		}
 
-		Frontend::FNodeHandle FGraphBuilder::AddInputTemplateNodeHandle(UObject& InMetaSound, Frontend::FNodeHandle& InputNodeHandle)
-		{
-			using namespace Frontend;
-
-			FMetasoundAssetBase* MetaSoundAsset = IMetasoundUObjectRegistry::Get().GetObjectAsAssetBase(&InMetaSound);
-			check(MetaSoundAsset);
-
-			UMetasoundEditorGraph* Graph = Cast<UMetasoundEditorGraph>(MetaSoundAsset->GetGraph());
-			if (!ensure(Graph))
-			{
-				return INodeController::GetInvalidHandle();
-			}
-
-			TArray<FConstOutputHandle> NodeOutputs = InputNodeHandle->GetConstOutputs();
-			if (!ensure(!NodeOutputs.IsEmpty()))
-			{
-				return INodeController::GetInvalidHandle();
-			}
-
-			if (!ensure(InputNodeHandle->GetClassMetadata().GetType() == EMetasoundFrontendClassType::Input))
-			{
-				return INodeController::GetInvalidHandle();
-			}
-
-			const INodeTemplate* InputTemplate = INodeTemplateRegistry::Get().FindTemplate(FInputNodeTemplate::ClassName);
-			check(InputTemplate);
-
-			const FName DataType = NodeOutputs.Last()->GetDataType();
-			FNodeTemplateGenerateInterfaceParams Params { { DataType }, { } };
-
-			FGraphHandle GraphHandle = Graph->GetGraphHandle();
-			FNodeHandle TemplateNodeHandle = GraphHandle->AddTemplateNode(*InputTemplate, MoveTemp(Params));
-
-			TemplateNodeHandle->GetInputs().Last()->Connect(*InputNodeHandle->GetOutputs().Last());
-			FMetasoundFrontendDocumentModifyContext& ModifyContext = GetOutermostMetaSoundChecked(*Graph).GetModifyContext();
-			ModifyContext.AddNodeIDsModified({ TemplateNodeHandle->GetID() });
-			return TemplateNodeHandle;
-		}
-
 		bool FGraphBuilder::GetPinLiteral(UEdGraphPin& InInputPin, FMetasoundFrontendLiteral& OutDefaultLiteral)
 		{
 			using namespace Frontend;
@@ -1375,16 +1355,13 @@ namespace Metasound
 			return true;
 		}
 
-		Frontend::FNodeHandle FGraphBuilder::AddInputNodeHandle(UObject& InMetaSound, const FCreateNodeVertexParams& InParams, const FMetasoundFrontendLiteral* InDefaultValue, const FName* InNameBase)
+		FMetasoundFrontendClassInput FGraphBuilder::CreateUniqueClassInput(UObject& InMetaSound, const FCreateNodeVertexParams& InParams, const FMetasoundFrontendLiteral* InDefaultValue, const FName* InNameBase)
 		{
-			FMetasoundAssetBase* MetaSoundAsset = IMetasoundUObjectRegistry::Get().GetObjectAsAssetBase(&InMetaSound);
-			check(MetaSoundAsset);
-			const FGuid VertexID = FGuid::NewGuid();
-
 			FMetasoundFrontendClassInput ClassInput;
 			ClassInput.Name = GenerateUniqueNameByClassType(InMetaSound, EMetasoundFrontendClassType::Input, InNameBase ? InNameBase->ToString() : TEXT("Input"));
 			ClassInput.TypeName = InParams.DataType;
-			ClassInput.VertexID = VertexID;
+			ClassInput.VertexID = FGuid::NewGuid();
+			ClassInput.NodeID = FGuid::NewGuid();
 
 			// Can be unset if attempting to mirror parameters from a reroute, so default to reference
 			ClassInput.AccessType = InParams.AccessType == EMetasoundFrontendVertexAccessType::Unset ? EMetasoundFrontendVertexAccessType::Reference : InParams.AccessType;
@@ -1399,33 +1376,28 @@ namespace Metasound
 				ClassInput.DefaultLiteral.SetFromLiteral(Literal);
 			}
 
-			return MetaSoundAsset->GetRootGraphHandle()->AddInputVertex(ClassInput);
+			return ClassInput;
 		}
 
-		Frontend::FNodeHandle FGraphBuilder::AddOutputNodeHandle(UObject& InMetaSound, const FCreateNodeVertexParams& InParams, const FName* InNameBase)
+		FMetasoundFrontendClassOutput FGraphBuilder::CreateUniqueClassOutput(UObject& InMetaSound, const FCreateNodeVertexParams& InParams, const FName* InNameBase)
 		{
-			FMetasoundAssetBase* MetaSoundAsset = IMetasoundUObjectRegistry::Get().GetObjectAsAssetBase(&InMetaSound);
-			check(MetaSoundAsset);
-			const FGuid VertexID = FGuid::NewGuid();
-
-			const FName NewName = GenerateUniqueNameByClassType(InMetaSound, EMetasoundFrontendClassType::Output, InNameBase ? InNameBase->ToString() : TEXT("Output"));
-
 			FMetasoundFrontendClassOutput ClassOutput;
-			ClassOutput.Name = NewName;
+			ClassOutput.Name = GenerateUniqueNameByClassType(InMetaSound, EMetasoundFrontendClassType::Output, InNameBase ? InNameBase->ToString() : TEXT("Output"));
 			ClassOutput.TypeName = InParams.DataType;
-			ClassOutput.VertexID = VertexID;
+			ClassOutput.VertexID = FGuid::NewGuid();
+			ClassOutput.NodeID = FGuid::NewGuid();
 
 			// Can be unset if attempting to mirror parameters from a reroute, so default to reference
 			ClassOutput.AccessType = InParams.AccessType == EMetasoundFrontendVertexAccessType::Unset ? EMetasoundFrontendVertexAccessType::Reference : InParams.AccessType;
 
-			return MetaSoundAsset->GetRootGraphHandle()->AddOutputVertex(ClassOutput);
+			return ClassOutput;
 		}
 
 		FName FGraphBuilder::GenerateUniqueVariableName(const Frontend::FConstGraphHandle& InFrontendGraph, const FString& InBaseName)
 		{
 			using namespace Frontend;
 
-			TArray<FName> ExistingVariableNames;
+			TSet<FName> ExistingVariableNames;
 
 			// Get all the names from the existing variables on the graph
 			// and place into the ExistingVariableNames array.
@@ -2369,7 +2341,7 @@ namespace Metasound
 						{
 							if (Node->GetClassMetadata().GetClassName() == FInputNodeTemplate::ClassName)
 							{
-								NewGraphNode = CastChecked<UMetasoundEditorGraphNode>(AddInputNode(MetaSoundObject, Node, false));
+								NewGraphNode = CastChecked<UMetasoundEditorGraphNode>(AddInputNode(MetaSoundObject, Node->GetID(), false));
 							}
 							else
 							{
@@ -2380,7 +2352,7 @@ namespace Metasound
 
 						case EMetasoundFrontendClassType::Output:
 						{
-							NewGraphNode = CastChecked<UMetasoundEditorGraphNode>(AddOutputNode(MetaSoundObject, Node, false));
+							NewGraphNode = CastChecked<UMetasoundEditorGraphNode>(AddOutputNode(MetaSoundObject, Node->GetID(), false));
 						}
 						break;
 
diff --git a/Engine/Plugins/Runtime/Metasound/Source/MetasoundEditor/Private/MetasoundEditorGraphBuilder.h b/Engine/Plugins/Runtime/Metasound/Source/MetasoundEditor/Private/MetasoundEditorGraphBuilder.h
--- a/Engine/Plugins/Runtime/Metasound/Source/MetasoundEditor/Private/MetasoundEditorGraphBuilder.h
+++ b/Engine/Plugins/Runtime/Metasound/Source/MetasoundEditor/Private/MetasoundEditorGraphBuilder.h
@@ -121,16 +121,6 @@ namespace Metasound
 			// Adds a new EdGraph comment node associated with the given MetaSoundFrontendGraph comment ID
 			static UMetasoundEditorGraphCommentNode* CreateCommentNode(UObject& InMetaSound, bool bInSelectNewNode = true, FGuid InCommentID = FGuid::NewGuid());
 
-			// Adds an input template node to the given MetaSound object's document model and returns an associated node handle
-			static Frontend::FNodeHandle AddInputTemplateNodeHandle(UObject& InMetaSound, Frontend::FNodeHandle& InputNodeHandle);
-
-			// Generates FNodeHandle for the given external node data. Does not bind or create EdGraph representation of given node.
-			static Frontend::FNodeHandle AddInputNodeHandle(
-				UObject& InMetaSound,
-				const FCreateNodeVertexParams& InVertexParams,
-				const FMetasoundFrontendLiteral* InDefaultValue = nullptr,
-				const FName* InNameBase = nullptr);
-
 			// Adds a corresponding UMetasoundEditorGraphExternalNode for the provided node handle.
 			static UMetasoundEditorGraphExternalNode* AddExternalNode(UObject& InMetaSound, const FGuid& InNodeID, const FMetasoundFrontendClassMetadata& InMetadata, bool bInSelectNewNode = true);
 
@@ -146,14 +136,11 @@ namespace Metasound
 			// Adds an variable node with the given node handle to the editor graph.
 			static UMetasoundEditorGraphVariableNode* AddVariableNode(UObject& InMetaSound, const Frontend::FConstNodeHandle& InNodeHandle, bool bInSelectNewNode = true);
 
-			// Adds an input node to the editor graph that corresponds to the provided input template node handle.
-			static UMetasoundEditorGraphInputNode* AddInputNode(UObject& InMetaSound, const Frontend::FConstNodeHandle& InInputTemplateNode, bool bInSelectNewNode = true);
+			// Adds an input node to the editor graph that corresponds to the provided input template node in the document with the given ID.
+			static UMetasoundEditorGraphInputNode* AddInputNode(UObject& InMetaSound, const FGuid& InInputTemplateNodeID, bool bInSelectNewNode = true);
 
-			// Adds an output node to the editor graph that corresponds to the provided node handle.
-			static UMetasoundEditorGraphOutputNode* AddOutputNode(UObject& InMetaSound, const Frontend::FConstNodeHandle& InNodeHandle, bool bInSelectNewNode = true);
-
-			// Generates analogous FNodeHandle for the given internal node data. Does not bind nor create EdGraph representation of given node.
-			static Frontend::FNodeHandle AddOutputNodeHandle(UObject& InMetaSound, const FCreateNodeVertexParams& InParams, const FName* InNameBase = nullptr);
+			// Adds an output node to the editor graph that corresponds to the provided output node ID.
+			static UMetasoundEditorGraphOutputNode* AddOutputNode(UObject& InMetaSound, const FGuid& InNodeID, bool bInSelectNewNode = true);
 
 			// Create a unique name for the variable.
 			static FName GenerateUniqueVariableName(const Frontend::FConstGraphHandle& InFrontendGraph, const FString& InBaseName);
@@ -189,13 +176,26 @@ namespace Metasound
 			// the Editor Graph representation of the pins.
 			static bool ConnectNodes(UEdGraphPin& InInputPin, UEdGraphPin& InOutputPin, bool bInConnectEdPins);
 
+			// Creates a unique class input with the given default data.
+			static FMetasoundFrontendClassInput CreateUniqueClassInput(
+				UObject& InMetaSound,
+				const FCreateNodeVertexParams& InVertexParams,
+				const FMetasoundFrontendLiteral* InDefaultValue = nullptr,
+				const FName* InNameBase = nullptr);
+
+			// Creates a unique class output with the given default data. Output is not assigned a NodeID.
+			static FMetasoundFrontendClassOutput CreateUniqueClassOutput(
+				UObject& InMetaSound,
+				const FCreateNodeVertexParams& InVertexParams,
+				const FName* InNameBase = nullptr);
+
 			// Disconnects pin's associated frontend vertex from any linked input
 			// or output nodes, and reflects change in the Frontend graph. Does *not*
 			// disconnect the EdGraph pins.
 			static void DisconnectPinVertex(UEdGraphPin& InPin, bool bAddLiteralInputs = true);
 
 			// Generates a unique output name for the given MetaSound object
-			static FName GenerateUniqueNameByClassType(const UObject& InMetaSound, EMetasoundFrontendClassType InClassType, const FString& InBaseName);
+			static FName GenerateUniqueNameByClassType(UObject& InMetaSound, EMetasoundFrontendClassType InClassType, const FString& InBaseName);
 
 			static TArray<FString> GetDataTypeNameCategories(const FName& InDataTypeName);
 
diff --git a/Engine/Plugins/Runtime/Metasound/Source/MetasoundEditor/Private/MetasoundEditorGraphSchema.cpp b/Engine/Plugins/Runtime/Metasound/Source/MetasoundEditor/Private/MetasoundEditorGraphSchema.cpp
--- a/Engine/Plugins/Runtime/Metasound/Source/MetasoundEditor/Private/MetasoundEditorGraphSchema.cpp
+++ b/Engine/Plugins/Runtime/Metasound/Source/MetasoundEditor/Private/MetasoundEditorGraphSchema.cpp
@@ -45,6 +45,7 @@
 #include "MetasoundVariableNodes.h"
 #include "MetasoundVertex.h"
 #include "NodeTemplates/MetasoundFrontendNodeTemplateAudioAnalyzer.h"
+#include "NodeTemplates/MetasoundFrontendNodeTemplateInput.h"
 #include "NodeTemplates/MetasoundFrontendNodeTemplateReroute.h"
 #include "ScopedTransaction.h"
 #include "Settings/EditorStyleSettings.h"
@@ -497,9 +498,9 @@ const FLinearColor& FMetasoundGraphSchemaAction_NewInput::GetIconColor() const
 UEdGraphNode* FMetasoundGraphSchemaAction_NewInput::PerformAction(UEdGraph* ParentGraph, UEdGraphPin* FromPin, const FVector2D InLocation, bool bSelectNewNode /* = true */)
 {
 	using namespace Metasound::Editor;
+	using namespace Metasound::Engine;
 	using namespace Metasound::Frontend;
 
-
 	UMetasoundEditorGraph* MetasoundGraph = CastChecked<UMetasoundEditorGraph>(ParentGraph);
 	UObject& ParentMetasound = MetasoundGraph->GetMetasoundChecked();
 
@@ -520,8 +521,9 @@ UEdGraphNode* FMetasoundGraphSchemaAction_NewInput::PerformAction(UEdGraph* Pare
 	MetasoundGraph->Modify();
 	Input->Modify();
 
-	FConstNodeHandle TemplateNodeHandle = FGraphBuilder::AddInputTemplateNodeHandle(ParentMetasound, InputNodeHandle);
-	if (UMetasoundEditorGraphNode* NewGraphNode = FGraphBuilder::AddInputNode(ParentMetasound, TemplateNodeHandle))
+	FMetaSoundFrontendDocumentBuilder& Builder = FDocumentBuilderRegistry::GetChecked().FindOrBeginBuilding(&ParentMetasound);
+	const FMetasoundFrontendNode* TemplateNode = FInputNodeTemplate::CreateNode(Builder, Input->GetMemberName());
+	if (UMetasoundEditorGraphNode* NewGraphNode = FGraphBuilder::AddInputNode(ParentMetasound, TemplateNode->GetID()))
 	{
 		NewGraphNode->Modify();
 		NewGraphNode->UpdateFrontendNodeLocation(InLocation);
@@ -546,6 +548,7 @@ FMetasoundGraphSchemaAction_PromoteToInput::FMetasoundGraphSchemaAction_PromoteT
 UEdGraphNode* FMetasoundGraphSchemaAction_PromoteToInput::PerformAction(UEdGraph* ParentGraph, UEdGraphPin* FromPin, const FVector2D InLocation, bool bSelectNewNode /* = true */)
 {
 	using namespace Metasound::Editor;
+	using namespace Metasound::Engine;
 	using namespace Metasound::Frontend;
 
 	FConstInputHandle InputHandle = FGraphBuilder::GetConstInputHandleFromPin(FromPin);
@@ -567,24 +570,27 @@ UEdGraphNode* FMetasoundGraphSchemaAction_PromoteToInput::PerformAction(UEdGraph
 	// The promoted input must have the same vertex access type in order for it to be connectible
 	const FCreateNodeVertexParams VertexParams = { InputHandle->GetDataType(), InputHandle->GetVertexAccessType() };
 
-	FNodeHandle InputNodeHandle = FGraphBuilder::AddInputNodeHandle(ParentMetasound, VertexParams, &DefaultValue, &InputName);
-	if (ensure(InputNodeHandle->IsValid()))
+	FMetaSoundFrontendDocumentBuilder& Builder = FDocumentBuilderRegistry::GetChecked().FindOrBeginBuilding(&ParentMetasound);
+	FMetasoundFrontendClassInput ClassInput = FGraphBuilder::CreateUniqueClassInput(ParentMetasound, VertexParams);
+	if (const FMetasoundFrontendNode* NewNode = Builder.AddGraphInput(ClassInput))
 	{
-		UMetasoundEditorGraphInput* Input = MetasoundGraph->FindOrAddInput(InputNodeHandle);
+		UMetasoundEditorGraphInput* Input = MetasoundGraph->FindOrAddInput(NewNode->GetID());
 		if (ensure(Input))
 		{
-			FConstNodeHandle TemplateNodeHandle = FGraphBuilder::AddInputTemplateNodeHandle(ParentMetasound, InputNodeHandle);
-			if (UMetasoundEditorGraphNode* NewGraphNode = FGraphBuilder::AddInputNode(ParentMetasound, TemplateNodeHandle))
+			if (const FMetasoundFrontendNode* NewTemplateNode = FInputNodeTemplate::CreateNode(Builder, ClassInput.Name))
 			{
-				NewGraphNode->UpdateFrontendNodeLocation(InLocation);
-				NewGraphNode->SyncLocationFromFrontendNode();
-				UEdGraphNode* EdGraphNode = CastChecked<UEdGraphNode>(NewGraphNode);
-
-				if (ensure(SchemaPrivate::TryConnectNewNodeToMatchingDataTypePin(*EdGraphNode, FromPin)))
+				if (UMetasoundEditorGraphNode* NewGraphNode = FGraphBuilder::AddInputNode(ParentMetasound, NewTemplateNode->GetID()))
 				{
-					FGraphBuilder::RegisterGraphWithFrontend(ParentMetasound);
-					SchemaPrivate::SelecteNodeInEditor(*MetasoundGraph, *NewGraphNode);
-					return EdGraphNode;
+					NewGraphNode->UpdateFrontendNodeLocation(InLocation);
+					NewGraphNode->SyncLocationFromFrontendNode();
+					UEdGraphNode* EdGraphNode = CastChecked<UEdGraphNode>(NewGraphNode);
+
+					if (ensure(SchemaPrivate::TryConnectNewNodeToMatchingDataTypePin(*EdGraphNode, FromPin)))
+					{
+						FGraphBuilder::RegisterGraphWithFrontend(ParentMetasound);
+						SchemaPrivate::SelecteNodeInEditor(*MetasoundGraph, *NewGraphNode);
+						return EdGraphNode;
+					}
 				}
 			}
 		}
@@ -724,17 +730,11 @@ UEdGraphNode* FMetasoundGraphSchemaAction_NewOutput::PerformAction(UEdGraph* Par
 		return nullptr;
 	}
 
-	FConstNodeHandle NodeHandle = Output->GetConstNodeHandle();
-	if (!ensure(NodeHandle->IsValid()))
-	{
-		return nullptr;
-	}
-
 	const FScopedTransaction Transaction(LOCTEXT("AddNewOutputNode2", "Add New MetaSound Output Node"));
 	ParentMetasound.Modify();
 	ParentGraph->Modify();
 
-	if (UMetasoundEditorGraphOutputNode* NewGraphNode = FGraphBuilder::AddOutputNode(ParentMetasound, NodeHandle, bSelectNewNode))
+	if (UMetasoundEditorGraphOutputNode* NewGraphNode = FGraphBuilder::AddOutputNode(ParentMetasound, Output->NodeID, bSelectNewNode))
 	{
 		NewGraphNode->UpdateFrontendNodeLocation(Location);
 		NewGraphNode->SyncLocationFromFrontendNode();
@@ -777,13 +777,16 @@ UEdGraphNode* FMetasoundGraphSchemaAction_PromoteToOutput::PerformAction(UEdGrap
 	const FVertexName NewNodeName = FGraphBuilder::GenerateUniqueNameByClassType(ParentMetasound, EMetasoundFrontendClassType::Output, OutputName);
 	const FCreateNodeVertexParams VertexParams = { OutputHandle->GetDataType(), OutputHandle->GetVertexAccessType() };
 
-	FConstNodeHandle NodeHandle = FGraphBuilder::AddOutputNodeHandle(ParentMetasound, VertexParams, &NewNodeName);
-	if (ensure(NodeHandle->IsValid()))
+	const FMetasoundFrontendClassOutput ClassOutput = FGraphBuilder::CreateUniqueClassOutput(ParentMetasound, VertexParams, &NewNodeName);
+	FMetaSoundFrontendDocumentBuilder& Builder = IDocumentBuilderRegistry::GetChecked().FindOrBeginBuilding(&ParentMetasound);
+	const FMetasoundFrontendNode* OutputNode = Builder.AddGraphOutput(ClassOutput);
+	if (ensure(OutputNode))
 	{
-		UMetasoundEditorGraphOutput* Output = MetasoundGraph->FindOrAddOutput(NodeHandle);
+		const FGuid OutputNodeID = OutputNode->GetID();
+		UMetasoundEditorGraphOutput* Output = MetasoundGraph->FindOrAddOutput(OutputNodeID);
 		if (ensure(Output))
 		{
-			if (UMetasoundEditorGraphOutputNode* NewGraphNode = FGraphBuilder::AddOutputNode(ParentMetasound, NodeHandle))
+			if (UMetasoundEditorGraphOutputNode* NewGraphNode = FGraphBuilder::AddOutputNode(ParentMetasound, OutputNodeID))
 			{ 
 				NewGraphNode->UpdateFrontendNodeLocation(InLocation);
 				NewGraphNode->SyncLocationFromFrontendNode();
diff --git a/Engine/Plugins/Runtime/Metasound/Source/MetasoundEditor/Public/MetasoundEditorGraph.h b/Engine/Plugins/Runtime/Metasound/Source/MetasoundEditor/Public/MetasoundEditorGraph.h
--- a/Engine/Plugins/Runtime/Metasound/Source/MetasoundEditor/Public/MetasoundEditorGraph.h
+++ b/Engine/Plugins/Runtime/Metasound/Source/MetasoundEditor/Public/MetasoundEditorGraph.h
@@ -166,6 +166,7 @@ public:
 	void InitializeLiteral();
 
 protected:
+	UE_DEPRECATED(5.5, "No longer supported in favor of internal Editor::FGraphBuilder API which is actively being transitioned to using Document Builder API")
 	virtual UMetasoundEditorGraphNode* AddNode(Metasound::Frontend::FNodeHandle InNodeHandle, bool bInSelectNewNode) PURE_VIRTUAL(UMetasoundEditorGraphMember::AddNode, return nullptr;);
 
 	/** Default literal value of member */
@@ -187,7 +188,7 @@ class METASOUNDEDITOR_API UMetasoundEditorGraphVertex : public UMetasoundEditorG
 	GENERATED_BODY()
 
 protected:
-	/** Adds the node handle for a newly created vertex. */
+	UE_DEPRECATED(5.5, "EditorGraph vertices no longer generates node handles (use Builder API instead)")
 	virtual Metasound::Frontend::FNodeHandle AddNodeHandle(const FName& InNodeName, const Metasound::Editor::FCreateNodeVertexParams& InParams) PURE_VIRTUAL(UMetasoundEditorGraphVertex::AddNodeHandle, return Metasound::Frontend::INodeController::GetInvalidHandle(); )
 
 	const FMetasoundFrontendNode* GetFrontendNode() const;
@@ -274,7 +275,10 @@ public:
 	virtual EMetasoundFrontendVertexAccessType GetVertexAccessType() const override;
 
 protected:
+	UE_DEPRECATED(5.5, "No longer supported in favor of internal Editor::FGraphBuilder API which is actively being transitioned to using Document Builder API")
 	virtual UMetasoundEditorGraphNode* AddNode(Metasound::Frontend::FNodeHandle InNodeHandle, bool bInSelectNewNode) override;
+
+	UE_DEPRECATED(5.5, "EditorGraph vertices no longer generates node handles (use Builder API instead)")
 	virtual Metasound::Frontend::FNodeHandle AddNodeHandle(const FName& InNodeName, const Metasound::Editor::FCreateNodeVertexParams& InParams) override;
 
 	virtual EMetasoundFrontendClassType GetClassType() const override { return EMetasoundFrontendClassType::Input; }
@@ -300,8 +304,12 @@ public:
 	virtual EMetasoundFrontendVertexAccessType GetVertexAccessType() const override;
 
 protected:
+	UE_DEPRECATED(5.5, "No longer supported in favor of internal Editor::FGraphBuilder API which is actively being transitioned to using Document Builder API")
 	virtual UMetasoundEditorGraphNode* AddNode(Metasound::Frontend::FNodeHandle InNodeHandle, bool bInSelectNewNode) override;
+
+	UE_DEPRECATED(5.5, "No longer supported in favor of internal Editor::FGraphBuilder API which is actively being transitioned to using Document Builder API")
 	virtual Metasound::Frontend::FNodeHandle AddNodeHandle(const FName& InNodeName, const Metasound::Editor::FCreateNodeVertexParams& InParams) override;
+
 	virtual EMetasoundFrontendClassType GetClassType() const override { return EMetasoundFrontendClassType::Output; }
 	virtual Metasound::Editor::ENodeSection GetSectionID() const override;
 };
@@ -348,6 +356,7 @@ public:
 	virtual bool CanRename() const override;
 
 protected:
+	UE_DEPRECATED(5.5, "No longer supported in favor of internal Editor::FGraphBuilder API which is actively being transitioned to using Document Builder API")
 	virtual UMetasoundEditorGraphNode* AddNode(Metasound::Frontend::FNodeHandle InNodeHandle, bool bInSelectNewNode) override;
 
 private:
@@ -433,10 +442,12 @@ private:
 public:
 	UMetasoundEditorGraphInput* FindInput(FGuid InNodeID) const;
 	UMetasoundEditorGraphInput* FindInput(FName InName) const;
+	UMetasoundEditorGraphInput* FindOrAddInput(const FGuid& InNodeID);
 	UMetasoundEditorGraphInput* FindOrAddInput(Metasound::Frontend::FConstNodeHandle InNodeHandle);
 
 	UMetasoundEditorGraphOutput* FindOutput(FGuid InNodeID) const;
 	UMetasoundEditorGraphOutput* FindOutput(FName InName) const;
+	UMetasoundEditorGraphOutput* FindOrAddOutput(const FGuid& InNodeID);
 	UMetasoundEditorGraphOutput* FindOrAddOutput(Metasound::Frontend::FConstNodeHandle InNodeHandle);
 
 	UMetasoundEditorGraphVariable* FindVariable(const FGuid& InVariableID) const;
diff --git a/Engine/Plugins/Runtime/Metasound/Source/MetasoundFrontend/Private/MetasoundFrontendDocumentBuilder.cpp b/Engine/Plugins/Runtime/Metasound/Source/MetasoundFrontend/Private/MetasoundFrontendDocumentBuilder.cpp
--- a/Engine/Plugins/Runtime/Metasound/Source/MetasoundFrontend/Private/MetasoundFrontendDocumentBuilder.cpp
+++ b/Engine/Plugins/Runtime/Metasound/Source/MetasoundFrontend/Private/MetasoundFrontendDocumentBuilder.cpp
@@ -2526,6 +2526,27 @@ void FMetaSoundFrontendDocumentBuilder::IterateNodesConnectedWithVertex(const FM
 	}
 }
 
+void FMetaSoundFrontendDocumentBuilder::IterateNodesByClassType(Metasound::Frontend::FConstClassAndNodeFunctionRef Func, EMetasoundFrontendClassType ClassType, const FGuid* InPageID) const
+{
+	using namespace Metasound::Frontend;
+
+	check(ClassType != EMetasoundFrontendClassType::Invalid);
+
+	const FGuid& PageID = InPageID ? *InPageID : BuildPageID;
+	const FMetasoundFrontendDocument& Doc = GetConstDocumentChecked();
+	const FMetasoundFrontendGraph& Graph = Doc.RootGraph.FindConstGraphChecked(PageID);
+	for (const FMetasoundFrontendNode& Node : Graph.Nodes)
+	{
+		if (const FMetasoundFrontendClass* Class = FindDependency(Node.ClassID))
+		{
+			if (Class->Metadata.GetType() == ClassType)
+			{
+				Func(*Class, Node);
+			}
+		}
+	}
+}
+
 bool FMetaSoundFrontendDocumentBuilder::ModifyInterfaces(Metasound::Frontend::FModifyInterfaceOptions&& InOptions)
 {
 	using namespace Metasound::Frontend;
diff --git a/Engine/Plugins/Runtime/Metasound/Source/MetasoundFrontend/Private/NodeTemplates/MetasoundFrontendNodeTemplateInput.cpp b/Engine/Plugins/Runtime/Metasound/Source/MetasoundFrontend/Private/NodeTemplates/MetasoundFrontendNodeTemplateInput.cpp
--- a/Engine/Plugins/Runtime/Metasound/Source/MetasoundFrontend/Private/NodeTemplates/MetasoundFrontendNodeTemplateInput.cpp
+++ b/Engine/Plugins/Runtime/Metasound/Source/MetasoundFrontend/Private/NodeTemplates/MetasoundFrontendNodeTemplateInput.cpp
@@ -22,7 +22,7 @@ namespace Metasound::Frontend
 	namespace InputNodeTemplatePrivate
 	{
 		// Creates an input template node, sets node position (should only ever be one in style location) from and connects it to the associated input with the given name.
-		const FMetasoundFrontendNode* InitTemplateNode(const INodeTemplate& InTemplate, FName InputName, FMetaSoundFrontendDocumentBuilder& InOutBuilder)
+		const FMetasoundFrontendNode* InitTemplateNode(const INodeTemplate& InTemplate, FName InputName, FMetaSoundFrontendDocumentBuilder& InOutBuilder, const FGuid* InPageID = nullptr)
 		{
 			FMetasoundFrontendEdge NewEdge;
 			FName TypeName;
@@ -32,7 +32,7 @@ namespace Metasound::Frontend
 
 			// Cache data from Input node pointer as needed as subsequent call to create new template node may invalidate the input node's pointer
 			{
-				const FMetasoundFrontendNode* InputNode = InOutBuilder.FindGraphInputNode(InputName);
+				const FMetasoundFrontendNode* InputNode = InOutBuilder.FindGraphInputNode(InputName, InPageID);
 				check(InputNode);
 				TypeName = InputNode->Interface.Outputs.Last().TypeName;
 				NewEdge.FromNodeID = InputNode->GetID(),
@@ -45,7 +45,7 @@ namespace Metasound::Frontend
 			}
 
 			FNodeTemplateGenerateInterfaceParams Params { { }, { TypeName } };
-			const FMetasoundFrontendNode* NewNode = InOutBuilder.AddNodeByTemplate(InTemplate, MoveTemp(Params));
+			const FMetasoundFrontendNode* NewNode = InOutBuilder.AddNodeByTemplate(InTemplate, MoveTemp(Params), FGuid::NewGuid(), InPageID);
 			check(NewNode);
 			NewEdge.ToNodeID = NewNode->GetID();
 			NewEdge.ToVertexID = NewNode->Interface.Inputs.Last().VertexID;
@@ -53,12 +53,12 @@ namespace Metasound::Frontend
 #if WITH_EDITORONLY_DATA
 			for (const TPair<FGuid, FVector2D>& Pair : Locations)
 			{
-				InOutBuilder.SetNodeLocation(NewEdge.ToNodeID, Pair.Value);
+				InOutBuilder.SetNodeLocation(NewEdge.ToNodeID, Pair.Value, nullptr, InPageID);
 			}
 #endif // WITH_EDITORONLY_DATA
 
 			// Add edge between input node and new template node
-			InOutBuilder.AddEdge(MoveTemp(NewEdge));
+			InOutBuilder.AddEdge(MoveTemp(NewEdge), InPageID);
 
 			return NewNode;
 		};
@@ -68,6 +68,20 @@ namespace Metasound::Frontend
 
 	const FMetasoundFrontendVersionNumber FInputNodeTemplate::VersionNumber = { 1, 0 } ;
 
+#if WITH_EDITOR
+	const FMetasoundFrontendNode* FInputNodeTemplate::CreateNode(FMetaSoundFrontendDocumentBuilder& InOutBuilder, FName InputName, const FGuid* InPageID)
+	{
+		if (const FMetasoundFrontendClassInput* Input = InOutBuilder.FindGraphInput(InputName))
+		{
+			const INodeTemplate* ThisTemplate = INodeTemplateRegistry::Get().FindTemplate(ClassName);
+			check(ThisTemplate);
+			return InputNodeTemplatePrivate::InitTemplateNode(*ThisTemplate, InputName, InOutBuilder, InPageID);
+		}
+
+		return nullptr;
+	}
+#endif // WITH_EDITOR
+
 	const FMetasoundFrontendClassName& FInputNodeTemplate::GetClassName() const
 	{
 		return FInputNodeTemplate::ClassName;
@@ -273,5 +287,6 @@ namespace Metasound::Frontend
 
 		return bInjectedNodes;
 	}
+
 #endif // WITH_EDITOR
 } // namespace Metasound::Frontend
\ No newline at end of file
diff --git a/Engine/Plugins/Runtime/Metasound/Source/MetasoundFrontend/Public/MetasoundFrontendDocumentBuilder.h b/Engine/Plugins/Runtime/Metasound/Source/MetasoundFrontend/Public/MetasoundFrontendDocumentBuilder.h
--- a/Engine/Plugins/Runtime/Metasound/Source/MetasoundFrontend/Public/MetasoundFrontendDocumentBuilder.h
+++ b/Engine/Plugins/Runtime/Metasound/Source/MetasoundFrontend/Public/MetasoundFrontendDocumentBuilder.h
@@ -25,6 +25,7 @@ namespace Metasound::Frontend
 	// Forward Declarations
 	class INodeTemplate;
 
+	using FConstClassAndNodeFunctionRef = TFunctionRef<void(const FMetasoundFrontendClass&, const FMetasoundFrontendNode&)>;
 	using FFinalizeNodeFunctionRef = TFunctionRef<void(FMetasoundFrontendNode&, const Metasound::Frontend::FNodeRegistryKey&)>;
 
 	enum class EInvalidEdgeReason : uint8
@@ -315,6 +316,9 @@ public:
 	// Note that this does not return whether or not the given edge exists, but rather if it could be legally applied to the given edge vertices.
 	Metasound::Frontend::EInvalidEdgeReason IsValidEdge(const FMetasoundFrontendEdge& InEdge, const FGuid* InPageID = nullptr) const;
 
+	// Iterates nodes that are filtered by only subscribing to a class with the given type (asserts if provided invalid class type).
+	void IterateNodesByClassType(Metasound::Frontend::FConstClassAndNodeFunctionRef Func, EMetasoundFrontendClassType ClassType, const FGuid* InPageID = nullptr) const;
+
 	bool ModifyInterfaces(Metasound::Frontend::FModifyInterfaceOptions&& InOptions);
 
 #if WITH_EDITORONLY_DATA
diff --git a/Engine/Plugins/Runtime/Metasound/Source/MetasoundFrontend/Public/NodeTemplates/MetasoundFrontendNodeTemplateInput.h b/Engine/Plugins/Runtime/Metasound/Source/MetasoundFrontend/Public/NodeTemplates/MetasoundFrontendNodeTemplateInput.h
--- a/Engine/Plugins/Runtime/Metasound/Source/MetasoundFrontend/Public/NodeTemplates/MetasoundFrontendNodeTemplateInput.h
+++ b/Engine/Plugins/Runtime/Metasound/Source/MetasoundFrontend/Public/NodeTemplates/MetasoundFrontendNodeTemplateInput.h
@@ -31,6 +31,12 @@ namespace Metasound::Frontend
 
 		virtual ~FInputNodeTemplate() = default;
 
+#if WITH_EDITOR
+		// Adds template node and connects with the input of the provided name on the page with
+		// the provided id (defaults to builder's build page ID if not provided)
+		static const FMetasoundFrontendNode* CreateNode(FMetaSoundFrontendDocumentBuilder& InOutBuilder, FName InputName, const FGuid* InPageID = nullptr);
+#endif // WITH_EDITOR
+
 		virtual const FMetasoundFrontendClassName& GetClassName() const override;
 
 #if WITH_EDITOR
