- Deprecate UMetaSoundQualityHelper & move function to editor-only MetaSoundSettings util function
- Remove config copy pasta
#rnx

diff --git a/Engine/Plugins/Runtime/Metasound/Source/MetasoundEngine/Private/MetasoundSettings.cpp b/Engine/Plugins/Runtime/Metasound/Source/MetasoundEngine/Private/MetasoundSettings.cpp
--- a/Engine/Plugins/Runtime/Metasound/Source/MetasoundEngine/Private/MetasoundSettings.cpp
+++ b/Engine/Plugins/Runtime/Metasound/Source/MetasoundEngine/Private/MetasoundSettings.cpp
@@ -282,7 +282,8 @@ FName UMetaSoundSettings::GetQualitySettingPropertyName()
 }
 #endif // WITH_EDITORONLY_DATA
 
-TArray<FName> UMetaSoundQualityHelper::GetQualityList()
+#if WITH_EDITOR
+TArray<FName> UMetaSoundSettings::GetQualityList()
 {
 	TArray<FName> Names;
 
@@ -298,4 +299,6 @@ TArray<FName> UMetaSoundQualityHelper::GetQualityList()
 
 	return Names;
 }
+#endif // WITH_EDITOR
+
 #undef LOCTEXT_NAMESPACE // MetaSound
diff --git a/Engine/Plugins/Runtime/Metasound/Source/MetasoundEngine/Public/MetasoundSettings.h b/Engine/Plugins/Runtime/Metasound/Source/MetasoundEngine/Public/MetasoundSettings.h
--- a/Engine/Plugins/Runtime/Metasound/Source/MetasoundEngine/Public/MetasoundSettings.h
+++ b/Engine/Plugins/Runtime/Metasound/Source/MetasoundEngine/Public/MetasoundSettings.h
@@ -51,8 +51,8 @@ public:
 	/**
 	* Returns a list of quality settings to present to a combobox
 	* */
-	UFUNCTION()
-	static TArray<FName> GetQualityList();
+	UFUNCTION(meta = (DeprecatedFunction, DeprecationMessage = "Use UMetaSoundSettings::GetQualityList instead"))
+	static TArray<FName> GetQualityList() { return { }; };
 };
 
 USTRUCT()
@@ -193,6 +193,14 @@ public:
 	bool SetTargetPage(FName PageName);
 
 #if WITH_EDITOR
+public:
+	/* Returns a list of quality settings to present to a combobox. Ex:
+	 * UPROPERTY(EditAnywhere, BlueprintReadWrite, meta=(GetOptions="MetasoundEngine.MetaSoundSettings.GetQualityList"))
+	 * FName QualitySetting;
+	*/
+	UFUNCTION()
+	static TArray<FName> GetQualityList();
+
 private:
 	void ConformPageSettingsDefault(bool bNotifyDefaultConformed);
 
diff --git a/Engine/Plugins/Runtime/Metasound/Source/MetasoundEngine/Public/MetasoundSource.h b/Engine/Plugins/Runtime/Metasound/Source/MetasoundEngine/Public/MetasoundSource.h
--- a/Engine/Plugins/Runtime/Metasound/Source/MetasoundEngine/Public/MetasoundSource.h
+++ b/Engine/Plugins/Runtime/Metasound/Source/MetasoundEngine/Public/MetasoundSource.h
@@ -68,7 +68,7 @@ DECLARE_TS_MULTICAST_DELEGATE_TwoParams(FOnGeneratorInstanceDestroyed, uint64, T
 /**
  * This Metasound type can be played as an audio source.
  */
-UCLASS(hidecategories = object, BlueprintType, config = Metasound, defaultconfig)
+UCLASS(hidecategories = object, BlueprintType)
 class METASOUNDENGINE_API UMetaSoundSource : public USoundWaveProcedural, public FMetasoundAssetBase, public IMetaSoundDocumentInterface
 {
 	GENERATED_BODY()
@@ -129,9 +129,8 @@ public:
 	EMetaSoundOutputAudioFormat OutputFormat;
 
 #if WITH_EDITORONLY_DATA
-
 	// The Quality this Metasound will use. These are defined in the MetaSounds project settings.
-	UPROPERTY(config, EditAnywhere, BlueprintReadWrite, meta = (GetOptions="MetasoundEngine.MetaSoundQualityHelper.GetQualityList"), Category = "Metasound")
+	UPROPERTY(EditAnywhere, BlueprintReadWrite, meta = (GetOptions="MetasoundEngine.MetaSoundSettings.GetQualityList"), Category = "Metasound")
 	FName QualitySetting;
 
 	// This a editor only look up for the Quality Setting above. Preventing orphaning of the original name.
@@ -145,7 +144,6 @@ public:
 	// Override the SampleRate for this Sound (overrides Quality). NOTE: A Zero value will have no effect and use either the Quality setting (if set), or the Device Rate
 	UPROPERTY(EditAnywhere, BlueprintReadOnly, AdvancedDisplay, Category = Metasound, meta = (UIMin = 0, UIMax = 96000, DisplayName = "Override Sample Rate (in Hz)"))
 	FPerPlatformInt SampleRateOverride = 0;
-	
 #endif //WITH_EDITOR_DATA
 
 	UPROPERTY(AssetRegistrySearchable)
