Fix MetaSound asset manager log spam on shutdown after validation

#rnx

diff --git a/Engine/Plugins/Runtime/Metasound/Source/MetasoundEngine/Private/MetasoundAssetSubsystem.cpp b/Engine/Plugins/Runtime/Metasound/Source/MetasoundEngine/Private/MetasoundAssetSubsystem.cpp
--- a/Engine/Plugins/Runtime/Metasound/Source/MetasoundEngine/Private/MetasoundAssetSubsystem.cpp
+++ b/Engine/Plugins/Runtime/Metasound/Source/MetasoundEngine/Private/MetasoundAssetSubsystem.cpp
@@ -178,6 +178,7 @@ namespace Metasound::Engine
 		virtual void RemoveAsset(const UObject& InObject) override;
 		virtual void RemoveAsset(const FAssetData& InAssetData) override;
 		virtual void RenameAsset(const FAssetData& InAssetData, const FString& InOldObjectPath) override;
+		virtual void SetLogActiveAssetsOnShutdown(bool bInLogActiveAssetsOnShutdown) override;
 		virtual FMetasoundAssetBase* TryLoadAssetFromKey(const FAssetKey& InKey) const override;
 		virtual bool TryGetAssetIDFromClassName(const FMetasoundFrontendClassName& InClassName, FGuid& OutGuid) const override;
 		virtual bool TryLoadReferencedAssets(const FMetasoundAssetBase& InAssetBase, TArray<FMetasoundAssetBase*>& OutReferencedAssets) const override;
@@ -202,6 +203,8 @@ namespace Metasound::Engine
 		TSet<FName> AutoUpdateDenyListCache;
 		std::atomic<bool> bIsInitialAssetScanComplete = false;
 		TMap<FAssetKey, TArray<FTopLevelAssetPath>> PathMap;
+
+		bool bLogActiveAssetsOnShutdown = true;
 	};
 
 	FMetaSoundAssetManager::~FMetaSoundAssetManager()
@@ -209,15 +212,18 @@ namespace Metasound::Engine
 #if !NO_LOGGING
 		if (!PathMap.IsEmpty())
 		{
-			TSet<FAssetKey> Keys;
-			if (int32 NumKeys = PathMap.GetKeys(Keys); NumKeys > 0)
+			if (bLogActiveAssetsOnShutdown)
 			{
-				UE_LOG(LogMetaSound, Display, TEXT("AssetManager is shutting down with the following %i assets active:"), NumKeys);
-				for (const TPair<FAssetKey, TArray<FTopLevelAssetPath>>& Pair : PathMap)
+				TSet<FAssetKey> Keys;
+				if (int32 NumKeys = PathMap.GetKeys(Keys); NumKeys > 0)
 				{
-					for (const FTopLevelAssetPath& Path : Pair.Value)
+					UE_LOG(LogMetaSound, Display, TEXT("AssetManager is shutting down with the following %i assets active:"), NumKeys);
+					for (const TPair<FAssetKey, TArray<FTopLevelAssetPath>>& Pair : PathMap)
 					{
-						UE_LOG(LogMetaSound, Display, TEXT("- %s"), *Path.ToString());
+						for (const FTopLevelAssetPath& Path : Pair.Value)
+						{
+							UE_LOG(LogMetaSound, Display, TEXT("- %s"), *Path.ToString());
+						}
 					}
 				}
 			}
@@ -778,6 +784,11 @@ namespace Metasound::Engine
 			}
 		});
 	}
+	
+	void FMetaSoundAssetManager::SetLogActiveAssetsOnShutdown(bool bInLogActiveAssetsOnShutdown)
+	{
+		bLogActiveAssetsOnShutdown = bInLogActiveAssetsOnShutdown;
+	}
 
 	FMetasoundAssetBase* FMetaSoundAssetManager::TryLoadAssetFromKey(const Metasound::Frontend::FAssetKey& InAssetKey) const
 	{
diff --git a/Engine/Plugins/Runtime/Metasound/Source/MetasoundEngine/Private/MetasoundEngineAsset.h b/Engine/Plugins/Runtime/Metasound/Source/MetasoundEngine/Private/MetasoundEngineAsset.h
--- a/Engine/Plugins/Runtime/Metasound/Source/MetasoundEngine/Private/MetasoundEngineAsset.h
+++ b/Engine/Plugins/Runtime/Metasound/Source/MetasoundEngine/Private/MetasoundEngineAsset.h
@@ -128,8 +128,13 @@ namespace Metasound::Engine
 				}
 			}
 
+			IMetaSoundAssetManager& AssetManager = IMetaSoundAssetManager::GetChecked();
+			// Validation has added assets to the asset manager
+			// and we don't remove them immediately after validation to optimize possible subsequent validation
+			// Set this flag to prevent log spam of active assets on shutdown
+			AssetManager.SetLogActiveAssetsOnShutdown(false);
+
 			// Add error for multiple assets with the same class name
-			const IMetaSoundAssetManager& AssetManager = IMetaSoundAssetManager::GetChecked();
 			const FAssetKey Key(InMetaSound.GetConstDocument().RootGraph.Metadata);
 			const TArray<FTopLevelAssetPath>* AssetPaths = AssetManager.FindAssetPaths(Key);
 			if (AssetPaths && AssetPaths->Num() > 1)
diff --git a/Engine/Plugins/Runtime/Metasound/Source/MetasoundFrontend/Public/MetasoundAssetManager.h b/Engine/Plugins/Runtime/Metasound/Source/MetasoundFrontend/Public/MetasoundAssetManager.h
--- a/Engine/Plugins/Runtime/Metasound/Source/MetasoundFrontend/Public/MetasoundAssetManager.h
+++ b/Engine/Plugins/Runtime/Metasound/Source/MetasoundFrontend/Public/MetasoundAssetManager.h
@@ -189,6 +189,9 @@ namespace Metasound::Frontend
 		UE_DEPRECATED(5.5, "Rescan no longer supported nor required by Frontend")
 		virtual void RescanAutoUpdateDenyList() { }
 
+		// Set flag for logging active assets on shutdown. In certain cases (ex. validation), it is expected that assets are active at shutdown
+		virtual void SetLogActiveAssetsOnShutdown(bool bLogActiveAssetsOnShutdown) = 0;
+
 		// Attempts to retrieve the AssetID from the given ClassName if the ClassName is from a valid asset.
 		virtual bool TryGetAssetIDFromClassName(const FMetasoundFrontendClassName& InClassName, FGuid& OutGuid) const = 0;
 
