Add Metasound operator cache csv stats

To better inform which sounds (and how many instances) we should cache, this CL adds several stats to track cache misses and utilzation of cached sounds.

The majority of the tracking logic is encapsulated in the MetaSoundOperatorCacheStatTracker which handles updating active counts for each cached operator and emitting csv events.

This change also adds the ability to write all active metasounds to the csv which can be situationally useful. It is disabled by default due to the overhead/csv bloat.

#rnx
#tests local replays, pie, horde RR

diff --git a/Engine/Plugins/Runtime/Metasound/Source/MetasoundGenerator/Private/MetasoundGenerator.cpp b/Engine/Plugins/Runtime/Metasound/Source/MetasoundGenerator/Private/MetasoundGenerator.cpp
--- a/Engine/Plugins/Runtime/Metasound/Source/MetasoundGenerator/Private/MetasoundGenerator.cpp
+++ b/Engine/Plugins/Runtime/Metasound/Source/MetasoundGenerator/Private/MetasoundGenerator.cpp
@@ -941,7 +941,7 @@ PRAGMA_ENABLE_DEPRECATION_WARNINGS
 		{
 			// See if the cache has an operator with matching OperatorID
 			OperatorPoolID = FOperatorPoolEntryID{InInitParams.Graph->GetInstanceID(), InInitParams.OperatorSettings};
-			FOperatorAndInputs GraphOperatorAndInputs = OperatorPool->ClaimOperator(*OperatorPoolID);
+			FOperatorAndInputs GraphOperatorAndInputs = OperatorPool->ClaimOperator(*OperatorPoolID, FOperatorContext::FromInitParams(InInitParams));
 			if (GraphOperatorAndInputs.Operator.IsValid())
 			{
 				UE_LOG(LogMetasoundGenerator, VeryVerbose, TEXT("Using cached operator %s for MetaSound %s"), *OperatorPoolID->ToString(), *InInitParams.MetaSoundName);
diff --git a/Engine/Plugins/Runtime/Metasound/Source/MetasoundGenerator/Private/MetasoundGeneratorModule.cpp b/Engine/Plugins/Runtime/Metasound/Source/MetasoundGenerator/Private/MetasoundGeneratorModule.cpp
--- a/Engine/Plugins/Runtime/Metasound/Source/MetasoundGenerator/Private/MetasoundGeneratorModule.cpp
+++ b/Engine/Plugins/Runtime/Metasound/Source/MetasoundGenerator/Private/MetasoundGeneratorModule.cpp
@@ -10,9 +10,19 @@
 #include "Misc/CString.h"
 #include "Modules/ModuleInterface.h"
 #include "Modules/ModuleManager.h"
+#include "ProfilingDebugging/CsvProfiler.h"
 
 DEFINE_LOG_CATEGORY(LogMetasoundGenerator);
 
+CSV_DEFINE_CATEGORY(MetaSound_ActiveOperators, true);
+
+// Note: disabled by default as it bloats the csvs quite a bit.
+static TAutoConsoleVariable<bool> CVarRecordActiveOperatorsToCsv(
+	TEXT("au.MetaSound.RecordActiveMetasoundsToCsv"),
+	false,
+	TEXT("Record the name of each active Metasound when csv profiling is recording.")
+);
+
 static FAutoConsoleCommand CommandMetaSoundExperimentalOperatorPoolSetMaxNumOperators(
 	TEXT("au.MetaSound.Experimental.OperatorPool.SetMaxNumOperators"),
 	TEXT("Set the maximum number of operators in the MetaSound operator cache."),
@@ -53,10 +63,33 @@ namespace Metasound
 
 		OperatorPool = MakeShared<FOperatorPool>(OperatorPoolSettings);
 		OperatorInstanceCounterManager = MakeShared<FConcurrentInstanceCounterManager>(InstanceCounterCategory);
+
+#if CSV_PROFILER && !UE_BUILD_SHIPPING
+		CsvEndFrameDelegateHandle = FCsvProfiler::Get()->OnCSVProfileEndFrame().AddLambda([WeakCounterManager = OperatorInstanceCounterManager.ToWeakPtr()]
+		{
+			if (!CVarRecordActiveOperatorsToCsv->GetBool())
+			{
+				return;
+			}
+
+			if (TSharedPtr<FConcurrentInstanceCounterManager> CounterManager = WeakCounterManager.Pin())
+			{
+				CounterManager->VisitStats([](const FName& StatName, int64 Value)
+				{
+					FCsvProfiler::RecordCustomStat(StatName, CSV_CATEGORY_INDEX(MetaSound_ActiveOperators), (int32)Value, ECsvCustomStatOp::Set);
+				});
+			}
+		});
+#endif // #if CSV_PROFILER && !UE_BUILD_SHIPPING
 	}
 
 	void FMetasoundGeneratorModule::ShutdownModule()
 	{
+#if CSV_PROFILER && !UE_BUILD_SHIPPING
+		FCsvProfiler::Get()->OnCSVProfileEndFrame().Remove(CsvEndFrameDelegateHandle);
+		CsvEndFrameDelegateHandle.Reset();
+#endif // #if CSV_PROFILER && !UE_BUILD_SHIPPING
+
 		if (OperatorPool.IsValid())
 		{
 			TSharedPtr<FOperatorPool> PoolShuttingDown = OperatorPool;
diff --git a/Engine/Plugins/Runtime/Metasound/Source/MetasoundGenerator/Private/MetasoundGeneratorModuleImpl.h b/Engine/Plugins/Runtime/Metasound/Source/MetasoundGenerator/Private/MetasoundGeneratorModuleImpl.h
--- a/Engine/Plugins/Runtime/Metasound/Source/MetasoundGenerator/Private/MetasoundGeneratorModuleImpl.h
+++ b/Engine/Plugins/Runtime/Metasound/Source/MetasoundGenerator/Private/MetasoundGeneratorModuleImpl.h
@@ -23,5 +23,7 @@ namespace Metasound
 
 		TSharedPtr<FOperatorPool> OperatorPool;
 		TSharedPtr<FConcurrentInstanceCounterManager> OperatorInstanceCounterManager;
+
+		FDelegateHandle CsvEndFrameDelegateHandle;
 	};
 }
diff --git a/Engine/Plugins/Runtime/Metasound/Source/MetasoundGenerator/Private/MetasoundInstanceCounter.cpp b/Engine/Plugins/Runtime/Metasound/Source/MetasoundGenerator/Private/MetasoundInstanceCounter.cpp
--- a/Engine/Plugins/Runtime/Metasound/Source/MetasoundGenerator/Private/MetasoundInstanceCounter.cpp
+++ b/Engine/Plugins/Runtime/Metasound/Source/MetasoundGenerator/Private/MetasoundInstanceCounter.cpp
@@ -43,6 +43,16 @@ namespace Metasound
 		return 0;
 	}
 
+	void FConcurrentInstanceCounterManager::VisitStats(TFunctionRef<void(const FName&, int64)> Visitor)
+	{
+		FScopeLock Lock(&MapCritSec);
+
+		for (const TPair<FName, FStats>& Pair : StatsMap)
+		{
+			Visitor(Pair.Key, Pair.Value.GetCount());
+		}
+	}
+
 #if COUNTERSTRACE_ENABLED
 	FConcurrentInstanceCounterManager::FStats::FStats(const FString& InName)
 	: TraceCounter(MakeUnique<FCountersTrace::FCounterInt>(TraceCounterNameType_Dynamic, *InName, TraceCounterDisplayHint_None))
@@ -64,7 +74,7 @@ namespace Metasound
 		TraceCounter->Decrement();
 	}
 
-	int64 FConcurrentInstanceCounterManager::FStats::GetCount()
+	int64 FConcurrentInstanceCounterManager::FStats::GetCount() const
 	{
 		ensure(TraceCounter);
 #if COUNTERSTRACE_ENABLED
@@ -74,7 +84,7 @@ namespace Metasound
 #endif
 	}
 
-	int64 FConcurrentInstanceCounterManager::FStats::GetPeakCount()
+	int64 FConcurrentInstanceCounterManager::FStats::GetPeakCount() const
 	{
 		return PeakCount;
 	}
diff --git a/Engine/Plugins/Runtime/Metasound/Source/MetasoundGenerator/Private/MetasoundOperatorCache.cpp b/Engine/Plugins/Runtime/Metasound/Source/MetasoundGenerator/Private/MetasoundOperatorCache.cpp
--- a/Engine/Plugins/Runtime/Metasound/Source/MetasoundGenerator/Private/MetasoundOperatorCache.cpp
+++ b/Engine/Plugins/Runtime/Metasound/Source/MetasoundGenerator/Private/MetasoundOperatorCache.cpp
@@ -10,12 +10,14 @@
 #include "HAL/IConsoleManager.h"
 #include "MetasoundGeneratorModuleImpl.h"
 #include "MetasoundGeneratorBuilder.h"
+#include "MetasoundOperatorCacheStatTracker.h"
 #include "MetasoundTrace.h"
 #include "Misc/Guid.h"
 #include "Misc/ScopeLock.h"
 #include "Modules/ModuleManager.h"
 #include "ProfilingDebugging/CountersTrace.h"
 #include "Templates/UniquePtr.h"
+#include "ProfilingDebugging/CsvProfiler.h"
 
 namespace Metasound
 {
@@ -23,6 +25,8 @@ namespace Metasound
 	TRACE_DECLARE_INT_COUNTER(MetaSound_OperatorPool_NumOperators, TEXT("MetaSound/OperatorPool/NumOperatorsInPool"));
 	TRACE_DECLARE_FLOAT_COUNTER(MetaSound_OperatorPool_HitRatio, TEXT("MetaSound/OperatorPool/HitRatio"));
 	TRACE_DECLARE_FLOAT_COUNTER(MetaSound_OperatorPool_WindowedHitRatio, TEXT("MetaSound/OperatorPool/WindowedHitRatio"));
+
+	CSV_DEFINE_CATEGORY(MetaSound_OperatorPool, true);
 #endif // #if METASOUND_OPERATORCACHEPROFILER_ENABLED
 
 	namespace OperatorPoolPrivate
@@ -37,6 +41,7 @@ namespace Metasound
 			TEXT("Retrieves graph on the requesting thread prior to asynchronous task to create instance.\n"),
 			ECVF_Default);
 
+		// TODO: Move this into the OperatorCacheStatTracker.
 #if METASOUND_OPERATORCACHEPROFILER_ENABLED
 		static std::atomic<uint32> CacheHitCount = 0;
 		static std::atomic<uint32> CacheAttemptCount = 0;
@@ -104,6 +109,8 @@ namespace Metasound
 			{
 				const float HitRatio = RunningHitCount / static_cast<float>(RunningTotal);
 				TRACE_COUNTER_SET(MetaSound_OperatorPool_WindowedHitRatio, HitRatio);
+
+				CSV_CUSTOM_STAT(MetaSound_OperatorPool, WindowedCacheHitRatio, HitRatio, ECsvCustomStatOp::Set);
 			}
 		}
 
@@ -178,6 +185,15 @@ namespace Metasound
 	{
 	}
 
+	FOperatorContext FOperatorContext::FromInitParams(const FMetasoundGeneratorInitParams& InParams)
+	{
+		return FOperatorContext
+		{
+			.GraphInstanceName = InParams.Graph ? InParams.Graph->GetInstanceName() : NAME_None,
+			.MetaSoundName = InParams.MetaSoundName
+		};
+	}
+
 	FOperatorPoolEntryID::FOperatorPoolEntryID(FGuid InOperatorID, FOperatorSettings InSettings)
 	: OperatorID(MoveTemp(InOperatorID))
 	, OperatorSettings(MoveTemp(InSettings))
@@ -215,6 +231,9 @@ namespace Metasound
 	: Settings(InSettings)
 	, AsyncBuildPipe(UE_SOURCE_LOCATION)
 	{
+#if METASOUND_OPERATORCACHEPROFILER_ENABLED
+		CacheStatTracker = MakeUnique<Engine::FOperatorCacheStatTracker>();
+#endif
 	}
 
 	FOperatorPool::~FOperatorPool()
@@ -224,10 +243,12 @@ namespace Metasound
 
 	FOperatorAndInputs FOperatorPool::ClaimOperator(const FGuid& InOperatorID)
 	{
-		return ClaimOperator(FOperatorPoolEntryID{InOperatorID, FOperatorSettings{OperatorPoolPrivate::DefaultSampleRateForDeprecatedAPI, OperatorPoolPrivate::BlockRate}});
+		return ClaimOperator(
+			FOperatorPoolEntryID{InOperatorID, FOperatorSettings{OperatorPoolPrivate::DefaultSampleRateForDeprecatedAPI, OperatorPoolPrivate::BlockRate}},
+			FOperatorContext{});
 	}
 
-	FOperatorAndInputs FOperatorPool::ClaimOperator(const FOperatorPoolEntryID& InOperatorID)
+	FOperatorAndInputs FOperatorPool::ClaimOperator(const FOperatorPoolEntryID& InOperatorID, const FOperatorContext& InContext)
 	{
 		FOperatorAndInputs OpAndInputs;
 
@@ -254,6 +275,7 @@ namespace Metasound
 			bCacheHit ? HitRateTracker.AddHit() : HitRateTracker.AddMiss();
 			OperatorPoolPrivate::CacheAttemptCount++;
 			TRACE_COUNTER_SET(MetaSound_OperatorPool_HitRatio, OperatorPoolPrivate::GetHitRatio());
+			CacheStatTracker->RecordCacheEvent(InOperatorID, bCacheHit, InContext);
 #endif
 			UE_LOG(LogMetasoundGenerator, VeryVerbose, TEXT("Attempt to claim operator with ID %s from operator pool resulted in %s"), *InOperatorID.ToString(), *::LexToString(bCacheHit));
 		}
@@ -334,6 +356,7 @@ namespace Metasound
 		Stack.Add(InOperatorID);
 #if METASOUND_OPERATORCACHEPROFILER_ENABLED
 		TRACE_COUNTER_INCREMENT(MetaSound_OperatorPool_NumOperators);
+		CacheStatTracker->OnOperatorAdded(InOperatorID);
 #endif // #if METASOUND_OPERATORCACHEPROFILER_ENABLED
 	   
 		UE_LOG(LogMetasoundGenerator, VeryVerbose, TEXT("Adding operator with ID %s to operator pool"), *InOperatorID.ToString());
@@ -469,6 +492,14 @@ namespace Metasound
 				}
 			}
 
+#if METASOUND_OPERATORCACHEPROFILER_ENABLED
+			if (TSharedPtr<FOperatorPool> OperatorPool = WeakPoolPtr.Pin();
+				OperatorPool.IsValid())
+			{
+				OperatorPool->CacheStatTracker->RecordPreCacheRequest(*PreCacheData, NumToBuild);
+			}
+#endif // METASOUND_OPERATORCACHEPROFILER_ENABLED
+
 			for (int32 i = 0; i < NumToBuild; ++i)
 			{
 				// These build operations can take a fair bit of time, so
@@ -565,6 +596,8 @@ namespace Metasound
 
 #if METASOUND_OPERATORCACHEPROFILER_ENABLED
 				TRACE_COUNTER_SUBTRACT(MetaSound_OperatorPool_NumOperators, int64(NumRemoved));
+
+				Pool->CacheStatTracker->OnOperatorRemoved(OperatorID);
 #endif // #if METASOUND_OPERATORCACHEPROFILER_ENABLED
 			}
 		});
@@ -662,12 +695,16 @@ namespace Metasound
 								GraphIdToAssetIdLookUp.Remove(Stack[i], AssetID);
 							}
 						}
+
+#if METASOUND_OPERATORCACHEPROFILER_ENABLED
+						CacheStatTracker->OnOperatorTrimmed(Stack[i]);
+#endif // METASOUND_OPERATORCACHEPROFILER_ENABLED
 					}
 				}
 			}
 			Stack.RemoveAt(0, NumToTrim);
 #if METASOUND_OPERATORCACHEPROFILER_ENABLED
-			TRACE_COUNTER_DECREMENT(MetaSound_OperatorPool_NumOperators);
+			TRACE_COUNTER_SUBTRACT(MetaSound_OperatorPool_NumOperators, NumToTrim);
 #endif // #if METASOUND_OPERATORCACHEPROFILER_ENABLED
 		}
 	}
diff --git a/Engine/Plugins/Runtime/Metasound/Source/MetasoundGenerator/Private/MetasoundOperatorCacheStatTracker.cpp b/Engine/Plugins/Runtime/Metasound/Source/MetasoundGenerator/Private/MetasoundOperatorCacheStatTracker.cpp
new file mode 100644
--- /dev/null
+++ b/Engine/Plugins/Runtime/Metasound/Source/MetasoundGenerator/Private/MetasoundOperatorCacheStatTracker.cpp
@@ -0,0 +1,220 @@
+// Copyright Epic Games, Inc. All Rights Reserved.
+
+#include "MetasoundOperatorCacheStatTracker.h"
+
+#include "MetasoundOperatorCache.h"
+#include "MetasoundGeneratorModule.h"
+#include "ProfilingDebugging/CsvProfiler.h"
+
+#if METASOUND_OPERATORCACHEPROFILER_ENABLED
+namespace Metasound
+{
+	CSV_DECLARE_CATEGORY_EXTERN(MetaSound_OperatorPool);
+}
+
+namespace Metasound::Engine
+{
+	CSV_DEFINE_CATEGORY(MetaSound_OperatorCacheUtilization, true);
+	CSV_DEFINE_CATEGORY(MetaSound_AvailableCachedOperators, true);
+	CSV_DEFINE_CATEGORY(Metasound_OperatorCacheMiss, true);
+
+	namespace Private
+	{
+		static TAutoConsoleVariable<bool> CVarCacheMissCsvStatsEnabled(
+			TEXT("au.MetaSound.OperatorPool.CacheMissCsvStatsEnabled"),
+			true,
+			TEXT("Record which metasounds incur a cache miss when building their graph.")
+		);
+
+		FString GetGraphName(FName GraphInstanceName)
+		{
+			return FPackageName::ObjectPathToPathWithinPackage(GraphInstanceName.ToString());
+		}
+	}
+
+	FOperatorCacheStatTracker::FOperatorCacheStatTracker()
+	{
+#if CSV_PROFILER
+		CsvEndFrameDelegateHandle = FCsvProfiler::Get()->OnCSVProfileEndFrame().AddRaw(this, &FOperatorCacheStatTracker::OnCsvProfileEndFrame);
+#endif // CSV_PROFILER
+	}
+
+	FOperatorCacheStatTracker::~FOperatorCacheStatTracker()
+	{
+#if CSV_PROFILER
+		FCsvProfiler::Get()->OnCSVProfileEndFrame().Remove(CsvEndFrameDelegateHandle);
+		CsvEndFrameDelegateHandle.Reset();
+#endif // CSV_PROFILER
+	}
+
+	void FOperatorCacheStatTracker::RecordPreCacheRequest(const FOperatorBuildData& BuildData, int32 NumIntancesToBuild)
+	{
+		if (BuildData.NumInstances <= 0)
+		{
+			return;
+		}
+
+		const FOperatorPoolEntryID EntryID{ BuildData.InitParams.Graph->GetInstanceID(), BuildData.InitParams.OperatorSettings };
+
+		FScopeLock Lock(&CriticalSection);
+
+		if (FStatEntry* Entry = StatEntries.Find(EntryID))
+		{
+			Entry->NumCacheSlots += NumIntancesToBuild;
+
+			// Show how much we're increasing the existing cache for this sound by.
+			UE_LOG(LogMetasoundGenerator, Log,
+				TEXT("Pre-cached Metasound: %s [Graph: %s]. Added %d instances, Total: %d."),
+				*BuildData.InitParams.MetaSoundName,
+				*Entry->GraphName.ToString(),
+				NumIntancesToBuild,
+				Entry->NumCacheSlots);
+		}
+		else
+		{
+			// Get the asset name from the package path
+			const FString GraphName = Private::GetGraphName(BuildData.InitParams.Graph->GetInstanceName());
+
+			FStatEntry& StatEntry = StatEntries.Add(EntryID, FStatEntry
+			{
+				.GraphName = FName(*GraphName),
+				.NumInstancesBuilt = NumIntancesToBuild,
+				.NumCacheSlots = NumIntancesToBuild
+			});
+
+			if (GraphName == BuildData.InitParams.MetaSoundName)
+			{
+				UE_LOG(LogMetasoundGenerator, Log,
+					TEXT("Pre-cached Metasound: %s. Requested: %d, Built: %d."),
+					*BuildData.InitParams.MetaSoundName,
+					BuildData.NumInstances,
+					NumIntancesToBuild);
+			}
+			else
+			{
+				// Include the parent graph so it's clearer which this contributes to.
+				UE_LOG(LogMetasoundGenerator, Log,
+					TEXT("Pre-cached Metasound: %s [Graph: %s] Requested: %d, Built: %d."),
+					*BuildData.InitParams.MetaSoundName,
+					*StatEntry.GraphName.ToString(),
+					BuildData.NumInstances,
+					NumIntancesToBuild);
+			}
+		}
+	}
+
+	void FOperatorCacheStatTracker::RecordCacheEvent(const FOperatorPoolEntryID& OperatorID, bool bCacheHit, const FOperatorContext& Context)
+	{
+		if (!bCacheHit)
+		{
+			if (Private::CVarCacheMissCsvStatsEnabled->GetBool() &&
+				Context.GraphInstanceName != NAME_Name)
+			{
+				const FString GraphName = Private::GetGraphName(Context.GraphInstanceName);
+				FCsvProfiler::Get()->RecordCustomStat(*GraphName, CSV_CATEGORY_INDEX(Metasound_OperatorCacheMiss), 1, ECsvCustomStatOp::Accumulate);
+			}
+			return;
+		}
+
+		FScopeLock Lock(&CriticalSection);
+
+		if (FStatEntry* StatEntry = StatEntries.Find(OperatorID))
+		{
+			--StatEntry->NumAvailableInCache;
+			check(StatEntry->NumAvailableInCache >= 0);
+		}
+
+		--NumInCache;
+		check(NumInCache >= 0);
+	}
+
+	void FOperatorCacheStatTracker::OnOperatorAdded(const FOperatorPoolEntryID& OperatorID)
+	{
+		FScopeLock Lock(&CriticalSection);
+
+		if (FStatEntry* StatEntry = StatEntries.Find(OperatorID))
+		{
+			++StatEntry->NumAvailableInCache;
+		}
+
+		++NumInCache;
+	}
+
+	void FOperatorCacheStatTracker::OnOperatorTrimmed(const FOperatorPoolEntryID& OperatorID)
+	{
+		FScopeLock Lock(&CriticalSection);
+
+		if (FStatEntry* StatEntry = StatEntries.Find(OperatorID))
+		{
+			--StatEntry->NumCacheSlots;
+			--StatEntry->NumAvailableInCache;
+			check(StatEntry->NumCacheSlots >= 0);
+			check(StatEntry->NumAvailableInCache >= 0);
+
+			if (StatEntry->NumCacheSlots == 0)
+			{
+				UE_LOG(LogMetasoundGenerator, Log, TEXT("Evicted %s from the Operator Pool."), *StatEntry->GraphName.ToString())
+			}
+			else
+			{
+				UE_LOG(LogMetasoundGenerator, Log,
+					TEXT("Trimmed 1 instance of %s from the Operator Pool. %d instances remaining."),
+					*StatEntry->GraphName.ToString(),
+					StatEntry->NumCacheSlots);
+			}
+		}
+
+		--NumInCache;
+		check(NumInCache >= 0);
+	}
+
+	void FOperatorCacheStatTracker::OnOperatorRemoved(const FOperatorPoolEntryID& OperatorID)
+	{
+		FScopeLock Lock(&CriticalSection);
+
+		if (FStatEntry* StatEntry = StatEntries.Find(OperatorID))
+		{
+			NumInCache -= StatEntry->NumAvailableInCache;
+		}
+
+		StatEntries.Remove(OperatorID);
+
+		check(NumInCache >= 0);
+	}
+
+	void FOperatorCacheStatTracker::OnCsvProfileEndFrame()
+	{
+#if CSV_PROFILER
+		QUICK_SCOPE_CYCLE_COUNTER(OperatorCacheStatTracker_RecordStats)
+
+		FScopeLock Lock(&CriticalSection);
+
+		CSV_CUSTOM_STAT(MetaSound_OperatorPool, TotalCachedOperators, NumInCache, ECsvCustomStatOp::Set);
+
+		for (auto It = StatEntries.CreateIterator(); It; ++It)
+		{
+			const FOperatorPoolEntryID& PoolEntryID = It->Key;
+			FStatEntry& Entry = It->Value;
+
+			// Remove any nodes that have been evicted from the cache.
+			if (Entry.NumCacheSlots <= 0)
+			{
+				It.RemoveCurrent();
+				continue;
+			}
+
+			// Record cache utilization stats.
+			{
+				check(Entry.NumCacheSlots > 0);
+				const int32 NumAvailableIncache = Entry.NumAvailableInCache;
+				const int32 NumUsed = Entry.NumCacheSlots - NumAvailableIncache;
+				const float UtilizationRatio = static_cast<float>(NumUsed) / static_cast<float>(Entry.NumCacheSlots);
+
+				FCsvProfiler::Get()->RecordCustomStat(Entry.GraphName, CSV_CATEGORY_INDEX(MetaSound_AvailableCachedOperators), NumAvailableIncache, ECsvCustomStatOp::Set);
+				FCsvProfiler::Get()->RecordCustomStat(Entry.GraphName, CSV_CATEGORY_INDEX(MetaSound_OperatorCacheUtilization), UtilizationRatio, ECsvCustomStatOp::Set);
+			}
+		}
+#endif // CSV_PROFILER
+	}
+} // namespace Metasound::Private
+#endif // METASOUND_OPERATORCACHEPROFILER_ENABLED
diff --git a/Engine/Plugins/Runtime/Metasound/Source/MetasoundGenerator/Private/MetasoundOperatorCacheStatTracker.h b/Engine/Plugins/Runtime/Metasound/Source/MetasoundGenerator/Private/MetasoundOperatorCacheStatTracker.h
new file mode 100644
--- /dev/null
+++ b/Engine/Plugins/Runtime/Metasound/Source/MetasoundGenerator/Private/MetasoundOperatorCacheStatTracker.h
@@ -0,0 +1,56 @@
+// Copyright Epic Games, Inc. All Rights Reserved.
+#pragma once
+
+#include "MetasoundGenerator.h"
+#include "ProfilingDebugging/CountersTrace.h"
+
+#ifndef METASOUND_OPERATORCACHEPROFILER_ENABLED
+#define METASOUND_OPERATORCACHEPROFILER_ENABLED COUNTERSTRACE_ENABLED
+#endif
+
+#if METASOUND_OPERATORCACHEPROFILER_ENABLED
+namespace Metasound
+{
+	struct FOperatorBuildData;
+	struct FOperatorContext;
+}
+
+namespace Metasound::Engine
+{
+	class FOperatorCacheStatTracker final
+	{
+	public:
+		FOperatorCacheStatTracker();
+		~FOperatorCacheStatTracker();
+
+		// Called when we're pre-caching an operator. Called before AddOperator.
+		void RecordPreCacheRequest(const FOperatorBuildData& BuildData, int32 NumIntancesToBuild);
+		void RecordCacheEvent(const FOperatorPoolEntryID& OperatorID, bool bCacheHit, const FOperatorContext& Context);
+		void OnOperatorAdded(const FOperatorPoolEntryID& OperatorID);
+		void OnOperatorTrimmed(const FOperatorPoolEntryID& OperatorID);
+		void OnOperatorRemoved(const FOperatorPoolEntryID& OperatorID);
+
+	private:
+		void OnCsvProfileEndFrame();
+
+		struct FStatEntry
+		{
+			// Stored as an fname to avoid conversions when writing out csv stats.
+			FName GraphName;
+			// The actual number of instances we built.
+			int32 NumInstancesBuilt;
+			// The number of instances we have space to cache.
+			// This can be different from the number currently in the stack as those are removed when in use.
+			int32 NumCacheSlots;
+			// The number of operators sitting in the pool waiting to be used.
+			// Set exclusively by OnOperatorAdded/OnOperatorTrimmed.
+			int32 NumAvailableInCache = 0;
+		};
+
+		TMap<FOperatorPoolEntryID, FStatEntry> StatEntries;
+		mutable FCriticalSection CriticalSection;
+		FDelegateHandle CsvEndFrameDelegateHandle;
+		int32 NumInCache = 0;
+	};
+} // namespace Metasound::Engine
+#endif // METASOUND_OPERATORCACHEPROFILER_ENABLED
\ No newline at end of file
diff --git a/Engine/Plugins/Runtime/Metasound/Source/MetasoundGenerator/Public/MetasoundInstanceCounter.h b/Engine/Plugins/Runtime/Metasound/Source/MetasoundGenerator/Public/MetasoundInstanceCounter.h
--- a/Engine/Plugins/Runtime/Metasound/Source/MetasoundGenerator/Public/MetasoundInstanceCounter.h
+++ b/Engine/Plugins/Runtime/Metasound/Source/MetasoundGenerator/Public/MetasoundInstanceCounter.h
@@ -33,6 +33,8 @@ namespace Metasound
 		int64 GetCountForName(const FName& InName);
 		int64 GetPeakCountForName(const FName& InName);
 
+		void VisitStats(TFunctionRef<void(const FName&, int64)> Visitor);
+
 	private:
 		struct FStats
 		{
@@ -47,8 +49,8 @@ namespace Metasound
 			void Increment();
 			void Decrement();
 	
-			int64 GetCount();
-			int64 GetPeakCount();
+			int64 GetCount() const;
+			int64 GetPeakCount() const;
 	
 		private:
 			TUniquePtr<CounterType> TraceCounter;
diff --git a/Engine/Plugins/Runtime/Metasound/Source/MetasoundGenerator/Public/MetasoundOperatorCache.h b/Engine/Plugins/Runtime/Metasound/Source/MetasoundGenerator/Public/MetasoundOperatorCache.h
--- a/Engine/Plugins/Runtime/Metasound/Source/MetasoundGenerator/Public/MetasoundOperatorCache.h
+++ b/Engine/Plugins/Runtime/Metasound/Source/MetasoundGenerator/Public/MetasoundOperatorCache.h
@@ -18,6 +18,11 @@
 namespace Metasound
 {
 #if METASOUND_OPERATORCACHEPROFILER_ENABLED
+	namespace Engine
+	{
+		class FOperatorCacheStatTracker;
+	}
+
 	namespace OperatorPoolPrivate
 	{
 		class FWindowedHitRate
@@ -86,6 +91,15 @@ namespace Metasound
 
 	}; // struct FOperatorPrecacheData
 
+	// Provides additional debug context for the operator the pool is interacting with.
+	struct METASOUNDGENERATOR_API FOperatorContext
+	{
+		FName GraphInstanceName;
+		FStringView MetaSoundName;
+
+		static FOperatorContext FromInitParams(const FMetasoundGeneratorInitParams& InParams);
+	};
+
 	// Pool of re-useable metasound operators to be used / put back by the metasound generator
 	// operators can also be pre-constructed via the UMetasoundCacheSubsystem BP api.
 	class METASOUNDGENERATOR_API FOperatorPool : public TSharedFromThis<FOperatorPool>
@@ -96,9 +110,9 @@ namespace Metasound
 		~FOperatorPool();
 
 
-		UE_DEPRECATED(5.5, "Use ClaimOperator(const FOperatorPoolEntryID&) instead")
+		UE_DEPRECATED(5.5, "Use ClaimOperator(const FOperatorPoolEntryID&, ...) instead")
 		FOperatorAndInputs ClaimOperator(const FGuid& InOperatorID);
-		FOperatorAndInputs ClaimOperator(const FOperatorPoolEntryID& InOperatorID);
+		FOperatorAndInputs ClaimOperator(const FOperatorPoolEntryID& InOperatorID, const FOperatorContext& InContext);
 
 		UE_DEPRECATED(5.5, "Use AddOperator(const FOperatorPoolEntryID&, ...) instead")
 		void AddOperator(const FGuid& InOperatorID, TUniquePtr<IOperator>&& InOperator, FInputVertexInterfaceData&& InputData);
@@ -156,6 +170,7 @@ namespace Metasound
 
 #if METASOUND_OPERATORCACHEPROFILER_ENABLED
 		OperatorPoolPrivate::FWindowedHitRate HitRateTracker;
+		TUniquePtr<Engine::FOperatorCacheStatTracker> CacheStatTracker;
 #endif // #if METASOUND_OPERATORCACHEPROFILER_ENABLED
 
 		// Notifies active build tasks to abort as soon as possible
