Add missing csvprofiler ifdef to metasoundoperator cache stat tracker.

#rnx

diff --git a/Engine/Plugins/Runtime/Metasound/Source/MetasoundGenerator/Private/MetasoundOperatorCacheStatTracker.cpp b/Engine/Plugins/Runtime/Metasound/Source/MetasoundGenerator/Private/MetasoundOperatorCacheStatTracker.cpp
--- a/Engine/Plugins/Runtime/Metasound/Source/MetasoundGenerator/Private/MetasoundOperatorCacheStatTracker.cpp
+++ b/Engine/Plugins/Runtime/Metasound/Source/MetasoundGenerator/Private/MetasoundOperatorCacheStatTracker.cpp
@@ -107,12 +107,14 @@ namespace Metasound::Engine
 	{
 		if (!bCacheHit)
 		{
+#if CSV_PROFILER
 			if (Private::CVarCacheMissCsvStatsEnabled->GetBool() &&
 				Context.GraphInstanceName != NAME_Name)
 			{
 				const FString GraphName = Private::GetGraphName(Context.GraphInstanceName);
 				FCsvProfiler::Get()->RecordCustomStat(*GraphName, CSV_CATEGORY_INDEX(Metasound_OperatorCacheMiss), 1, ECsvCustomStatOp::Accumulate);
 			}
+#endif // CSV_PROFILER
 			return;
 		}
 
