Fix for false negative ensure when swapping datatype/vertex access type on MetaSound via builder
#rnx

diff --git a/Engine/Plugins/Runtime/Metasound/Source/MetasoundFrontend/Private/MetasoundFrontendDocumentBuilder.cpp b/Engine/Plugins/Runtime/Metasound/Source/MetasoundFrontend/Private/MetasoundFrontendDocumentBuilder.cpp
--- a/Engine/Plugins/Runtime/Metasound/Source/MetasoundFrontend/Private/MetasoundFrontendDocumentBuilder.cpp
+++ b/Engine/Plugins/Runtime/Metasound/Source/MetasoundFrontend/Private/MetasoundFrontendDocumentBuilder.cpp
@@ -3377,17 +3377,18 @@ bool FMetaSoundFrontendDocumentBuilder::SetGraphInputAccessType(FName InputName,
 		return false;
 	}
 
-	if (const int32* Index = DocumentCache->GetInterfaceCache().FindInputIndex(InputName))
+	const int32* Index = DocumentCache->GetInterfaceCache().FindInputIndex(InputName);
+	if (!Index)
 	{
-		FMetasoundFrontendDocument& Document = GetDocumentChecked();
-		FMetasoundFrontendGraphClass& RootGraph = Document.RootGraph;
-		FMetasoundFrontendClassInput& GraphInput = RootGraph.Interface.Inputs[*Index];
+		return false;
+	}
 
-		if (GraphInput.AccessType == AccessType)
-		{
-			return true;
-		}
+	FMetasoundFrontendDocument& Document = GetDocumentChecked();
+	FMetasoundFrontendGraphClass& RootGraph = Document.RootGraph;
+	FMetasoundFrontendClassInput& GraphInput = RootGraph.Interface.Inputs[*Index];
 
+	if (GraphInput.AccessType != AccessType)
+	{
 		GraphInput.AccessType = AccessType;
 		if (AccessType == EMetasoundFrontendVertexAccessType::Reference)
 		{
@@ -3431,18 +3432,18 @@ bool FMetaSoundFrontendDocumentBuilder::SetGraphInputAccessType(FName InputName,
 			});
 
 			const bool bNodeConformed = ConformGraphInputNodeToClass(GraphInput);
-			if (bNodeConformed)
+			if (!bNodeConformed)
 			{
+				return false;
+			}
+
 #if WITH_EDITORONLY_DATA
-				Document.Metadata.ModifyContext.AddMemberIDModified(GraphInput.NodeID);
+			Document.Metadata.ModifyContext.AddMemberIDModified(GraphInput.NodeID);
 #endif // WITH_EDITORONLY_DATA
-
-				return true;
-			}
 		}
 	}
 
-	return false;
+	return true;
 }
 
 bool FMetaSoundFrontendDocumentBuilder::SetGraphInputDataType(FName InputName, FName DataType)
@@ -3451,16 +3452,17 @@ bool FMetaSoundFrontendDocumentBuilder::SetGraphInputDataType(FName InputName, F
 
 	if (IDataTypeRegistry::Get().IsRegistered(DataType))
 	{
-		if (const int32* Index = DocumentCache->GetInterfaceCache().FindInputIndex(InputName))
+		const int32* Index = DocumentCache->GetInterfaceCache().FindInputIndex(InputName);
+		if (!Index)
 		{
-			FMetasoundFrontendDocument& Document = GetDocumentChecked();
-			FMetasoundFrontendGraphClass& RootGraph = Document.RootGraph;
-			FMetasoundFrontendClassInput& GraphInput = RootGraph.Interface.Inputs[*Index];
-			if (GraphInput.TypeName == DataType)
-			{
-				return true;
-			}
+			return false;
+		}
 
+		FMetasoundFrontendDocument& Document = GetDocumentChecked();
+		FMetasoundFrontendGraphClass& RootGraph = Document.RootGraph;
+		FMetasoundFrontendClassInput& GraphInput = RootGraph.Interface.Inputs[*Index];
+		if (GraphInput.TypeName != DataType)
+		{
 			FMetasoundFrontendLiteral DefaultLiteral;
 			DefaultLiteral.SetFromLiteral(IDataTypeRegistry::Get().CreateDefaultLiteral(DataType));
 			GraphInput.DefaultLiteral = DefaultLiteral;
@@ -3497,22 +3499,22 @@ bool FMetaSoundFrontendDocumentBuilder::SetGraphInputDataType(FName InputName, F
 			});
 
 			const bool bNodeConformed = ConformGraphInputNodeToClass(GraphInput);
-			if (bNodeConformed)
+			if (!bNodeConformed)
 			{
-				RemoveUnusedDependencies();
+				return false;
+			}
+
+			RemoveUnusedDependencies();
 
 #if WITH_EDITORONLY_DATA
-				ClearMemberMetadata(GraphInput.NodeID);
-				Document.Metadata.ModifyContext.AddMemberIDModified(GraphInput.NodeID);
-				Document.Metadata.ModifyContext.AddNodeIDModified(GraphInput.NodeID);
+			ClearMemberMetadata(GraphInput.NodeID);
+			Document.Metadata.ModifyContext.AddMemberIDModified(GraphInput.NodeID);
+			Document.Metadata.ModifyContext.AddNodeIDModified(GraphInput.NodeID);
 #endif // WITH_EDITORONLY_DATA
-
-				return true;
-			}
 		}
 	}
 
-	return false;
+	return true;
 }
 
 bool FMetaSoundFrontendDocumentBuilder::SetGraphInputDefault(FName InputName, const FMetasoundFrontendLiteral& InDefaultLiteral)
@@ -3576,16 +3578,17 @@ bool FMetaSoundFrontendDocumentBuilder::SetGraphOutputAccessType(FName OutputNam
 		return false;
 	}
 
-	if (const int32* Index = DocumentCache->GetInterfaceCache().FindOutputIndex(OutputName))
+	const int32* Index = DocumentCache->GetInterfaceCache().FindOutputIndex(OutputName);
+	if (!Index)
 	{
-		FMetasoundFrontendDocument& Document = GetDocumentChecked();
-		FMetasoundFrontendGraphClass& RootGraph = Document.RootGraph;
-		FMetasoundFrontendClassOutput& GraphOutput = RootGraph.Interface.Outputs[*Index];
-		if (GraphOutput.AccessType == AccessType)
-		{
-			return true;
-		}
+		return false;
+	}
 
+	FMetasoundFrontendDocument& Document = GetDocumentChecked();
+	FMetasoundFrontendGraphClass& RootGraph = Document.RootGraph;
+	FMetasoundFrontendClassOutput& GraphOutput = RootGraph.Interface.Outputs[*Index];
+	if (GraphOutput.AccessType != AccessType)
+	{
 		GraphOutput.AccessType = AccessType;
 		if (AccessType == EMetasoundFrontendVertexAccessType::Value)
 		{
@@ -3613,69 +3616,71 @@ bool FMetaSoundFrontendDocumentBuilder::SetGraphOutputAccessType(FName OutputNam
 		}
 
 		const bool bNodeConformed = ConformGraphOutputNodeToClass(GraphOutput);
-		if (bNodeConformed)
+		if (!bNodeConformed)
 		{
+			return false;
+		}
+
 #if WITH_EDITORONLY_DATA
-			Document.Metadata.ModifyContext.AddMemberIDModified(GraphOutput.NodeID);
+		Document.Metadata.ModifyContext.AddMemberIDModified(GraphOutput.NodeID);
 #endif // WITH_EDITORONLY_DATA
-
-			return true;
-		}
 	}
 
-	return false;
+	return true;
 }
 
 bool FMetaSoundFrontendDocumentBuilder::SetGraphOutputDataType(FName OutputName, FName DataType)
 {
 	using namespace Metasound::Frontend;
 
-	if (IDataTypeRegistry::Get().IsRegistered(DataType))
+	if (!IDataTypeRegistry::Get().IsRegistered(DataType))
 	{
-		if (const int32* Index = DocumentCache->GetInterfaceCache().FindOutputIndex(OutputName))
-		{
-			FMetasoundFrontendDocument& Document = GetDocumentChecked();
-			FMetasoundFrontendGraphClass& RootGraph = Document.RootGraph;
-			FMetasoundFrontendClassOutput& GraphOutput = RootGraph.Interface.Outputs[*Index];
-			if (GraphOutput.TypeName == DataType)
-			{
-				return true;
-			}
+		return false;
+	}
 
-			GraphOutput.TypeName = DataType;
+	const int32* Index = DocumentCache->GetInterfaceCache().FindOutputIndex(OutputName);
+	if (!Index)
+	{
+		return false;
+	}
 
-			RootGraph.IterateGraphPages([this, &GraphOutput, &DataType](FMetasoundFrontendGraph& Graph)
+	FMetasoundFrontendDocument& Document = GetDocumentChecked();
+	FMetasoundFrontendGraphClass& RootGraph = Document.RootGraph;
+	FMetasoundFrontendClassOutput& GraphOutput = RootGraph.Interface.Outputs[*Index];
+	if (GraphOutput.TypeName != DataType)
+	{
+		GraphOutput.TypeName = DataType;
+
+		RootGraph.IterateGraphPages([this, &GraphOutput, &DataType](FMetasoundFrontendGraph& Graph)
+		{
+			const IDocumentGraphNodeCache& NodeCache = DocumentCache->GetNodeCache();
+			if (const int32* NodeIndex = NodeCache.FindNodeIndex(GraphOutput.NodeID))
 			{
-				const IDocumentGraphNodeCache& NodeCache = DocumentCache->GetNodeCache();
-				if (const int32* NodeIndex = NodeCache.FindNodeIndex(GraphOutput.NodeID))
-				{
-					FMetasoundFrontendNode& Node = FindBuildGraphChecked().Nodes[*NodeIndex];
+				FMetasoundFrontendNode& Node = FindBuildGraphChecked().Nodes[*NodeIndex];
 
-					FMetasoundFrontendLiteral DefaultLiteral;
-					DefaultLiteral.SetFromLiteral(IDataTypeRegistry::Get().CreateDefaultLiteral(DataType));
-					FMetasoundFrontendVertex& NodeInput = Node.Interface.Inputs.Last();
-					Node.InputLiterals = { FMetasoundFrontendVertexLiteral { NodeInput.VertexID, MoveTemp(DefaultLiteral) } };
+				FMetasoundFrontendLiteral DefaultLiteral;
+				DefaultLiteral.SetFromLiteral(IDataTypeRegistry::Get().CreateDefaultLiteral(DataType));
+				FMetasoundFrontendVertex& NodeInput = Node.Interface.Inputs.Last();
+				Node.InputLiterals = { FMetasoundFrontendVertexLiteral { NodeInput.VertexID, MoveTemp(DefaultLiteral) } };
 
-					RemoveEdgeToNodeInput(GraphOutput.NodeID, NodeInput.VertexID);
-					GraphOutput.TypeName = DataType;
-				}
-			});
+				RemoveEdgeToNodeInput(GraphOutput.NodeID, NodeInput.VertexID);
+				GraphOutput.TypeName = DataType;
+			}
+		});
 
-			const bool bNodeConformed = ConformGraphOutputNodeToClass(GraphOutput);
-			if (bNodeConformed)
-			{
+		const bool bNodeConformed = ConformGraphOutputNodeToClass(GraphOutput);
+		if (!bNodeConformed)
+		{
+			return false;
+		}
 
 #if WITH_EDITORONLY_DATA
-				ClearMemberMetadata(GraphOutput.NodeID);
-				Document.Metadata.ModifyContext.AddMemberIDModified(GraphOutput.NodeID);
+		ClearMemberMetadata(GraphOutput.NodeID);
+		Document.Metadata.ModifyContext.AddMemberIDModified(GraphOutput.NodeID);
 #endif // WITH_EDITORONLY_DATA
-
-				return true;
-			}
-		}
 	}
 
-	return false;
+	return true;
 }
 
 #if WITH_EDITOR
