- Trivial: fix for focus not syncing correctly with MetaSound build page in page panel
- Minor add of ability to focus by PageID to MetaSoundEditorSubsystem

#rnx

diff --git a/Engine/Plugins/Runtime/Metasound/Source/MetasoundEditor/Private/MetasoundDetailCustomization.cpp b/Engine/Plugins/Runtime/Metasound/Source/MetasoundEditor/Private/MetasoundDetailCustomization.cpp
--- a/Engine/Plugins/Runtime/Metasound/Source/MetasoundEditor/Private/MetasoundDetailCustomization.cpp
+++ b/Engine/Plugins/Runtime/Metasound/Source/MetasoundEditor/Private/MetasoundDetailCustomization.cpp
@@ -577,8 +577,13 @@ namespace Metasound::Editor
 			}
 			else
 			{
-				BuildPageName = Frontend::DefaultGraphPageName;
-				DocBuilder.SetBuildPageID(Frontend::DefaultGraphPageID);
+				constexpr bool bFocusEditor = false; // Already focused by user action
+				EMetaSoundBuilderResult Result = EMetaSoundBuilderResult::Failed;
+				UMetaSoundEditorSubsystem::GetConstChecked().SetFocusedPage(Builder.Get(), Frontend::DefaultGraphPageName, bFocusEditor, Result);
+				if (Result == EMetaSoundBuilderResult::Succeeded)
+				{
+					BuildPageName = Frontend::DefaultGraphPageName;
+				}
 			}
 		}
 		else
diff --git a/Engine/Plugins/Runtime/Metasound/Source/MetasoundEditor/Private/MetasoundEditorSubsystem.cpp b/Engine/Plugins/Runtime/Metasound/Source/MetasoundEditor/Private/MetasoundEditorSubsystem.cpp
--- a/Engine/Plugins/Runtime/Metasound/Source/MetasoundEditor/Private/MetasoundEditorSubsystem.cpp
+++ b/Engine/Plugins/Runtime/Metasound/Source/MetasoundEditor/Private/MetasoundEditorSubsystem.cpp
@@ -239,7 +239,7 @@ void UMetaSoundEditorSubsystem::InitEdGraph(UObject& InMetaSound)
 	Metasound::Editor::FGraphBuilder::BindEditorGraph(IDocumentBuilderRegistry::GetChecked().FindOrBeginBuilding(&InMetaSound));
 }
 
-void UMetaSoundEditorSubsystem::RegisterGraphWithFrontend(UObject& InMetaSound, bool bInForceViewSynchronization)
+void UMetaSoundEditorSubsystem::RegisterGraphWithFrontend(UObject& InMetaSound, bool bInForceViewSynchronization) const
 {
 	Metasound::Editor::FGraphBuilder::RegisterGraphWithFrontend(InMetaSound, bInForceViewSynchronization);
 }
@@ -249,10 +249,10 @@ void UMetaSoundEditorSubsystem::RegisterToolbarExtender(TSharedRef<FExtender> In
 	EditorToolbarExtenders.AddUnique(InExtender);
 }
 
-void UMetaSoundEditorSubsystem::SetFocusedPage(UMetaSoundBuilderBase* InBuilder, FName InPageName, bool bFocusPageEditor, EMetaSoundBuilderResult& OutResult) const
+void UMetaSoundEditorSubsystem::SetFocusedPage(UMetaSoundBuilderBase* Builder, FName PageName, bool bFocusPageEditor, EMetaSoundBuilderResult& OutResult) const
 {
 	using namespace Metasound::Frontend;
-	if (!InBuilder)
+	if (!Builder)
 	{
 		OutResult = EMetaSoundBuilderResult::Failed;
 		return;
@@ -260,31 +260,64 @@ void UMetaSoundEditorSubsystem::SetFocusedPage(UMetaSoundBuilderBase* InBuilder,
 
 	const UMetaSoundSettings* Settings = GetDefault<UMetaSoundSettings>();
 	check(Settings);
-	if (const FMetaSoundPageSettings* PageSettings = Settings->FindPageSettings(InPageName))
+	if (const FMetaSoundPageSettings* PageSettings = Settings->FindPageSettings(PageName))
 	{
-		const FScopedTransaction Transaction(FText::Format(LOCTEXT("SetFocusedBuildPageTransactionFormat", "Set Focused Builder Page '{0}'"), FText::FromName(InPageName)));
-		InBuilder->Modify();
-		if (InBuilder->GetBuilder().SetBuildPageID(PageSettings->UniqueId))
+		const bool bFocusedPage = SetFocusedPageInternal(*PageSettings, *Builder, bFocusPageEditor);
+		if (bFocusedPage)
 		{
-			if (UMetasoundEditorSettings* EditorSettings = GetMutableDefault<UMetasoundEditorSettings>())
+			OutResult = EMetaSoundBuilderResult::Succeeded;
+			return;
+		}
+	}
+
+	OutResult = EMetaSoundBuilderResult::Failed;
+}
+
+bool UMetaSoundEditorSubsystem::SetFocusedPage(UMetaSoundBuilderBase& Builder, const FGuid& InPageID, bool bFocusPageEditor) const
+{
+	using namespace Metasound::Frontend;
+
+	const UMetaSoundSettings* Settings = GetDefault<UMetaSoundSettings>();
+	check(Settings);
+	if (const FMetaSoundPageSettings* PageSettings = Settings->FindPageSettings(InPageID))
+	{
+		return SetFocusedPageInternal(*PageSettings, Builder, bFocusPageEditor);
+	}
+
+	return false;
+}
+
+bool UMetaSoundEditorSubsystem::SetFocusedPageInternal(const FMetaSoundPageSettings& InPageSettings, UMetaSoundBuilderBase& Builder, bool bFocusPageEditor) const
+{
+	using namespace Metasound::Frontend;
+
+	const FScopedTransaction Transaction(FText::Format(LOCTEXT("SetFocusedPageTransactionFormat", "Set Focused Page '{0}'"), FText::FromName(InPageSettings.Name)));
+	Builder.Modify();
+	if (Builder.GetBuilder().SetBuildPageID(InPageSettings.UniqueId))
+	{
+		if (UMetasoundEditorSettings* EditorSettings = GetMutableDefault<UMetasoundEditorSettings>())
+		{
+			if (EditorSettings->AuditionPageMode == EAuditionPageMode::Focused)
 			{
-				EditorSettings->AuditionTargetPage = InPageName;
+				EditorSettings->AuditionTargetPage = InPageSettings.Name;
+
+				// Reregister to ensure all future audible instances are using the new page implementation.
+				RegisterGraphWithFrontend(Builder.GetBuilder().CastDocumentObjectChecked<UObject>());
 			}
+		}
 
-			if (GEditor && bFocusPageEditor)
+		if (GEditor && bFocusPageEditor)
+		{
+			if (UAssetEditorSubsystem* AssetEditorSubsystem = GEditor->GetEditorSubsystem<UAssetEditorSubsystem>())
 			{
-				if (UAssetEditorSubsystem* AssetEditorSubsystem = GEditor->GetEditorSubsystem<UAssetEditorSubsystem>())
-				{
-					AssetEditorSubsystem->OpenEditorForAsset(&InBuilder->GetConstBuilder().CastDocumentObjectChecked<UObject>());
-				}
+				AssetEditorSubsystem->OpenEditorForAsset(&Builder.GetConstBuilder().CastDocumentObjectChecked<UObject>());
 			}
-
-			OutResult = EMetaSoundBuilderResult::Succeeded;
-			return;
 		}
-	};
 
-	OutResult = EMetaSoundBuilderResult::Failed;
+		return true;
+	}
+
+	return false;
 }
 
 bool UMetaSoundEditorSubsystem::UnregisterToolbarExtender(TSharedRef<FExtender> InExtender)
diff --git a/Engine/Plugins/Runtime/Metasound/Source/MetasoundEditor/Public/MetasoundEditorSubsystem.h b/Engine/Plugins/Runtime/Metasound/Source/MetasoundEditor/Public/MetasoundEditorSubsystem.h
--- a/Engine/Plugins/Runtime/Metasound/Source/MetasoundEditor/Public/MetasoundEditorSubsystem.h
+++ b/Engine/Plugins/Runtime/Metasound/Source/MetasoundEditor/Public/MetasoundEditorSubsystem.h
@@ -15,6 +15,9 @@ class UMetaSoundBuilderBase;
 class UMetasoundEditorGraphMember;
 class UMetasoundEditorGraphMemberDefaultLiteral;
 
+struct FMetaSoundPageSettings;
+
+
 /** The subsystem in charge of editor MetaSound functionality */
 UCLASS()
 class METASOUNDEDITOR_API UMetaSoundEditorSubsystem : public UEditorSubsystem
@@ -70,7 +73,7 @@ public:
 	// @param InMetaSound - MetaSound to register
 	// @param bInForceSynchronize - Forces the synchronize flag for all open graphs being registered by this call (all referenced graphs and
 	// referencing graphs open in editors)
-	void RegisterGraphWithFrontend(UObject& InMetaSound, bool bInForceViewSynchronization = false);
+	void RegisterGraphWithFrontend(UObject& InMetaSound, bool bInForceViewSynchronization = false) const;
 
 	// Register toolbar extender that will be displayed in the MetaSound Asset Editor.
 	void RegisterToolbarExtender(TSharedRef<FExtender> InExtender);
@@ -79,7 +82,14 @@ public:
 	// the provided builder to the given page name if and sets the audition target page to
 	// the provided name. If the given builder has an asset editor open, optionally opens or brings
 	// that editor's associated page into user focus.
-	void SetFocusedPage(UMetaSoundBuilderBase* InBuilder, FName InPageName, bool bFocusPageEditor, EMetaSoundBuilderResult& OutResult) const;
+	UFUNCTION(BlueprintCallable, Category = "Audio|MetaSound|Builder|Editor", meta = (ExpandEnumAsExecs = "OutResult"))
+	void SetFocusedPage(UMetaSoundBuilderBase* Builder, FName PageName, bool bFocusPageEditor, EMetaSoundBuilderResult& OutResult) const;
+
+	// If the given PageID is implemented on the provided builder, sets the focused page of
+	// the provided builder to the given PageID if and sets the audition target page to
+	// the provided ID. If the given builder has an asset editor open, optionally opens or brings
+	// that editor's associated PageID into user focus.
+	bool SetFocusedPage(UMetaSoundBuilderBase& Builder, const FGuid& InPageID, bool bFocusPageEditor) const;
 
 	// Unregisters toolbar extender that is displayed in the MetaSound Asset Editor.
 	bool UnregisterToolbarExtender(TSharedRef<FExtender> InExtender);
@@ -94,6 +104,8 @@ public:
 	static const UMetaSoundEditorSubsystem& GetConstChecked();
 
 private:
+	bool SetFocusedPageInternal(const FMetaSoundPageSettings& InPageSettings, UMetaSoundBuilderBase& Builder, bool bFocusPageEditor) const;
+
 	// Copy over sound wave settings such as attenuation, modulation, and sound class from the template sound wave to the MetaSound
 	void SetSoundWaveSettingsFromTemplate(USoundWave& NewMetasound, const USoundWave& TemplateSoundWave) const;
 
