Fix bug in operator cache resulting in too many precached metasounds.
#rnx

diff --git a/Engine/Plugins/Runtime/Metasound/Source/MetasoundGenerator/Private/MetasoundOperatorCache.cpp b/Engine/Plugins/Runtime/Metasound/Source/MetasoundGenerator/Private/MetasoundOperatorCache.cpp
--- a/Engine/Plugins/Runtime/Metasound/Source/MetasoundGenerator/Private/MetasoundOperatorCache.cpp
+++ b/Engine/Plugins/Runtime/Metasound/Source/MetasoundGenerator/Private/MetasoundOperatorCache.cpp
@@ -479,6 +479,7 @@ namespace Metasound
 			}
 
 			int32 NumToBuild = PreCacheData->NumInstances;
+			const FOperatorPoolEntryID EntryID{PreCacheData->InitParams.Graph->GetInstanceID(), PreCacheData->InitParams.OperatorSettings};
 
 			if (PreCacheData->bTouchExisting)
 			{
@@ -486,8 +487,8 @@ namespace Metasound
 				if (OperatorPool.IsValid())
 				{
 					// Get the number of instances already in the cache & move pre-existing to the top of the cache
-					const int32 NumInCache = OperatorPool->GetNumCachedOperatorsWithAssetClassID(PreCacheData->AssetClassID);
-					OperatorPool->TouchOperatorsViaAssetClassID(PreCacheData->AssetClassID, FMath::Min(NumInCache, NumToBuild));
+					const int32 NumInCache = OperatorPool->GetNumCachedOperatorsWithID(EntryID);
+					OperatorPool->TouchOperators(EntryID, FMath::Min(NumInCache, NumToBuild));
 					NumToBuild -= NumInCache;
 				}
 			}
@@ -516,7 +517,6 @@ namespace Metasound
 				FOperatorAndInputs OperatorAndInputs = GeneratorBuilder::BuildGraphOperator(PreCacheData->InitParams.OperatorSettings, PreCacheData->InitParams, BuildResults);
 				GeneratorBuilder::LogBuildErrors(PreCacheData->InitParams.MetaSoundName, BuildResults);
 
-				FOperatorPoolEntryID EntryID{PreCacheData->InitParams.Graph->GetInstanceID(), PreCacheData->InitParams.OperatorSettings};
 				OperatorPool->AddOperatorInternal(EntryID, MoveTemp(OperatorAndInputs));
 				OperatorPool->AddAssetIdToGraphIdLookUpInternal(PreCacheData->AssetClassID, EntryID);
 			}
