Trivial: Add cvar for setting MetaSound target page for convenience

#rnx

diff --git a/Engine/Plugins/Runtime/Metasound/Source/MetasoundEngine/Private/MetasoundSettings.cpp b/Engine/Plugins/Runtime/Metasound/Source/MetasoundEngine/Private/MetasoundSettings.cpp
--- a/Engine/Plugins/Runtime/Metasound/Source/MetasoundEngine/Private/MetasoundSettings.cpp
+++ b/Engine/Plugins/Runtime/Metasound/Source/MetasoundEngine/Private/MetasoundSettings.cpp
@@ -3,6 +3,7 @@
 #include "MetasoundSettings.h"
 
 #include "Algo/Count.h"
+#include "HAL/IConsoleManager.h"
 #include "MetasoundFrontendDocument.h"
 
 #define LOCTEXT_NAMESPACE "MetaSound"
@@ -10,6 +11,21 @@
 
 namespace Metasound::SettingsPrivate
 {
+	FAutoConsoleCommand CVarMetaSoundSetTargetPage(
+		TEXT("au.MetaSound.Pages.SetTarget"),
+		TEXT("Sets the target page to that with the given name. If name not specified or not found, command is ignored.\n"),
+		FConsoleCommandWithArgsDelegate::CreateLambda([](const TArray<FString>& Args)
+		{
+			if (!Args.IsEmpty())
+			{
+				if (UMetaSoundSettings* Settings = GetMutableDefault<UMetaSoundSettings>())
+				{
+					Settings->SetTargetPage(FName { *Args.Last() });
+				}
+			}
+		})
+	);
+
 #if WITH_EDITOR
 	template<typename SettingsStructType>
 	TSet<FName> GetStructNames(const TArray<SettingsStructType>& InSettings, int32 IgnoreIndex = INDEX_NONE)
