Metasounds - Add cvar to disable operator pool cache.

#rnx

diff --git a/Engine/Plugins/Runtime/Metasound/Source/MetasoundEngine/Private/MetasoundOperatorCacheSubsystem.cpp b/Engine/Plugins/Runtime/Metasound/Source/MetasoundEngine/Private/MetasoundOperatorCacheSubsystem.cpp
--- a/Engine/Plugins/Runtime/Metasound/Source/MetasoundEngine/Private/MetasoundOperatorCacheSubsystem.cpp
+++ b/Engine/Plugins/Runtime/Metasound/Source/MetasoundEngine/Private/MetasoundOperatorCacheSubsystem.cpp
@@ -10,6 +10,13 @@
 
 namespace Metasound::OperatorCachePrivate
 {
+	static bool bOperatorPrecacheEnabled = true;
+	static FAutoConsoleVariableRef CVarOperatorPrecacheEnabled(
+		TEXT("au.MetaSound.OperatorCache.EnablePrecache"),
+		bOperatorPrecacheEnabled,
+		TEXT("If precaching metasound operators via the UMetaSoundCacheSubsystem is enabled.")
+	);
+
 	TOptional<FMetasoundGeneratorInitParams> CreateInitParams(UMetaSoundSource& InMetaSound, const FSoundGeneratorInitParams& InParams)
 	{
 		using namespace Metasound::Frontend;
@@ -83,6 +90,12 @@ void UMetaSoundCacheSubsystem::PrecacheMetaSoundInternal(UMetaSoundSource* InMet
 	using namespace Audio;
 	using namespace Metasound;
 
+	if (!OperatorCachePrivate::bOperatorPrecacheEnabled)
+	{
+		UE_LOG(LogMetaSound, Log, TEXT("Ignoring PrecacheMetaSound request since au.MetaSound.OperatorCache.EnablePrecache is false."));
+		return;
+	}
+
 	IMetasoundGeneratorModule* Module = FModuleManager::GetModulePtr<IMetasoundGeneratorModule>("MetasoundGenerator");
 	if (!ensure(Module))
 	{
@@ -147,6 +160,9 @@ void UMetaSoundCacheSubsystem::RemoveCachedOperatorsForMetaSound(UMetaSoundSourc
 {
 	using namespace Metasound;
 
+	// Note: we're not checking the bOperatorPrecacheEnabled cvar here in case it was disable after some sounds had already been cached.
+	// If nothing is cached this will do very little.
+
 	if (!InMetaSound)
 	{
 		UE_LOG(LogMetaSound, Warning, TEXT("Remove Cached Operators called without being provided a MetaSound, ignoring request"));
