Categorize metasound operator pre-caching under LLM Audio/Metasound/OperatorPool

#rnx
#tests local llm replay

diff --git a/Engine/Plugins/Runtime/Metasound/Source/MetasoundGenerator/Private/MetasoundOperatorCache.cpp b/Engine/Plugins/Runtime/Metasound/Source/MetasoundGenerator/Private/MetasoundOperatorCache.cpp
--- a/Engine/Plugins/Runtime/Metasound/Source/MetasoundGenerator/Private/MetasoundOperatorCache.cpp
+++ b/Engine/Plugins/Runtime/Metasound/Source/MetasoundGenerator/Private/MetasoundOperatorCache.cpp
@@ -8,6 +8,7 @@
 #include "Containers/Map.h"
 #include "HAL/CriticalSection.h"
 #include "HAL/IConsoleManager.h"
+#include "HAL/LowLevelMemTracker.h"
 #include "MetasoundGeneratorModuleImpl.h"
 #include "MetasoundGeneratorBuilder.h"
 #include "MetasoundOperatorCacheStatTracker.h"
@@ -21,6 +22,8 @@
 
 namespace Metasound
 {
+	LLM_DEFINE_TAG(Audio_Metasound_OperatorPool);
+
 #if METASOUND_OPERATORCACHEPROFILER_ENABLED
 	TRACE_DECLARE_INT_COUNTER(MetaSound_OperatorPool_NumOperators, TEXT("MetaSound/OperatorPool/NumOperatorsInPool"));
 	TRACE_DECLARE_FLOAT_COUNTER(MetaSound_OperatorPool_HitRatio, TEXT("MetaSound/OperatorPool/HitRatio"));
@@ -447,7 +450,7 @@ namespace Metasound
 		{
 			using namespace OperatorPoolPrivate;
 
-			METASOUND_LLM_SCOPE;
+			LLM_SCOPE_BYTAG(Audio_Metasound_OperatorPool);
 			METASOUND_TRACE_CPUPROFILER_EVENT_SCOPE(Metasound::FOperatorPool::AsyncOperatorPrecache)
 
 			if (!ensure(PreCacheData))
