Fix for missed MetaSound default value/input name overrides in paste & promotion logic from earlier CL
- Minor clean-up replacing old input handle logic with builder API
- Minor tool-tip fix-up
#rnx
#tests in-editor asset actions

diff --git a/Engine/Plugins/Runtime/Metasound/Source/MetasoundEditor/Private/MetasoundEditorDocumentClipboardUtils.cpp b/Engine/Plugins/Runtime/Metasound/Source/MetasoundEditor/Private/MetasoundEditorDocumentClipboardUtils.cpp
--- a/Engine/Plugins/Runtime/Metasound/Source/MetasoundEditor/Private/MetasoundEditorDocumentClipboardUtils.cpp
+++ b/Engine/Plugins/Runtime/Metasound/Source/MetasoundEditor/Private/MetasoundEditorDocumentClipboardUtils.cpp
@@ -67,7 +67,7 @@ namespace Metasound::Editor
 							FCreateNodeVertexParams VertexParams;
 							VertexParams.DataType = Breadcrumb.DataType;
 
-							FMetasoundFrontendClassInput ClassInput = FGraphBuilder::CreateUniqueClassInput(*OutAsset.GetOwningAsset(), VertexParams);
+							FMetasoundFrontendClassInput ClassInput = FGraphBuilder::CreateUniqueClassInput(*OutAsset.GetOwningAsset(), VertexParams, &Breadcrumb.DefaultLiteral, &Breadcrumb.MemberName);
 							if (const FMetasoundFrontendNode* NewNode = Builder.AddGraphInput(ClassInput))
 							{
 								Input = Graph.FindOrAddInput(NewNode->GetID());
diff --git a/Engine/Plugins/Runtime/Metasound/Source/MetasoundEditor/Private/MetasoundEditorGraphSchema.cpp b/Engine/Plugins/Runtime/Metasound/Source/MetasoundEditor/Private/MetasoundEditorGraphSchema.cpp
--- a/Engine/Plugins/Runtime/Metasound/Source/MetasoundEditor/Private/MetasoundEditorGraphSchema.cpp
+++ b/Engine/Plugins/Runtime/Metasound/Source/MetasoundEditor/Private/MetasoundEditorGraphSchema.cpp
@@ -548,30 +548,35 @@ FMetasoundGraphSchemaAction_PromoteToInput::FMetasoundGraphSchemaAction_PromoteT
 UEdGraphNode* FMetasoundGraphSchemaAction_PromoteToInput::PerformAction(UEdGraph* ParentGraph, UEdGraphPin* FromPin, const FVector2D InLocation, bool bSelectNewNode /* = true */)
 {
 	using namespace Metasound::Editor;
-	using namespace Metasound::Engine;
 	using namespace Metasound::Frontend;
 
-	FConstInputHandle InputHandle = FGraphBuilder::GetConstInputHandleFromPin(FromPin);
-	if (!ensure(InputHandle->IsValid()))
+	UMetasoundEditorGraph* MetasoundGraph = CastChecked<UMetasoundEditorGraph>(ParentGraph);
+	UObject& ParentMetasound = MetasoundGraph->GetMetasoundChecked();
+	FMetaSoundFrontendDocumentBuilder& Builder = IDocumentBuilderRegistry::GetChecked().FindOrBeginBuilding(&ParentMetasound);
+	FMetasoundFrontendVertexHandle InputVertexHandle = FGraphBuilder::GetPinVertexHandle(Builder, FromPin);
+	if (!ensure(InputVertexHandle.IsSet()))
+	{
+		return nullptr;
+	}
+
+	const FMetasoundFrontendVertex* InputVertex = Builder.FindNodeInput(InputVertexHandle.NodeID, InputVertexHandle.VertexID);;
+	if (!ensure(InputVertex))
 	{
 		return nullptr;
 	}
 
 	const FScopedTransaction Transaction(LOCTEXT("PromoteNodeInputToGraphInput", "Promote MetaSound Node Input to Graph Input"));
-	UMetasoundEditorGraph* MetasoundGraph = CastChecked<UMetasoundEditorGraph>(ParentGraph);
-	UObject& ParentMetasound = MetasoundGraph->GetMetasoundChecked();
 	ParentMetasound.Modify();
 	MetasoundGraph->Modify();
 
 	FMetasoundFrontendLiteral DefaultValue;
 	FGraphBuilder::GetPinLiteral(*FromPin, DefaultValue);
-
-	const FName InputName = InputHandle->GetName();
-	// The promoted input must have the same vertex access type in order for it to be connectible
-	const FCreateNodeVertexParams VertexParams = { InputHandle->GetDataType(), InputHandle->GetVertexAccessType() };
-
-	FMetaSoundFrontendDocumentBuilder& Builder = FDocumentBuilderRegistry::GetChecked().FindOrBeginBuilding(&ParentMetasound);
-	FMetasoundFrontendClassInput ClassInput = FGraphBuilder::CreateUniqueClassInput(ParentMetasound, VertexParams);
+	const FCreateNodeVertexParams VertexParams =
+	{
+		InputVertex->TypeName,
+		Builder.GetNodeInputAccessType(InputVertexHandle.NodeID, InputVertex->VertexID)
+	};
+	FMetasoundFrontendClassInput ClassInput = FGraphBuilder::CreateUniqueClassInput(ParentMetasound, VertexParams, &DefaultValue, &InputVertex->Name);
 	if (const FMetasoundFrontendNode* NewNode = Builder.AddGraphInput(ClassInput))
 	{
 		UMetasoundEditorGraphInput* Input = MetasoundGraph->FindOrAddInput(NewNode->GetID());
@@ -757,28 +762,34 @@ FMetasoundGraphSchemaAction_PromoteToOutput::FMetasoundGraphSchemaAction_Promote
 
 UEdGraphNode* FMetasoundGraphSchemaAction_PromoteToOutput::PerformAction(UEdGraph* ParentGraph, UEdGraphPin* FromPin, const FVector2D InLocation, bool bSelectNewNode /* = true */)
 {
-	using namespace Metasound;
 	using namespace Metasound::Editor;
 	using namespace Metasound::Frontend;
 
-	FConstOutputHandle OutputHandle = FGraphBuilder::GetConstOutputHandleFromPin(FromPin);
-	if (!ensure(OutputHandle->IsValid()))
+	UMetasoundEditorGraph* MetasoundGraph = CastChecked<UMetasoundEditorGraph>(ParentGraph);
+	UObject& ParentMetasound = MetasoundGraph->GetMetasoundChecked();
+	FMetaSoundFrontendDocumentBuilder& Builder = IDocumentBuilderRegistry::GetChecked().FindOrBeginBuilding(&ParentMetasound);
+	FMetasoundFrontendVertexHandle OutputVertexHandle = FGraphBuilder::GetPinVertexHandle(Builder, FromPin);
+	if (!ensure(OutputVertexHandle.IsSet()))
+	{
+		return nullptr;
+	}
+
+	const FMetasoundFrontendVertex* OutputVertex = Builder.FindNodeOutput(OutputVertexHandle.NodeID, OutputVertexHandle.VertexID);;
+	if (!ensure(OutputVertex))
 	{
 		return nullptr;
 	}
 
 	const FScopedTransaction Transaction(LOCTEXT("PromoteNodeOutputToGraphOutput", "Promote MetaSound Node Output to Graph Output"));
-	UMetasoundEditorGraph* MetasoundGraph = CastChecked<UMetasoundEditorGraph>(ParentGraph);
-	UObject& ParentMetasound = MetasoundGraph->GetMetasoundChecked();
 	ParentMetasound.Modify();
 	MetasoundGraph->Modify();
 
-	const FString OutputName = OutputHandle->GetName().ToString();
-	const FVertexName NewNodeName = FGraphBuilder::GenerateUniqueNameByClassType(ParentMetasound, EMetasoundFrontendClassType::Output, OutputName);
-	const FCreateNodeVertexParams VertexParams = { OutputHandle->GetDataType(), OutputHandle->GetVertexAccessType() };
-
-	const FMetasoundFrontendClassOutput ClassOutput = FGraphBuilder::CreateUniqueClassOutput(ParentMetasound, VertexParams, &NewNodeName);
-	FMetaSoundFrontendDocumentBuilder& Builder = IDocumentBuilderRegistry::GetChecked().FindOrBeginBuilding(&ParentMetasound);
+	const FCreateNodeVertexParams VertexParams =
+	{
+		OutputVertex->TypeName,
+		Builder.GetNodeOutputAccessType(OutputVertexHandle.NodeID, OutputVertex->VertexID)
+	};
+	const FMetasoundFrontendClassOutput ClassOutput = FGraphBuilder::CreateUniqueClassOutput(ParentMetasound, VertexParams, &OutputVertex->Name);
 	const FMetasoundFrontendNode* OutputNode = Builder.AddGraphOutput(ClassOutput);
 	if (ensure(OutputNode))
 	{
diff --git a/Engine/Plugins/Runtime/Metasound/Source/MetasoundEngine/Public/MetasoundSettings.h b/Engine/Plugins/Runtime/Metasound/Source/MetasoundEngine/Public/MetasoundSettings.h
--- a/Engine/Plugins/Runtime/Metasound/Source/MetasoundEngine/Public/MetasoundSettings.h
+++ b/Engine/Plugins/Runtime/Metasound/Source/MetasoundEngine/Public/MetasoundSettings.h
@@ -139,8 +139,8 @@ private:
 	Metasound::Engine::FOnPageSettingsUpdated OnPageSettingsUpdated;
 #endif //WITH_EDITORONLY_DATA
 
-	/** PageID to target when attempting to execute MetaSound. If target page is not implemented for the set
-	  *  platform, uses order of cooked pages (see 'Page Settings' for order) falling back to lower index-ordered page
+	/** Page Name to target when attempting to execute MetaSound. If target page is not implemented (or cooked in a runtime build)
+	  * for the active platform, uses order of cooked pages (see 'Page Settings' for order) falling back to lower index-ordered page
 	  * implemented in MetaSound asset. */
 	UPROPERTY(EditAnywhere, config, Category = Pages)
 	FName TargetPageName = Metasound::Frontend::DefaultGraphPageName;
