UE-219566 - Conversion node output does not connect to its previous exec when using Redo

diff --git a/Engine/Plugins/Runtime/Metasound/Source/MetasoundEditor/Private/MetasoundEditorGraphSchema.cpp b/Engine/Plugins/Runtime/Metasound/Source/MetasoundEditor/Private/MetasoundEditorGraphSchema.cpp
--- a/Engine/Plugins/Runtime/Metasound/Source/MetasoundEditor/Private/MetasoundEditorGraphSchema.cpp
+++ b/Engine/Plugins/Runtime/Metasound/Source/MetasoundEditor/Private/MetasoundEditorGraphSchema.cpp
@@ -1591,6 +1591,7 @@ bool UMetasoundEditorGraphSchema::TryCreateConnection(UEdGraphPin* PinA, UEdGrap
 		UMetasoundEditorGraph* MetaSoundGraph = CastChecked<UMetasoundEditorGraph>(InputPin->GetOwningNode()->GetGraph());
 		UObject& ParentMetaSound = MetaSoundGraph->GetMetasoundChecked();
 
+		MetaSoundGraph->Modify();
 		ParentMetaSound.Modify();
 
 		if (Connectability.PossibleConverterNodeClasses.Num() == 0)
@@ -1607,6 +1608,8 @@ bool UMetasoundEditorGraphSchema::TryCreateConnection(UEdGraphPin* PinA, UEdGrap
 		{
 			UEdGraphNode* InputNode = InputPin->GetOwningNode();
 			UEdGraphNode* OutputNode = OutputPin->GetOwningNode();
+
+			OutputPin->Modify();
 			
 			FVector2D Location = FVector2D();
 			Location += FVector2D(InputNode->NodePosX, InputNode->NodePosY);
