Fix naming inconsistencies (feedback from prior CL)

#rnx

diff --git a/Engine/Plugins/Runtime/Metasound/Source/MetasoundFrontend/Private/MetasoundAssetBase.cpp b/Engine/Plugins/Runtime/Metasound/Source/MetasoundFrontend/Private/MetasoundAssetBase.cpp
--- a/Engine/Plugins/Runtime/Metasound/Source/MetasoundFrontend/Private/MetasoundAssetBase.cpp
+++ b/Engine/Plugins/Runtime/Metasound/Source/MetasoundFrontend/Private/MetasoundAssetBase.cpp
@@ -73,11 +73,12 @@ namespace Metasound
 
 			// Registers node by copying document. Updates to document require re-registration.
 			// This registry entry does not support node creation as it is only intended to be
-			// used when cooking MetaSounds. 
-			class FDocumentNodeRegistryEntryForCook : public INodeRegistryEntry
+			// used when serializing MetaSounds in contexts not requiring any runtime model to
+			// be generated (ex. cooking commandlets that don't play or are validating MetaSounds, etc.).
+			class FDocumentNodeRegistryEntryForSerialization : public INodeRegistryEntry
 			{
 			public:
-				FDocumentNodeRegistryEntryForCook(const FMetasoundFrontendDocument& InDocument, const FTopLevelAssetPath& InAssetPath)
+				FDocumentNodeRegistryEntryForSerialization(const FMetasoundFrontendDocument& InDocument, const FTopLevelAssetPath& InAssetPath)
 					: Interfaces(InDocument.Interfaces)
 					, FrontendClass(InDocument.RootGraph)
 					, ClassInfo(InDocument.RootGraph, InAssetPath)
@@ -86,9 +87,9 @@ namespace Metasound
 					FrontendClass.Metadata.SetType(EMetasoundFrontendClassType::External);
 				}
 
-				FDocumentNodeRegistryEntryForCook(const FDocumentNodeRegistryEntryForCook& InOther) = default;
+				FDocumentNodeRegistryEntryForSerialization(const FDocumentNodeRegistryEntryForSerialization& InOther) = default;
 
-				virtual ~FDocumentNodeRegistryEntryForCook() = default;
+				virtual ~FDocumentNodeRegistryEntryForSerialization() = default;
 
 				virtual const FNodeClassInfo& GetClassInfo() const override
 				{
@@ -107,7 +108,7 @@ namespace Metasound
 
 				virtual TUniquePtr<INodeRegistryEntry> Clone() const override
 				{
-					return MakeUnique<FDocumentNodeRegistryEntryForCook>(*this);
+					return MakeUnique<FDocumentNodeRegistryEntryForSerialization>(*this);
 				}
 
 				virtual const TSet<FMetasoundFrontendVersion>* GetImplementedInterfaces() const override
@@ -388,7 +389,7 @@ void FMetasoundAssetBase::UpdateAndRegisterForSerialization()
 		return;
 	}
 
-	PreSaveReferencedDocuments();
+	UpdateAndRegisterReferencesForSerialization();
 	IMetaSoundAssetManager::GetChecked().AddOrUpdateAsset(*GetOwningAsset());
 
 	// Auto update must be done after all referenced asset classes are registered
@@ -430,7 +431,7 @@ void FMetasoundAssetBase::UpdateAndRegisterForSerialization()
 		TScriptInterface<IMetaSoundDocumentInterface> DocInterface(Owner);
 		const FMetasoundFrontendDocument& Document = DocInterface->GetConstDocument();
 		const FTopLevelAssetPath AssetPath = DocInterface->GetAssetPathChecked();
-		TUniquePtr<INodeRegistryEntry> RegistryEntry = MakeUnique<AssetBasePrivate::FDocumentNodeRegistryEntryForCook>(Document, AssetPath);
+		TUniquePtr<INodeRegistryEntry> RegistryEntry = MakeUnique<AssetBasePrivate::FDocumentNodeRegistryEntryForSerialization>(Document, AssetPath);
 
 		const FNodeRegistryKey NodeKey = FRegistryContainerImpl::Get().RegisterNode(MoveTemp(RegistryEntry));
 		GraphRegistryKey = FGraphRegistryKey { NodeKey, AssetPath };
@@ -938,12 +939,12 @@ void FMetasoundAssetBase::RegisterAssetDependencies(const Metasound::Frontend::F
 void FMetasoundAssetBase::CookReferencedMetaSounds()
 {
 #if WITH_EDITORONLY_DATA
-	PreSaveReferencedDocuments();
+	UpdateAndRegisterReferencesForSerialization();
 #endif // WITH_EDITORONLY_DATA
 }
 
 #if WITH_EDITORONLY_DATA
-void FMetasoundAssetBase::PreSaveReferencedDocuments()
+void FMetasoundAssetBase::UpdateAndRegisterReferencesForSerialization()
 {
 	using namespace Metasound::Frontend;
 
diff --git a/Engine/Plugins/Runtime/Metasound/Source/MetasoundFrontend/Public/MetasoundAssetBase.h b/Engine/Plugins/Runtime/Metasound/Source/MetasoundFrontend/Public/MetasoundAssetBase.h
--- a/Engine/Plugins/Runtime/Metasound/Source/MetasoundFrontend/Public/MetasoundAssetBase.h
+++ b/Engine/Plugins/Runtime/Metasound/Source/MetasoundFrontend/Public/MetasoundAssetBase.h
@@ -96,10 +96,9 @@ public:
 	void CookMetaSound();
 
 #if WITH_EDITORONLY_DATA
-	// Executes on this and referenced MetaSound document objects. Presave includes autoupdating and
-	// optimizing the document for runtime use, which is then registered with the MetaSound Frontend.
-	// Unlike 'RegisterGraphWithFrontend', this call does not generate required runtime data for graph
-	// execution.
+	// Updates and registers this and referenced MetaSound document objects with the NodeClass Registry. AutoUpdates and
+	// optimizes aforementioned documents for serialization. Unlike 'UpdateAndRegisterForRuntime', does not generate required
+	// runtime data for graph execution.
 	void UpdateAndRegisterForSerialization();
 #endif // WITH_EDITORONLY_DATA
 
@@ -270,7 +269,7 @@ protected:
 
 	bool AutoUpdate(bool bInLogWarningsOnDroppedConnection);
 
-	UE_DEPRECATED(5.5, "Moved to private, PreSaveReferencedDocuments implementation")
+	UE_DEPRECATED(5.5, "Moved to private, non-cook specific implementation")
 	void CookReferencedMetaSounds();
 
 	// Ensures all referenced graph classes are registered (or re-registers depending on options).
@@ -282,11 +281,8 @@ protected:
 private:
 #if WITH_EDITORONLY_DATA
 	void UpdateAssetRegistry();
-
-	bool bVersionedOnLoad = false;
-
-	void PreSaveReferencedDocuments();
-#endif
+	void UpdateAndRegisterReferencesForSerialization();
+#endif // WITH_EDITORONLY_DATA
 
 	// Checks if version is up-to-date. If so, returns true. If false, updates the interfaces within the given asset's document to the most recent version.
 	bool TryUpdateInterfaceFromVersion(const FMetasoundFrontendVersion& Version);
@@ -295,6 +291,10 @@ private:
 	// required, returns invalid interface (interface with no name and invalid version number).
 	FMetasoundFrontendInterface GetInterfaceToVersion(const FMetasoundFrontendVersion& InterfaceVersion) const;
 
+#if WITH_EDITORONLY_DATA
+	bool bVersionedOnLoad = false;
+#endif // WITH_EDITORONLY_DATA
+
 	Metasound::Frontend::FGraphRegistryKey GraphRegistryKey;
 };
 
