UE-215851 -  [MetaSound] - Renaming inputs/outputs highlights the name in different windows depending on the action

diff --git a/Engine/Plugins/Runtime/Metasound/Source/MetasoundEditor/Private/MetasoundEditorGraphSchema.cpp b/Engine/Plugins/Runtime/Metasound/Source/MetasoundEditor/Private/MetasoundEditorGraphSchema.cpp
--- a/Engine/Plugins/Runtime/Metasound/Source/MetasoundEditor/Private/MetasoundEditorGraphSchema.cpp
+++ b/Engine/Plugins/Runtime/Metasound/Source/MetasoundEditor/Private/MetasoundEditorGraphSchema.cpp
@@ -243,7 +243,7 @@ PRAGMA_ENABLE_DEPRECATION_WARNINGS
 				}
 			}
 
-			void SelecteNodeInEditor(UMetasoundEditorGraph& InMetaSoundGraph, UMetasoundEditorGraphNode& InNode)
+			void SelectNodeInEditor(UMetasoundEditorGraph& InMetaSoundGraph, UMetasoundEditorGraphNode& InNode)
 			{
 				TSharedPtr<FEditor> MetasoundEditor = FGraphBuilder::GetEditorForGraph(InMetaSoundGraph);
 				if (MetasoundEditor.IsValid())
@@ -252,6 +252,16 @@ PRAGMA_ENABLE_DEPRECATION_WARNINGS
 				}
 			}
 
+			void SelectNodeInEditorForRename(UMetasoundEditorGraph& InMetaSoundGraph, UMetasoundEditorGraphNode& InNode)
+			{
+				TSharedPtr<FEditor> MetasoundEditor = FGraphBuilder::GetEditorForGraph(InMetaSoundGraph);
+				if (MetasoundEditor.IsValid())
+				{
+					MetasoundEditor->ClearSelectionAndSelectNode(&InNode);
+					MetasoundEditor->SetDelayedRename();
+				}
+			}
+
 			UEdGraphNode* PromoteToVariable(const FName BaseName, UEdGraphPin& FromPin, const FName DataType, const FMetasoundFrontendClass& InVariableClass, const FVector2D& InLocation, bool bSelectNode)
 			{
 				using namespace Frontend;
@@ -317,7 +327,7 @@ PRAGMA_ENABLE_DEPRECATION_WARNINGS
 							if (ensure(SchemaPrivate::TryConnectNewNodeToMatchingDataTypePin(*EdGraphNode, &FromPin)))
 							{
 								FGraphBuilder::RegisterGraphWithFrontend(ParentMetaSound);
-								SelecteNodeInEditor(*MetaSoundGraph, *NewGraphNode);
+								SelectNodeInEditorForRename(*MetaSoundGraph, *NewGraphNode);
 								return EdGraphNode;
 							}
 						}
@@ -467,7 +477,7 @@ UEdGraphNode* FMetasoundGraphSchemaAction_NewNode::PerformAction(UEdGraph* Paren
 		NewGraphNode->UpdateFrontendNodeLocation(Location);
 		NewGraphNode->SyncLocationFromFrontendNode();
 		SchemaPrivate::TryConnectNewNodeToMatchingDataTypePin(*NewGraphNode, FromPin);
-		SchemaPrivate::SelecteNodeInEditor(*MetaSoundGraph, *NewGraphNode);
+		SchemaPrivate::SelectNodeInEditorForRename(*MetaSoundGraph, *NewGraphNode);
 		return NewGraphNode;
 	}
 
@@ -593,7 +603,7 @@ UEdGraphNode* FMetasoundGraphSchemaAction_PromoteToInput::PerformAction(UEdGraph
 					if (ensure(SchemaPrivate::TryConnectNewNodeToMatchingDataTypePin(*EdGraphNode, FromPin)))
 					{
 						FGraphBuilder::RegisterGraphWithFrontend(ParentMetasound);
-						SchemaPrivate::SelecteNodeInEditor(*MetasoundGraph, *NewGraphNode);
+						SchemaPrivate::SelectNodeInEditorForRename(*MetasoundGraph, *NewGraphNode);
 						return EdGraphNode;
 					}
 				}
@@ -807,7 +817,7 @@ UEdGraphNode* FMetasoundGraphSchemaAction_PromoteToOutput::PerformAction(UEdGrap
 				if (ensure(SchemaPrivate::TryConnectNewNodeToMatchingDataTypePin(*EdGraphNode, FromPin)))
 				{
 					FGraphBuilder::RegisterGraphWithFrontend(ParentMetasound);
-					SchemaPrivate::SelecteNodeInEditor(*MetasoundGraph, *NewGraphNode);
+					SchemaPrivate::SelectNodeInEditorForRename(*MetasoundGraph, *NewGraphNode);
 					return EdGraphNode;
 				}
 			}
