[MetaSound] Move node Alignment commands to be under Organization section
#rnx

diff --git a/Engine/Plugins/Runtime/Metasound/Source/MetasoundEditor/Private/MetasoundEditorGraphNode.cpp b/Engine/Plugins/Runtime/Metasound/Source/MetasoundEditor/Private/MetasoundEditorGraphNode.cpp
--- a/Engine/Plugins/Runtime/Metasound/Source/MetasoundEditor/Private/MetasoundEditorGraphNode.cpp
+++ b/Engine/Plugins/Runtime/Metasound/Source/MetasoundEditor/Private/MetasoundEditorGraphNode.cpp
@@ -710,7 +710,7 @@ void UMetasoundEditorGraphNode::GetNodeContextMenuActions(UToolMenu* Menu, UGrap
 	if (Context->Node)
 	{
 		{
-			FToolMenuSection& Section = Menu->AddSection("MetasoundEditorGraphNodeAlignment");
+			FToolMenuSection& Section = Menu->FindOrAddSection("MetasoundGraphNodeActionsOrganization", LOCTEXT("NodeActionsOrganizationMenuHeader", "Organization"));
 			Section.AddSubMenu("Alignment", LOCTEXT("AlignmentHeader", "Alignment"), FText(), FNewToolMenuDelegate::CreateLambda([](UToolMenu* SubMenu)
 			{
 				{
diff --git a/Engine/Plugins/Runtime/Metasound/Source/MetasoundEditor/Private/MetasoundEditorGraphSchema.cpp b/Engine/Plugins/Runtime/Metasound/Source/MetasoundEditor/Private/MetasoundEditorGraphSchema.cpp
--- a/Engine/Plugins/Runtime/Metasound/Source/MetasoundEditor/Private/MetasoundEditorGraphSchema.cpp
+++ b/Engine/Plugins/Runtime/Metasound/Source/MetasoundEditor/Private/MetasoundEditorGraphSchema.cpp
@@ -1338,7 +1338,7 @@ void UMetasoundEditorGraphSchema::GetContextMenuActions(class UToolMenu* Menu, c
 			LOCTEXT("FindReferencesInGraph_Tooltip", "Find References to the selected Node in the current Graph"),
 			FSlateIcon());
 
-		FToolMenuSection& OrganizationSection = Menu->AddSection("MetasoundGraphSchemaNodeActionsOrganization", LOCTEXT("NodeActionsOrganizationMenuHeader", "Organization"));
+		FToolMenuSection& OrganizationSection = Menu->FindOrAddSection("MetasoundGraphNodeActionsOrganization", LOCTEXT("NodeActionsOrganizationMenuHeader", "Organization"));
 
 		// Only display update ability if node is of type external
 		// and node registry is reporting a major update is available.
