Add initial page implementation of input class default literals

#rnx

diff --git a/Engine/Plugins/Runtime/Metasound/Source/MetasoundEditor/Private/MetasoundDetailCustomization.cpp b/Engine/Plugins/Runtime/Metasound/Source/MetasoundEditor/Private/MetasoundDetailCustomization.cpp
--- a/Engine/Plugins/Runtime/Metasound/Source/MetasoundEditor/Private/MetasoundDetailCustomization.cpp
+++ b/Engine/Plugins/Runtime/Metasound/Source/MetasoundEditor/Private/MetasoundDetailCustomization.cpp
@@ -579,16 +579,16 @@ namespace Metasound::Editor
 			{
 				constexpr bool bFocusEditor = false; // Already focused by user action
 				EMetaSoundBuilderResult Result = EMetaSoundBuilderResult::Failed;
-				UMetaSoundEditorSubsystem::GetConstChecked().SetFocusedPage(Builder.Get(), Frontend::DefaultGraphPageName, bFocusEditor, Result);
+				UMetaSoundEditorSubsystem::GetConstChecked().SetFocusedPage(Builder.Get(), Frontend::DefaultPageName, bFocusEditor, Result);
 				if (Result == EMetaSoundBuilderResult::Succeeded)
 				{
-					BuildPageName = Frontend::DefaultGraphPageName;
+					BuildPageName = Frontend::DefaultPageName;
 				}
 			}
 		}
 		else
 		{
-			BuildPageName = Frontend::DefaultGraphPageName;
+			BuildPageName = Frontend::DefaultPageName;
 		}
 
 		UpdateItemNames();
diff --git a/Engine/Plugins/Runtime/Metasound/Source/MetasoundEditor/Private/MetasoundEditor.cpp b/Engine/Plugins/Runtime/Metasound/Source/MetasoundEditor/Private/MetasoundEditor.cpp
--- a/Engine/Plugins/Runtime/Metasound/Source/MetasoundEditor/Private/MetasoundEditor.cpp
+++ b/Engine/Plugins/Runtime/Metasound/Source/MetasoundEditor/Private/MetasoundEditor.cpp
@@ -913,9 +913,16 @@ namespace Metasound
 
 			if (AssetEditorPrivate::EnablePageEditor)
 			{
-				InTabManager->RegisterTabSpawner(TabNamesPrivate::Pages, FOnSpawnTab::CreateLambda([InPagesDetails = PagesDetails](const FSpawnTabArgs&)
+				InTabManager->RegisterTabSpawner(TabNamesPrivate::Pages, FOnSpawnTab::CreateLambda([this, InPagesDetails = PagesDetails](const FSpawnTabArgs&)
 				{
-					return SNew(SDockTab).Label(LOCTEXT("MetasoundPagesDetailsTitle", "Pages"))[InPagesDetails.ToSharedRef()];
+					return SNew(SDockTab)
+						.Visibility(TAttribute<EVisibility>::Create([this]()
+						{
+							return Builder.IsValid() && Builder->IsPreset()
+								? EVisibility::Hidden
+								: EVisibility::Visible;
+						}))
+						.Label(LOCTEXT("MetasoundPagesDetailsTitle", "Pages"))[ InPagesDetails.ToSharedRef() ];
 				}))
 				.SetDisplayName(LOCTEXT("PagesTab", "Pages"))
 					.SetGroup(WorkspaceMenuCategoryRef)
@@ -3429,7 +3436,7 @@ namespace Metasound
 			{
 				TSharedPtr<FMetasoundGraphMemberSchemaAction> MetasoundAction = StaticCastSharedPtr<FMetasoundGraphMemberSchemaAction>(Action);
 				if (MetasoundAction.IsValid())
-				{						
+				{
 					if (const UMetasoundEditorGraphVertex* GraphVertex = Cast<UMetasoundEditorGraphVertex>(MetasoundAction->GetGraphMember()))
 					{
 						if (GraphVertex->IsInterfaceMember())
@@ -3438,7 +3445,7 @@ namespace Metasound
 						}
 					}
 				}
-			}		
+			}
 
 			return true;
 		}
@@ -3507,7 +3514,7 @@ namespace Metasound
 						{
 							if (const FMetasoundFrontendClassInput* SourceInput = DocumentBuilder.FindGraphInput(SourceMemberName))
 							{
-								if (const FMetasoundFrontendNode* FrontendNode = DocumentBuilder.DuplicateGraphInput(*SourceInput, SourceInput->DefaultLiteral, Name))
+								if (const FMetasoundFrontendNode* FrontendNode = DocumentBuilder.DuplicateGraphInput(*SourceInput, Name))
 								{
 									FGraphBuilder::SynchronizeGraphMembers(DocumentBuilder, Graph);
 									NewGraphMember = Graph.FindInput(FrontendNode->Name);
diff --git a/Engine/Plugins/Runtime/Metasound/Source/MetasoundEditor/Private/MetasoundEditorGraph.cpp b/Engine/Plugins/Runtime/Metasound/Source/MetasoundEditor/Private/MetasoundEditorGraph.cpp
--- a/Engine/Plugins/Runtime/Metasound/Source/MetasoundEditor/Private/MetasoundEditorGraph.cpp
+++ b/Engine/Plugins/Runtime/Metasound/Source/MetasoundEditor/Private/MetasoundEditorGraph.cpp
@@ -1712,8 +1712,7 @@ UMetasoundEditorGraphInput* UMetasoundEditorGraph::FindInput(FName InName) const
 
 UMetasoundEditorGraphInput* UMetasoundEditorGraph::FindOrAddInput(const FGuid& InNodeID)
 {
-	using namespace Metasound::Editor;
-	using namespace Metasound::Frontend;
+	using namespace Metasound;
 
 	if (TObjectPtr<UMetasoundEditorGraphInput> Input = FindInput(InNodeID))
 	{
@@ -1725,12 +1724,12 @@ UMetasoundEditorGraphInput* UMetasoundEditorGraph::FindOrAddInput(const FGuid& I
 	{
 		if (const FMetasoundFrontendClassInput* ClassInput = Builder.FindGraphInput(Node->Name))
 		{
-			const FMetasoundFrontendLiteral& DefaultLiteral = ClassInput->DefaultLiteral;
+			const FMetasoundFrontendLiteral& DefaultLiteral = ClassInput->FindConstDefaultChecked(Frontend::DefaultPageID);
 			if (const FMetasoundFrontendClass* Class = Builder.FindDependency(Node->ClassID))
 			{
 				FMetasoundFrontendClassName ClassName = Class->Metadata.GetClassName();
 
-				UMetasoundEditorGraphInput* NewInput = NewObject<UMetasoundEditorGraphInput>(this, GraphPrivate::GetUniqueTransientMemberName(), RF_Transactional);
+				UMetasoundEditorGraphInput* NewInput = NewObject<UMetasoundEditorGraphInput>(this, Editor::GraphPrivate::GetUniqueTransientMemberName(), RF_Transactional);
 				if (ensure(NewInput))
 				{
 					NewInput->InitMember(ClassInput->TypeName, DefaultLiteral, InNodeID, MoveTemp(ClassName));
diff --git a/Engine/Plugins/Runtime/Metasound/Source/MetasoundEditor/Private/MetasoundEditorGraphBuilder.cpp b/Engine/Plugins/Runtime/Metasound/Source/MetasoundEditor/Private/MetasoundEditorGraphBuilder.cpp
--- a/Engine/Plugins/Runtime/Metasound/Source/MetasoundEditor/Private/MetasoundEditorGraphBuilder.cpp
+++ b/Engine/Plugins/Runtime/Metasound/Source/MetasoundEditor/Private/MetasoundEditorGraphBuilder.cpp
@@ -1366,14 +1366,14 @@ namespace Metasound
 			// Can be unset if attempting to mirror parameters from a reroute, so default to reference
 			ClassInput.AccessType = InParams.AccessType == EMetasoundFrontendVertexAccessType::Unset ? EMetasoundFrontendVertexAccessType::Reference : InParams.AccessType;
 
-			if (nullptr != InDefaultValue)
+			if (InDefaultValue)
 			{
-				ClassInput.DefaultLiteral = *InDefaultValue;
+				ClassInput.InitDefault(*InDefaultValue);
 			}
 			else
 			{
 				Metasound::FLiteral Literal = Frontend::IDataTypeRegistry::Get().CreateDefaultLiteral(InParams.DataType);
-				ClassInput.DefaultLiteral.SetFromLiteral(Literal);
+				ClassInput.InitDefault().SetFromLiteral(Literal);
 			}
 
 			return ClassInput;
diff --git a/Engine/Plugins/Runtime/Metasound/Source/MetasoundEngine/Private/MetasoundBuilderSubsystem.cpp b/Engine/Plugins/Runtime/Metasound/Source/MetasoundEngine/Private/MetasoundBuilderSubsystem.cpp
--- a/Engine/Plugins/Runtime/Metasound/Source/MetasoundEngine/Private/MetasoundBuilderSubsystem.cpp
+++ b/Engine/Plugins/Runtime/Metasound/Source/MetasoundEngine/Private/MetasoundBuilderSubsystem.cpp
@@ -394,7 +394,7 @@ void UMetaSoundSourceBuilder::OnInputAdded(int32 InputIndex)
 			}
 		}
 
-		const FLiteral NewInputLiteral = NewInput.DefaultLiteral.ToLiteral(NewInput.TypeName);
+		const FLiteral NewInputLiteral = NewInput.FindConstDefaultChecked(Frontend::DefaultPageID).ToLiteral(NewInput.TypeName);
 		Transactor.AddInputDataDestination(NewInput.NodeID, NewInput.Name, NewInputLiteral, &UMetaSoundSourceBuilder::CreateDataReference);
 		return true;
 	});
@@ -451,12 +451,12 @@ void UMetaSoundSourceBuilder::OnNodeAdded(int32 NodeIndex) const
 				// Check for default literal on class input
 				else if (const FMetasoundFrontendClassInput* GraphInput = Builder.FindGraphInput(InputVertex.Name))
 				{
-					DefaultLiteral = &(GraphInput->DefaultLiteral);
+					DefaultLiteral = &(GraphInput->FindConstDefaultChecked(Frontend::DefaultPageID));
 				}
 				else
 				{
 					// As a last resort, get default literal on node class 
-					DefaultLiteral = &(NodeClass->Interface.Inputs.Last().DefaultLiteral);
+					DefaultLiteral = &(NodeClass->Interface.Inputs.Last().FindConstDefaultChecked(Frontend::DefaultPageID));
 				}
 
 				FInputNodeConstructorParams InitData
diff --git a/Engine/Plugins/Runtime/Metasound/Source/MetasoundEngine/Private/MetasoundSettings.cpp b/Engine/Plugins/Runtime/Metasound/Source/MetasoundEngine/Private/MetasoundSettings.cpp
--- a/Engine/Plugins/Runtime/Metasound/Source/MetasoundEngine/Private/MetasoundSettings.cpp
+++ b/Engine/Plugins/Runtime/Metasound/Source/MetasoundEngine/Private/MetasoundSettings.cpp
@@ -156,10 +156,10 @@ void UMetaSoundSettings::ConformPageSettingsDefault(bool bNotifyDefaultConformed
 	for (int32 Index = PageSettings.Num() - 1; Index >= 0; --Index)
 	{
 		FMetaSoundPageSettings& Page = PageSettings[Index];
-		const bool bIsDefaultName = Page.Name == Frontend::DefaultGraphPageName;
+		const bool bIsDefaultName = Page.Name == Frontend::DefaultPageName;
 		if (bIsDefaultName)
 		{
-			if (Page.UniqueId != Frontend::DefaultGraphPageID)
+			if (Page.UniqueId != Frontend::DefaultPageID)
 			{
 				Page.UniqueId = { };
 				bDefaultConformed = true;
@@ -169,7 +169,7 @@ void UMetaSoundSettings::ConformPageSettingsDefault(bool bNotifyDefaultConformed
 		}
 		else
 		{
-			if (Page.UniqueId == Frontend::DefaultGraphPageID)
+			if (Page.UniqueId == Frontend::DefaultPageID)
 			{
 				Page.UniqueId = FGuid::NewGuid();
 				bDefaultConformed = true;
@@ -180,7 +180,7 @@ void UMetaSoundSettings::ConformPageSettingsDefault(bool bNotifyDefaultConformed
 	if (!bContainsPageDefault)
 	{
 		FMetaSoundPageSettings DefaultSettings;
-		DefaultSettings.Name = Frontend::DefaultGraphPageName;
+		DefaultSettings.Name = Frontend::DefaultPageName;
 		PageSettings.Insert(MoveTemp(DefaultSettings), 0);
 		bDefaultConformed = true;
 	}
@@ -221,7 +221,7 @@ const FGuid& UMetaSoundSettings::GetTargetPageID() const
 
 	if (PageSettings.IsEmpty())
 	{
-		return Metasound::Frontend::DefaultGraphPageID;
+		return Metasound::Frontend::DefaultPageID;
 	}
 
 	return PageSettings.Last().UniqueId;
diff --git a/Engine/Plugins/Runtime/Metasound/Source/MetasoundEngine/Private/MetasoundSource.cpp b/Engine/Plugins/Runtime/Metasound/Source/MetasoundEngine/Private/MetasoundSource.cpp
--- a/Engine/Plugins/Runtime/Metasound/Source/MetasoundEngine/Private/MetasoundSource.cpp
+++ b/Engine/Plugins/Runtime/Metasound/Source/MetasoundEngine/Private/MetasoundSource.cpp
@@ -710,11 +710,11 @@ void UMetaSoundSource::OnAsyncReferencedAssetsLoaded(const TArray<FMetasoundAsse
 void UMetaSoundSource::MigrateEditorGraph(FMetaSoundFrontendDocumentBuilder& OutBuilder)
 {
 	PRAGMA_DISABLE_DEPRECATION_WARNINGS
-		if (Graph)
-		{
-			Graph->MigrateEditorDocumentData(OutBuilder);
-			Graph = nullptr;
-		}
+	if (Graph)
+	{
+		Graph->MigrateEditorDocumentData(OutBuilder);
+		Graph = nullptr;
+	}
 	PRAGMA_ENABLE_DEPRECATION_WARNINGS
 }
 
@@ -1726,7 +1726,8 @@ UMetaSoundSource::FRuntimeInput UMetaSoundSource::CreateRuntimeInput(const Metas
 		UE_LOG(LogMetaSound, Warning, TEXT("Failed to find data type '%s' in registry. Assuming data type is not transmittable"), *Input.TypeName.ToString());
 	}
 
-	FAudioParameter DefaultParameter = SourcePrivate::MakeAudioParameter(Registry, Input.Name, Input.TypeName, Input.DefaultLiteral, bCreateUObjectProxies);
+	const FMetasoundFrontendLiteral& DefaultLiteral = Input.FindConstDefaultChecked(Frontend::DefaultPageID);
+	FAudioParameter DefaultParameter = SourcePrivate::MakeAudioParameter(Registry, Input.Name, Input.TypeName, DefaultLiteral, bCreateUObjectProxies);
 
 	return FRuntimeInput { Input.Name, Input.TypeName, Input.AccessType, DefaultParameter, bIsTransmittable };
 }
@@ -1773,7 +1774,8 @@ Metasound::TSortedVertexNameMap<UMetaSoundSource::FRuntimeInput> UMetaSoundSourc
 
 	// Add the parameter pack input that ALL Metasounds have
 	FMetasoundFrontendClassInput ParameterPackInput = UMetasoundParameterPack::GetClassInput();
-	FAudioParameter ParameterPackDefaultParameter = SourcePrivate::MakeAudioParameter(Registry, ParameterPackInput.Name, ParameterPackInput.TypeName, ParameterPackInput.DefaultLiteral, bCreateUObjectProxies) ;
+	const FMetasoundFrontendLiteral& DefaultLiteral = ParameterPackInput.FindConstDefaultChecked(Frontend::DefaultPageID);
+	FAudioParameter ParameterPackDefaultParameter = SourcePrivate::MakeAudioParameter(Registry, ParameterPackInput.Name, ParameterPackInput.TypeName, DefaultLiteral, bCreateUObjectProxies) ;
 	PublicInputs.Add(ParameterPackInput.Name, FRuntimeInput{ParameterPackInput.Name, ParameterPackInput.TypeName, ParameterPackInput.AccessType, ParameterPackDefaultParameter, true /* bIsTransmittable */});
 	
 	return PublicInputs;
diff --git a/Engine/Plugins/Runtime/Metasound/Source/MetasoundEngine/Public/MetasoundBuilderBase.cpp b/Engine/Plugins/Runtime/Metasound/Source/MetasoundEngine/Public/MetasoundBuilderBase.cpp
--- a/Engine/Plugins/Runtime/Metasound/Source/MetasoundEngine/Public/MetasoundBuilderBase.cpp
+++ b/Engine/Plugins/Runtime/Metasound/Source/MetasoundEngine/Public/MetasoundBuilderBase.cpp
@@ -59,11 +59,11 @@ void UMetaSoundBuilderBase::BeginDestroy()
 
 FMetaSoundBuilderNodeOutputHandle UMetaSoundBuilderBase::AddGraphInputNode(FName Name, FName DataType, FMetasoundFrontendLiteral DefaultValue, EMetaSoundBuilderResult& OutResult, bool bIsConstructorInput)
 {
-	using namespace Metasound::Frontend;
+	using namespace Metasound;
 
 	FMetaSoundBuilderNodeOutputHandle NewHandle;
 
-	if (IDataTypeRegistry::Get().FindDataTypeRegistryEntry(DataType) == nullptr)
+	if (Frontend::IDataTypeRegistry::Get().FindDataTypeRegistryEntry(DataType) == nullptr)
 	{
 		UE_LOG(LogMetaSound, Error, TEXT("AddGraphInputNode Failed on builder '%s' when attempting to add '%s': '%s' is not a registered DataType"), *GetName(), *Name.ToString(), *DataType.ToString());
 	}
@@ -76,7 +76,7 @@ FMetaSoundBuilderNodeOutputHandle UMetaSoundBuilderBase::AddGraphInputNode(FName
 		}
 		else
 		{
-			FDocumentIDGenerator& IDGenerator = FDocumentIDGenerator::Get();
+			Frontend::FDocumentIDGenerator& IDGenerator = Frontend::FDocumentIDGenerator::Get();
 			const FMetasoundFrontendDocument& Doc = GetConstBuilder().GetConstDocumentChecked();
 
 			FMetasoundFrontendClassInput Description;
@@ -84,8 +84,10 @@ FMetaSoundBuilderNodeOutputHandle UMetaSoundBuilderBase::AddGraphInputNode(FName
 			Description.TypeName = DataType;
 			Description.NodeID = IDGenerator.CreateNodeID(Doc);
 			Description.VertexID = IDGenerator.CreateVertexID(Doc);
-			Description.DefaultLiteral = static_cast<FMetasoundFrontendLiteral>(DefaultValue);
 			Description.AccessType = bIsConstructorInput ? EMetasoundFrontendVertexAccessType::Value : EMetasoundFrontendVertexAccessType::Reference;
+
+			Description.InitDefault(MoveTemp(DefaultValue));
+
 			Node = Builder.AddGraphInput(Description);
 		}
 
@@ -803,8 +805,8 @@ void UMetaSoundBuilderBase::Initialize()
 	TScriptInterface<IMetaSoundDocumentInterface> DocObject = NewObject<UObject>(GetTransientPackage(), &GetBaseMetaSoundUClass(), { }, NewObjectFlags);
 	TSharedRef<FDocumentModifyDelegates> DocumentDelegates = MakeShared<FDocumentModifyDelegates>(DocObject->GetConstDocument());
 	Builder = FMetaSoundFrontendDocumentBuilder(DocObject, DocumentDelegates);
-	InitDelegates(*DocumentDelegates);
 	Builder.InitDocument();
+	InitDelegates(*DocumentDelegates);
 }
 
 void UMetaSoundBuilderBase::InitNodeLocations()
diff --git a/Engine/Plugins/Runtime/Metasound/Source/MetasoundEngine/Public/MetasoundSettings.h b/Engine/Plugins/Runtime/Metasound/Source/MetasoundEngine/Public/MetasoundSettings.h
--- a/Engine/Plugins/Runtime/Metasound/Source/MetasoundEngine/Public/MetasoundSettings.h
+++ b/Engine/Plugins/Runtime/Metasound/Source/MetasoundEngine/Public/MetasoundSettings.h
@@ -143,7 +143,7 @@ private:
 	  * for the active platform, uses order of cooked pages (see 'Page Settings' for order) falling back to lower index-ordered page
 	  * implemented in MetaSound asset. */
 	UPROPERTY(EditAnywhere, config, Category = Pages)
-	FName TargetPageName = Metasound::Frontend::DefaultGraphPageName;
+	FName TargetPageName = Metasound::Frontend::DefaultPageName;
 
 	/** Array of possible page settings that can be added to a MetaSound object. Order
 	  * defines default fallback logic whereby a higher index-ordered page
diff --git a/Engine/Plugins/Runtime/Metasound/Source/MetasoundEngineTest/Private/EngineTestMetaSoundAutomatedNodeTest.cpp b/Engine/Plugins/Runtime/Metasound/Source/MetasoundEngineTest/Private/EngineTestMetaSoundAutomatedNodeTest.cpp
--- a/Engine/Plugins/Runtime/Metasound/Source/MetasoundEngineTest/Private/EngineTestMetaSoundAutomatedNodeTest.cpp
+++ b/Engine/Plugins/Runtime/Metasound/Source/MetasoundEngineTest/Private/EngineTestMetaSoundAutomatedNodeTest.cpp
@@ -694,7 +694,7 @@ namespace Metasound::EngineTest{
 					TestNodeName, 
 					TestNodeID, 
 					TestVertexName,
-					NodeClass.Interface.Inputs[0].DefaultLiteral.ToLiteral(DataTypeName)
+					NodeClass.Interface.Inputs[0].FindConstDefaultChecked(Frontend::DefaultPageID).ToLiteral(DataTypeName)
 				};
 
 				Node = DataTypeRegistry.CreateInputNode(DataTypeName, MoveTemp(NodeInitData));
diff --git a/Engine/Plugins/Runtime/Metasound/Source/MetasoundEngineTest/Private/EngineTestMetasoundSource.cpp b/Engine/Plugins/Runtime/Metasound/Source/MetasoundEngineTest/Private/EngineTestMetasoundSource.cpp
--- a/Engine/Plugins/Runtime/Metasound/Source/MetasoundEngineTest/Private/EngineTestMetasoundSource.cpp
+++ b/Engine/Plugins/Runtime/Metasound/Source/MetasoundEngineTest/Private/EngineTestMetasoundSource.cpp
@@ -177,6 +177,7 @@ namespace EngineTestMetaSoundSourcePrivate
 
 		Document.RootGraph.Metadata.SetClassName(FMetasoundFrontendClassName { "Namespace", "Unit Test Node", *LexToString(FGuid::NewGuid()) });
 		Document.RootGraph.Metadata.SetType(EMetasoundFrontendClassType::Graph);
+		Document.RootGraph.InitDefaultPage();
 
 		FDocumentHandle DocumentHandle = IDocumentController::CreateDocumentHandle(Document);
 		FGraphHandle RootGraph = DocumentHandle->GetRootGraph();
@@ -201,7 +202,7 @@ namespace EngineTestMetaSoundSourcePrivate
 		FrequencyInput.Name = "Frequency";
 		FrequencyInput.TypeName = GetMetasoundDataTypeName<float>();
 		FrequencyInput.VertexID = FGuid::NewGuid();
-		FrequencyInput.DefaultLiteral.Set(100.f);
+		FrequencyInput.InitDefault().Set(100.f);
 		FNodeHandle FrequencyInputNode = RootGraph->AddInputVertex(FrequencyInput);
 		check(FrequencyInputNode->IsValid());
 
diff --git a/Engine/Plugins/Runtime/Metasound/Source/MetasoundEngineTest/Private/NodeTestGraphBuilder.cpp b/Engine/Plugins/Runtime/Metasound/Source/MetasoundEngineTest/Private/NodeTestGraphBuilder.cpp
--- a/Engine/Plugins/Runtime/Metasound/Source/MetasoundEngineTest/Private/NodeTestGraphBuilder.cpp
+++ b/Engine/Plugins/Runtime/Metasound/Source/MetasoundEngineTest/Private/NodeTestGraphBuilder.cpp
@@ -14,6 +14,7 @@ namespace Metasound::Test
 	{
 		Document.RootGraph.Metadata.SetClassName({ "Metasound", "TestNodes", *LexToString(FGuid::NewGuid()) });
 		Document.RootGraph.Metadata.SetType(EMetasoundFrontendClassType::Graph);
+		Document.RootGraph.InitDefaultPage();
 
 		DocumentHandle = IDocumentController::CreateDocumentHandle(Document);
 		RootGraph = DocumentHandle->GetRootGraph();
@@ -57,7 +58,8 @@ namespace Metasound::Test
 		Input.TypeName = TypeName;
 		Input.AccessType = AccessType;
 		Input.VertexID = FGuid::NewGuid();
-		
+		Input.InitDefault();
+
 		return RootGraph->AddInputVertex(Input);
 	}
 
diff --git a/Engine/Plugins/Runtime/Metasound/Source/MetasoundEngineTest/Public/NodeTestGraphBuilder.h b/Engine/Plugins/Runtime/Metasound/Source/MetasoundEngineTest/Public/NodeTestGraphBuilder.h
--- a/Engine/Plugins/Runtime/Metasound/Source/MetasoundEngineTest/Public/NodeTestGraphBuilder.h
+++ b/Engine/Plugins/Runtime/Metasound/Source/MetasoundEngineTest/Public/NodeTestGraphBuilder.h
@@ -70,7 +70,7 @@ namespace Metasound::Test
 			Input.TypeName = GetMetasoundDataTypeName<DataType>();
 			Input.VertexID = FGuid::NewGuid();
 			Input.AccessType = EMetasoundFrontendVertexAccessType::Value;
-			Input.DefaultLiteral.Set(Value);
+			Input.InitDefault().Set(Value);
 			return RootGraph->AddInputVertex(Input);
 		}
 
@@ -129,7 +129,7 @@ namespace Metasound::Test
 			Input.TypeName = GetMetasoundDataTypeName<InputDataType>();
 			Input.VertexID = FGuid::NewGuid();
 			Input.AccessType = EMetasoundFrontendVertexAccessType::Value;
-			Input.DefaultLiteral.Set(Value);
+			Input.InitDefault().Set(Value);
 			
 			Frontend::FNodeHandle InputNode = RootGraph->AddInputVertex(Input);
 
diff --git a/Engine/Plugins/Runtime/Metasound/Source/MetasoundFrontend/Private/MetasoundAssetBase.cpp b/Engine/Plugins/Runtime/Metasound/Source/MetasoundFrontend/Private/MetasoundAssetBase.cpp
--- a/Engine/Plugins/Runtime/Metasound/Source/MetasoundFrontend/Private/MetasoundAssetBase.cpp
+++ b/Engine/Plugins/Runtime/Metasound/Source/MetasoundFrontend/Private/MetasoundAssetBase.cpp
@@ -534,15 +534,10 @@ void FMetasoundAssetBase::SetDocument(FMetasoundFrontendDocument InDocument, boo
 bool FMetasoundAssetBase::VersionAsset(FMetaSoundFrontendDocumentBuilder& Builder)
 {
 	using namespace Metasound;
-	using namespace Metasound::Frontend;
 
 	METASOUND_TRACE_CPUPROFILER_EVENT_SCOPE(MetaSoundAssetBase::VersionAsset);
 
-	bool bDidEdit = false;
-
-	// Version Document Model
-	bDidEdit = Builder.VersionInterfaces();
-	bDidEdit |= Metasound::Frontend::VersionDocument(*this, Builder);
+	bool bDidEdit = Frontend::VersionDocument(Builder);
 
 	// TODO: Move this logic to builder API above, which will require rewriting update transforms to
 	// take in builder instead of DocumentHandle.
@@ -558,6 +553,7 @@ bool FMetasoundAssetBase::VersionAsset(FMetaSoundFrontendDocumentBuilder& Builde
 			bPassUpdated = false;
 
 			const TArray<FMetasoundFrontendVersion> Versions = Document.Interfaces.Array();
+
 			for (const FMetasoundFrontendVersion& Version : Versions)
 			{
 				bPassUpdated |= TryUpdateInterfaceFromVersion(Version);
diff --git a/Engine/Plugins/Runtime/Metasound/Source/MetasoundFrontend/Private/MetasoundFrontendDocument.cpp b/Engine/Plugins/Runtime/Metasound/Source/MetasoundFrontend/Private/MetasoundFrontendDocument.cpp
--- a/Engine/Plugins/Runtime/Metasound/Source/MetasoundFrontend/Private/MetasoundFrontendDocument.cpp
+++ b/Engine/Plugins/Runtime/Metasound/Source/MetasoundFrontend/Private/MetasoundFrontendDocument.cpp
@@ -501,7 +501,6 @@ FMetasoundFrontendClassInterface FMetasoundFrontendClassInterface::GenerateClass
 		for (const FInputDataVertex& InputVertex : InputInterface)
 		{
 			FMetasoundFrontendClassInput ClassInput;
-
 			ClassInput.Name = InputVertex.VertexName;
 			ClassInput.TypeName = InputVertex.DataTypeName;
 			ClassInput.AccessType = Metasound::DocumentPrivate::CoreVertexAccessTypeToFrontendVertexAccessType(InputVertex.AccessType);
@@ -527,10 +526,9 @@ FMetasoundFrontendClassInterface FMetasoundFrontendClassInterface::GenerateClass
 			FLiteral DefaultLiteral = InputVertex.GetDefaultLiteral();
 			if (DefaultLiteral.GetType() != ELiteralType::Invalid)
 			{
-				ClassInput.DefaultLiteral.SetFromLiteral(DefaultLiteral);
+				ClassInput.InitDefault().SetFromLiteral(DefaultLiteral);
 			}
 
-
 			ClassInterface.Inputs.Add(MoveTemp(ClassInput));
 		}
 
@@ -813,21 +811,21 @@ FMetasoundFrontendClassMetadata FMetasoundFrontendClassMetadata::GenerateClassMe
 FMetasoundFrontendClassInput::FMetasoundFrontendClassInput(const FMetasoundFrontendClassVertex& InOther)
 :	FMetasoundFrontendClassVertex(InOther)
 {
-	using namespace Metasound::Frontend;
-
-	EMetasoundFrontendLiteralType DefaultType = GetMetasoundFrontendLiteralType(IDataTypeRegistry::Get().GetDesiredLiteralType(InOther.TypeName));
+	using namespace Metasound;
 
-	DefaultLiteral.SetType(DefaultType);
+	const ELiteralType LiteralType = Frontend::IDataTypeRegistry::Get().GetDesiredLiteralType(InOther.TypeName);
+	const EMetasoundFrontendLiteralType DefaultType = Frontend::GetMetasoundFrontendLiteralType(LiteralType);
+	InitDefault().SetType(DefaultType);
 }
 
 FMetasoundFrontendClassInput::FMetasoundFrontendClassInput(const Audio::FParameterInterface::FInput& InInput)
 {
-	using namespace Metasound::Frontend;
+	using namespace Metasound;
 
 	Name = InInput.InitValue.ParamName;
-	DefaultLiteral = FMetasoundFrontendLiteral(InInput.InitValue);
-	TypeName = Metasound::DocumentPrivate::ResolveMemberDataType(InInput.DataType, InInput.InitValue.ParamType);
-	VertexID = FClassIDGenerator::Get().CreateInputID(InInput);
+	InitDefault(FMetasoundFrontendLiteral(InInput.InitValue));
+	TypeName = DocumentPrivate::ResolveMemberDataType(InInput.DataType, InInput.InitValue.ParamType);
+	VertexID = Frontend::FClassIDGenerator::Get().CreateInputID(InInput);
 
 #if WITH_EDITOR
 	// Interfaces should never serialize text to avoid desync between
@@ -839,6 +837,97 @@ FMetasoundFrontendClassInput::FMetasoundFrontendClassInput(const Audio::FParamet
 #endif // WITH_EDITOR
 }
 
+#if WITH_EDITORONLY_DATA
+FMetasoundFrontendLiteral& FMetasoundFrontendClassInput::AddDefault(const FGuid& InPageID)
+{
+	checkf(!ContainsDefault(InPageID), TEXT("Page default with given ID already exists"));
+	return Defaults.Add_GetRef(FMetasoundFrontendClassInputDefault { InPageID }).Literal;
+}
+#endif // WITH_EDITORONLY_DATA
+
+bool FMetasoundFrontendClassInput::ContainsDefault(const FGuid& InPageID) const
+{
+	auto IsPage = [&InPageID](const FMetasoundFrontendClassInputDefault& Default) { return Default.PageID == InPageID; };
+	return Defaults.ContainsByPredicate(IsPage);
+}
+
+const FMetasoundFrontendLiteral* FMetasoundFrontendClassInput::FindConstDefault(const FGuid& InPageID) const
+{
+	auto IsPage = [&InPageID](FMetasoundFrontendClassInputDefault& Default) { return Default.PageID == InPageID; };
+	if (const FMetasoundFrontendClassInputDefault* Default = Defaults.FindByPredicate(IsPage))
+	{
+		return &Default->Literal;
+	}
+
+	return nullptr;
+}
+
+const FMetasoundFrontendLiteral& FMetasoundFrontendClassInput::FindConstDefaultChecked(const FGuid& InPageID) const
+{
+	const FMetasoundFrontendLiteral* Literal = FindConstDefault(InPageID);
+	check(Literal);
+	return *Literal;
+}
+
+FMetasoundFrontendLiteral* FMetasoundFrontendClassInput::FindDefault(const FGuid& InPageID)
+{
+	auto IsPage = [&InPageID](FMetasoundFrontendClassInputDefault& Default) { return Default.PageID == InPageID; };
+	if (FMetasoundFrontendClassInputDefault* Default = Defaults.FindByPredicate(IsPage))
+	{
+		return &Default->Literal;
+	}
+
+	return nullptr;
+}
+
+FMetasoundFrontendLiteral& FMetasoundFrontendClassInput::FindDefaultChecked(const FGuid& InPageID)
+{
+	FMetasoundFrontendLiteral* Literal = FindDefault(InPageID);
+	check(Literal);
+	return *Literal;
+}
+
+FMetasoundFrontendLiteral& FMetasoundFrontendClassInput::InitDefault()
+{
+	checkf(Defaults.IsEmpty(), TEXT("Default(s) already initialized"));
+	return Defaults.Add_GetRef(FMetasoundFrontendClassInputDefault { Metasound::Frontend::DefaultPageID }).Literal;
+}
+
+void FMetasoundFrontendClassInput::InitDefault(FMetasoundFrontendLiteral InitLiteral)
+{
+	checkf(Defaults.IsEmpty(), TEXT("Default(s) already initialized"));
+	Defaults.Add_GetRef(FMetasoundFrontendClassInputDefault{ Metasound::Frontend::DefaultPageID }).Literal = MoveTemp(InitLiteral);
+}
+
+void FMetasoundFrontendClassInput::IterateDefaults(TFunctionRef<void(const FGuid&, FMetasoundFrontendLiteral&)> IterFunc)
+{
+	for (FMetasoundFrontendClassInputDefault& Default : Defaults)
+	{
+		IterFunc(Default.PageID, Default.Literal);
+	}
+}
+
+void FMetasoundFrontendClassInput::IterateDefaults(TFunctionRef<void(const FGuid&, const FMetasoundFrontendLiteral&)> IterFunc) const
+{
+	for (const FMetasoundFrontendClassInputDefault& Default : Defaults)
+	{
+		IterFunc(Default.PageID, Default.Literal);
+	}
+}
+
+#if WITH_EDITORONLY_DATA
+void FMetasoundFrontendClassInput::RemoveAllDefaults()
+{
+	Defaults.Reset();
+}
+
+bool FMetasoundFrontendClassInput::RemoveDefault(const FGuid& InPageID)
+{
+	auto IsPage = [&InPageID](const FMetasoundFrontendClassInputDefault& Default) { return Default.PageID == InPageID; };
+	return Defaults.RemoveAllSwap(IsPage) > 0;
+}
+#endif // WITH_EDITORONLY_DATA
+
 FMetasoundFrontendClassVariable::FMetasoundFrontendClassVariable(const FMetasoundFrontendClassVertex& InOther)
 	: FMetasoundFrontendClassVertex(InOther)
 {
@@ -888,15 +977,13 @@ FMetasoundFrontendGraphClass::FMetasoundFrontendGraphClass()
 #if WITH_EDITORONLY_DATA
 const FMetasoundFrontendGraph& FMetasoundFrontendGraphClass::AddGraphPage(const FGuid& InPageID, bool bDuplicateLastGraph, bool bSetAsBuildGraph)
 {
-	checkf(InPageID != Metasound::Frontend::DefaultGraphPageID, TEXT("Cannot add graph with DefaultPageID"))
 	checkf(!ContainsGraphPage(InPageID), TEXT("Cannot add new graph page with existing PageID"));
 
 	FMetasoundFrontendGraph* NewGraph = nullptr;
 	if (bDuplicateLastGraph)
 	{
-PRAGMA_DISABLE_DEPRECATION_WARNINGS
-		FMetasoundFrontendGraph ToDuplicate = PagedGraphs.IsEmpty() ? Graph : PagedGraphs.Last();
-PRAGMA_ENABLE_DEPRECATION_WARNINGS
+		checkf(!PagedGraphs.IsEmpty(), TEXT("Cannot duplicate graph. No graph to duplicate"));
+		FMetasoundFrontendGraph ToDuplicate = PagedGraphs.Last();
 		NewGraph = &PagedGraphs.Add_GetRef(MoveTemp(ToDuplicate));
 	}
 	else
@@ -911,13 +998,6 @@ PRAGMA_ENABLE_DEPRECATION_WARNINGS
 
 bool FMetasoundFrontendGraphClass::ContainsGraphPage(const FGuid& InPageID) const
 {
-PRAGMA_DISABLE_DEPRECATION_WARNINGS
-	if (InPageID == Graph.PageID)
-	{
-		return true;
-	}
-PRAGMA_ENABLE_DEPRECATION_WARNINGS
-
 	auto MatchesPageID = [&InPageID](const FMetasoundFrontendGraph& Iter) { return Iter.PageID == InPageID; };
 	return PagedGraphs.ContainsByPredicate(MatchesPageID);
 }
@@ -926,10 +1006,6 @@ PRAGMA_ENABLE_DEPRECATION_WARNINGS
 void FMetasoundFrontendGraphClass::RemoveAllGraphPages()
 {
 	PagedGraphs.Empty();
-
-PRAGMA_DISABLE_DEPRECATION_WARNINGS
-	Graph.PageID = FGuid();
-PRAGMA_ENABLE_DEPRECATION_WARNINGS
 }
 
 bool FMetasoundFrontendGraphClass::RemoveGraphPage(const FGuid& InPageID, FGuid* OutAdjacentPageID)
@@ -963,13 +1039,6 @@ bool FMetasoundFrontendGraphClass::RemoveGraphPage(const FGuid& InPageID, FGuid*
 
 FMetasoundFrontendGraph* FMetasoundFrontendGraphClass::FindGraph(const FGuid& InPageID)
 {
-PRAGMA_DISABLE_DEPRECATION_WARNINGS
-	if (InPageID == Graph.PageID)
-	{
-		return &Graph;
-	}
-PRAGMA_ENABLE_DEPRECATION_WARNINGS
-
 	auto MatchesPageID = [this, &InPageID](const FMetasoundFrontendGraph& Iter) { return Iter.PageID == InPageID; };
 	FMetasoundFrontendGraph* PageGraph = PagedGraphs.FindByPredicate(MatchesPageID);
 	return PageGraph;
@@ -984,13 +1053,6 @@ FMetasoundFrontendGraph& FMetasoundFrontendGraphClass::FindGraphChecked(const FG
 
 const FMetasoundFrontendGraph* FMetasoundFrontendGraphClass::FindConstGraph(const FGuid& InPageID) const
 {
-PRAGMA_DISABLE_DEPRECATION_WARNINGS
-	if (InPageID == Graph.PageID)
-	{
-		return &Graph;
-	}
-PRAGMA_ENABLE_DEPRECATION_WARNINGS
-
 	auto MatchesPageID = [this, &InPageID](const FMetasoundFrontendGraph& Iter) { return Iter.PageID == InPageID; };
 	const FMetasoundFrontendGraph* PageGraph = PagedGraphs.FindByPredicate(MatchesPageID);
 	return PageGraph;
@@ -1005,20 +1067,24 @@ const FMetasoundFrontendGraph& FMetasoundFrontendGraphClass::FindConstGraphCheck
 
 FMetasoundFrontendGraph& FMetasoundFrontendGraphClass::GetDefaultGraph()
 {
-	return FindGraphChecked(FGuid());
+	return FindGraphChecked(Metasound::Frontend::DefaultPageID);
 }
 
 const FMetasoundFrontendGraph& FMetasoundFrontendGraphClass::GetConstDefaultGraph() const
 {
-	return FindConstGraphChecked(FGuid());
+	return FindConstGraphChecked(Metasound::Frontend::DefaultPageID);
 }
 
-void FMetasoundFrontendGraphClass::IterateGraphPages(TFunctionRef<void(FMetasoundFrontendGraph&)> IterFunc)
+FMetasoundFrontendGraph& FMetasoundFrontendGraphClass::InitDefaultPage()
 {
-PRAGMA_DISABLE_DEPRECATION_WARNINGS
-	IterFunc(Graph);
-PRAGMA_ENABLE_DEPRECATION_WARNINGS
+	checkf(PagedGraphs.IsEmpty(), TEXT("Attempting to initialize default page for graph class with existing graph implementation"));
+	FMetasoundFrontendGraph& NewGraph = PagedGraphs.AddDefaulted_GetRef();
+	NewGraph.PageID = Metasound::Frontend::DefaultPageID;
+	return NewGraph;
+}
 
+void FMetasoundFrontendGraphClass::IterateGraphPages(TFunctionRef<void(FMetasoundFrontendGraph&)> IterFunc)
+{
 	for (FMetasoundFrontendGraph& Iter : PagedGraphs)
 	{
 		IterFunc(Iter);
@@ -1027,10 +1093,6 @@ PRAGMA_ENABLE_DEPRECATION_WARNINGS
 
 void FMetasoundFrontendGraphClass::IterateGraphPages(TFunctionRef<void(const FMetasoundFrontendGraph&)> IterFunc) const
 {
-PRAGMA_DISABLE_DEPRECATION_WARNINGS
-	IterFunc(Graph);
-PRAGMA_ENABLE_DEPRECATION_WARNINGS
-
 	for (const FMetasoundFrontendGraph& Iter : PagedGraphs)
 	{
 		IterFunc(Iter);
@@ -1040,17 +1102,19 @@ PRAGMA_ENABLE_DEPRECATION_WARNINGS
 void FMetasoundFrontendGraphClass::ResetGraphs()
 {
 	PagedGraphs.Empty();
+}
 
-PRAGMA_DISABLE_DEPRECATION_WARNINGS
-	Graph.Nodes.Reset();
-	Graph.Edges.Reset();
-PRAGMA_ENABLE_DEPRECATION_WARNINGS
+#if WITH_EDITORONLY_DATA
+TArray<FMetasoundFrontendGraph>& FMetasoundFrontendGraphClass::IPropertyVersionTransform::GetPagesUnsafe(FMetasoundFrontendGraphClass& GraphClass)
+{
+	return GraphClass.PagedGraphs;
 }
 
 FMetasoundFrontendVersionNumber FMetasoundFrontendDocument::GetMaxVersion()
 {
 	return Metasound::Frontend::GetMaxDocumentVersion();
 }
+#endif // WITH_EDITORONLY_DATA
 
 FMetasoundFrontendDocument::FMetasoundFrontendDocument()
 {
@@ -1196,7 +1260,8 @@ namespace Metasound::Frontend
 	{
 		for (const FMetasoundFrontendClassInput& ClassInput : InClass.Interface.Inputs)
 		{
-			OnLiteral(ClassInput.TypeName, ClassInput.DefaultLiteral);
+			const FMetasoundFrontendLiteral& DefaultLiteral = ClassInput.FindConstDefaultChecked(Frontend::DefaultPageID);
+			OnLiteral(ClassInput.TypeName, DefaultLiteral);
 		}
 	}
 
diff --git a/Engine/Plugins/Runtime/Metasound/Source/MetasoundFrontend/Private/MetasoundFrontendDocumentBuilder.cpp b/Engine/Plugins/Runtime/Metasound/Source/MetasoundFrontend/Private/MetasoundFrontendDocumentBuilder.cpp
--- a/Engine/Plugins/Runtime/Metasound/Source/MetasoundFrontend/Private/MetasoundFrontendDocumentBuilder.cpp
+++ b/Engine/Plugins/Runtime/Metasound/Source/MetasoundFrontend/Private/MetasoundFrontendDocumentBuilder.cpp
@@ -175,12 +175,14 @@ namespace Metasound::Frontend
 				if (FMetasoundFrontendVertexLiteral* VertexLiteral = InOutNode.InputLiterals.FindByPredicate(IsVertexLiteralWithVertexID))
 				{
 					// Update existing literal default value with value from class input.
-					VertexLiteral->Value = InClassInput.DefaultLiteral;
+					const FMetasoundFrontendLiteral& DefaultLiteral = InClassInput.FindConstDefaultChecked(Frontend::DefaultPageID);
+					VertexLiteral->Value = DefaultLiteral;
 				}
 				else
 				{
 					// Add literal default value with value from class input.
-					InOutNode.InputLiterals.Add(FMetasoundFrontendVertexLiteral{InputVertex->VertexID, InClassInput.DefaultLiteral});
+					const FMetasoundFrontendLiteral& DefaultLiteral = InClassInput.FindConstDefaultChecked(Frontend::DefaultPageID);
+					InOutNode.InputLiterals.Add(FMetasoundFrontendVertexLiteral { InputVertex->VertexID, DefaultLiteral } );
 				}
 			}
 			else
@@ -731,7 +733,7 @@ void FMetaSoundFrontendDocumentBuilder::AddEdge(FMetasoundFrontendEdge&& InNewEd
 #if DO_CHECK
 	{
 		const IDocumentGraphEdgeCache& EdgeCache = DocumentCache->GetEdgeCache(PageID);
-		checkf(!EdgeCache.IsNodeInputConnected(InNewEdge.ToNodeID, InNewEdge.ToVertexID), TEXT("Failed to add edge in MetaSound Builder: Destination input already connected"));
+		checkf(!EdgeCache.IsNodeInputConnected(InNewEdge.ToNodeID, InNewEdge.ToVertexID), TEXT("Failed to add edge in MetaSound Builder: Destination input a lready connected"));
 
 		const EInvalidEdgeReason Reason = IsValidEdge(InNewEdge, &PageID);
 		checkf(Reason == Metasound::Frontend::EInvalidEdgeReason::None, TEXT("Attempted call to AddEdge in MetaSound Builder where edge is invalid: %s."), *LexToString(Reason));
@@ -1361,13 +1363,7 @@ void FMetaSoundFrontendDocumentBuilder::ClearDocument(TSharedRef<Metasound::Fron
 	checkf(!DocObject.IsAsset(), TEXT("Cannot call clear document on asset '%s': builder API does not support document mutation on serialized objects without editor data loaded"), *GetDebugName());
 #endif // !WITH_EDITORONLY_DATA
 
-	// Empty implementation for single remaining graph.
-	GraphClass.IterateGraphPages([](FMetasoundFrontendGraph& Graph)
-	{
-		Graph.Variables.Empty();
-		Graph.Nodes.Empty();
-		Graph.Edges.Empty();
-	});
+	GraphClass.InitDefaultPage();
 
 	Doc.Interfaces.Empty();
 	Doc.Dependencies.Empty();
@@ -1557,7 +1553,7 @@ bool FMetaSoundFrontendDocumentBuilder::ConvertToPreset(const FMetasoundFrontend
 		PresetAssetRootGraph.PresetOptions.bIsPreset = true;
 	}
 
-	// Apply root graph transform 
+	// Apply root graph transform
 	FRebuildPresetRootGraph RebuildPresetRootGraph(InReferencedDocument);
 	if (RebuildPresetRootGraph.Transform(GetDocumentChecked()))
 	{
@@ -1576,11 +1572,11 @@ bool FMetaSoundFrontendDocumentBuilder::ConvertToPreset(const FMetasoundFrontend
 	return false;
 }
 
-const FMetasoundFrontendNode* FMetaSoundFrontendDocumentBuilder::DuplicateGraphInput(const FMetasoundFrontendClassInput& InClassInput, FMetasoundFrontendLiteral DefaultValue, const FName InName, const FGuid* InPageID)
+const FMetasoundFrontendNode* FMetaSoundFrontendDocumentBuilder::DuplicateGraphInput(const FMetasoundFrontendClassInput& InClassInput, const FName InName, const FGuid* InPageID)
 {
-	using namespace Metasound::Frontend;
+	using namespace Metasound;
 
-	FDocumentIDGenerator& IDGenerator = FDocumentIDGenerator::Get();
+	Frontend::FDocumentIDGenerator& IDGenerator = Frontend::FDocumentIDGenerator::Get();
 	const FMetasoundFrontendDocument& Doc = GetConstDocumentChecked();
 
 	const FGuid& PageID = InPageID ? *InPageID : BuildPageID;
@@ -1592,7 +1588,6 @@ const FMetasoundFrontendNode* FMetaSoundFrontendDocumentBuilder::DuplicateGraphI
 	ClassInput.Metadata.SetDisplayName(FText::GetEmpty());
 #endif // WITH_EDITORONLY_DATA
 	ClassInput.Name = InName;
-	ClassInput.DefaultLiteral = static_cast<FMetasoundFrontendLiteral>(DefaultValue);
 
 	return AddGraphInput(ClassInput, &PageID);
 }
@@ -2161,10 +2156,10 @@ EMetasoundFrontendVertexAccessType FMetaSoundFrontendDocumentBuilder::GetNodeInp
 
 const FMetasoundFrontendLiteral* FMetaSoundFrontendDocumentBuilder::GetNodeInputClassDefault(const FGuid& InNodeID, const FGuid& InVertexID, const FGuid* InPageID) const
 {
-	using namespace Metasound::Frontend;
+	using namespace Metasound;
 
 	const FGuid& PageID = InPageID ? *InPageID : BuildPageID;
-	const IDocumentGraphNodeCache& NodeCache = DocumentCache->GetNodeCache(PageID);
+	const Frontend::IDocumentGraphNodeCache& NodeCache = DocumentCache->GetNodeCache(PageID);
 	if (const int32* NodeIndex = NodeCache.FindNodeIndex(InNodeID))
 	{
 		const FMetasoundFrontendDocument& Document = GetConstDocumentChecked();
@@ -2180,7 +2175,7 @@ const FMetasoundFrontendLiteral* FMetaSoundFrontendDocumentBuilder::GetNodeInput
 					case EMetasoundFrontendClassType::Output:
 					{
 						const FMetasoundFrontendClassInput& ClassInput = Class->Interface.Inputs.Last();
-						return &ClassInput.DefaultLiteral;
+						return ClassInput.FindConstDefault(Frontend::DefaultPageID);
 					}
 					break;
 
@@ -2189,7 +2184,7 @@ const FMetasoundFrontendLiteral* FMetaSoundFrontendDocumentBuilder::GetNodeInput
 						auto IsClassInput = [VertexName = Vertex->Name](const FMetasoundFrontendClassInput& Input) { return Input.Name == VertexName; };
 						if (const FMetasoundFrontendClassInput* ClassInput = Class->Interface.Inputs.FindByPredicate(IsClassInput))
 						{
-							return &ClassInput->DefaultLiteral;
+							return ClassInput->FindConstDefault(Frontend::DefaultPageID);
 						}
 						static_assert(static_cast<uint32>(EMetasoundFrontendClassType::Invalid) == 10, "Potential missing case coverage for EMetasoundFrontendClassType "
 							"(default may not be sufficient for newly added class types)");
@@ -2297,6 +2292,7 @@ void FMetaSoundFrontendDocumentBuilder::InitDocument(const FMetasoundFrontendDoc
 	METASOUND_TRACE_CPUPROFILER_EVENT_SCOPE(FMetaSoundFrontendDocumentBuilder::InitDocument);
 
 	FMetasoundFrontendDocument& Document = GetDocumentChecked();
+	Document.RootGraph.InitDefaultPage();
 
 	// 1. Set default class Metadata.
 	if (InDocumentTemplate)
@@ -2311,11 +2307,13 @@ void FMetaSoundFrontendDocumentBuilder::InitDocument(const FMetasoundFrontendDoc
 		FMetasoundFrontendClassMetadata& ClassMetadata = Document.RootGraph.Metadata;
 		InitGraphClassMetadata(Document.RootGraph.Metadata, bResetVersion, InNewClassName);
 
+#if WITH_EDITORONLY_DATA
 		// 1b. Set default doc version Metadata
 		{
 			FMetasoundFrontendDocumentMetadata& DocMetadata = Document.Metadata;
 			DocMetadata.Version.Number = GetMaxDocumentVersion();
 		}
+#endif // WITH_EDITORONLY_DATA
 
 		// 1c. Add default interfaces for given UClass
 		{
@@ -2691,14 +2689,10 @@ void FMetaSoundFrontendDocumentBuilder::RemoveAllGraphPages()
 	FMetasoundFrontendGraphClass& RootGraph = GetDocumentChecked().RootGraph;
 	RootGraph.IterateGraphPages([this](FMetasoundFrontendGraph& Graph)
 	{
-		using namespace Metasound;
-		if (Graph.PageID != Frontend::DefaultGraphPageID)
-		{
-			DocumentDelegates->PageDelegates.OnRemovingPage.Broadcast(Frontend::FDocumentMutatePageArgs { Graph.PageID });
-		}
+		DocumentDelegates->PageDelegates.OnRemovingPage.Broadcast(Frontend::FDocumentMutatePageArgs { Graph.PageID });
 	});
 	RootGraph.RemoveAllGraphPages();
-	SetBuildPageID(Frontend::DefaultGraphPageID);
+	SetBuildPageID(Frontend::DefaultPageID);
 }
 #endif // WITH_EDITORONLY_DATA
 
@@ -3448,9 +3442,9 @@ bool FMetaSoundFrontendDocumentBuilder::SetGraphInputAccessType(FName InputName,
 
 bool FMetaSoundFrontendDocumentBuilder::SetGraphInputDataType(FName InputName, FName DataType)
 {
-	using namespace Metasound::Frontend;
+	using namespace Metasound;
 
-	if (IDataTypeRegistry::Get().IsRegistered(DataType))
+	if (Frontend::IDataTypeRegistry::Get().IsRegistered(DataType))
 	{
 		const int32* Index = DocumentCache->GetInterfaceCache().FindInputIndex(InputName);
 		if (!Index)
@@ -3464,13 +3458,13 @@ bool FMetaSoundFrontendDocumentBuilder::SetGraphInputDataType(FName InputName, F
 		if (GraphInput.TypeName != DataType)
 		{
 			FMetasoundFrontendLiteral DefaultLiteral;
-			DefaultLiteral.SetFromLiteral(IDataTypeRegistry::Get().CreateDefaultLiteral(DataType));
-			GraphInput.DefaultLiteral = DefaultLiteral;
+			DefaultLiteral.SetFromLiteral(Frontend::IDataTypeRegistry::Get().CreateDefaultLiteral(DataType));
+			GraphInput.FindDefaultChecked(Frontend::DefaultPageID) = DefaultLiteral;
 			GraphInput.TypeName = DataType;
 
 			RootGraph.IterateGraphPages([this, &DataType, &GraphInput](FMetasoundFrontendGraph& Graph)
 			{
-				const IDocumentGraphNodeCache& NodeCache = DocumentCache->GetNodeCache();
+				const Frontend::IDocumentGraphNodeCache& NodeCache = DocumentCache->GetNodeCache();
 				if (const int32* NodeIndex = NodeCache.FindNodeIndex(GraphInput.NodeID))
 				{
 					FMetasoundFrontendNode& Node = Graph.Nodes[*NodeIndex];
@@ -3483,7 +3477,7 @@ bool FMetaSoundFrontendDocumentBuilder::SetGraphInputDataType(FName InputName, F
 							// If connected to an input template node, disconnect the template node from other nodes as the data type is
 							// about to be mismatched.  Otherwise, direct connection to other nodes (i.e. at runtime when template
 							// nodes aren't injected) forcefully remove to avoid data type mismatch.
-							if (ConnectedNodeClass->Metadata.GetClassName() == FInputNodeTemplate::ClassName)
+							if (ConnectedNodeClass->Metadata.GetClassName() == Frontend::FInputNodeTemplate::ClassName)
 							{
 								RemoveEdgesFromNodeOutput(Edge.ToNodeID, ConnectedNode.Interface.Outputs.Last().VertexID, &Graph.PageID);
 								ConnectedNode.Interface.Inputs.Last().TypeName = DataType;
@@ -3519,7 +3513,7 @@ bool FMetaSoundFrontendDocumentBuilder::SetGraphInputDataType(FName InputName, F
 
 bool FMetaSoundFrontendDocumentBuilder::SetGraphInputDefault(FName InputName, const FMetasoundFrontendLiteral& InDefaultLiteral)
 {
-	using namespace Metasound::Frontend;
+	using namespace Metasound;
 
 	auto NameMatchesInput = [&InputName](const FMetasoundFrontendClassInput& Input) { return Input.Name == InputName; };
 	TArray<FMetasoundFrontendClassInput>& Inputs = GetDocumentChecked().RootGraph.Interface.Inputs;
@@ -3528,9 +3522,9 @@ bool FMetaSoundFrontendDocumentBuilder::SetGraphInputDefault(FName InputName, co
 	if (Index != INDEX_NONE)
 	{
 		FMetasoundFrontendClassInput& Input = Inputs[Index];
-		if (IDataTypeRegistry::Get().IsLiteralTypeSupported(Input.TypeName, InDefaultLiteral.GetType()))
+		if (Frontend::IDataTypeRegistry::Get().IsLiteralTypeSupported(Input.TypeName, InDefaultLiteral.GetType()))
 		{
-			Input.DefaultLiteral = InDefaultLiteral;
+			Input.FindDefaultChecked(Frontend::DefaultPageID) = InDefaultLiteral;
 			DocumentDelegates->InterfaceDelegates.OnInputDefaultChanged.Broadcast(Index);
 
 			// Set the input as no longer inheriting default for presets
@@ -4100,4 +4094,9 @@ bool FMetaSoundFrontendDocumentBuilder::VersionInterfaces()
 
 	return false;
 }
+
+FMetasoundFrontendDocument& FMetaSoundFrontendDocumentBuilder::IPropertyVersionTransform::GetDocumentUnsafe(const FMetaSoundFrontendDocumentBuilder& Builder)
+{
+	return Builder.GetDocumentChecked();
+}
 #endif // WITH_EDITORONLY_DATA
diff --git a/Engine/Plugins/Runtime/Metasound/Source/MetasoundFrontend/Private/MetasoundFrontendDocumentModifyDelegates.cpp b/Engine/Plugins/Runtime/Metasound/Source/MetasoundFrontend/Private/MetasoundFrontendDocumentModifyDelegates.cpp
--- a/Engine/Plugins/Runtime/Metasound/Source/MetasoundFrontend/Private/MetasoundFrontendDocumentModifyDelegates.cpp
+++ b/Engine/Plugins/Runtime/Metasound/Source/MetasoundFrontend/Private/MetasoundFrontendDocumentModifyDelegates.cpp
@@ -78,7 +78,7 @@ PRAGMA_ENABLE_DEPRECATION_WARNINGS
 	FNodeModifyDelegates& FDocumentModifyDelegates::FindNodeDelegatesChecked(const FGuid& InPageID)
 	{
 PRAGMA_DISABLE_DEPRECATION_WARNINGS
-		return InPageID == DefaultGraphPageID
+		return InPageID == DefaultPageID
 		? NodeDelegates
 		: PageNodeDelegates.FindChecked(InPageID);
 PRAGMA_ENABLE_DEPRECATION_WARNINGS
@@ -87,7 +87,7 @@ PRAGMA_ENABLE_DEPRECATION_WARNINGS
 	FEdgeModifyDelegates& FDocumentModifyDelegates::FindEdgeDelegatesChecked(const FGuid& InPageID)
 	{
 PRAGMA_DISABLE_DEPRECATION_WARNINGS
-		return InPageID == DefaultGraphPageID
+		return InPageID == DefaultPageID
 		? EdgeDelegates
 		: PageEdgeDelegates.FindChecked(InPageID);
 PRAGMA_ENABLE_DEPRECATION_WARNINGS
diff --git a/Engine/Plugins/Runtime/Metasound/Source/MetasoundFrontend/Private/MetasoundFrontendDocumentVersioning.cpp b/Engine/Plugins/Runtime/Metasound/Source/MetasoundFrontend/Private/MetasoundFrontendDocumentVersioning.cpp
--- a/Engine/Plugins/Runtime/Metasound/Source/MetasoundFrontend/Private/MetasoundFrontendDocumentVersioning.cpp
+++ b/Engine/Plugins/Runtime/Metasound/Source/MetasoundFrontend/Private/MetasoundFrontendDocumentVersioning.cpp
@@ -1,6 +1,7 @@
 // Copyright Epic Games, Inc. All Rights Reserved.
 #include "MetasoundFrontendDocumentVersioning.h"
 
+#if WITH_EDITORONLY_DATA
 #include "Algo/Transform.h"
 #include "CoreGlobals.h"
 #include "Interfaces/MetasoundFrontendInterface.h"
@@ -18,604 +19,713 @@
 
 namespace Metasound::Frontend
 {
-	class FVersionDocumentTransform
+	namespace VersioningPrivate
 	{
-		protected:
-			virtual FMetasoundFrontendVersionNumber GetTargetVersion() const = 0;
-			virtual void TransformInternal(FDocumentHandle InDocument) const
-			{
-				checkNoEntry();
-			}
-
-			virtual void TransformInternal(FMetasoundFrontendDocument& OutDocument) const
-			{
-				FDocumentAccessPtr DocAccessPtr = MakeAccessPtr<FDocumentAccessPtr>(OutDocument.AccessPoint, OutDocument);
-				return TransformInternal(FDocumentController::CreateDocumentHandle(DocAccessPtr));
-			}
+		class FMigratePagePropertiesTransform : public FMetaSoundFrontendDocumentBuilder::IPropertyVersionTransform
+		{
+		public:
+			virtual ~FMigratePagePropertiesTransform() = default;
 
-			virtual void TransformInternal(FMetaSoundFrontendDocumentBuilder& OutBuilder) const
+			bool Transform(FMetaSoundFrontendDocumentBuilder& OutBuilder) const override
 			{
-			}
+				using namespace Metasound;
 
-		public:
-			bool Transform(FDocumentHandle InDocument) const
-			{
-				if (FMetasoundFrontendDocumentMetadata* Metadata = InDocument->GetMetadata())
+				bool bUpdated = false;
+				auto MigrateInterfaceInputDefaults = [&](FMetasoundFrontendClassInterface& OutInterface)
 				{
-					const FMetasoundFrontendVersionNumber TargetVersion = GetTargetVersion();
-					if (Metadata->Version.Number < TargetVersion)
+					for (FMetasoundFrontendClassInput& Input : OutInterface.Inputs)
 					{
-						TransformInternal(InDocument);
-						Metadata->Version.Number = TargetVersion;
-						return true;
+PRAGMA_DISABLE_DEPRECATION_WARNINGS
+						if (Input.DefaultLiteral.IsValid())
+						{
+							Input.AddDefault(Frontend::DefaultPageID) = MoveTemp(Input.DefaultLiteral);
+							Input.DefaultLiteral = FMetasoundFrontendLiteral::GetInvalid();
+							bUpdated = true;
+						}
+PRAGMA_ENABLE_DEPRECATION_WARNINGS
 					}
+				};
+
+				FMetasoundFrontendDocument& Document = GetDocumentUnsafe(OutBuilder);
+				FMetasoundFrontendGraphClass& GraphClass = Document.RootGraph;
+				MigrateInterfaceInputDefaults(GraphClass.Interface);
+				for (FMetasoundFrontendClass& Dependency : Document.Dependencies)
+				{
+					MigrateInterfaceInputDefaults(Dependency.Interface);
 				}
 
-				return false;
+				struct FMigratePageGraphs : public FMetasoundFrontendGraphClass::IPropertyVersionTransform
+				{
+				public:
+					virtual ~FMigratePageGraphs() = default;
+
+					virtual bool Transform(FMetasoundFrontendGraphClass& OutClass) const override
+					{
+						bool bUpdated = false;
+PRAGMA_DISABLE_DEPRECATION_WARNINGS
+						TArray<FMetasoundFrontendGraph>& Pages = GetPagesUnsafe(OutClass);
+						if (Pages.IsEmpty())
+						{
+							Pages.Add(MoveTemp(OutClass.Graph));
+							return true;
+						}
+PRAGMA_ENABLE_DEPRECATION_WARNINGS
+
+						return false;
+					}
+				};
+
+				bUpdated |= FMigratePageGraphs().Transform(GraphClass);
+				return bUpdated;
 			}
+		};
 
-			bool Transform(FMetaSoundFrontendDocumentBuilder& OutDocumentBuilder) const
-			{
-				const FMetasoundFrontendDocumentMetadata& Metadata = OutDocumentBuilder.GetConstDocumentChecked().Metadata;
+		class FVersionDocumentInterfacesTransform : public FMetaSoundFrontendDocumentBuilder::IPropertyVersionTransform
+		{
+		public:
+			virtual ~FVersionDocumentInterfacesTransform() = default;
 
-				const FMetasoundFrontendVersionNumber TargetVersion = GetTargetVersion();
-				if (Metadata.Version.Number < TargetVersion)
+			bool Transform(FMetaSoundFrontendDocumentBuilder& OutBuilder) const override
+			{
+				FMetasoundFrontendDocument& Document = GetDocumentUnsafe(OutBuilder);
+				if (Document.RequiresInterfaceVersioning())
 				{
-					TransformInternal(OutDocumentBuilder);
-					OutDocumentBuilder.SetVersionNumber(TargetVersion);
+					Document.VersionInterfaces();
 					return true;
 				}
 
 				return false;
 			}
-	};
-
-	/** Versions document from 1.0 to 1.1. */
-	class FVersionDocument_1_1 : public FVersionDocumentTransform
-	{
-		FName Name;
-		const FString& Path;
+		};
 
-	public:
-		FVersionDocument_1_1(FName InName, const FString& InPath)
-		: Name(InName)
-		, Path(InPath)
+		class FVersionDocumentTransform
 		{
-		}
+			public:
+				virtual ~FVersionDocumentTransform() = default;
 
-		FMetasoundFrontendVersionNumber GetTargetVersion() const override
-		{
-			return { 1, 1 };
-		}
+			protected:
+				virtual FMetasoundFrontendVersionNumber GetTargetVersion() const = 0;
+
+				virtual void TransformInternal(FDocumentHandle) const
+				{
+					checkNoEntry();
+				}
+
+				virtual void TransformInternal(FMetasoundFrontendDocument& OutDocument) const
+				{
+					FDocumentAccessPtr DocAccessPtr = MakeAccessPtr<FDocumentAccessPtr>(OutDocument.AccessPoint, OutDocument);
+					return TransformInternal(FDocumentController::CreateDocumentHandle(DocAccessPtr));
+				}
+
+				virtual void TransformInternal(FMetaSoundFrontendDocumentBuilder&) const
+				{
+				}
+
+			public:
+				bool Transform(FDocumentHandle InDocument) const
+				{
+					if (FMetasoundFrontendDocumentMetadata* Metadata = InDocument->GetMetadata())
+					{
+						const FMetasoundFrontendVersionNumber TargetVersion = GetTargetVersion();
+						if (Metadata->Version.Number < TargetVersion)
+						{
+							TransformInternal(InDocument);
+							Metadata->Version.Number = TargetVersion;
+							return true;
+						}
+					}
 
-		void TransformInternal(FDocumentHandle InDocument) const override
+					return false;
+				}
+
+				virtual bool Transform(FMetaSoundFrontendDocumentBuilder& OutDocumentBuilder) const 
+				{
+					const FMetasoundFrontendDocumentMetadata& Metadata = OutDocumentBuilder.GetConstDocumentChecked().Metadata;
+
+					const FMetasoundFrontendVersionNumber TargetVersion = GetTargetVersion();
+					if (Metadata.Version.Number < TargetVersion)
+					{
+						TransformInternal(OutDocumentBuilder);
+						OutDocumentBuilder.SetVersionNumber(TargetVersion);
+						return true;
+					}
+
+					return false;
+				}
+		};
+
+		/** Versions document from 1.0 to 1.1. */
+		class FVersionDocument_1_1 : public FVersionDocumentTransform
 		{
-#if WITH_EDITOR
-			FGraphHandle GraphHandle = InDocument->GetRootGraph();
-			TArray<FNodeHandle> FrontendNodes = GraphHandle->GetNodes();
+			FName Name;
+			const FString& Path;
+
+		public:
+			virtual ~FVersionDocument_1_1() = default;
+
+			FVersionDocument_1_1(FName InName, const FString& InPath)
+			: Name(InName)
+			, Path(InPath)
+			{
+			}
 
-			// Before literals could be stored on node inputs directly, they were stored
-			// by creating hidden input nodes. Update the doc by finding all hidden input
-			// nodes, placing the literal value of the input node directly on the
-			// downstream node's input. Then delete the hidden input node.
-			for (FNodeHandle& NodeHandle : FrontendNodes)
+			FMetasoundFrontendVersionNumber GetTargetVersion() const override
 			{
-				const bool bIsHiddenNode = NodeHandle->GetNodeStyle().Display.Visibility == EMetasoundFrontendNodeStyleDisplayVisibility::Hidden;
-				const bool bIsInputNode = EMetasoundFrontendClassType::Input == NodeHandle->GetClassMetadata().GetType();
-				const bool bIsHiddenInputNode = bIsHiddenNode && bIsInputNode;
+				return { 1, 1 };
+			}
 
-				if (bIsHiddenInputNode)
+			void TransformInternal(FDocumentHandle InDocument) const override
+			{
+	#if WITH_EDITOR
+				FGraphHandle GraphHandle = InDocument->GetRootGraph();
+				TArray<FNodeHandle> FrontendNodes = GraphHandle->GetNodes();
+
+				// Before literals could be stored on node inputs directly, they were stored
+				// by creating hidden input nodes. Update the doc by finding all hidden input
+				// nodes, placing the literal value of the input node directly on the
+				// downstream node's input. Then delete the hidden input node.
+				for (FNodeHandle& NodeHandle : FrontendNodes)
 				{
-					// Get literal value from input node.
-					const FGuid VertexID = GraphHandle->GetVertexIDForInputVertex(NodeHandle->GetNodeName());
-					const FMetasoundFrontendLiteral DefaultLiteral = GraphHandle->GetDefaultInput(VertexID);
+					const bool bIsHiddenNode = NodeHandle->GetNodeStyle().Display.Visibility == EMetasoundFrontendNodeStyleDisplayVisibility::Hidden;
+					const bool bIsInputNode = EMetasoundFrontendClassType::Input == NodeHandle->GetClassMetadata().GetType();
+					const bool bIsHiddenInputNode = bIsHiddenNode && bIsInputNode;
 
-					// Apply literal value to downstream node's inputs.
-					TArray<FOutputHandle> OutputHandles = NodeHandle->GetOutputs();
-					if (ensure(OutputHandles.Num() == 1))
+					if (bIsHiddenInputNode)
 					{
-						FOutputHandle OutputHandle = OutputHandles[0];
-						TArray<FInputHandle> Inputs = OutputHandle->GetConnectedInputs();
-						OutputHandle->Disconnect();
+						// Get literal value from input node.
+						const FGuid VertexID = GraphHandle->GetVertexIDForInputVertex(NodeHandle->GetNodeName());
+						const FMetasoundFrontendLiteral DefaultLiteral = GraphHandle->GetDefaultInput(VertexID);
 
-						for (FInputHandle& Input : Inputs)
+						// Apply literal value to downstream node's inputs.
+						TArray<FOutputHandle> OutputHandles = NodeHandle->GetOutputs();
+						if (ensure(OutputHandles.Num() == 1))
 						{
-							if (const FMetasoundFrontendLiteral* Literal = Input->GetClassDefaultLiteral())
+							FOutputHandle OutputHandle = OutputHandles[0];
+							TArray<FInputHandle> Inputs = OutputHandle->GetConnectedInputs();
+							OutputHandle->Disconnect();
+
+							for (FInputHandle& Input : Inputs)
 							{
-								if (!Literal->IsEqual(DefaultLiteral))
+								if (const FMetasoundFrontendLiteral* Literal = Input->GetClassDefaultLiteral())
+								{
+									if (!Literal->IsEqual(DefaultLiteral))
+									{
+										Input->SetLiteral(DefaultLiteral);
+									}
+								}
+								else
 								{
 									Input->SetLiteral(DefaultLiteral);
 								}
 							}
-							else
-							{
-								Input->SetLiteral(DefaultLiteral);
-							}
 						}
+						GraphHandle->RemoveNode(*NodeHandle);
 					}
-					GraphHandle->RemoveNode(*NodeHandle);
 				}
+	#else
+				UE_LOG(LogMetaSound, Error, TEXT("Asset '%s' at '%s' must be saved with editor enabled in order to version document to target version '%s'."), *Name.ToString(), *Path, *GetTargetVersion().ToString());
+	#endif // !WITH_EDITOR
 			}
-#else
-			UE_LOG(LogMetaSound, Error, TEXT("Asset '%s' at '%s' must be saved with editor enabled in order to version document to target version '%s'."), *Name.ToString(), *Path, *GetTargetVersion().ToString());
-#endif // !WITH_EDITOR
-		}
-	};
-
-	/** Versions document from 1.1 to 1.2. */
-	class FVersionDocument_1_2 : public FVersionDocumentTransform
-	{
-	private:
-		const FName Name;
-		const FString& Path;
+		};
 
-	public:
-		FVersionDocument_1_2(const FName InName, const FString& InPath)
-			: Name(InName)
-			, Path(InPath)
+		/** Versions document from 1.1 to 1.2. */
+		class FVersionDocument_1_2 : public FVersionDocumentTransform
 		{
-		}
+		private:
+			const FName Name;
+			const FString& Path;
 
-		FMetasoundFrontendVersionNumber GetTargetVersion() const override
-		{
-			return { 1, 2 };
-		}
+		public:
+			FVersionDocument_1_2(const FName InName, const FString& InPath)
+				: Name(InName)
+				, Path(InPath)
+			{
+			}
+			virtual ~FVersionDocument_1_2() = default;
 
-		void TransformInternal(FDocumentHandle InDocument) const override
-		{
-#if WITH_EDITOR
-			const FMetasoundFrontendGraphClass& GraphClass = InDocument->GetRootGraphClass();
-			FMetasoundFrontendClassMetadata Metadata = GraphClass.Metadata;
-
-			Metadata.SetClassName({ "GraphAsset", Name, *Path });
-			Metadata.SetDisplayName(FText::FromString(Name.ToString()));
-			InDocument->GetRootGraph()->SetGraphMetadata(Metadata);
-#else
-			UE_LOG(LogMetaSound, Error, TEXT("Asset '%s' at '%s' must be saved with editor enabled in order to version document to target version '%s'."), *Name.ToString(), *Path, *GetTargetVersion().ToString());
-#endif // !WITH_EDITOR
-		}
-	};
+			FMetasoundFrontendVersionNumber GetTargetVersion() const override
+			{
+				return { 1, 2 };
+			}
 
-	/** Versions document from 1.2 to 1.3. */
-	class FVersionDocument_1_3 : public FVersionDocumentTransform
-	{
-	public:
-		FVersionDocument_1_3()
-		{
-		}
+			void TransformInternal(FDocumentHandle InDocument) const override
+			{
+	#if WITH_EDITOR
+				const FMetasoundFrontendGraphClass& GraphClass = InDocument->GetRootGraphClass();
+				FMetasoundFrontendClassMetadata Metadata = GraphClass.Metadata;
+
+				Metadata.SetClassName({ "GraphAsset", Name, *Path });
+				Metadata.SetDisplayName(FText::FromString(Name.ToString()));
+				InDocument->GetRootGraph()->SetGraphMetadata(Metadata);
+	#else
+				UE_LOG(LogMetaSound, Error, TEXT("Asset '%s' at '%s' must be saved with editor enabled in order to version document to target version '%s'."), *Name.ToString(), *Path, *GetTargetVersion().ToString());
+	#endif // !WITH_EDITOR
+			}
+		};
 
-		FMetasoundFrontendVersionNumber GetTargetVersion() const override
+		/** Versions document from 1.2 to 1.3. */
+		class FVersionDocument_1_3 : public FVersionDocumentTransform
 		{
-			return {1, 3};
-		}
+		public:
+			virtual ~FVersionDocument_1_3() = default;
 
-		void TransformInternal(FDocumentHandle InDocument) const override
-		{
-			const FMetasoundFrontendGraphClass& GraphClass = InDocument->GetRootGraphClass();
-			FMetasoundFrontendClassMetadata Metadata = GraphClass.Metadata;
+			FMetasoundFrontendVersionNumber GetTargetVersion() const override
+			{
+				return {1, 3};
+			}
 
-			Metadata.SetClassName(FMetasoundFrontendClassName { FName(), *FGuid::NewGuid().ToString(), FName() });
-			InDocument->GetRootGraph()->SetGraphMetadata(Metadata);
-		}
-	};
+			void TransformInternal(FDocumentHandle InDocument) const override
+			{
+				const FMetasoundFrontendGraphClass& GraphClass = InDocument->GetRootGraphClass();
+				FMetasoundFrontendClassMetadata Metadata = GraphClass.Metadata;
 
-	/** Versions document from 1.3 to 1.4. */
-	class FVersionDocument_1_4 : public FVersionDocumentTransform
-	{
-	public:
-		FVersionDocument_1_4()
-		{
-		}
+				Metadata.SetClassName(FMetasoundFrontendClassName { FName(), *FGuid::NewGuid().ToString(), FName() });
+				InDocument->GetRootGraph()->SetGraphMetadata(Metadata);
+			}
+		};
 
-		FMetasoundFrontendVersionNumber GetTargetVersion() const override
+		/** Versions document from 1.3 to 1.4. */
+		class FVersionDocument_1_4 : public FVersionDocumentTransform
 		{
-			return {1, 4};
-		}
+		public:
+			virtual ~FVersionDocument_1_4() = default;
 
-		void TransformInternal(FDocumentHandle InDocument) const override
-		{
-			FMetasoundFrontendDocumentMetadata* Metadata = InDocument->GetMetadata();
-			check(Metadata);
-			check(Metadata->Version.Number.Major == 1);
-			check(Metadata->Version.Number.Minor == 3);
-
-			const TSet<FMetasoundFrontendVersion>& Interfaces = InDocument->GetInterfaceVersions();
-
-			// Version 1.3 did not have an "InterfaceVersion" property on the
-			// document, so any document that is being updated should start off
-			// with an "Invalid" interface version.
-			if (ensure(Interfaces.IsEmpty()))
-			{
-				// At the time when version 1.4 of the document was introduced, 
-				// these were the only available interfaces. 
-				static const FMetasoundFrontendVersion PreexistingInterfaceVersions[] = {
-					FMetasoundFrontendVersion{"MetaSound", {1, 0}},
-					FMetasoundFrontendVersion{"MonoSource", {1, 0}},
-					FMetasoundFrontendVersion{"StereoSource", {1, 0}},
-					FMetasoundFrontendVersion{"MonoSource", {1, 1}},
-					FMetasoundFrontendVersion{"StereoSource", {1, 1}}
-				};
-				static const int32 NumPreexistingInterfaceVersions = sizeof(PreexistingInterfaceVersions) / sizeof(PreexistingInterfaceVersions[0]);
+			FMetasoundFrontendVersionNumber GetTargetVersion() const override
+			{
+				return {1, 4};
+			}
+
+			void TransformInternal(FDocumentHandle InDocument) const override
+			{
+				FMetasoundFrontendDocumentMetadata* Metadata = InDocument->GetMetadata();
+				check(Metadata);
+				check(Metadata->Version.Number.Major == 1);
+				check(Metadata->Version.Number.Minor == 3);
 
-				TArray<FMetasoundFrontendInterface> CandidateInterfaces;
-				IInterfaceRegistry& InterfaceRegistry = IInterfaceRegistry::Get();
-				for (int32 i = 0; i < NumPreexistingInterfaceVersions; i++)
+				const TSet<FMetasoundFrontendVersion>& Interfaces = InDocument->GetInterfaceVersions();
+
+				// Version 1.3 did not have an "InterfaceVersion" property on the
+				// document, so any document that is being updated should start off
+				// with an "Invalid" interface version.
+				if (ensure(Interfaces.IsEmpty()))
 				{
-					FMetasoundFrontendInterface Interface;
-					if (InterfaceRegistry.FindInterface(GetInterfaceRegistryKey(PreexistingInterfaceVersions[i]), Interface))
+					// At the time when version 1.4 of the document was introduced, 
+					// these were the only available interfaces. 
+					static const FMetasoundFrontendVersion PreexistingInterfaceVersions[] = {
+						FMetasoundFrontendVersion{"MetaSound", {1, 0}},
+						FMetasoundFrontendVersion{"MonoSource", {1, 0}},
+						FMetasoundFrontendVersion{"StereoSource", {1, 0}},
+						FMetasoundFrontendVersion{"MonoSource", {1, 1}},
+						FMetasoundFrontendVersion{"StereoSource", {1, 1}}
+					};
+					static const int32 NumPreexistingInterfaceVersions = sizeof(PreexistingInterfaceVersions) / sizeof(PreexistingInterfaceVersions[0]);
+
+					TArray<FMetasoundFrontendInterface> CandidateInterfaces;
+					IInterfaceRegistry& InterfaceRegistry = IInterfaceRegistry::Get();
+					for (int32 i = 0; i < NumPreexistingInterfaceVersions; i++)
 					{
-						CandidateInterfaces.Add(Interface);
+						FMetasoundFrontendInterface Interface;
+						if (InterfaceRegistry.FindInterface(GetInterfaceRegistryKey(PreexistingInterfaceVersions[i]), Interface))
+						{
+							CandidateInterfaces.Add(Interface);
+						}
 					}
-				}
 
-				const FMetasoundFrontendGraphClass& RootGraph = InDocument->GetRootGraphClass();
-				const TArray<FMetasoundFrontendClass>& Dependencies = InDocument->GetDependencies();
-				const TArray<FMetasoundFrontendGraphClass>& Subgraphs = InDocument->GetSubgraphs();
+					const FMetasoundFrontendGraphClass& RootGraph = InDocument->GetRootGraphClass();
+					const TArray<FMetasoundFrontendClass>& Dependencies = InDocument->GetDependencies();
+					const TArray<FMetasoundFrontendGraphClass>& Subgraphs = InDocument->GetSubgraphs();
 
-				if (const FMetasoundFrontendInterface* Interface = FindMostSimilarInterfaceSupportingEnvironment(RootGraph, Dependencies, Subgraphs, CandidateInterfaces))
-				{
-					UE_LOG(LogMetaSound, Display, TEXT("Assigned interface [InterfaceVersion:%s] to document [RootGraphClassName:%s]"),
-						*Interface->Version.ToString(), *RootGraph.Metadata.GetClassName().ToString());
+					if (const FMetasoundFrontendInterface* Interface = FindMostSimilarInterfaceSupportingEnvironment(RootGraph, Dependencies, Subgraphs, CandidateInterfaces))
+					{
+						UE_LOG(LogMetaSound, Display, TEXT("Assigned interface [InterfaceVersion:%s] to document [RootGraphClassName:%s]"),
+							*Interface->Version.ToString(), *RootGraph.Metadata.GetClassName().ToString());
 
-					InDocument->AddInterfaceVersion(Interface->Version);
-				}
-				else
-				{
-					UE_LOG(LogMetaSound, Warning, TEXT("Failed to find interface for document [RootGraphClassName:%s]"),
-						*RootGraph.Metadata.GetClassName().ToString());
+						InDocument->AddInterfaceVersion(Interface->Version);
+					}
+					else
+					{
+						UE_LOG(LogMetaSound, Warning, TEXT("Failed to find interface for document [RootGraphClassName:%s]"),
+							*RootGraph.Metadata.GetClassName().ToString());
+					}
 				}
 			}
-		}
-	};
+		};
 
-	/** Versions document from 1.4 to 1.5. */
-	class FVersionDocument_1_5 : public FVersionDocumentTransform
-	{
-		FName Name;
-		const FString& Path;
-
-	public:
-		FVersionDocument_1_5(FName InName, const FString& InPath)
-			: Name(InName)
-			, Path(InPath)
+		/** Versions document from 1.4 to 1.5. */
+		class FVersionDocument_1_5 : public FVersionDocumentTransform
 		{
-		}
+			FName Name;
+			const FString& Path;
 
-		FMetasoundFrontendVersionNumber GetTargetVersion() const override
-		{
-			return { 1, 5 };
-		}
+		public:
+			FVersionDocument_1_5(FName InName, const FString& InPath)
+				: Name(InName)
+				, Path(InPath)
+			{
+			}
+			virtual ~FVersionDocument_1_5() = default;
 
-		void TransformInternal(FDocumentHandle InDocument) const override
-		{
-#if WITH_EDITOR
-			const FMetasoundFrontendClassMetadata& Metadata = InDocument->GetRootGraphClass().Metadata;
-			const FText NewAssetName = FText::FromString(Name.ToString());
-			if (Metadata.GetDisplayName().CompareTo(NewAssetName) != 0)
+			FMetasoundFrontendVersionNumber GetTargetVersion() const override
 			{
-				FMetasoundFrontendClassMetadata NewMetadata = Metadata;
-				NewMetadata.SetDisplayName(NewAssetName);
-				InDocument->GetRootGraph()->SetGraphMetadata(NewMetadata);
+				return { 1, 5 };
 			}
-#else
-			UE_LOG(LogMetaSound, Error, TEXT("Asset '%s' at '%s' must be saved with editor enabled in order to version document to target version '%s'."), *Name.ToString(), *Path, *GetTargetVersion().ToString());
-#endif // !WITH_EDITOR
-		}
-	};
 
-	/** Versions document from 1.5 to 1.6. */
-	class FVersionDocument_1_6 : public FVersionDocumentTransform
-	{
-	public:
-		FVersionDocument_1_6() = default;
+			void TransformInternal(FDocumentHandle InDocument) const override
+			{
+	#if WITH_EDITOR
+				const FMetasoundFrontendClassMetadata& Metadata = InDocument->GetRootGraphClass().Metadata;
+				const FText NewAssetName = FText::FromString(Name.ToString());
+				if (Metadata.GetDisplayName().CompareTo(NewAssetName) != 0)
+				{
+					FMetasoundFrontendClassMetadata NewMetadata = Metadata;
+					NewMetadata.SetDisplayName(NewAssetName);
+					InDocument->GetRootGraph()->SetGraphMetadata(NewMetadata);
+				}
+	#else
+				UE_LOG(LogMetaSound, Error, TEXT("Asset '%s' at '%s' must be saved with editor enabled in order to version document to target version '%s'."), *Name.ToString(), *Path, *GetTargetVersion().ToString());
+	#endif // !WITH_EDITOR
+			}
+		};
 
-		FMetasoundFrontendVersionNumber GetTargetVersion() const override
+		/** Versions document from 1.5 to 1.6. */
+		class FVersionDocument_1_6 : public FVersionDocumentTransform
 		{
-			return { 1, 6 };
-		}
+		public:
+			FVersionDocument_1_6() = default;
+			virtual ~FVersionDocument_1_6() = default;
 
-		void TransformInternal(FDocumentHandle InDocument) const override
-		{
-			const FGuid NewAssetClassID = FGuid::NewGuid();
-			FMetasoundFrontendGraphClass Class = InDocument->GetRootGraphClass();
-			Class.Metadata.SetClassName(FMetasoundFrontendClassName({ }, FName(*NewAssetClassID.ToString()), { }));
-		}
-	};
+			FMetasoundFrontendVersionNumber GetTargetVersion() const override
+			{
+				return { 1, 6 };
+			}
 
-	/** Versions document from 1.6 to 1.7. */
-	class FVersionDocument_1_7 : public FVersionDocumentTransform
-	{
-		FName Name;
-		const FString& Path;
+			void TransformInternal(FDocumentHandle InDocument) const override
+			{
+				const FGuid NewAssetClassID = FGuid::NewGuid();
+				FMetasoundFrontendGraphClass Class = InDocument->GetRootGraphClass();
+				Class.Metadata.SetClassName(FMetasoundFrontendClassName({ }, FName(*NewAssetClassID.ToString()), { }));
+			}
+		};
 
-	public:
-		FVersionDocument_1_7(FName InName, const FString& InPath)
-			: Name(InName)
-			, Path(InPath)
+		/** Versions document from 1.6 to 1.7. */
+		class FVersionDocument_1_7 : public FVersionDocumentTransform
 		{
-		}
+			FName Name;
+			const FString& Path;
 
+		public:
+			FVersionDocument_1_7(FName InName, const FString& InPath)
+				: Name(InName)
+				, Path(InPath)
+			{
+			}
+			virtual ~FVersionDocument_1_7() = default;
 
-		FMetasoundFrontendVersionNumber GetTargetVersion() const override
-		{
-			return { 1, 7 };
-		}
+			FMetasoundFrontendVersionNumber GetTargetVersion() const override
+			{
+				return { 1, 7 };
+			}
 
-		void TransformInternal(FDocumentHandle InDocument) const override
-		{
-#if WITH_EDITOR
-			auto RenameTransform = [](FNodeHandle NodeHandle)
-			{
-				// Required nodes are all (at the point of this transform) providing
-				// unique names and customized display names (ex. 'Audio' for both mono &
-				// L/R output, On Play, & 'On Finished'), so do not replace them by nulling
-				// out the guid as a name and using the converted FName of the FText DisplayName.
-				if (!NodeHandle->IsInterfaceMember())
+			void TransformInternal(FDocumentHandle InDocument) const override
+			{
+	#if WITH_EDITOR
+				auto RenameTransform = [](FNodeHandle NodeHandle)
 				{
-					const FName NewNodeName = *NodeHandle->GetDisplayName().ToString();
-					NodeHandle->IterateInputs([&](FInputHandle InputHandle)
+					// Required nodes are all (at the point of this transform) providing
+					// unique names and customized display names (ex. 'Audio' for both mono &
+					// L/R output, On Play, & 'On Finished'), so do not replace them by nulling
+					// out the guid as a name and using the converted FName of the FText DisplayName.
+					if (!NodeHandle->IsInterfaceMember())
 					{
-						InputHandle->SetName(NewNodeName);
-					});
-
-					NodeHandle->IterateOutputs([&](FOutputHandle OutputHandle)
-					{
-						OutputHandle->SetName(NewNodeName);
-					});
+						const FName NewNodeName = *NodeHandle->GetDisplayName().ToString();
+						NodeHandle->IterateInputs([&](FInputHandle InputHandle)
+						{
+							InputHandle->SetName(NewNodeName);
+						});
 
-					NodeHandle->SetDisplayName(FText());
-					NodeHandle->SetNodeName(NewNodeName);
-				}
-			};
+						NodeHandle->IterateOutputs([&](FOutputHandle OutputHandle)
+						{
+							OutputHandle->SetName(NewNodeName);
+						});
 
-			InDocument->GetRootGraph()->IterateNodes(RenameTransform, EMetasoundFrontendClassType::Input);
-			InDocument->GetRootGraph()->IterateNodes(RenameTransform, EMetasoundFrontendClassType::Output);
-#else
-			UE_LOG(LogMetaSound, Error, TEXT("Asset '%s' at '%s' must be saved with editor enabled in order to version document to target version '%s'."), *Name.ToString(), *Path, *GetTargetVersion().ToString());
-#endif // !WITH_EDITOR
-		}
-	};
+						NodeHandle->SetDisplayName(FText());
+						NodeHandle->SetNodeName(NewNodeName);
+					}
+				};
 
-	/** Versions document from 1.7 to 1.8. */
-	class FVersionDocument_1_8 : public FVersionDocumentTransform
-	{
-		FName Name;
-		const FString& Path;
+				InDocument->GetRootGraph()->IterateNodes(RenameTransform, EMetasoundFrontendClassType::Input);
+				InDocument->GetRootGraph()->IterateNodes(RenameTransform, EMetasoundFrontendClassType::Output);
+	#else
+				UE_LOG(LogMetaSound, Error, TEXT("Asset '%s' at '%s' must be saved with editor enabled in order to version document to target version '%s'."), *Name.ToString(), *Path, *GetTargetVersion().ToString());
+	#endif // !WITH_EDITOR
+			}
+		};
 
-	public:
-		FVersionDocument_1_8(FName InName, const FString& InPath)
-			: Name(InName)
-			, Path(InPath)
+		/** Versions document from 1.7 to 1.8. */
+		class FVersionDocument_1_8 : public FVersionDocumentTransform
 		{
-		}
+			FName Name;
+			const FString& Path;
 
-		FMetasoundFrontendVersionNumber GetTargetVersion() const override
-		{
-			return { 1, 8 };
-		}
+		public:
+			FVersionDocument_1_8(FName InName, const FString& InPath)
+				: Name(InName)
+				, Path(InPath)
+			{
+			}
+			virtual ~FVersionDocument_1_8() = default;
 
-		void TransformInternal(FDocumentHandle InDocument) const override
-		{
-#if WITH_EDITOR
-			// Do not serialize MetaData text for dependencies as
-			// CacheRegistryData dynamically provides this.
-			InDocument->IterateDependencies([](FMetasoundFrontendClass& Dependency)
+			FMetasoundFrontendVersionNumber GetTargetVersion() const override
 			{
-				constexpr bool bSerializeText = false;
-				Dependency.Metadata.SetSerializeText(bSerializeText);
+				return { 1, 8 };
+			}
 
-				for (FMetasoundFrontendClassInput& Input : Dependency.Interface.Inputs)
+			void TransformInternal(FDocumentHandle InDocument) const override
+			{
+	#if WITH_EDITOR
+				// Do not serialize MetaData text for dependencies as
+				// CacheRegistryData dynamically provides this.
+				InDocument->IterateDependencies([](FMetasoundFrontendClass& Dependency)
 				{
-					Input.Metadata.SetSerializeText(false);
-				}
+					constexpr bool bSerializeText = false;
+					Dependency.Metadata.SetSerializeText(bSerializeText);
 
-				for (FMetasoundFrontendClassOutput& Output : Dependency.Interface.Outputs)
-				{
-					Output.Metadata.SetSerializeText(false);
-				}
-			});
+					for (FMetasoundFrontendClassInput& Input : Dependency.Interface.Inputs)
+					{
+						Input.Metadata.SetSerializeText(false);
+					}
 
-			const TSet<FMetasoundFrontendVersion>& InterfaceVersions = InDocument->GetInterfaceVersions();
+					for (FMetasoundFrontendClassOutput& Output : Dependency.Interface.Outputs)
+					{
+						Output.Metadata.SetSerializeText(false);
+					}
+				});
 
-			using FNameDataTypePair = TPair<FName, FName>;
-			TSet<FNameDataTypePair> InterfaceInputs;
-			TSet<FNameDataTypePair> InterfaceOutputs;
+				const TSet<FMetasoundFrontendVersion>& InterfaceVersions = InDocument->GetInterfaceVersions();
 
-			for (const FMetasoundFrontendVersion& Version : InterfaceVersions)
-			{
-				FInterfaceRegistryKey RegistryKey = GetInterfaceRegistryKey(Version);
-				const IInterfaceRegistryEntry* Entry = IInterfaceRegistry::Get().FindInterfaceRegistryEntry(RegistryKey);
-				if (ensure(Entry))
+				using FNameDataTypePair = TPair<FName, FName>;
+				TSet<FNameDataTypePair> InterfaceInputs;
+				TSet<FNameDataTypePair> InterfaceOutputs;
+
+				for (const FMetasoundFrontendVersion& Version : InterfaceVersions)
 				{
-					const FMetasoundFrontendInterface& Interface = Entry->GetInterface();
-					Algo::Transform(Interface.Inputs, InterfaceInputs, [](const FMetasoundFrontendClassInput& Input)
+					FInterfaceRegistryKey RegistryKey = GetInterfaceRegistryKey(Version);
+					const IInterfaceRegistryEntry* Entry = IInterfaceRegistry::Get().FindInterfaceRegistryEntry(RegistryKey);
+					if (ensure(Entry))
 					{
-						return FNameDataTypePair(Input.Name, Input.TypeName);
-					});
+						const FMetasoundFrontendInterface& Interface = Entry->GetInterface();
+						Algo::Transform(Interface.Inputs, InterfaceInputs, [](const FMetasoundFrontendClassInput& Input)
+						{
+							return FNameDataTypePair(Input.Name, Input.TypeName);
+						});
 
-					Algo::Transform(Interface.Outputs, InterfaceOutputs, [](const FMetasoundFrontendClassOutput& Output)
-					{
-						return FNameDataTypePair(Output.Name, Output.TypeName);
-					});
+						Algo::Transform(Interface.Outputs, InterfaceOutputs, [](const FMetasoundFrontendClassOutput& Output)
+						{
+							return FNameDataTypePair(Output.Name, Output.TypeName);
+						});
+					}
 				}
+
+				// Only serialize MetaData text for inputs owned by the graph (not by interfaces)
+				FMetasoundFrontendGraphClass RootGraphClass = InDocument->GetRootGraphClass();
+				for (FMetasoundFrontendClassInput& Input : RootGraphClass.Interface.Inputs)
+				{
+					const bool bSerializeText = !InterfaceInputs.Contains(FNameDataTypePair(Input.Name, Input.TypeName));
+					Input.Metadata.SetSerializeText(bSerializeText);
+				}
+
+				// Only serialize MetaData text for outputs owned by the graph (not by interfaces)
+				for (FMetasoundFrontendClassOutput& Output : RootGraphClass.Interface.Outputs)
+				{
+					const bool bSerializeText = !InterfaceOutputs.Contains(FNameDataTypePair(Output.Name, Output.TypeName));
+					Output.Metadata.SetSerializeText(bSerializeText);
+				}
+
+				InDocument->SetRootGraphClass(MoveTemp(RootGraphClass));
+	#else
+			UE_LOG(LogMetaSound, Error, TEXT("Asset '%s' at '%s' must be saved with editor enabled in order to version document to target version '%s'."), *Name.ToString(), *Path, *GetTargetVersion().ToString());
+	#endif // !WITH_EDITOR
 			}
+		};
 
-			// Only serialize MetaData text for inputs owned by the graph (not by interfaces)
-			FMetasoundFrontendGraphClass RootGraphClass = InDocument->GetRootGraphClass();
-			for (FMetasoundFrontendClassInput& Input : RootGraphClass.Interface.Inputs)
+		/** Versions document from 1.8 to 1.9. */
+		class FVersionDocument_1_9 : public FVersionDocumentTransform
+		{
+			FName Name;
+			const FString& Path;
+
+		public:
+			FVersionDocument_1_9(FName InName, const FString& InPath)
+				: Name(InName)
+				, Path(InPath)
 			{
-				const bool bSerializeText = !InterfaceInputs.Contains(FNameDataTypePair(Input.Name, Input.TypeName));
-				Input.Metadata.SetSerializeText(bSerializeText);
 			}
+			virtual ~FVersionDocument_1_9() = default;
 
-			// Only serialize MetaData text for outputs owned by the graph (not by interfaces)
-			for (FMetasoundFrontendClassOutput& Output : RootGraphClass.Interface.Outputs)
+			FMetasoundFrontendVersionNumber GetTargetVersion() const override
 			{
-				const bool bSerializeText = !InterfaceOutputs.Contains(FNameDataTypePair(Output.Name, Output.TypeName));
-				Output.Metadata.SetSerializeText(bSerializeText);
+				return { 1, 9 };
 			}
 
-			InDocument->SetRootGraphClass(MoveTemp(RootGraphClass));
-#else
-		UE_LOG(LogMetaSound, Error, TEXT("Asset '%s' at '%s' must be saved with editor enabled in order to version document to target version '%s'."), *Name.ToString(), *Path, *GetTargetVersion().ToString());
-#endif // !WITH_EDITOR
-		}
-	};
-
-	/** Versions document from 1.8 to 1.9. */
-	class FVersionDocument_1_9 : public FVersionDocumentTransform
-	{
-		FName Name;
-		const FString& Path;
+			void TransformInternal(FDocumentHandle InDocument) const override
+			{
+	#if WITH_EDITOR
+				// Display name text is no longer copied at this versioning point for assets
+				// from the asset's FName to avoid FText warnings regarding generation from
+				// an FString.  It also avoids desync if asset gets moved.
+				FMetasoundFrontendGraphClass RootGraphClass = InDocument->GetRootGraphClass();
+				RootGraphClass.Metadata.SetDisplayName(FText());
+				InDocument->SetRootGraphClass(MoveTemp(RootGraphClass));
+	#else
+				UE_LOG(LogMetaSound, Error, TEXT("Asset '%s' at '%s' must be saved with editor enabled in order to version document to target version '%s'."), *Name.ToString(), *Path, *GetTargetVersion().ToString());
+	#endif // !WITH_EDITOR
+			}
+		};
 
-	public:
-		FVersionDocument_1_9(FName InName, const FString& InPath)
-			: Name(InName)
-			, Path(InPath)
+		/** Versions document from 1.9 to 1.10. */
+		class FVersionDocument_1_10 : public FVersionDocumentTransform
 		{
-		}
+		public:
+			virtual ~FVersionDocument_1_10() = default;
 
-		FMetasoundFrontendVersionNumber GetTargetVersion() const override
-		{
-			return { 1, 9 };
-		}
+			FMetasoundFrontendVersionNumber GetTargetVersion() const override
+			{
+				return { 1, 10 };
+			}
+
+			void TransformInternal(FDocumentHandle InDocument) const override
+			{
+				FMetasoundFrontendGraphClass Class = InDocument->GetRootGraphClass();
+				FMetasoundFrontendGraphClassPresetOptions PresetOptions = Class.PresetOptions;
+				Class.PresetOptions.bIsPreset = Class.Metadata.GetAndClearAutoUpdateManagesInterface_Deprecated();
+				InDocument->SetRootGraphClass(MoveTemp(Class));
+			}
+		};
 
-		void TransformInternal(FDocumentHandle InDocument) const override
+		/** Versions document from 1.10 to 1.11. */
+		class FVersionDocument_1_11 : public FVersionDocumentTransform
 		{
-#if WITH_EDITOR
-			// Display name text is no longer copied at this versioning point for assets
-			// from the asset's FName to avoid FText warnings regarding generation from
-			// an FString.  It also avoids desync if asset gets moved.
-			FMetasoundFrontendGraphClass RootGraphClass = InDocument->GetRootGraphClass();
-			RootGraphClass.Metadata.SetDisplayName(FText());
-			InDocument->SetRootGraphClass(MoveTemp(RootGraphClass));
-#else
-			UE_LOG(LogMetaSound, Error, TEXT("Asset '%s' at '%s' must be saved with editor enabled in order to version document to target version '%s'."), *Name.ToString(), *Path, *GetTargetVersion().ToString());
-#endif // !WITH_EDITOR
-		}
-	};
+		public:
+			virtual ~FVersionDocument_1_11() = default;
 
-	/** Versions document from 1.9 to 1.10. */
-	class FVersionDocument_1_10 : public FVersionDocumentTransform
-	{
-	public:
-		FVersionDocument_1_10() = default;
+			FMetasoundFrontendVersionNumber GetTargetVersion() const override
+			{
+				return { 1, 11 };
+			}
 
-		FMetasoundFrontendVersionNumber GetTargetVersion() const override
-		{
-			return { 1, 10 };
-		}
+			void TransformInternal(FDocumentHandle InDocument) const override
+			{
+				// Clear object literals on inputs that are connected 
+				// to prevent referencing assets that are not used in the graph
+				InDocument->GetRootGraph()->IterateNodes([](FNodeHandle NodeHandle)
+				{
+					TArray<FInputHandle> NodeInputs = NodeHandle->GetInputs();
+					for (FInputHandle NodeInput : NodeInputs)
+					{
+						NodeInput->ClearConnectedObjectLiterals();
+					}
+				});
+			}
+		};
 
-		void TransformInternal(FDocumentHandle InDocument) const override
+		/** Versions document from 1.11 to 1.12. */
+		class FVersionDocument_1_12 : public FVersionDocumentTransform
 		{
-			FMetasoundFrontendGraphClass Class = InDocument->GetRootGraphClass();
-			FMetasoundFrontendGraphClassPresetOptions PresetOptions = Class.PresetOptions;
-			Class.PresetOptions.bIsPreset = Class.Metadata.GetAndClearAutoUpdateManagesInterface_Deprecated();
-			InDocument->SetRootGraphClass(MoveTemp(Class));
-		}
-	};
+			const FName Name;
+			const FSoftObjectPath* Path = nullptr;
 
-	/** Versions document from 1.10 to 1.11. */
-	class FVersionDocument_1_11 : public FVersionDocumentTransform
-	{
-	public:
-		FVersionDocument_1_11() = default;
+		public:
+			static const FMetasoundFrontendVersionNumber TargetVersion;
 
-		FMetasoundFrontendVersionNumber GetTargetVersion() const override
-		{
-			return { 1, 11 };
-		}
+			FVersionDocument_1_12(FName InName, const FSoftObjectPath& InAssetPath)
+				: Name(InName)
+				, Path(&InAssetPath)
+			{
+			}
+			virtual ~FVersionDocument_1_12() = default;
 
-		void TransformInternal(FDocumentHandle InDocument) const override
-		{
-			// Clear object literals on inputs that are connected 
-			// to prevent referencing assets that are not used in the graph
-			InDocument->GetRootGraph()->IterateNodes([](FNodeHandle NodeHandle)
+			FMetasoundFrontendVersionNumber GetTargetVersion() const override
 			{
-				TArray<FInputHandle> NodeInputs = NodeHandle->GetInputs();
-				for (FInputHandle NodeInput : NodeInputs)
+				return { 1, 12 };
+			}
+
+			void TransformInternal(FMetaSoundFrontendDocumentBuilder& OutBuilder) const override
+			{
+				using namespace VersioningPrivate;
+
+				if (IsRunningCookCommandlet())
 				{
-					NodeInput->ClearConnectedObjectLiterals();
+					UE_LOG(LogMetaSound, Display, TEXT("Resave recommended: Asset '%s' at '%s' skipped migrated editor data/creation of input template nodes during cook to target document version '%s'."), *Name.ToString(), *Path->ToString(), *GetTargetVersion().ToString());
 				}
-			});
-		}
-	};
-
-	/** Versions document from 1.11 to 1.12. */
-	class FVersionDocument_1_12 : public FVersionDocumentTransform
-	{
-		FMetasoundAssetBase* AssetBase = nullptr;
-		const FName Name;
-		const FSoftObjectPath* Path = nullptr;
-
-	public:
-		FVersionDocument_1_12(FMetasoundAssetBase& InAssetBase, FName InName, const FSoftObjectPath& InAssetPath)
-			: AssetBase(&InAssetBase)
-			, Name(InName)
-			, Path(&InAssetPath)
-		{
-		}
+				else
+				{
+					FMigratePagePropertiesTransform().Transform(OutBuilder);
+					OutBuilder.GetMetasoundAsset().MigrateEditorGraph(OutBuilder);
+					UE_LOG(LogMetaSound, Display, TEXT("Resave recommended: Asset '%s' at '%s' successfully migrated editor data in target document version '%s'."), *Name.ToString(), *Path->ToString(), *GetTargetVersion().ToString());
+				}
+			}
+		};
 
-		FMetasoundFrontendVersionNumber GetTargetVersion() const override
+			/** Versions document from 1.12 to 1.13. */
+		class FVersionDocument_1_13 : public FVersionDocumentTransform
 		{
-			return { 1, 12 };
-		}
+		public:
+			static const FMetasoundFrontendVersionNumber TargetVersion;
 
-		void TransformInternal(FMetaSoundFrontendDocumentBuilder& OutBuilder) const override
-		{
-			using namespace Metasound::Frontend;
+			virtual ~FVersionDocument_1_13() = default;
 
-#if WITH_EDITORONLY_DATA
-			if (IsRunningCookCommandlet())
+			FMetasoundFrontendVersionNumber GetTargetVersion() const override
 			{
-				UE_LOG(LogMetaSound, Display, TEXT("Resave recommended: Asset '%s' at '%s' skipped migrated editor data/creation of input template nodes during cook to target document version '%s'."), *Name.ToString(), *Path->ToString(), *GetTargetVersion().ToString());
+				return { 1, 13 };
 			}
-			else
+
+			void TransformInternal(FMetaSoundFrontendDocumentBuilder& OutBuilder) const override
 			{
-				check(AssetBase);
-				AssetBase->MigrateEditorGraph(OutBuilder);
-				UE_LOG(LogMetaSound, Display, TEXT("Resave recommended: Asset '%s' at '%s' successfully migrated editor data in target document version '%s'."), *Name.ToString(), *Path->ToString(), *GetTargetVersion().ToString());
+				FMigratePagePropertiesTransform().Transform(OutBuilder);
 			}
-#else // !WITH_EDITORONLY_DATA
-			UE_LOG(LogMetaSound, Warning, TEXT("Asset '%s' at '%s' must be re-saved with editor data enabled in order to migrate document to target version '%s'."), *Name.ToString(), *Path->ToString(), *GetTargetVersion().ToString());
-#endif // !WITH_EDITORONLY_DATA
-		}
-	};
+		};
 
-	namespace VersioningPrivate
-	{
-		bool VersionBuilderDocument(FMetasoundAssetBase& InAssetBase, FMetaSoundFrontendDocumentBuilder& Builder)
+		bool VersionBuilderDocument(FMetaSoundFrontendDocumentBuilder& Builder)
 		{
 			UObject& DocObject = Builder.CastDocumentObjectChecked<UObject>();
 			const FName Name = DocObject.GetFName();
 			const FString Path = DocObject.GetPathName();
 
 			bool bWasUpdated = false;
-			bWasUpdated |= FVersionDocument_1_12(InAssetBase, Name, Path).Transform(Builder);
+			bWasUpdated |= FVersionDocument_1_12(Name, Path).Transform(Builder);
+			bWasUpdated |= FVersionDocument_1_13().Transform(Builder);
 
 			return bWasUpdated;
 		}
 	} // namespace VersioningPrivate
 
-	bool VersionDocument(FMetasoundAssetBase& InAssetBase, FMetaSoundFrontendDocumentBuilder& Builder)
+	bool VersionDocument(FMetaSoundFrontendDocumentBuilder& Builder)
 	{
+		using namespace VersioningPrivate;
+
 		bool bWasUpdated = false;
 
-		UObject* OwningAsset = InAssetBase.GetOwningAsset();
-		TScriptInterface<IMetaSoundDocumentInterface> DocumentInterface(OwningAsset);
-		const FMetasoundFrontendDocument& Document = DocumentInterface->GetConstDocument();
-		const FName Name = FName(OwningAsset->GetName());
-		const FString Path = OwningAsset->GetPathName();
+		UObject& MetaSoundAsset = Builder.CastDocumentObjectChecked<UObject>();
+		const FName Name(*MetaSoundAsset.GetName());
+		const FString Path = MetaSoundAsset.GetPathName();
 
 		// Copied as value will be mutated with each applicable transform below
-		const FMetasoundFrontendVersionNumber InitVersionNumber = Document.Metadata.Version.Number;
+		const FMetasoundFrontendVersionNumber InitVersionNumber = Builder.GetConstDocumentChecked().Metadata.Version.Number;
+
+		// Old manual property transform that was applied prior to versioning schema being added.
+		// Only runs if internal logic finds necessary.
+		bWasUpdated = FVersionDocumentInterfacesTransform().Transform(Builder);
+
 		if (InitVersionNumber < GetMaxDocumentVersion())
 		{
 			// Controller (Soft Deprecated) Transforms
 			if (InitVersionNumber.Major == 1 && InitVersionNumber.Minor < 12)
 			{
-				UObject* DocObject = DocumentInterface.GetObject();
-				FDocumentHandle DocHandle = InAssetBase.GetDocumentHandle();
+				// Page Graph migration must be completed for graph accessor back
+				// compat prior to all controller versioning, so just do it here.
+				FMigratePagePropertiesTransform().Transform(Builder);
+
+				FDocumentHandle DocHandle = Builder.GetMetasoundAsset().GetDocumentHandle();
 
 				bWasUpdated |= FVersionDocument_1_1(Name, Path).Transform(DocHandle);
 				bWasUpdated |= FVersionDocument_1_2(Name, Path).Transform(DocHandle);
@@ -631,10 +741,10 @@ namespace Metasound::Frontend
 				// No longer supported, new versions should go in VersioningPrivate::VersionBuilderDocument
 			}
 
-			bWasUpdated |= VersioningPrivate::VersionBuilderDocument(InAssetBase, Builder);
+			bWasUpdated |= VersionBuilderDocument(Builder);
 			if (bWasUpdated)
 			{
-				const FMetasoundFrontendVersionNumber& NewVersionNumber = Document.Metadata.Version.Number;
+				const FMetasoundFrontendVersionNumber& NewVersionNumber = Builder.GetConstDocumentChecked().Metadata.Version.Number;
 				UE_LOG(LogMetaSound, Verbose, TEXT("MetaSound at '%s' Document Versioned: '%s' --> '%s'"), *Path, *InitVersionNumber.ToString(), *NewVersionNumber.ToString());
 			}
 		}
@@ -642,3 +752,4 @@ namespace Metasound::Frontend
 		return bWasUpdated;
 	}
 } // namespace Metasound::Frontend
+#endif // WITH_EDITORONLY_DATA
\ No newline at end of file
diff --git a/Engine/Plugins/Runtime/Metasound/Source/MetasoundFrontend/Private/MetasoundFrontendDocumentVersioning.h b/Engine/Plugins/Runtime/Metasound/Source/MetasoundFrontend/Private/MetasoundFrontendDocumentVersioning.h
--- a/Engine/Plugins/Runtime/Metasound/Source/MetasoundFrontend/Private/MetasoundFrontendDocumentVersioning.h
+++ b/Engine/Plugins/Runtime/Metasound/Source/MetasoundFrontend/Private/MetasoundFrontendDocumentVersioning.h
@@ -1,14 +1,13 @@
 // Copyright Epic Games, Inc. All Rights Reserved.
 #pragma once
 
-#include "Containers/UnrealString.h"
-#include "MetasoundDocumentInterface.h"
-#include "MetasoundFrontendController.h"
-#include "MetasoundFrontendTransform.h"
+#include "Misc/CoreMiscDefines.h"
+
+#if WITH_EDITORONLY_DATA
 #include "MetasoundFrontendDocument.h"
-#include "UObject/NameTypes.h"
-#include "UObject/ScriptInterface.h"
 
+
+// Forward Declartions
 class IMetaSoundDocumentInterface;
 struct FMetaSoundFrontendDocumentBuilder;
 
@@ -17,10 +16,11 @@ namespace Metasound::Frontend
 {
 	static FMetasoundFrontendVersionNumber GetMaxDocumentVersion()
 	{
-		return FMetasoundFrontendVersionNumber { 1, 12 };
+		return FMetasoundFrontendVersionNumber { 1, 13 };
 	}
 
 	// Versions Frontend Document. Passed as AssetBase for backward compat to
 	// version asset documents predating the IMetaSoundDocumentInterface
-	bool VersionDocument(FMetasoundAssetBase& InAssetBase, FMetaSoundFrontendDocumentBuilder& Builder);
+	bool VersionDocument(FMetaSoundFrontendDocumentBuilder& Builder);
 } // namespace Metasound::Frontend
+#endif // WITH_EDITORONLY_DATA
\ No newline at end of file
diff --git a/Engine/Plugins/Runtime/Metasound/Source/MetasoundFrontend/Private/MetasoundFrontendGraph.cpp b/Engine/Plugins/Runtime/Metasound/Source/MetasoundFrontend/Private/MetasoundFrontendGraph.cpp
--- a/Engine/Plugins/Runtime/Metasound/Source/MetasoundFrontend/Private/MetasoundFrontendGraph.cpp
+++ b/Engine/Plugins/Runtime/Metasound/Source/MetasoundFrontend/Private/MetasoundFrontendGraph.cpp
@@ -331,9 +331,10 @@ namespace Metasound
 		if (nullptr == Literal)
 		{
 			// Find Class Default that is not invalid
-			if (InOwningGraphClassInput.DefaultLiteral.IsValid())
+			const FMetasoundFrontendLiteral& DefaultLiteral = InOwningGraphClassInput.FindConstDefaultChecked(Frontend::DefaultPageID);
+			if (DefaultLiteral.IsValid())
 			{
-				Literal = &InOwningGraphClassInput.DefaultLiteral;
+				Literal = &DefaultLiteral;
 			}
 		}
 
@@ -342,9 +343,10 @@ namespace Metasound
 		{
 			const FMetasoundFrontendClassInput& InputNodeClassInput = InInputNodeClass.Interface.Inputs[0];
 
-			if (InputNodeClassInput.DefaultLiteral.IsValid())
+			const FMetasoundFrontendLiteral& DefaultLiteral = InOwningGraphClassInput.FindConstDefaultChecked(Frontend::DefaultPageID);
+			if (DefaultLiteral.IsValid())
 			{
-				Literal = &InputNodeClassInput.DefaultLiteral;
+				Literal = &DefaultLiteral;
 			}
 		}
 
diff --git a/Engine/Plugins/Runtime/Metasound/Source/MetasoundFrontend/Private/MetasoundFrontendGraphController.cpp b/Engine/Plugins/Runtime/Metasound/Source/MetasoundFrontend/Private/MetasoundFrontendGraphController.cpp
--- a/Engine/Plugins/Runtime/Metasound/Source/MetasoundFrontend/Private/MetasoundFrontendGraphController.cpp
+++ b/Engine/Plugins/Runtime/Metasound/Source/MetasoundFrontend/Private/MetasoundFrontendGraphController.cpp
@@ -1258,7 +1258,8 @@ namespace Metasound
 		{
 			if (const FMetasoundFrontendClassInput* Desc = FindInputDescriptionWithVertexID(InVertexID))
 			{
-				return Desc->DefaultLiteral;
+				const FMetasoundFrontendLiteral& DefaultLiteral = Desc->FindConstDefaultChecked(Frontend::DefaultPageID);
+				return DefaultLiteral;
 			}
 			return FMetasoundFrontendLiteral{};
 		}
@@ -1269,7 +1270,8 @@ namespace Metasound
 			{
 				if (ensure(IDataTypeRegistry::Get().IsLiteralTypeSupported(Desc->TypeName, InLiteral.GetType())))
 				{
-					Desc->DefaultLiteral = InLiteral;
+					FMetasoundFrontendLiteral& DefaultLiteral = Desc->FindDefaultChecked(Frontend::DefaultPageID);
+					DefaultLiteral = InLiteral;
 					return true;
 				}
 				else
@@ -1286,8 +1288,9 @@ namespace Metasound
 			if (FMetasoundFrontendClassInput* Desc = FindInputDescriptionWithVertexID(InVertexID))
 			{
 				Metasound::FLiteral Literal = IDataTypeRegistry::Get().CreateDefaultLiteral(Desc->TypeName);
-				Desc->DefaultLiteral.SetFromLiteral(Literal);
-				return Desc->DefaultLiteral.IsValid();
+				FMetasoundFrontendLiteral& DefaultLiteral = Desc->FindDefaultChecked(Frontend::DefaultPageID);
+				DefaultLiteral.SetFromLiteral(Literal);
+				return DefaultLiteral.IsValid();
 			}
 
 			return false;
@@ -1414,7 +1417,8 @@ namespace Metasound
 		{
 			if (FMetasoundFrontendClassInput* Desc = FindInputDescriptionWithName(InInputName))
 			{
-				Desc->DefaultLiteral.Clear();
+				FMetasoundFrontendLiteral& DefaultLiteral = Desc->FindDefaultChecked(Frontend::DefaultPageID);
+				DefaultLiteral.Clear();
 			}
 
 			return false;
@@ -1631,7 +1635,7 @@ namespace Metasound
 				FString UnknownAsset = TEXT("UnknownAsset");
 				FProxyDataCache ProxyCache;
 				const FMetasoundFrontendDocument& Doc = *(OwningDocument->GetDocumentPtr().Get());
-				ProxyCache.CreateAndCacheProxies(Doc, Frontend::DefaultGraphPageID);
+				ProxyCache.CreateAndCacheProxies(Doc, Frontend::DefaultPageID);
 
 				TUniquePtr<FFrontendGraph> Graph = FFrontendGraphBuilder::CreateGraph(*GraphClass, Subgraphs, Dependencies, ProxyCache, UnknownAsset, Frontend::CreateLocallyUniqueId(), { });
 
diff --git a/Engine/Plugins/Runtime/Metasound/Source/MetasoundFrontend/Private/MetasoundFrontendInputController.cpp b/Engine/Plugins/Runtime/Metasound/Source/MetasoundFrontend/Private/MetasoundFrontendInputController.cpp
--- a/Engine/Plugins/Runtime/Metasound/Source/MetasoundFrontend/Private/MetasoundFrontendInputController.cpp
+++ b/Engine/Plugins/Runtime/Metasound/Source/MetasoundFrontend/Private/MetasoundFrontendInputController.cpp
@@ -143,7 +143,7 @@ namespace Metasound
 		{
 			if (const FMetasoundFrontendClassInput* ClassInput = ClassInputPtr.Get())
 			{
-				return &(ClassInput->DefaultLiteral);
+				return ClassInput->FindConstDefault(Frontend::DefaultPageID);
 			}
 			return nullptr;
 		}
diff --git a/Engine/Plugins/Runtime/Metasound/Source/MetasoundFrontend/Private/MetasoundFrontendLiteral.cpp b/Engine/Plugins/Runtime/Metasound/Source/MetasoundFrontend/Private/MetasoundFrontendLiteral.cpp
--- a/Engine/Plugins/Runtime/Metasound/Source/MetasoundFrontend/Private/MetasoundFrontendLiteral.cpp
+++ b/Engine/Plugins/Runtime/Metasound/Source/MetasoundFrontend/Private/MetasoundFrontendLiteral.cpp
@@ -879,6 +879,13 @@ void FMetasoundFrontendLiteral::Clear()
 	AsNumDefault = 1;
 }
 
+FMetasoundFrontendLiteral FMetasoundFrontendLiteral::GetInvalid()
+{
+	FMetasoundFrontendLiteral InvalidValue;
+	InvalidValue.Type = EMetasoundFrontendLiteralType::Invalid;
+	return InvalidValue;
+}
+
 void FMetasoundFrontendLiteral::Empty()
 {
 	AsNumDefault = 0;
diff --git a/Engine/Plugins/Runtime/Metasound/Source/MetasoundFrontend/Private/MetasoundFrontendNodeController.cpp b/Engine/Plugins/Runtime/Metasound/Source/MetasoundFrontend/Private/MetasoundFrontendNodeController.cpp
--- a/Engine/Plugins/Runtime/Metasound/Source/MetasoundFrontend/Private/MetasoundFrontendNodeController.cpp
+++ b/Engine/Plugins/Runtime/Metasound/Source/MetasoundFrontend/Private/MetasoundFrontendNodeController.cpp
@@ -134,7 +134,7 @@ namespace Metasound
 						if (ClassInput)
 						{
 							// Check if setting back to class default literal
-							const FMetasoundFrontendLiteral& ClassDefaultLiteral = ClassInput->DefaultLiteral;
+							const FMetasoundFrontendLiteral& ClassDefaultLiteral = ClassInput->FindConstDefaultChecked(Frontend::DefaultPageID);
 							bSettingToClassDefaultLiteral = ClassDefaultLiteral.IsEqual(InVertexLiteral.Value);
 
 							// Check if setting from class default literal (which may have a None type) to an appropriate type 
@@ -964,12 +964,13 @@ namespace Metasound
 			Algo::Transform(OutInterfaceUpdates.RegistryClass.Interface.Inputs, OutInterfaceUpdates.AddedInputs, [&](const FMetasoundFrontendClassInput& Input) { return &Input; });
 			for (const FMetasoundFrontendClassInput& Input : NodeClassInterface.Inputs)
 			{
-				auto IsEquivalent = [NodeClassInput = &Input](const FMetasoundFrontendClassInput* Iter)
+				const FMetasoundFrontendLiteral& NodeClassDefault = Input.FindConstDefaultChecked(Frontend::DefaultPageID);
+				auto IsEquivalent = [&NodeClassDefault, &Input](const FMetasoundFrontendClassInput* Iter)
 				{
-					const bool bDefaultEquivalent = Iter->DefaultLiteral.IsEqual(NodeClassInput->DefaultLiteral);
+					const bool bDefaultEquivalent = Iter->FindConstDefaultChecked(Frontend::DefaultPageID).IsEqual(NodeClassDefault);
 					if (bDefaultEquivalent)
 					{
-						return FMetasoundFrontendClassVertex::IsFunctionalEquivalent(*NodeClassInput, *Iter);
+						return FMetasoundFrontendClassVertex::IsFunctionalEquivalent(Input, *Iter);
 					}
 					return false;
 				};
diff --git a/Engine/Plugins/Runtime/Metasound/Source/MetasoundFrontend/Private/MetasoundFrontendRegistryContainerImpl.cpp b/Engine/Plugins/Runtime/Metasound/Source/MetasoundFrontend/Private/MetasoundFrontendRegistryContainerImpl.cpp
--- a/Engine/Plugins/Runtime/Metasound/Source/MetasoundFrontend/Private/MetasoundFrontendRegistryContainerImpl.cpp
+++ b/Engine/Plugins/Runtime/Metasound/Source/MetasoundFrontend/Private/MetasoundFrontendRegistryContainerImpl.cpp
@@ -480,7 +480,7 @@ namespace Metasound::Frontend
 
 		// Use the asset path of the provided document interface object for identification, *NOT* the
 		// built version as the build process may in fact create a new object with a transient path.
-		FGuid PageID = Frontend::DefaultGraphPageID;
+		FGuid PageID = Frontend::DefaultPageID;
 		const TScriptInterface<IMetaSoundDocumentInterface> RegistryDocInterface = RegistryPrivate::BuildRegistryDocument(InDocumentInterface, bAsync, PageID);
 
 		UObject* OwningObject = RegistryDocInterface.GetObject();
diff --git a/Engine/Plugins/Runtime/Metasound/Source/MetasoundFrontend/Private/MetasoundFrontendTransform.cpp b/Engine/Plugins/Runtime/Metasound/Source/MetasoundFrontend/Private/MetasoundFrontendTransform.cpp
--- a/Engine/Plugins/Runtime/Metasound/Source/MetasoundFrontend/Private/MetasoundFrontendTransform.cpp
+++ b/Engine/Plugins/Runtime/Metasound/Source/MetasoundFrontend/Private/MetasoundFrontendTransform.cpp
@@ -288,7 +288,8 @@ namespace Metasound
 				{
 					if (FMetasoundFrontendVertex::IsFunctionalEquivalent(*ClassInput, OriginalVertex))
 					{
-						NewVertex.DefaultLiteral = ClassInput->DefaultLiteral;
+						const FMetasoundFrontendLiteral& DefaultLiteral = ClassInput->FindConstDefaultChecked(Frontend::DefaultPageID);
+						NewVertex.FindDefaultChecked(Frontend::DefaultPageID) = DefaultLiteral;
 						NewVertex.NodeID = ClassInput->NodeID;
 						FNodeHandle OriginalInputNode = GraphHandle->GetInputNodeWithName(OriginalVertex.Name);
 
@@ -866,6 +867,7 @@ namespace Metasound
 						ClassInput.NodeID = ExistingClassInput->NodeID;
 					}
 
+					FMetasoundFrontendLiteral& DefaultLiteral = ClassInput.InitDefault();
 					if (InPresetGraph->ContainsInputVertex(NodeName, ClassInput.TypeName))
 					{
 						// If the input vertex already exists in the parent graph,
@@ -874,12 +876,12 @@ namespace Metasound
 						if (InPresetGraph->GetInputsInheritingDefault().Contains(NodeName))
 						{
 							const FGuid ReferencedVertexID = ReferencedGraph->GetVertexIDForInputVertex(NodeName);
-							ClassInput.DefaultLiteral = ReferencedGraph->GetDefaultInput(ReferencedVertexID);
+							DefaultLiteral = ReferencedGraph->GetDefaultInput(ReferencedVertexID);
 						}
 						else
 						{
 							FGuid VertexID = InPresetGraph->GetVertexIDForInputVertex(NodeName);
-							ClassInput.DefaultLiteral = InPresetGraph->GetDefaultInput(VertexID);
+							DefaultLiteral = InPresetGraph->GetDefaultInput(VertexID);
 						}
 					}
 					else
@@ -888,7 +890,7 @@ namespace Metasound
 						// then it is a new vertex and should use the default value
 						// of the referenced graph.
 						const FGuid ReferencedVertexID = ReferencedGraph->GetVertexIDForInputVertex(NodeName);
-						ClassInput.DefaultLiteral = ReferencedGraph->GetDefaultInput(ReferencedVertexID);
+						DefaultLiteral = ReferencedGraph->GetDefaultInput(ReferencedVertexID);
 					}
 
 					ClassInputs.Add(MoveTemp(ClassInput));
diff --git a/Engine/Plugins/Runtime/Metasound/Source/MetasoundFrontend/Private/MetasoundParameterPack.cpp b/Engine/Plugins/Runtime/Metasound/Source/MetasoundFrontend/Private/MetasoundParameterPack.cpp
--- a/Engine/Plugins/Runtime/Metasound/Source/MetasoundFrontend/Private/MetasoundParameterPack.cpp
+++ b/Engine/Plugins/Runtime/Metasound/Source/MetasoundFrontend/Private/MetasoundParameterPack.cpp
@@ -29,17 +29,20 @@ TSharedPtr<Audio::IProxyData> UMetasoundParameterPack::CreateProxyData(const Aud
 
 Metasound::FSendAddress UMetasoundParameterPack::CreateSendAddressFromEnvironment(const Metasound::FMetasoundEnvironment& InEnvironment)
 {
-	PRAGMA_DISABLE_DEPRECATION_WARNINGS
+PRAGMA_DISABLE_DEPRECATION_WARNINGS
 	return Metasound::FMetaSoundParameterTransmitter::CreateSendAddressFromEnvironment(InEnvironment, MetasoundParameterPack::RoutingName, Metasound::GetMetasoundDataTypeName<FMetasoundParameterStorageWrapper>());
-	PRAGMA_ENABLE_DEPRECATION_WARNINGS
+PRAGMA_ENABLE_DEPRECATION_WARNINGS
 }
 
 FMetasoundFrontendClassInput UMetasoundParameterPack::GetClassInput()
 {
 	FMetasoundFrontendClassInput ClassInput;
 	ClassInput.Name = MetasoundParameterPack::RoutingName;
-	ClassInput.DefaultLiteral = FMetasoundFrontendLiteral({ MetasoundParameterPack::RoutingName, nullptr });
 	ClassInput.TypeName = Metasound::GetMetasoundDataTypeName<FMetasoundParameterStorageWrapper>();
+
+	constexpr UObject* NullObjectPtr = nullptr;
+	ClassInput.InitDefault() = FAudioParameter { MetasoundParameterPack::RoutingName, NullObjectPtr };
+
 	return ClassInput;
 }
 
diff --git a/Engine/Plugins/Runtime/Metasound/Source/MetasoundFrontend/Public/MetasoundFrontendDocument.h b/Engine/Plugins/Runtime/Metasound/Source/MetasoundFrontend/Public/MetasoundFrontendDocument.h
--- a/Engine/Plugins/Runtime/Metasound/Source/MetasoundFrontend/Public/MetasoundFrontendDocument.h
+++ b/Engine/Plugins/Runtime/Metasound/Source/MetasoundFrontend/Public/MetasoundFrontendDocument.h
@@ -22,6 +22,7 @@
 // Forward Declarations
 struct FMetasoundFrontendClass;
 struct FMetasoundFrontendClassInterface;
+struct FMetaSoundFrontendDocumentBuilder;
 
 enum class EMetasoundFrontendClassType : uint8;
 
@@ -35,8 +36,8 @@ namespace Metasound
 
 	namespace Frontend
 	{
-		constexpr FGuid DefaultGraphPageID(0, 0, 0, 0);
-		constexpr TCHAR DefaultGraphPageName[] = TEXT("Default");
+		constexpr FGuid DefaultPageID(0, 0, 0, 0);
+		constexpr TCHAR DefaultPageName[] = TEXT("Default");
 
 #if WITH_EDITORONLY_DATA
 		extern const FText METASOUNDFRONTEND_API DefaultGraphPageDisplayName;
@@ -934,6 +935,36 @@ struct METASOUNDFRONTEND_API FMetasoundFrontendClassStyleDisplay
 };
 
 
+USTRUCT()
+struct METASOUNDFRONTEND_API FMetasoundFrontendClassInputDefault
+{
+	GENERATED_BODY()
+
+	FMetasoundFrontendClassInputDefault() = default;
+
+	FMetasoundFrontendClassInputDefault(const FMetasoundFrontendLiteral& Literal)
+		: Literal(Literal)
+		, PageID(Metasound::Frontend::DefaultPageID)
+	{
+	}
+
+	FMetasoundFrontendClassInputDefault(const FGuid& InPageID)
+		: PageID(InPageID)
+	{
+	}
+
+	FMetasoundFrontendClassInputDefault(const FAudioParameter& InParameter)
+		: Literal(InParameter)
+	{
+	}
+
+	UPROPERTY()
+	FMetasoundFrontendLiteral Literal;
+
+	UPROPERTY()
+	FGuid PageID = Metasound::Frontend::DefaultPageID;
+};
+
 // Contains info for input vertex of a Metasound class.
 USTRUCT() 
 struct METASOUNDFRONTEND_API FMetasoundFrontendClassInput : public FMetasoundFrontendClassVertex
@@ -941,13 +972,37 @@ struct METASOUNDFRONTEND_API FMetasoundFrontendClassInput : public FMetasoundFro
 	GENERATED_BODY()
 
 	FMetasoundFrontendClassInput() = default;
-
 	FMetasoundFrontendClassInput(const FMetasoundFrontendClassVertex& InOther);
 	FMetasoundFrontendClassInput(const Audio::FParameterInterface::FInput& InInput);
 
-	// Default value for this input.
-	UPROPERTY(EditAnywhere, Category = Parameters)
+#if WITH_EDITORONLY_DATA
+	UPROPERTY(meta = (DeprecationMessage = "5.5 - Direct access will be revoked and page manipulation limited to public API in future builds. Field has been rolled into DefaultLiterals Array."))
 	FMetasoundFrontendLiteral DefaultLiteral;
+#endif // WITH_EDITORONLY_DATA
+
+private:
+	UPROPERTY(EditAnywhere, Category = Parameters)
+	TArray<FMetasoundFrontendClassInputDefault> Defaults;
+
+public:
+#if WITH_EDITORONLY_DATA
+	FMetasoundFrontendLiteral& AddDefault(const FGuid& InPageID);
+#endif // WITH_EDITORONLY_DATA
+
+	bool ContainsDefault(const FGuid& InPageID) const;
+	const FMetasoundFrontendLiteral* FindConstDefault(const FGuid& InPageID) const;
+	const FMetasoundFrontendLiteral& FindConstDefaultChecked(const FGuid& InPageID) const;
+	FMetasoundFrontendLiteral* FindDefault(const FGuid& InPageID);
+	FMetasoundFrontendLiteral& FindDefaultChecked(const FGuid& InPageID);
+	FMetasoundFrontendLiteral& InitDefault();
+	void InitDefault(FMetasoundFrontendLiteral InitLiteral);
+	void IterateDefaults(TFunctionRef<void(const FGuid&, FMetasoundFrontendLiteral&)> IterFunc);
+	void IterateDefaults(TFunctionRef<void(const FGuid&, const FMetasoundFrontendLiteral&)> IterFunc) const;
+
+#if WITH_EDITORONLY_DATA
+	void RemoveAllDefaults();
+	bool RemoveDefault(const FGuid& InPageID);
+#endif // WITH_EDITORONLY_DATA
 };
 
 // Contains info for variable vertex of a Metasound class.
@@ -1664,8 +1719,10 @@ public:
 	FMetasoundFrontendGraphClass();
 	virtual ~FMetasoundFrontendGraphClass() = default;
 
-	UPROPERTY(meta = (DeprecationMessage = "5.5 - Direct Graph access will be revoked and page manipulation applied via public API below in future builds."))
+#if WITH_EDITORONLY_DATA
+	UPROPERTY(meta = (DeprecationMessage = "5.5 - GraphClasses now support multiple paged graphs. Use the provided page graph accessors"))
 	FMetasoundFrontendGraph Graph;
+#endif // WITH_EDITORONLY_DATA
 
 private:
 	UPROPERTY()
@@ -1687,6 +1744,9 @@ public:
 #endif // WITH_EDITORONLY_DATA
 
 	bool ContainsGraphPage(const FGuid& InPageID) const;
+
+	FMetasoundFrontendGraph& InitDefaultPage();
+
 	void IterateGraphPages(TFunctionRef<void(FMetasoundFrontendGraph&)> IterFunc);
 	void IterateGraphPages(TFunctionRef<void(const FMetasoundFrontendGraph&)> IterFunc) const;
 
@@ -1697,6 +1757,20 @@ public:
 	FMetasoundFrontendGraph& GetDefaultGraph();
 	const FMetasoundFrontendGraph& GetConstDefaultGraph() const;
 	void ResetGraphs();
+
+#if WITH_EDITORONLY_DATA
+	struct IPropertyVersionTransform
+	{
+	public:
+		virtual ~IPropertyVersionTransform() = default;
+
+	protected:
+		virtual bool Transform(FMetasoundFrontendGraphClass& OutClass) const = 0;
+
+		// Allows for unsafe access to a document for property migration.
+		static TArray<FMetasoundFrontendGraph>& GetPagesUnsafe(FMetasoundFrontendGraphClass& GraphClass);
+	};
+#endif // WITH_EDITORONLY_DATA
 };
 
 UCLASS()
@@ -1744,7 +1818,9 @@ struct METASOUNDFRONTEND_API FMetasoundFrontendDocument
 	GENERATED_BODY()
 
 public:
+#if WITH_EDITORONLY_DATA
 	static FMetasoundFrontendVersionNumber GetMaxVersion();
+#endif // WITH_EDITORONLY_DATA
 
 	Metasound::Frontend::FAccessPoint AccessPoint;
 
diff --git a/Engine/Plugins/Runtime/Metasound/Source/MetasoundFrontend/Public/MetasoundFrontendDocumentBuilder.h b/Engine/Plugins/Runtime/Metasound/Source/MetasoundFrontend/Public/MetasoundFrontendDocumentBuilder.h
--- a/Engine/Plugins/Runtime/Metasound/Source/MetasoundFrontend/Public/MetasoundFrontendDocumentBuilder.h
+++ b/Engine/Plugins/Runtime/Metasound/Source/MetasoundFrontend/Public/MetasoundFrontendDocumentBuilder.h
@@ -180,7 +180,7 @@ public:
 
 #if WITH_EDITORONLY_DATA
 	// Adds a graph page to the given builder's document
-	const FMetasoundFrontendGraph& AddGraphPage(const FGuid& InPageID, bool bDuplicateLastGraph = true, bool bSetAsBuildGraph = true);
+	const FMetasoundFrontendGraph& AddGraphPage(const FGuid& InPageID, bool bDuplicateLastGraph, bool bSetAsBuildGraph = true);
 #endif // WITH_EDITORONLY_DATA
 
 	// Returns whether or not the given edge can be added, which requires that its input
@@ -206,7 +206,7 @@ public:
 	bool ConvertFromPreset();
 	bool ConvertToPreset(const FMetasoundFrontendDocument& InReferencedDocument, TSharedRef<Metasound::Frontend::FDocumentModifyDelegates> ModifyDelegates = { });
 
-	const FMetasoundFrontendNode* DuplicateGraphInput(const FMetasoundFrontendClassInput& InClassInput, FMetasoundFrontendLiteral DefaultValue, const FName InName, const FGuid* InPageID = nullptr);
+	const FMetasoundFrontendNode* DuplicateGraphInput(const FMetasoundFrontendClassInput& InClassInput, const FName InName, const FGuid* InPageID = nullptr);
 	const FMetasoundFrontendNode* DuplicateGraphOutput(const FMetasoundFrontendClassOutput& InClassOutput, const FName InName, const FGuid* InPageID = nullptr);
 
 #if WITH_EDITORONLY_DATA
@@ -440,7 +440,21 @@ public:
 	bool TransformTemplateNodes();
 
 	// Versions legacy document members that contained interface information
+	UE_DEPRECATED(5.5, "Moved to internally implemented versioning logic")
 	bool VersionInterfaces();
+
+	// Struct enabling property migration of data that must be applied prior to versioning logic
+	struct IPropertyVersionTransform
+	{
+	public:
+		virtual ~IPropertyVersionTransform() = default;
+
+	protected:
+		virtual bool Transform(FMetaSoundFrontendDocumentBuilder& Builder) const = 0;
+
+		// Allows for unsafe access to a document for property migration.
+		static FMetasoundFrontendDocument& GetDocumentUnsafe(const FMetaSoundFrontendDocumentBuilder& Builder);
+	};
 #endif // WITH_EDITORONLY_DATA
 
 private:
diff --git a/Engine/Plugins/Runtime/Metasound/Source/MetasoundFrontend/Public/MetasoundFrontendLiteral.h b/Engine/Plugins/Runtime/Metasound/Source/MetasoundFrontend/Public/MetasoundFrontendLiteral.h
--- a/Engine/Plugins/Runtime/Metasound/Source/MetasoundFrontend/Public/MetasoundFrontendLiteral.h
+++ b/Engine/Plugins/Runtime/Metasound/Source/MetasoundFrontend/Public/MetasoundFrontendLiteral.h
@@ -161,6 +161,8 @@ public:
 	// Remove any stored data and set to an invalid state.
 	void Clear();
 
+	static FMetasoundFrontendLiteral GetInvalid();
+
 private:
 	// Remove all values.
 	void Empty();
