Fix MetaSound frontend comment node ids and locations not being updated properly

diff --git a/Engine/Plugins/Runtime/Metasound/Source/MetasoundEditor/Private/MetasoundEditorDocumentClipboardUtils.cpp b/Engine/Plugins/Runtime/Metasound/Source/MetasoundEditor/Private/MetasoundEditorDocumentClipboardUtils.cpp
--- a/Engine/Plugins/Runtime/Metasound/Source/MetasoundEditor/Private/MetasoundEditorDocumentClipboardUtils.cpp
+++ b/Engine/Plugins/Runtime/Metasound/Source/MetasoundEditor/Private/MetasoundEditorDocumentClipboardUtils.cpp
@@ -303,16 +303,40 @@ namespace Metasound::Editor
 		}
 	}
 
-	void FDocumentClipboardUtils::ProcessPastedCommentNodes(FMetasoundAssetBase& OutAsset, const TArrayView<UMetasoundEditorGraphCommentNode*> CommentNodes)
+	void FDocumentClipboardUtils::ProcessPastedCommentNodes(FMetasoundAssetBase& OutAsset, const FVector2D& InLocation, const TArrayView<UMetasoundEditorGraphCommentNode*> CommentNodes)
 	{
 		using namespace Engine;
 		using namespace Frontend;
 
 		UMetasoundEditorGraph& Graph = *CastChecked<UMetasoundEditorGraph>(&OutAsset.GetGraphChecked());
 		UMetaSoundBuilderBase& Builder = FDocumentBuilderRegistry::GetChecked().FindOrBeginBuilding(*OutAsset.GetOwningAsset());
+
+		// Calculate average node position 
+		FVector2D AvgNodePosition = FVector2D::ZeroVector;
+		for (UEdGraphNode* Node : CommentNodes)
+		{
+			AvgNodePosition.X += Node->NodePosX;
+			AvgNodePosition.Y += Node->NodePosY;
+		}
+
+		if (!CommentNodes.IsEmpty())
+		{
+			float InvNumNodes = 1.0f / CommentNodes.Num();
+			AvgNodePosition.X *= InvNumNodes;
+			AvgNodePosition.Y *= InvNumNodes;
+		}
+
 		for (UMetasoundEditorGraphCommentNode* CommentNode : CommentNodes)
 		{
+			// Regenerate id
 			CommentNode->CreateNewGuid();
+			CommentNode->SetCommentID(CommentNode->NodeGuid);
+
+			// Add offset based on paste location to avoid comments automatically being captured by the original node if using group movement 
+			CommentNode->NodePosX = (CommentNode->NodePosX - AvgNodePosition.X) + InLocation.X;
+			CommentNode->NodePosY = (CommentNode->NodePosY - AvgNodePosition.Y) + InLocation.Y;
+
+			// Update frontend node
 			FMetaSoundFrontendGraphComment& NewComment = Builder.FindOrAddGraphComment(CommentNode->GetCommentID());
 			UMetasoundEditorGraphCommentNode::ConvertToFrontendComment(*CommentNode, NewComment);
 		}
@@ -432,7 +456,7 @@ namespace Metasound::Editor
 		OutMetaSound.Modify();
 		Asset->GetGraphChecked().Modify();
 
-		ProcessPastedCommentNodes(*Asset, PastedCommentNodes);
+		ProcessPastedCommentNodes(*Asset, InLocation, PastedCommentNodes);
 
 		ProcessPastedInputNodes(*Asset, PastedGraphNodes);
 		ProcessPastedOutputNodes(*Asset, PastedGraphNodes);
diff --git a/Engine/Plugins/Runtime/Metasound/Source/MetasoundEditor/Private/MetasoundEditorDocumentClipboardUtils.h b/Engine/Plugins/Runtime/Metasound/Source/MetasoundEditor/Private/MetasoundEditorDocumentClipboardUtils.h
--- a/Engine/Plugins/Runtime/Metasound/Source/MetasoundEditor/Private/MetasoundEditorDocumentClipboardUtils.h
+++ b/Engine/Plugins/Runtime/Metasound/Source/MetasoundEditor/Private/MetasoundEditorDocumentClipboardUtils.h
@@ -26,7 +26,7 @@ namespace Metasound::Editor
 		static TArray<UEdGraphNode*> PasteClipboardString(const FText& InTransactionText, const FString& InClipboardString, const FVector2D& InLocation, UObject& OutMetaSound, FDocumentPasteNotifications& OutNotifications);
 
 	private:
-		static void ProcessPastedCommentNodes(FMetasoundAssetBase& OutAsset, const TArrayView<UMetasoundEditorGraphCommentNode*> CommentNodes);
+		static void ProcessPastedCommentNodes(FMetasoundAssetBase& OutAsset, const FVector2D& InLocation, const TArrayView<UMetasoundEditorGraphCommentNode*> CommentNodes);
 
 		static void ProcessPastedExternalNodes(FMetasoundAssetBase& OutAsset, TArray<UMetasoundEditorGraphNode*>& OutPastedNodes, FDocumentPasteNotifications& OutNotifications);
 
diff --git a/Engine/Plugins/Runtime/Metasound/Source/MetasoundEditor/Private/MetasoundEditorGraphCommentNode.cpp b/Engine/Plugins/Runtime/Metasound/Source/MetasoundEditor/Private/MetasoundEditorGraphCommentNode.cpp
--- a/Engine/Plugins/Runtime/Metasound/Source/MetasoundEditor/Private/MetasoundEditorGraphCommentNode.cpp
+++ b/Engine/Plugins/Runtime/Metasound/Source/MetasoundEditor/Private/MetasoundEditorGraphCommentNode.cpp
@@ -94,6 +94,12 @@ FGuid UMetasoundEditorGraphCommentNode::GetCommentID() const
 	return CommentID;
 }
 
+UObject& UMetasoundEditorGraphCommentNode::GetMetasoundChecked() const
+{
+	UMetasoundEditorGraph* EdGraph = CastChecked<UMetasoundEditorGraph>(GetGraph());
+	return EdGraph->GetMetasoundChecked();
+}
+
 bool UMetasoundEditorGraphCommentNode::RemoveFromDocument() const
 {
 	return GetBuilderChecked().RemoveGraphComment(CommentID);
@@ -102,10 +108,12 @@ bool UMetasoundEditorGraphCommentNode::RemoveFromDocument() const
 void UMetasoundEditorGraphCommentNode::SetBounds(const class FSlateRect& Rect)
 {
 	Super::SetBounds(Rect);
+	UpdateFrontendNodeLocation();
+}
 
-	FMetaSoundFrontendGraphComment& FrontendComment = GetBuilderChecked().FindOrAddGraphComment(CommentID);
-	FrontendComment.Position = FVector2D(NodePosX, NodePosY);
-	FrontendComment.Size = FVector2D(NodeWidth, NodeHeight);
+void UMetasoundEditorGraphCommentNode::SetCommentID(const FGuid& InGuid)
+{ 
+	CommentID = InGuid; 
 }
 
 void UMetasoundEditorGraphCommentNode::OnRenameNode(const FString& NewName)
@@ -115,4 +123,12 @@ void UMetasoundEditorGraphCommentNode::OnRenameNode(const FString& NewName)
 	UMetaSoundBuilderBase& Builder = GetBuilderChecked();
 	Builder.FindOrAddGraphComment(CommentID).Comment = NewName;
 }
+
+void UMetasoundEditorGraphCommentNode::UpdateFrontendNodeLocation()
+{
+	using namespace Metasound::Frontend;
+	FMetaSoundFrontendGraphComment& FrontendComment = GetBuilderChecked().FindOrAddGraphComment(CommentID);
+	FrontendComment.Position = FVector2D(NodePosX, NodePosY);
+	FrontendComment.Size = FVector2D(NodeWidth, NodeHeight);
+}
 #undef LOCTEXT_NAMESPACE
diff --git a/Engine/Plugins/Runtime/Metasound/Source/MetasoundEditor/Private/SMetasoundGraphNodeComment.cpp b/Engine/Plugins/Runtime/Metasound/Source/MetasoundEditor/Private/SMetasoundGraphNodeComment.cpp
--- a/Engine/Plugins/Runtime/Metasound/Source/MetasoundEditor/Private/SMetasoundGraphNodeComment.cpp
+++ b/Engine/Plugins/Runtime/Metasound/Source/MetasoundEditor/Private/SMetasoundGraphNodeComment.cpp
@@ -14,6 +14,14 @@ namespace Metasound
 		{
 			SGraphNodeComment::MoveTo(NewPosition, NodeFilter, bMarkDirty);
 
+			// Update frontend node position for current node 
+			UEdGraphNode* Node = GetNodeObj();
+			if (UMetasoundEditorGraphCommentNode* MetaSoundCommentNode = Cast<UMetasoundEditorGraphCommentNode>(Node))
+			{
+				MetaSoundCommentNode->GetMetasoundChecked().Modify();
+				MetaSoundCommentNode->UpdateFrontendNodeLocation();
+			}
+			
 			// Update Frontend node positions for unselected nodes that are dragged along with the comment box
 			// partially copied from SGraphNodeComment::MoveTo
 			// Don't drag note content if either of the shift keys are down.
@@ -34,10 +42,14 @@ namespace Metasound
 						{
 							FVector2D MetasoundNodePosition = FVector2D(MetasoundGraphNode->NodePosX, MetasoundGraphNode->NodePosY);
 							MetasoundNodePosition += PositionDelta;
-
 							MetasoundGraphNode->GetMetasoundChecked().Modify();
 							MetasoundGraphNode->UpdateFrontendNodeLocation(MetasoundNodePosition);
 						}
+						else if (UMetasoundEditorGraphCommentNode* MetasoundCommentNode = Cast<UMetasoundEditorGraphCommentNode>(*NodeIt))
+						{
+							MetasoundCommentNode->GetMetasoundChecked().Modify();
+							MetasoundCommentNode->UpdateFrontendNodeLocation();
+						}
 					}
 				}
 			}
diff --git a/Engine/Plugins/Runtime/Metasound/Source/MetasoundEditor/Public/MetasoundEditorGraphCommentNode.h b/Engine/Plugins/Runtime/Metasound/Source/MetasoundEditor/Public/MetasoundEditorGraphCommentNode.h
--- a/Engine/Plugins/Runtime/Metasound/Source/MetasoundEditor/Public/MetasoundEditorGraphCommentNode.h
+++ b/Engine/Plugins/Runtime/Metasound/Source/MetasoundEditor/Public/MetasoundEditorGraphCommentNode.h
@@ -51,11 +51,18 @@ public:
 
 	UMetaSoundBuilderBase& GetBuilderChecked() const;
 
-	// Returns frontend comment ID
+	// Get/Set frontend comment ID
 	FGuid GetCommentID() const;
+	void SetCommentID(const FGuid& InGuid);
+
+	// Get this node's MetaSound
+	UObject& GetMetasoundChecked() const;
 
 	bool RemoveFromDocument() const;
 
+	// Update the frontend location with this editor node location
+	void UpdateFrontendNodeLocation();
+
 private:
 	FMetasoundAssetBase& GetAssetChecked();
 
