Part 1: AudioUnitTest fix-ups in response to prior default page graph/input literal changes
#rnx

diff --git a/Engine/Plugins/Runtime/Metasound/Source/MetasoundFrontend/Private/MetasoundFrontendDocument.cpp b/Engine/Plugins/Runtime/Metasound/Source/MetasoundFrontend/Private/MetasoundFrontendDocument.cpp
--- a/Engine/Plugins/Runtime/Metasound/Source/MetasoundFrontend/Private/MetasoundFrontendDocument.cpp
+++ b/Engine/Plugins/Runtime/Metasound/Source/MetasoundFrontend/Private/MetasoundFrontendDocument.cpp
@@ -1075,7 +1075,7 @@ const FMetasoundFrontendGraph& FMetasoundFrontendGraphClass::GetConstDefaultGrap
 	return FindConstGraphChecked(Metasound::Frontend::DefaultPageID);
 }
 
-FMetasoundFrontendGraph& FMetasoundFrontendGraphClass::InitDefaultPage()
+FMetasoundFrontendGraph& FMetasoundFrontendGraphClass::InitDefaultGraph()
 {
 	checkf(PagedGraphs.IsEmpty(), TEXT("Attempting to initialize default page for graph class with existing graph implementation"));
 	FMetasoundFrontendGraph& NewGraph = PagedGraphs.AddDefaulted_GetRef();
diff --git a/Engine/Plugins/Runtime/Metasound/Source/MetasoundFrontend/Private/MetasoundFrontendDocumentBuilder.cpp b/Engine/Plugins/Runtime/Metasound/Source/MetasoundFrontend/Private/MetasoundFrontendDocumentBuilder.cpp
--- a/Engine/Plugins/Runtime/Metasound/Source/MetasoundFrontend/Private/MetasoundFrontendDocumentBuilder.cpp
+++ b/Engine/Plugins/Runtime/Metasound/Source/MetasoundFrontend/Private/MetasoundFrontendDocumentBuilder.cpp
@@ -1363,7 +1363,7 @@ void FMetaSoundFrontendDocumentBuilder::ClearDocument(TSharedRef<Metasound::Fron
 	checkf(!DocObject.IsAsset(), TEXT("Cannot call clear document on asset '%s': builder API does not support document mutation on serialized objects without editor data loaded"), *GetDebugName());
 #endif // !WITH_EDITORONLY_DATA
 
-	GraphClass.InitDefaultPage();
+	GraphClass.InitDefaultGraph();
 
 	Doc.Interfaces.Empty();
 	Doc.Dependencies.Empty();
@@ -2292,7 +2292,7 @@ void FMetaSoundFrontendDocumentBuilder::InitDocument(const FMetasoundFrontendDoc
 	METASOUND_TRACE_CPUPROFILER_EVENT_SCOPE(FMetaSoundFrontendDocumentBuilder::InitDocument);
 
 	FMetasoundFrontendDocument& Document = GetDocumentChecked();
-	Document.RootGraph.InitDefaultPage();
+	Document.RootGraph.InitDefaultGraph();
 
 	// 1. Set default class Metadata.
 	if (InDocumentTemplate)
diff --git a/Engine/Plugins/Runtime/Metasound/Source/MetasoundFrontend/Private/MetasoundFrontendDocumentController.cpp b/Engine/Plugins/Runtime/Metasound/Source/MetasoundFrontend/Private/MetasoundFrontendDocumentController.cpp
--- a/Engine/Plugins/Runtime/Metasound/Source/MetasoundFrontend/Private/MetasoundFrontendDocumentController.cpp
+++ b/Engine/Plugins/Runtime/Metasound/Source/MetasoundFrontend/Private/MetasoundFrontendDocumentController.cpp
@@ -403,7 +403,7 @@ namespace Metasound
 							FMetasoundFrontendGraphClass NewClass;
 							NewClass.ID = FGuid::NewGuid();
 							NewClass.Metadata = InMetadata;
-
+							NewClass.InitDefaultGraph();
 							Document->Subgraphs.Add(NewClass);
 						}
 						break;
diff --git a/Engine/Plugins/Runtime/Metasound/Source/MetasoundFrontend/Public/MetasoundFrontendDocument.h b/Engine/Plugins/Runtime/Metasound/Source/MetasoundFrontend/Public/MetasoundFrontendDocument.h
--- a/Engine/Plugins/Runtime/Metasound/Source/MetasoundFrontend/Public/MetasoundFrontendDocument.h
+++ b/Engine/Plugins/Runtime/Metasound/Source/MetasoundFrontend/Public/MetasoundFrontendDocument.h
@@ -1745,7 +1745,7 @@ public:
 
 	bool ContainsGraphPage(const FGuid& InPageID) const;
 
-	FMetasoundFrontendGraph& InitDefaultPage();
+	FMetasoundFrontendGraph& InitDefaultGraph();
 
 	void IterateGraphPages(TFunctionRef<void(FMetasoundFrontendGraph&)> IterFunc);
 	void IterateGraphPages(TFunctionRef<void(const FMetasoundFrontendGraph&)> IterFunc) const;
