Part 2: AudioUnitTest fix-ups in response to prior default page graph/input literal changes P2
#rnx

diff --git a/Engine/Plugins/Runtime/Metasound/Source/MetasoundEngineTest/Private/EngineTestMetasoundSource.cpp b/Engine/Plugins/Runtime/Metasound/Source/MetasoundEngineTest/Private/EngineTestMetasoundSource.cpp
--- a/Engine/Plugins/Runtime/Metasound/Source/MetasoundEngineTest/Private/EngineTestMetasoundSource.cpp
+++ b/Engine/Plugins/Runtime/Metasound/Source/MetasoundEngineTest/Private/EngineTestMetasoundSource.cpp
@@ -177,7 +177,7 @@ namespace EngineTestMetaSoundSourcePrivate
 
 		Document.RootGraph.Metadata.SetClassName(FMetasoundFrontendClassName { "Namespace", "Unit Test Node", *LexToString(FGuid::NewGuid()) });
 		Document.RootGraph.Metadata.SetType(EMetasoundFrontendClassType::Graph);
-		Document.RootGraph.InitDefaultPage();
+		Document.RootGraph.InitDefaultGraph();
 
 		FDocumentHandle DocumentHandle = IDocumentController::CreateDocumentHandle(Document);
 		FGraphHandle RootGraph = DocumentHandle->GetRootGraph();
diff --git a/Engine/Plugins/Runtime/Metasound/Source/MetasoundEngineTest/Private/NodeTestGraphBuilder.cpp b/Engine/Plugins/Runtime/Metasound/Source/MetasoundEngineTest/Private/NodeTestGraphBuilder.cpp
--- a/Engine/Plugins/Runtime/Metasound/Source/MetasoundEngineTest/Private/NodeTestGraphBuilder.cpp
+++ b/Engine/Plugins/Runtime/Metasound/Source/MetasoundEngineTest/Private/NodeTestGraphBuilder.cpp
@@ -14,7 +14,7 @@ namespace Metasound::Test
 	{
 		Document.RootGraph.Metadata.SetClassName({ "Metasound", "TestNodes", *LexToString(FGuid::NewGuid()) });
 		Document.RootGraph.Metadata.SetType(EMetasoundFrontendClassType::Graph);
-		Document.RootGraph.InitDefaultPage();
+		Document.RootGraph.InitDefaultGraph();
 
 		DocumentHandle = IDocumentController::CreateDocumentHandle(Document);
 		RootGraph = DocumentHandle->GetRootGraph();
