Remove shadowing unused variable causing compiler warnings
#rnx

diff --git a/Engine/Plugins/Runtime/Metasound/Source/MetasoundFrontend/Private/MetasoundFrontendDocumentVersioning.cpp b/Engine/Plugins/Runtime/Metasound/Source/MetasoundFrontend/Private/MetasoundFrontendDocumentVersioning.cpp
--- a/Engine/Plugins/Runtime/Metasound/Source/MetasoundFrontend/Private/MetasoundFrontendDocumentVersioning.cpp
+++ b/Engine/Plugins/Runtime/Metasound/Source/MetasoundFrontend/Private/MetasoundFrontendDocumentVersioning.cpp
@@ -61,7 +61,6 @@ PRAGMA_ENABLE_DEPRECATION_WARNINGS
 
 					virtual bool Transform(FMetasoundFrontendGraphClass& OutClass) const override
 					{
-						bool bUpdated = false;
 PRAGMA_DISABLE_DEPRECATION_WARNINGS
 						TArray<FMetasoundFrontendGraph>& Pages = GetPagesUnsafe(OutClass);
 						if (Pages.IsEmpty())
