Fix copy pasta from input literal serialization migration which can cause wrong default input literal to be observed
#rnx

diff --git a/Engine/Plugins/Runtime/Metasound/Source/MetasoundFrontend/Private/MetasoundFrontendGraph.cpp b/Engine/Plugins/Runtime/Metasound/Source/MetasoundFrontend/Private/MetasoundFrontendGraph.cpp
--- a/Engine/Plugins/Runtime/Metasound/Source/MetasoundFrontend/Private/MetasoundFrontendGraph.cpp
+++ b/Engine/Plugins/Runtime/Metasound/Source/MetasoundFrontend/Private/MetasoundFrontendGraph.cpp
@@ -343,7 +343,7 @@ namespace Metasound
 		{
 			const FMetasoundFrontendClassInput& InputNodeClassInput = InInputNodeClass.Interface.Inputs[0];
 
-			const FMetasoundFrontendLiteral& DefaultLiteral = InOwningGraphClassInput.FindConstDefaultChecked(Frontend::DefaultPageID);
+			const FMetasoundFrontendLiteral& DefaultLiteral = InputNodeClassInput.FindConstDefaultChecked(Frontend::DefaultPageID);
 			if (DefaultLiteral.IsValid())
 			{
 				Literal = &DefaultLiteral;
