MetaSound Wave Player 'simple view' by default

#rnx

diff --git a/Engine/Plugins/Runtime/Metasound/Source/MetasoundEngine/Private/MetasoundWavePlayerNode.cpp b/Engine/Plugins/Runtime/Metasound/Source/MetasoundEngine/Private/MetasoundWavePlayerNode.cpp
--- a/Engine/Plugins/Runtime/Metasound/Source/MetasoundEngine/Private/MetasoundWavePlayerNode.cpp
+++ b/Engine/Plugins/Runtime/Metasound/Source/MetasoundEngine/Private/MetasoundWavePlayerNode.cpp
@@ -1137,11 +1137,34 @@ namespace Metasound
 		{
 			using namespace WavePlayerVertexNames;
 
+			//Adds uncommonly used pins to the Advanced View, to reduce the size of the node.
+			FDataVertexMetadata InputLoopStartMetaData = METASOUND_GET_PARAM_METADATA(InputLoopStart);
+			InputLoopStartMetaData.bIsAdvancedDisplay = true;
+
+			FDataVertexMetadata InputLoopDurationMetaData = METASOUND_GET_PARAM_METADATA(InputLoopDuration);
+			InputLoopDurationMetaData.bIsAdvancedDisplay = true;
+
+			FDataVertexMetadata OutputTriggerOnNearlyDoneMetaData = METASOUND_GET_PARAM_METADATA(OutputTriggerOnNearlyDone);
+			OutputTriggerOnNearlyDoneMetaData.bIsAdvancedDisplay = true;		
+
+			FDataVertexMetadata OutputTriggerOnCuePointMetaData = METASOUND_GET_PARAM_METADATA(OutputTriggerOnCuePoint);
+			OutputTriggerOnCuePointMetaData.bIsAdvancedDisplay = true;
+
+			FDataVertexMetadata OutputCuePointIDMetaData = METASOUND_GET_PARAM_METADATA(OutputCuePointID);
+			OutputCuePointIDMetaData.bIsAdvancedDisplay = true;
+
+			FDataVertexMetadata OutputCuePointLabelMetaData = METASOUND_GET_PARAM_METADATA(OutputCuePointLabel);
+			OutputCuePointLabelMetaData.bIsAdvancedDisplay = true;
+
+			FDataVertexMetadata OutputPlaybackLocationMetaData = METASOUND_GET_PARAM_METADATA(OutputPlaybackLocation);
+			OutputPlaybackLocationMetaData.bIsAdvancedDisplay = true;
+
 			// Workaround to override display name of OutputLoopRatio
 			static const FDataVertexMetadata OutputLoopRatioMetadata
 			{ 
 				METASOUND_GET_PARAM_TT(OutputLoopRatio), // description 
-				METASOUND_LOCTEXT("OutputLoopRatioNotPercentDisplayName", "Loop Ratio") // display name  
+				METASOUND_LOCTEXT("OutputLoopRatioNotPercentDisplayName", "Loop Ratio"), // display name  
+				true // Is Advanced Display
 			};
 
 			FVertexInterface VertexInterface(
@@ -1152,19 +1175,19 @@ namespace Metasound
 					TInputDataVertex<FTime>(METASOUND_GET_PARAM_NAME_AND_METADATA(InputStartTime), 0.0f),
 					TInputDataVertex<float>(METASOUND_GET_PARAM_NAME_AND_METADATA(InputPitchShift), 0.0f),
 					TInputDataVertex<bool>(METASOUND_GET_PARAM_NAME_AND_METADATA(InputLoop), false),
-					TInputDataVertex<FTime>(METASOUND_GET_PARAM_NAME_AND_METADATA(InputLoopStart), 0.0f),
-					TInputDataVertex<FTime>(METASOUND_GET_PARAM_NAME_AND_METADATA(InputLoopDuration), -1.0f)
+					TInputDataVertex<FTime>(METASOUND_GET_PARAM_NAME(InputLoopStart), InputLoopStartMetaData, 0.0f),
+					TInputDataVertex<FTime>(METASOUND_GET_PARAM_NAME(InputLoopDuration), InputLoopDurationMetaData, -1.0f)
 					),
 				FOutputVertexInterface(
 					TOutputDataVertex<FTrigger>(METASOUND_GET_PARAM_NAME_AND_METADATA(OutputTriggerOnPlay)),
 					TOutputDataVertex<FTrigger>(METASOUND_GET_PARAM_NAME_AND_METADATA(OutputTriggerOnDone)),
-					TOutputDataVertex<FTrigger>(METASOUND_GET_PARAM_NAME_AND_METADATA(OutputTriggerOnNearlyDone)),
+					TOutputDataVertex<FTrigger>(METASOUND_GET_PARAM_NAME(OutputTriggerOnNearlyDone), OutputTriggerOnNearlyDoneMetaData),
 					TOutputDataVertex<FTrigger>(METASOUND_GET_PARAM_NAME_AND_METADATA(OutputTriggerOnLooped)),
-					TOutputDataVertex<FTrigger>(METASOUND_GET_PARAM_NAME_AND_METADATA(OutputTriggerOnCuePoint)),
-					TOutputDataVertex<int32>(METASOUND_GET_PARAM_NAME_AND_METADATA(OutputCuePointID)),
-					TOutputDataVertex<FString>(METASOUND_GET_PARAM_NAME_AND_METADATA(OutputCuePointLabel)),
+					TOutputDataVertex<FTrigger>(METASOUND_GET_PARAM_NAME(OutputTriggerOnCuePoint), OutputTriggerOnCuePointMetaData),
+					TOutputDataVertex<int32>(METASOUND_GET_PARAM_NAME(OutputCuePointID), OutputCuePointIDMetaData),
+					TOutputDataVertex<FString>(METASOUND_GET_PARAM_NAME(OutputCuePointLabel), OutputCuePointLabelMetaData),
 					TOutputDataVertex<float>(METASOUND_GET_PARAM_NAME(OutputLoopRatio), OutputLoopRatioMetadata),
-					TOutputDataVertex<float>(METASOUND_GET_PARAM_NAME_AND_METADATA(OutputPlaybackLocation)),
+					TOutputDataVertex<float>(METASOUND_GET_PARAM_NAME(OutputPlaybackLocation), OutputPlaybackLocationMetaData),
 					TOutputDataVertex<FTime>(METASOUND_GET_PARAM_NAME_AND_METADATA(OutputPlaybackTime))
 				)
 			);
