Editor Input clean-up in anticipation of refactor to support input default page literals
#rnx

diff --git a/Engine/Plugins/Runtime/Metasound/Source/MetasoundEditor/Private/MetasoundEditorGraph.cpp b/Engine/Plugins/Runtime/Metasound/Source/MetasoundEditor/Private/MetasoundEditorGraph.cpp
--- a/Engine/Plugins/Runtime/Metasound/Source/MetasoundEditor/Private/MetasoundEditorGraph.cpp
+++ b/Engine/Plugins/Runtime/Metasound/Source/MetasoundEditor/Private/MetasoundEditorGraph.cpp
@@ -69,6 +69,24 @@ namespace Metasound::Editor
 			}
 		}
 
+		void SetAndUpdatePreviewLiteral(FMetaSoundFrontendDocumentBuilder& Builder, UMetasoundEditorGraph* MetaSoundGraph, FName MemberName, UMetasoundEditorGraphMemberDefaultLiteral& Literal)
+		{
+			Builder.SetGraphInputDefault(MemberName, Literal.GetDefault());
+			if (GEditor)
+			{
+				if (MetaSoundGraph && MetaSoundGraph->IsPreviewing())
+				{
+					UAudioComponent* PreviewComponent = GEditor->GetPreviewAudioComponent();
+					check(PreviewComponent);
+
+					if (TScriptInterface<IAudioParameterControllerInterface> ParamInterface = PreviewComponent)
+					{
+						Literal.UpdatePreviewInstance(MemberName, ParamInterface);
+					}
+				}
+			}
+		}
+
 		void SetLiteralOrClearIfMatchesDefault(Frontend::FInputHandle& InInputHandle, const FMetasoundFrontendLiteral& InDefaultLiteral)
 		{
 			using namespace Frontend;
@@ -147,16 +165,11 @@ void UMetasoundEditorGraphMember::InitializeLiteral()
 	{
 		FMetaSoundFrontendDocumentBuilder& Builder = GetFrontendBuilderChecked();
 		const bool bIsNew = UMetaSoundEditorSubsystem::GetChecked().BindMemberMetadata(Builder, *this, LiteralClass);
+		checkf(Literal, TEXT("Bind is required to initialize literal field on this member"));
+
 		if (bIsNew)
 		{
-			// Set default widget type from editor settings for float inputs 
-			if (IsA<UMetasoundEditorGraphInput>() && LiteralClass == UMetasoundEditorGraphMemberDefaultFloat::StaticClass())
-			{
-				if (const UMetasoundEditorSettings* EditorSettings = ::GetDefault<UMetasoundEditorSettings>())
-				{
-					Cast<UMetasoundEditorGraphMemberDefaultFloat>(Literal)->WidgetType = EditorSettings->DefaultInputWidgetType;
-				}
-			}
+			Literal->Initialize();
 		}
 	}
 }
@@ -728,32 +741,20 @@ void UMetasoundEditorGraphInput::ResetToClassDefault()
 	using namespace Metasound::Editor;
 	using namespace Metasound::Frontend;
 
-	UMetasoundEditorGraph* MetaSoundGraph = GetOwningGraph();
-	FGraphHandle GraphHandle = MetaSoundGraph->GetGraphHandle();
-	FConstNodeHandle NodeHandle = GraphHandle->GetNodeWithID(NodeID);
-
-	FMetasoundFrontendLiteral DefaultLiteral;
-	DefaultLiteral.SetFromLiteral(IDataTypeRegistry::Get().CreateDefaultLiteral(GetDataType()));
-
-	Literal->Modify();
-	Literal->SetFromLiteral(DefaultLiteral);
+	if (ensure(Literal))
+	{
+		FMetaSoundFrontendDocumentBuilder& Builder = GetFrontendBuilderChecked();
 
-	const Metasound::FVertexName& NodeName = NodeHandle->GetNodeName();
-	const FGuid VertexID = GraphHandle->GetVertexIDForInputVertex(NodeName);
-	GraphHandle->SetDefaultInput(VertexID, DefaultLiteral);
+		Builder.CastDocumentObjectChecked<UObject>().Modify();
+		Literal->Modify();
 
-	if (GEditor && MetaSoundGraph->IsPreviewing())
-	{
-		UAudioComponent* PreviewComponent = GEditor->GetPreviewAudioComponent();
-		check(PreviewComponent);
+		FMetasoundFrontendLiteral DefaultLiteral;
+		DefaultLiteral.SetFromLiteral(IDataTypeRegistry::Get().CreateDefaultLiteral(GetDataType()));
+		Literal->SetFromLiteral(DefaultLiteral);
 
-		if (TScriptInterface<IAudioParameterControllerInterface> ParamInterface = PreviewComponent)
-		{
-			Literal->UpdatePreviewInstance(NodeName, ParamInterface);
-		}
+		const FName MemberName = GetMemberName();
+		GraphPrivate::SetAndUpdatePreviewLiteral(Builder, GetOwningGraph(), MemberName, *Literal);
 	}
-
-	FGraphBuilder::GetOutermostMetaSoundChecked(*MetaSoundGraph).GetModifyContext().AddMemberIDsModified({ GetMemberID() });
 }
 
 void UMetasoundEditorGraphInput::SetDataType(FName InNewType, bool bPostTransaction)
@@ -827,61 +828,28 @@ void UMetasoundEditorGraphInput::SetMemberName(const FName& InNewName, bool bPos
 
 void UMetasoundEditorGraphInput::UpdateFrontendDefaultLiteral(bool bPostTransaction)
 {
-	using namespace Metasound;
-	using namespace Metasound::Frontend;
-
-	UMetasoundEditorGraph* MetaSoundGraph = GetOwningGraph();
-	UObject* Metasound = MetaSoundGraph->GetMetasound();
-	if (!ensure(Metasound))
-	{
-		return;
-	}
+	using namespace Metasound::Editor;
 
-	if (!ensure(Literal))
+	if (Literal)
 	{
-		return;
-	}
-
-	const FScopedTransaction Transaction(LOCTEXT("Set Input Default", "Set MetaSound Input Default"), bPostTransaction);
-	Metasound->Modify();
-
-	FMetasoundAssetBase* MetasoundAsset = IMetasoundUObjectRegistry::Get().GetObjectAsAssetBase(Metasound);
-	check(MetasoundAsset);
-
-	FGraphHandle GraphHandle = MetasoundAsset->GetRootGraphHandle();
-	FConstNodeHandle NodeHandle = GraphHandle->GetNodeWithID(NodeID);
-
-	const Metasound::FVertexName& NodeName = NodeHandle->GetNodeName();
-	const FGuid VertexID = GraphHandle->GetVertexIDForInputVertex(NodeName);
-	GraphHandle->SetDefaultInput(VertexID, Literal->GetDefault());
+		FMetaSoundFrontendDocumentBuilder& Builder = GetFrontendBuilderChecked();
 
-	if (GEditor && MetaSoundGraph->IsPreviewing())
-	{
-		UAudioComponent* PreviewComponent = GEditor->GetPreviewAudioComponent();
-		check(PreviewComponent);
+		const FScopedTransaction Transaction(LOCTEXT("Set Input Default", "Set MetaSound Input Default"), bPostTransaction);
+		Builder.CastDocumentObjectChecked<UObject>().Modify();
+		Literal->Modify();
 
-		if (TScriptInterface<IAudioParameterControllerInterface> ParamInterface = PreviewComponent)
-		{
-			Metasound::Frontend::FConstNodeHandle ConstNodeHandle = GetConstNodeHandle();
-			Metasound::FVertexName VertexKey = NodeHandle->GetNodeName();
-			Literal->UpdatePreviewInstance(VertexKey, ParamInterface);
-		}
+		const FName MemberName = GetMemberName();
+		GraphPrivate::SetAndUpdatePreviewLiteral(Builder, GetOwningGraph(), MemberName, *Literal);
 	}
 }
 
 EMetasoundFrontendVertexAccessType UMetasoundEditorGraphInput::GetVertexAccessType() const
 {
-	using namespace Metasound::Frontend;
-
 	const FName MemberName = GetMemberName();
-	const UMetasoundEditorGraph* MetaSoundGraph = GetOwningGraph();
-	FConstGraphHandle GraphHandle = MetaSoundGraph->GetGraphHandle();
-
-	FConstClassInputAccessPtr ClassInputPtr = GraphHandle->FindClassInputWithName(MemberName);
-
-	if (const FMetasoundFrontendClassInput* InputPtr = ClassInputPtr.Get())
+	const FMetaSoundFrontendDocumentBuilder& Builder = GetFrontendBuilderChecked();
+	if (const FMetasoundFrontendClassInput* Input = Builder.FindGraphInput(MemberName))
 	{
-		return InputPtr->AccessType;
+		return Input->AccessType;
 	}
 
 	return EMetasoundFrontendVertexAccessType::Reference;
@@ -1114,17 +1082,11 @@ void UMetasoundEditorGraphOutput::UpdateFrontendDefaultLiteral(bool bPostTransac
 
 EMetasoundFrontendVertexAccessType UMetasoundEditorGraphOutput::GetVertexAccessType() const
 {
-	using namespace Metasound::Frontend;
-
 	const FName MemberName = GetMemberName();
-	const UMetasoundEditorGraph* MetaSoundGraph = GetOwningGraph();
-	FConstGraphHandle GraphHandle = MetaSoundGraph->GetGraphHandle();
-
-	FConstClassOutputAccessPtr ClassOutputPtr = GraphHandle->FindClassOutputWithName(MemberName);
-
-	if (const FMetasoundFrontendClassOutput* OutputPtr = ClassOutputPtr.Get())
+	const FMetaSoundFrontendDocumentBuilder& Builder = GetFrontendBuilderChecked();
+	if (const FMetasoundFrontendClassOutput* Output = Builder.FindGraphOutput(MemberName))
 	{
-		return OutputPtr->AccessType;
+		return Output->AccessType;
 	}
 
 	return EMetasoundFrontendVertexAccessType::Reference;
diff --git a/Engine/Plugins/Runtime/Metasound/Source/MetasoundEditor/Private/MetasoundEditorGraphMemberDefaults.cpp b/Engine/Plugins/Runtime/Metasound/Source/MetasoundEditor/Private/MetasoundEditorGraphMemberDefaults.cpp
--- a/Engine/Plugins/Runtime/Metasound/Source/MetasoundEditor/Private/MetasoundEditorGraphMemberDefaults.cpp
+++ b/Engine/Plugins/Runtime/Metasound/Source/MetasoundEditor/Private/MetasoundEditorGraphMemberDefaults.cpp
@@ -11,6 +11,7 @@
 #include "MetasoundEditorGraph.h"
 #include "MetasoundEditorGraphBuilder.h"
 #include "MetasoundEditorGraphNode.h"
+#include "MetasoundEditorSettings.h"
 #include "MetasoundFrontend.h"
 #include "MetasoundFrontendController.h"
 #include "MetasoundFrontendDocument.h"
@@ -173,6 +174,14 @@ EMetasoundFrontendLiteralType UMetasoundEditorGraphMemberDefaultFloat::GetLitera
 	return EMetasoundFrontendLiteralType::Float;
 }
 
+void UMetasoundEditorGraphMemberDefaultFloat::Initialize()
+{
+	if (const UMetasoundEditorSettings* EditorSettings = ::GetDefault<UMetasoundEditorSettings>())
+	{
+		WidgetType = EditorSettings->DefaultInputWidgetType;
+	}
+}
+
 void UMetasoundEditorGraphMemberDefaultFloat::ForceRefresh()
 {
 	OnRangeChanged.Broadcast(Range);
diff --git a/Engine/Plugins/Runtime/Metasound/Source/MetasoundEditor/Public/MetasoundEditorGraph.h b/Engine/Plugins/Runtime/Metasound/Source/MetasoundEditor/Public/MetasoundEditorGraph.h
--- a/Engine/Plugins/Runtime/Metasound/Source/MetasoundEditor/Public/MetasoundEditorGraph.h
+++ b/Engine/Plugins/Runtime/Metasound/Source/MetasoundEditor/Public/MetasoundEditorGraph.h
@@ -56,6 +56,10 @@ public:
 	{
 	}
 
+	// Called when literal is initialized for the first time to allow for setting
+	// initial editor-only fields based on context within editor/document model.
+	virtual void Initialize() { };
+
 #if WITH_EDITOR
 	virtual void PostEditChangeProperty(FPropertyChangedEvent& InPropertyChangedEvent) override;
 	virtual void PostEditChangeChainProperty(FPropertyChangedChainEvent& InPropertyChangedEvent) override;
diff --git a/Engine/Plugins/Runtime/Metasound/Source/MetasoundEditor/Public/MetasoundEditorGraphMemberDefaults.h b/Engine/Plugins/Runtime/Metasound/Source/MetasoundEditor/Public/MetasoundEditorGraphMemberDefaults.h
--- a/Engine/Plugins/Runtime/Metasound/Source/MetasoundEditor/Public/MetasoundEditorGraphMemberDefaults.h
+++ b/Engine/Plugins/Runtime/Metasound/Source/MetasoundEditor/Public/MetasoundEditorGraphMemberDefaults.h
@@ -197,6 +197,7 @@ public:
 	virtual void ForceRefresh() override;
 	virtual FMetasoundFrontendLiteral GetDefault() const override;
 	virtual EMetasoundFrontendLiteralType GetLiteralType() const override;
+	virtual void Initialize() override;
 	virtual void SetFromLiteral(const FMetasoundFrontendLiteral& InLiteral) override;
 	virtual void UpdatePreviewInstance(const Metasound::FVertexName& InParameterName, TScriptInterface<IAudioParameterControllerInterface>& InParameterInterface) const override;
 	virtual void PostEditChangeChainProperty(FPropertyChangedChainEvent& InPropertyChangedEvent) override;
