Fix ensure when promoting MetaSound graph output to input

diff --git a/Engine/Plugins/Runtime/Metasound/Source/MetasoundEditor/Private/MetasoundEditorGraphBuilder.cpp b/Engine/Plugins/Runtime/Metasound/Source/MetasoundEditor/Private/MetasoundEditorGraphBuilder.cpp
--- a/Engine/Plugins/Runtime/Metasound/Source/MetasoundEditor/Private/MetasoundEditorGraphBuilder.cpp
+++ b/Engine/Plugins/Runtime/Metasound/Source/MetasoundEditor/Private/MetasoundEditorGraphBuilder.cpp
@@ -897,13 +897,13 @@ namespace Metasound
 
 				case EMetasoundFrontendClassType::Input:
 				{
-					Vertex = &Node->Interface.Inputs.Last();
+					Vertex = &Node->Interface.Outputs.Last();
 					break;
 				}
 
 				case EMetasoundFrontendClassType::Output:
 				{
-					Vertex = &Node->Interface.Outputs.Last();
+					Vertex = &Node->Interface.Inputs.Last();
 					break;
 				}
 
