Fix MetaSound comment node undo and comment node positions not being pasted relative to other nodes

diff --git a/Engine/Plugins/Runtime/Metasound/Source/MetasoundEditor/Private/MetasoundEditorDocumentClipboardUtils.cpp b/Engine/Plugins/Runtime/Metasound/Source/MetasoundEditor/Private/MetasoundEditorDocumentClipboardUtils.cpp
--- a/Engine/Plugins/Runtime/Metasound/Source/MetasoundEditor/Private/MetasoundEditorDocumentClipboardUtils.cpp
+++ b/Engine/Plugins/Runtime/Metasound/Source/MetasoundEditor/Private/MetasoundEditorDocumentClipboardUtils.cpp
@@ -303,7 +303,7 @@ namespace Metasound::Editor
 		}
 	}
 
-	void FDocumentClipboardUtils::ProcessPastedCommentNodes(FMetasoundAssetBase& OutAsset, const FVector2D& InLocation, const TArrayView<UMetasoundEditorGraphCommentNode*> CommentNodes)
+	void FDocumentClipboardUtils::ProcessPastedCommentNodes(FMetasoundAssetBase& OutAsset, const TArrayView<UMetasoundEditorGraphCommentNode*> CommentNodes)
 	{
 		using namespace Engine;
 		using namespace Frontend;
@@ -311,38 +311,19 @@ namespace Metasound::Editor
 		UMetasoundEditorGraph& Graph = *CastChecked<UMetasoundEditorGraph>(&OutAsset.GetGraphChecked());
 		UMetaSoundBuilderBase& Builder = FDocumentBuilderRegistry::GetChecked().FindOrBeginBuilding(*OutAsset.GetOwningAsset());
 
-		// Calculate average node position 
-		FVector2D AvgNodePosition = FVector2D::ZeroVector;
-		for (UEdGraphNode* Node : CommentNodes)
-		{
-			AvgNodePosition.X += Node->NodePosX;
-			AvgNodePosition.Y += Node->NodePosY;
-		}
-
-		if (!CommentNodes.IsEmpty())
-		{
-			float InvNumNodes = 1.0f / CommentNodes.Num();
-			AvgNodePosition.X *= InvNumNodes;
-			AvgNodePosition.Y *= InvNumNodes;
-		}
-
 		for (UMetasoundEditorGraphCommentNode* CommentNode : CommentNodes)
 		{
 			// Regenerate id
 			CommentNode->CreateNewGuid();
 			CommentNode->SetCommentID(CommentNode->NodeGuid);
 
-			// Add offset based on paste location to avoid comments automatically being captured by the original node if using group movement 
-			CommentNode->NodePosX = (CommentNode->NodePosX - AvgNodePosition.X) + InLocation.X;
-			CommentNode->NodePosY = (CommentNode->NodePosY - AvgNodePosition.Y) + InLocation.Y;
-
 			// Update frontend node
 			FMetaSoundFrontendGraphComment& NewComment = Builder.FindOrAddGraphComment(CommentNode->GetCommentID());
 			UMetasoundEditorGraphCommentNode::ConvertToFrontendComment(*CommentNode, NewComment);
 		}
 	}
 
-	void FDocumentClipboardUtils::ProcessPastedNodePositions(FMetasoundAssetBase& OutAsset, const FVector2D& InLocation, TArray<UMetasoundEditorGraphNode*>& OutPastedNodes)
+	void FDocumentClipboardUtils::ProcessPastedNodePositions(FMetasoundAssetBase& OutAsset, const FVector2D& InLocation, TArray<UMetasoundEditorGraphNode*>& OutPastedNodes, const TArrayView<UMetasoundEditorGraphCommentNode*> CommentNodes)
 	{
 		using namespace Frontend;
 
@@ -353,14 +334,20 @@ namespace Metasound::Editor
 			AvgNodePosition.X += Node->NodePosX;
 			AvgNodePosition.Y += Node->NodePosY;
 		}
+		for (UEdGraphNode* Node : CommentNodes)
+		{
+			AvgNodePosition.X += Node->NodePosX;
+			AvgNodePosition.Y += Node->NodePosY;
+		}
 
 		if (!OutPastedNodes.IsEmpty())
 		{
-			float InvNumNodes = 1.0f / OutPastedNodes.Num();
+			float InvNumNodes = 1.0f / (OutPastedNodes.Num() + CommentNodes.Num());
 			AvgNodePosition.X *= InvNumNodes;
 			AvgNodePosition.Y *= InvNumNodes;
 		}
 
+		// Set new node positions
 		for (UEdGraphNode* GraphNode : OutPastedNodes)
 		{
 			GraphNode->NodePosX = (GraphNode->NodePosX - AvgNodePosition.X) + InLocation.X;
@@ -379,6 +366,14 @@ namespace Metasound::Editor
 				}
 			}
 		}
+
+		// Set new comment node positions 
+		for (UMetasoundEditorGraphCommentNode* CommentNode : CommentNodes)
+		{
+			CommentNode->NodePosX = (CommentNode->NodePosX - AvgNodePosition.X) + InLocation.X;
+			CommentNode->NodePosY = (CommentNode->NodePosY - AvgNodePosition.Y) + InLocation.Y;
+			CommentNode->UpdateFrontendNodeLocation();
+		}
 	}
 
 	void FDocumentClipboardUtils::ProcessPastedNodeConnections(FMetasoundAssetBase& OutAsset, TArray<UMetasoundEditorGraphNode*>& OutPastedNodes)
@@ -456,13 +451,13 @@ namespace Metasound::Editor
 		OutMetaSound.Modify();
 		Asset->GetGraphChecked().Modify();
 
-		ProcessPastedCommentNodes(*Asset, InLocation, PastedCommentNodes);
+		ProcessPastedCommentNodes(*Asset, PastedCommentNodes);
 
 		ProcessPastedInputNodes(*Asset, PastedGraphNodes);
 		ProcessPastedOutputNodes(*Asset, PastedGraphNodes);
 		ProcessPastedVariableNodes(*Asset, PastedGraphNodes, OutNotifications);
 		ProcessPastedExternalNodes(*Asset, PastedGraphNodes, OutNotifications);
-		ProcessPastedNodePositions(*Asset, InLocation, PastedGraphNodes);
+		ProcessPastedNodePositions(*Asset, InLocation, PastedGraphNodes, PastedCommentNodes);
 		ProcessPastedNodeConnections(*Asset, PastedGraphNodes);
 
 		PastedNodes.Append(MoveTemp(PastedGraphNodes));
diff --git a/Engine/Plugins/Runtime/Metasound/Source/MetasoundEditor/Private/MetasoundEditorDocumentClipboardUtils.h b/Engine/Plugins/Runtime/Metasound/Source/MetasoundEditor/Private/MetasoundEditorDocumentClipboardUtils.h
--- a/Engine/Plugins/Runtime/Metasound/Source/MetasoundEditor/Private/MetasoundEditorDocumentClipboardUtils.h
+++ b/Engine/Plugins/Runtime/Metasound/Source/MetasoundEditor/Private/MetasoundEditorDocumentClipboardUtils.h
@@ -26,13 +26,13 @@ namespace Metasound::Editor
 		static TArray<UEdGraphNode*> PasteClipboardString(const FText& InTransactionText, const FString& InClipboardString, const FVector2D& InLocation, UObject& OutMetaSound, FDocumentPasteNotifications& OutNotifications);
 
 	private:
-		static void ProcessPastedCommentNodes(FMetasoundAssetBase& OutAsset, const FVector2D& InLocation, const TArrayView<UMetasoundEditorGraphCommentNode*> CommentNodes);
+		static void ProcessPastedCommentNodes(FMetasoundAssetBase& OutAsset, const TArrayView<UMetasoundEditorGraphCommentNode*> CommentNodes);
 
 		static void ProcessPastedExternalNodes(FMetasoundAssetBase& OutAsset, TArray<UMetasoundEditorGraphNode*>& OutPastedNodes, FDocumentPasteNotifications& OutNotifications);
 
 		static void ProcessPastedInputNodes(FMetasoundAssetBase& OutAsset, TArray<UMetasoundEditorGraphNode*>& OutPastedNodes);
 
-		static void ProcessPastedNodePositions(FMetasoundAssetBase& OutAsset, const FVector2D& InLocation, TArray<UMetasoundEditorGraphNode*>& OutPastedNodes);
+		static void ProcessPastedNodePositions(FMetasoundAssetBase& OutAsset, const FVector2D& InLocation, TArray<UMetasoundEditorGraphNode*>& OutPastedNodes, const TArrayView<UMetasoundEditorGraphCommentNode*> CommentNodes);
 
 		static void ProcessPastedNodeConnections(FMetasoundAssetBase& OutAsset, TArray<UMetasoundEditorGraphNode*>& OutPastedNodes);
 
diff --git a/Engine/Plugins/Runtime/Metasound/Source/MetasoundEditor/Private/MetasoundEditorGraphCommentNode.cpp b/Engine/Plugins/Runtime/Metasound/Source/MetasoundEditor/Private/MetasoundEditorGraphCommentNode.cpp
--- a/Engine/Plugins/Runtime/Metasound/Source/MetasoundEditor/Private/MetasoundEditorGraphCommentNode.cpp
+++ b/Engine/Plugins/Runtime/Metasound/Source/MetasoundEditor/Private/MetasoundEditorGraphCommentNode.cpp
@@ -67,7 +67,10 @@ void UMetasoundEditorGraphCommentNode::PostEditChangeProperty(FPropertyChangedEv
 
 	FMetaSoundFrontendGraphComment Comment;
 	ConvertToFrontendComment(*this, Comment);
-	GetBuilderChecked().FindOrAddGraphComment(CommentID) = MoveTemp(Comment);
+	if (FMetaSoundFrontendGraphComment* FrontendComment = GetBuilderChecked().FindGraphComment(CommentID))
+	{
+		*FrontendComment = MoveTemp(Comment);
+	}
 }
 
 void UMetasoundEditorGraphCommentNode::ResizeNode(const FVector2D& NewSize)
