Metasound graph node visualization registry.

diff --git a/Engine/Plugins/Runtime/Metasound/Source/MetasoundEditor/Private/MetasoundEditorGraphBuilder.cpp b/Engine/Plugins/Runtime/Metasound/Source/MetasoundEditor/Private/MetasoundEditorGraphBuilder.cpp
--- a/Engine/Plugins/Runtime/Metasound/Source/MetasoundEditor/Private/MetasoundEditorGraphBuilder.cpp
+++ b/Engine/Plugins/Runtime/Metasound/Source/MetasoundEditor/Private/MetasoundEditorGraphBuilder.cpp
@@ -1196,7 +1196,7 @@ namespace Metasound
 			}
 		}
 
-		bool FGraphBuilder::GetPinLiteral(UEdGraphPin& InInputPin, FMetasoundFrontendLiteral& OutDefaultLiteral)
+		bool FGraphBuilder::GetPinLiteral(const UEdGraphPin& InInputPin, FMetasoundFrontendLiteral& OutDefaultLiteral)
 		{
 			using namespace Frontend;
 
diff --git a/Engine/Plugins/Runtime/Metasound/Source/MetasoundEditor/Private/MetasoundEditorGraphBuilder.h b/Engine/Plugins/Runtime/Metasound/Source/MetasoundEditor/Private/MetasoundEditorGraphBuilder.h
--- a/Engine/Plugins/Runtime/Metasound/Source/MetasoundEditor/Private/MetasoundEditorGraphBuilder.h
+++ b/Engine/Plugins/Runtime/Metasound/Source/MetasoundEditor/Private/MetasoundEditorGraphBuilder.h
@@ -229,7 +229,7 @@ namespace Metasound
 			static void FindReroutedInputPins(UEdGraphPin* InPinToCheck, TArray<UEdGraphPin*>& InOutInputPins);
 
 			// Returns the default literal stored on the respective Frontend Node's Input.
-			static bool GetPinLiteral(UEdGraphPin& InInputPin, FMetasoundFrontendLiteral& OutLiteralDefault);
+			static bool GetPinLiteral(const UEdGraphPin& InInputPin, FMetasoundFrontendLiteral& OutLiteralDefault);
 
 			// Retrieves the proper pin color for the given PinType
 			static FLinearColor GetPinCategoryColor(const FEdGraphPinType& PinType);
diff --git a/Engine/Plugins/Runtime/Metasound/Source/MetasoundEditor/Private/MetasoundEditorGraphNode.cpp b/Engine/Plugins/Runtime/Metasound/Source/MetasoundEditor/Private/MetasoundEditorGraphNode.cpp
--- a/Engine/Plugins/Runtime/Metasound/Source/MetasoundEditor/Private/MetasoundEditorGraphNode.cpp
+++ b/Engine/Plugins/Runtime/Metasound/Source/MetasoundEditor/Private/MetasoundEditorGraphNode.cpp
@@ -16,6 +16,7 @@
 #include "MetasoundEditorGraph.h"
 #include "MetasoundEditorGraphBuilder.h"
 #include "MetasoundEditorGraphMemberDefaults.h"
+#include "MetasoundEditorGraphNodeVisualization.h"
 #include "MetasoundEditorGraphSchema.h"
 #include "MetasoundEditorGraphValidation.h"
 #include "MetasoundEditorModule.h"
@@ -765,6 +766,21 @@ FString UMetasoundEditorGraphNode::GetDocumentationExcerptName() const
 	return FString::Printf(TEXT("%s%s"), UMetaSoundPatch::StaticClass()->GetPrefixCPP(), *UMetaSoundPatch::StaticClass()->GetName());
 }
 
+bool UMetasoundEditorGraphNode::TryGetPinVisualizationValue(FName InPinName, bool& OutValue) const
+{
+	return Metasound::Editor::FGraphNodeVisualizationUtils::TryGetPinValue(*this, InPinName, OutValue);
+}
+
+bool UMetasoundEditorGraphNode::TryGetPinVisualizationValue(FName InPinName, int32& OutValue) const
+{
+	return Metasound::Editor::FGraphNodeVisualizationUtils::TryGetPinValue(*this, InPinName, OutValue);
+}
+
+bool UMetasoundEditorGraphNode::TryGetPinVisualizationValue(FName InPinName, float& OutValue) const
+{
+	return Metasound::Editor::FGraphNodeVisualizationUtils::TryGetPinValue(*this, InPinName, OutValue);
+}
+
 bool UMetasoundEditorGraphMemberNode::ClampFloatLiteral(const UMetasoundEditorGraphMemberDefaultFloat* DefaultFloatLiteral, FMetasoundFrontendLiteral& LiteralValue)
 {
 	bool bClampedFloatLiteral = false;
diff --git a/Engine/Plugins/Runtime/Metasound/Source/MetasoundEditor/Private/MetasoundEditorGraphNodeFactory.h b/Engine/Plugins/Runtime/Metasound/Source/MetasoundEditor/Private/MetasoundEditorGraphNodeFactory.h
--- a/Engine/Plugins/Runtime/Metasound/Source/MetasoundEditor/Private/MetasoundEditorGraphNodeFactory.h
+++ b/Engine/Plugins/Runtime/Metasound/Source/MetasoundEditor/Private/MetasoundEditorGraphNodeFactory.h
@@ -10,7 +10,6 @@
 #include "NodeTemplates/MetasoundFrontendNodeTemplateAudioAnalyzer.h"
 #include "NodeTemplates/MetasoundFrontendNodeTemplateReroute.h"
 #include "SGraphNode.h"
-#include "SMetasoundFilterGraphNodes.h"
 #include "SMetasoundGraphNode.h"
 #include "SMetasoundGraphNodeComment.h"
 #include "SMetasoundSpectrumAnalyzerGraphNode.h"
@@ -39,26 +38,6 @@ class FMetasoundGraphNodeFactory : public FGraphPanelNodeFactory
 				{
 					return SNew(SMetaSoundSpectrumAnalyzerGraphNode, InNode);
 				}
-				else if (ClassName == FMetasoundFrontendClassName{ Metasound::StandardNodes::Namespace, TEXT("Biquad Filter"), Metasound::StandardNodes::AudioVariant })
-				{
-					return SNew(SMetaSoundBiquadFilterGraphNode, InNode);
-				}
-				else if (ClassName == FMetasoundFrontendClassName{ Metasound::StandardNodes::Namespace, TEXT("Ladder Filter"), Metasound::StandardNodes::AudioVariant })
-				{
-					return SNew(SMetaSoundLadderFilterGraphNode, InNode);
-				}
-				else if (ClassName == FMetasoundFrontendClassName{ Metasound::StandardNodes::Namespace, TEXT("One-Pole High Pass Filter"), Metasound::StandardNodes::AudioVariant })
-				{
-					return SNew(SMetaSoundOnePoleHighPassFilterGraphNode, InNode);
-				}
-				else if (ClassName == FMetasoundFrontendClassName{ Metasound::StandardNodes::Namespace, TEXT("One-Pole Low Pass Filter"), Metasound::StandardNodes::AudioVariant })
-				{
-					return SNew(SMetaSoundOnePoleLowPassFilterGraphNode, InNode);
-				}
-				else if (ClassName == FMetasoundFrontendClassName{ Metasound::StandardNodes::Namespace, TEXT("State Variable Filter"), Metasound::StandardNodes::AudioVariant })
-				{
-					return SNew(SMetaSoundStateVariableFilterGraphNode, InNode);
-				}
 			}
 			return SNew(SMetaSoundGraphNode, InNode);
 		}
diff --git a/Engine/Plugins/Runtime/Metasound/Source/MetasoundEditor/Private/MetasoundEditorGraphNodeVisualization.cpp b/Engine/Plugins/Runtime/Metasound/Source/MetasoundEditor/Private/MetasoundEditorGraphNodeVisualization.cpp
new file mode 100644
--- /dev/null
+++ b/Engine/Plugins/Runtime/Metasound/Source/MetasoundEditor/Private/MetasoundEditorGraphNodeVisualization.cpp
@@ -0,0 +1,33 @@
+// Copyright Epic Games, Inc. All Rights Reserved.
+
+#include "MetasoundEditorGraphNodeVisualization.h"
+
+#include "Misc/LazySingleton.h"
+
+namespace Metasound::Editor
+{
+	FGraphNodeVisualizationRegistry& FGraphNodeVisualizationRegistry::Get()
+	{
+		return TLazySingleton<FGraphNodeVisualizationRegistry>::Get();
+	}
+
+	void FGraphNodeVisualizationRegistry::TearDown()
+	{
+		TLazySingleton<FGraphNodeVisualizationRegistry>::TearDown();
+	}
+
+	void FGraphNodeVisualizationRegistry::RegisterVisualization(FName Key, FOnCreateGraphNodeVisualizationWidget OnCreateGraphNodeVisualizationWidget)
+	{
+		RegisteredVisualizationDelegates.Add(Key, OnCreateGraphNodeVisualizationWidget);
+	}
+
+	TSharedPtr<SWidget> FGraphNodeVisualizationRegistry::CreateVisualizationWidget(FName Key, const FCreateGraphNodeVisualizationWidgetParams& Params)
+	{
+		if (const FOnCreateGraphNodeVisualizationWidget* VisualizationDelegate = RegisteredVisualizationDelegates.Find(Key))
+		{
+			return VisualizationDelegate->Execute(Params);
+		}
+
+		return nullptr;
+	}
+}
diff --git a/Engine/Plugins/Runtime/Metasound/Source/MetasoundEditor/Private/MetasoundEditorGraphNodeVisualization.h b/Engine/Plugins/Runtime/Metasound/Source/MetasoundEditor/Private/MetasoundEditorGraphNodeVisualization.h
new file mode 100644
--- /dev/null
+++ b/Engine/Plugins/Runtime/Metasound/Source/MetasoundEditor/Private/MetasoundEditorGraphNodeVisualization.h
@@ -0,0 +1,108 @@
+// Copyright Epic Games, Inc. All Rights Reserved.
+#pragma once
+
+#include "Containers/Map.h"
+#include "MetasoundEditorGraphBuilder.h"
+#include "MetasoundEditorGraphNode.h"
+#include "MetasoundEditorModule.h"
+
+namespace Metasound::Editor
+{
+	class FGraphNodeVisualizationRegistry
+	{
+	public:
+		static FGraphNodeVisualizationRegistry& Get();
+		static void TearDown();
+
+		// For the given key, register a delegate that can be used for creating node visualization widgets.
+		void RegisterVisualization(FName Key, FOnCreateGraphNodeVisualizationWidget OnCreateGraphNodeVisualizationWidget);
+
+		// Creates a visualization widget if a delegate has been registered with the given key.
+		TSharedPtr<SWidget> CreateVisualizationWidget(FName Key, const FCreateGraphNodeVisualizationWidgetParams& Params);
+
+	private:
+		TMap<FName, FOnCreateGraphNodeVisualizationWidget> RegisteredVisualizationDelegates;
+	};
+
+	class FGraphNodeVisualizationUtils
+	{
+	public:
+
+		template<class T>
+		static bool TryGetPinValue(const UMetasoundEditorGraphNode& MetaSoundNode, FName PinName, T& OutValue)
+		{
+			UEdGraphPin* const* FoundNamedInputPin = MetaSoundNode.Pins.FindByPredicate([&PinName](const UEdGraphPin* InPin)
+				{
+					return (InPin->Direction == EGPD_Input && InPin->PinType.PinCategory != FGraphBuilder::PinCategoryAudio && InPin->PinName == PinName);
+				});
+
+			if (!FoundNamedInputPin)
+			{
+				return false;
+			}
+
+			const UEdGraphPin* Pin = *FoundNamedInputPin;
+			if (Pin->LinkedTo.IsEmpty())
+			{
+				FMetasoundFrontendLiteral DefaultLiteral;
+				if (FGraphBuilder::GetPinLiteral(*Pin, DefaultLiteral))
+				{
+					if (DefaultLiteral.TryGet(OutValue))
+					{
+						return true;
+					}
+				}
+			}
+			else
+			{
+				// Find connected output for the input (only ever one):
+				const UEdGraphPin* SourcePin = Pin->LinkedTo.Last();
+				ensure(SourcePin->Direction == EGPD_Output);
+
+				const UEdGraphPin* ReroutedOutputPin = FGraphBuilder::FindReroutedOutputPin(SourcePin);
+
+				if (const UMetasoundEditorGraphNode* Node = Cast<UMetasoundEditorGraphNode>(ReroutedOutputPin->GetOwningNode()))
+				{
+					TSharedPtr<FEditor> Editor = FGraphBuilder::GetEditorForNode(MetaSoundNode);
+					if (Editor.IsValid())
+					{
+						const FGuid NodeID = Node->GetNodeID();
+						const FName OutputName = ReroutedOutputPin->GetFName();
+						if (Editor->GetConnectionManager().GetValue(NodeID, OutputName, OutValue))
+						{
+							return true;
+						}
+					}
+
+					if (const UMetasoundEditorGraphMemberNode* MemberNode = Cast<UMetasoundEditorGraphMemberNode>(Node))
+					{
+						if (const UMetasoundEditorGraphMember* Member = MemberNode->GetMember())
+						{
+							if (Editor.IsValid())
+							{
+								// For an input member, we must use this NodeID/VertexName pair with the GraphConnectionManager:
+								const FGuid MemberID = Member->GetMemberID();
+								const FName MemberName = Member->GetMemberName();
+								if (Editor->GetConnectionManager().GetValue(MemberID, MemberName, OutValue))
+								{
+									return true;
+								}
+							}
+
+							if (const UMetasoundEditorGraphMemberDefaultLiteral* MemberDefaultLiteral = Member->GetLiteral())
+							{
+								const FMetasoundFrontendLiteral DefaultLiteral = MemberDefaultLiteral->GetDefault();
+								if (DefaultLiteral.TryGet(OutValue))
+								{
+									return true;
+								}
+							}
+						}
+					}
+				}
+			}
+
+			return false;
+		}
+	};
+}
diff --git a/Engine/Plugins/Runtime/Metasound/Source/MetasoundEditor/Private/MetasoundEditorModule.cpp b/Engine/Plugins/Runtime/Metasound/Source/MetasoundEditor/Private/MetasoundEditorModule.cpp
--- a/Engine/Plugins/Runtime/Metasound/Source/MetasoundEditor/Private/MetasoundEditorModule.cpp
+++ b/Engine/Plugins/Runtime/Metasound/Source/MetasoundEditor/Private/MetasoundEditorModule.cpp
@@ -24,6 +24,7 @@
 #include "MetasoundEditorGraphConnectionDrawingPolicy.h"
 #include "MetasoundEditorGraphMemberDefaults.h"
 #include "MetasoundEditorGraphNodeFactory.h"
+#include "MetasoundEditorGraphNodeVisualization.h"
 #include "MetasoundEditorSettings.h"
 #include "MetasoundFrontendDataTypeRegistry.h"
 #include "MetasoundFrontendDocument.h"
@@ -40,6 +41,7 @@
 #include "PackageMigrationContext.h"
 #include "PropertyEditorDelegates.h"
 #include "PropertyEditorModule.h"
+#include "SMetasoundFilterFrequencyResponsePlots.h"
 #include "Styling/AppStyle.h"
 #include "Styling/CoreStyle.h"
 #include "Styling/SlateStyle.h"
@@ -399,6 +401,11 @@ namespace Metasound
 
 			}
 
+			virtual void RegisterGraphNodeVisualization(FName InNodeClassName, FOnCreateGraphNodeVisualizationWidget OnCreateGraphNodeVisualizationWidget) override
+			{
+				FGraphNodeVisualizationRegistry::Get().RegisterVisualization(InNodeClassName, OnCreateGraphNodeVisualizationWidget);
+			}
+
 			void RegisterSettingsDelegates()
 			{
 				using namespace Engine;
@@ -615,6 +622,26 @@ namespace Metasound
 				GraphPanelPinFactory = MakeShared<FMetasoundGraphPanelPinFactory>();
 				FEdGraphUtilities::RegisterVisualPinFactory(GraphPanelPinFactory);
 
+				RegisterGraphNodeVisualization(
+					"UE.Biquad Filter.Audio",
+					FOnCreateGraphNodeVisualizationWidget::CreateStatic(&CreateMetaSoundBiquadFilterGraphNodeVisualizationWidget));
+
+				RegisterGraphNodeVisualization(
+					"UE.Ladder Filter.Audio",
+					FOnCreateGraphNodeVisualizationWidget::CreateStatic(&CreateMetaSoundLadderFilterGraphNodeVisualizationWidget));
+
+				RegisterGraphNodeVisualization(
+					"UE.One-Pole High Pass Filter.Audio",
+					FOnCreateGraphNodeVisualizationWidget::CreateStatic(&CreateMetaSoundOnePoleHighPassFilterGraphNodeVisualizationWidget));
+
+				RegisterGraphNodeVisualization(
+					"UE.One-Pole Low Pass Filter.Audio",
+					FOnCreateGraphNodeVisualizationWidget::CreateStatic(&CreateMetaSoundOnePoleLowPassFilterGraphNodeVisualizationWidget));
+
+				RegisterGraphNodeVisualization(
+					"UE.State Variable Filter.Audio",
+					FOnCreateGraphNodeVisualizationWidget::CreateStatic(&CreateMetaSoundStateVariableFilterGraphNodeVisualizationWidget));
+
 				ISettingsModule& SettingsModule = FModuleManager::LoadModuleChecked<ISettingsModule>("Settings");
 
 				SettingsModule.RegisterSettings("Editor", "ContentEditors", "MetaSound Editor",
@@ -688,6 +715,8 @@ namespace Metasound
 
 				AssetActions.Reset();
 				PinTypes.Reset();
+
+				FGraphNodeVisualizationRegistry::TearDown();
 			}
 
 			void OnPackageMigration(UE::AssetTools::FPackageMigrationContext& MigrationContext)
@@ -755,4 +784,4 @@ namespace Metasound
 
 IMPLEMENT_MODULE(Metasound::Editor::FModule, MetasoundEditor);
 
-#undef LOCTEXT_NAMESPACE
\ No newline at end of file
+#undef LOCTEXT_NAMESPACE
diff --git a/Engine/Plugins/Runtime/Metasound/Source/MetasoundEditor/Private/SMetasoundFilterFrequencyResponsePlots.cpp b/Engine/Plugins/Runtime/Metasound/Source/MetasoundEditor/Private/SMetasoundFilterFrequencyResponsePlots.cpp
new file mode 100644
--- /dev/null
+++ b/Engine/Plugins/Runtime/Metasound/Source/MetasoundEditor/Private/SMetasoundFilterFrequencyResponsePlots.cpp
@@ -0,0 +1,480 @@
+// Copyright Epic Games, Inc. All Rights Reserved.
+#include "SMetasoundFilterFrequencyResponsePlots.h"
+
+#include "Algo/Count.h"
+#include "DSP/FloatArrayMath.h"
+#include "MetasoundEditorGraphBuilder.h"
+#include "MetasoundEditorGraphNode.h"
+#include "MetasoundEditorModule.h"
+#include "Widgets/Layout/SBox.h"
+
+#define LOCTEXT_NAMESPACE "SMetasoundFilterFrequencyResponsePlots"
+
+namespace Metasound::Editor
+{
+	namespace SMetasoundFilterFrequencyResponsePlotsPrivate
+	{
+		inline bool IsConnectedAudioOutputPin(const UEdGraphPin* Pin)
+		{
+			const bool bIsPinConnected = !Pin->LinkedTo.IsEmpty();
+			const bool bIsAudioOutputPin = Pin->Direction == EGPD_Output && Pin->PinType.PinCategory == FGraphBuilder::PinCategoryAudio;
+			return (bIsAudioOutputPin && bIsPinConnected);
+		}
+	} // namespace SMetasoundFilterFrequencyResponsePlotsPrivate
+
+	TSharedRef<SWidget> CreateMetaSoundBiquadFilterGraphNodeVisualizationWidget(const FCreateGraphNodeVisualizationWidgetParams& InParams)
+	{
+		const FName FilterType(TEXT("Type"));
+		const FName CutoffFrequency(TEXT("Cutoff Frequency"));
+		const FName Bandwidth(TEXT("Bandwidth"));
+		const FName GainDb(TEXT("Gain"));
+
+		return SNew(SBox)
+			.MinDesiredHeight(125.0f)
+			[
+				SNew(SMetaSoundBiquadFilterFrequencyResponsePlot)
+					.FilterType_UObject(InParams.MetaSoundNode, &UMetasoundEditorGraphNode::GetPinVisualizationValue<Audio::EBiquadFilter::Type>, FilterType)
+					.CutoffFrequency_UObject(InParams.MetaSoundNode, &UMetasoundEditorGraphNode::GetPinVisualizationValue<float>, CutoffFrequency)
+					.Bandwidth_UObject(InParams.MetaSoundNode, &UMetasoundEditorGraphNode::GetPinVisualizationValue<float>, Bandwidth)
+					.GainDb_UObject(InParams.MetaSoundNode, &UMetasoundEditorGraphNode::GetPinVisualizationValue<float>, GainDb)
+			];
+	}
+
+	TSharedRef<SWidget> CreateMetaSoundLadderFilterGraphNodeVisualizationWidget(const FCreateGraphNodeVisualizationWidgetParams& InParams)
+	{
+		const FName CutoffFrequency(TEXT("Cutoff Frequency"));
+		const FName Resonance(TEXT("Resonance"));
+
+		return SNew(SBox)
+			.MinDesiredHeight(125.0f)
+			[
+				SNew(SMetaSoundLadderFilterFrequencyResponsePlot)
+					.CutoffFrequency_UObject(InParams.MetaSoundNode, &UMetasoundEditorGraphNode::GetPinVisualizationValue<float>, CutoffFrequency)
+					.Resonance_UObject(InParams.MetaSoundNode, &UMetasoundEditorGraphNode::GetPinVisualizationValue<float>, Resonance)
+			];
+	}
+
+	TSharedRef<SWidget> CreateMetaSoundOnePoleHighPassFilterGraphNodeVisualizationWidget(const FCreateGraphNodeVisualizationWidgetParams& InParams)
+	{
+		const FName CutoffFrequency(TEXT("Cutoff Frequency"));
+
+		return SNew(SBox)
+			.MinDesiredHeight(125.0f)
+			[
+				SNew(SMetaSoundOnePoleHighPassFilterFrequencyResponsePlot)
+					.CutoffFrequency_UObject(InParams.MetaSoundNode, &UMetasoundEditorGraphNode::GetPinVisualizationValue<float>, CutoffFrequency)
+			];
+	}
+
+	TSharedRef<SWidget> CreateMetaSoundOnePoleLowPassFilterGraphNodeVisualizationWidget(const FCreateGraphNodeVisualizationWidgetParams& InParams)
+	{
+		const FName CutoffFrequency(TEXT("Cutoff Frequency"));
+
+		return SNew(SBox)
+			.MinDesiredHeight(125.0f)
+			[
+				SNew(SMetaSoundOnePoleLowPassFilterFrequencyResponsePlot)
+					.CutoffFrequency_UObject(InParams.MetaSoundNode, &UMetasoundEditorGraphNode::GetPinVisualizationValue<float>, CutoffFrequency)
+			];
+	}
+
+	TSharedRef<SWidget> CreateMetaSoundStateVariableFilterGraphNodeVisualizationWidget(const FCreateGraphNodeVisualizationWidgetParams& InParams)
+	{
+		const FName CutoffFrequency(TEXT("Cutoff Frequency"));
+		const FName Resonance(TEXT("Resonance"));
+		const FName BandStopControl(TEXT("Band Stop Control"));
+
+		return SNew(SBox)
+			.MinDesiredHeight(125.0f)
+			[
+				SNew(SMetaSoundStateVariableFilterFrequencyResponsePlot, InParams.MetaSoundNode)
+					.CutoffFrequency_UObject(InParams.MetaSoundNode, &UMetasoundEditorGraphNode::GetPinVisualizationValue<float>, CutoffFrequency)
+					.Resonance_UObject(InParams.MetaSoundNode, &UMetasoundEditorGraphNode::GetPinVisualizationValue<float>, Resonance)
+					.BandStopControl_UObject(InParams.MetaSoundNode, &UMetasoundEditorGraphNode::GetPinVisualizationValue<float>, BandStopControl)
+			];
+	}
+
+
+	const float SMetaSoundFilterFrequencyResponsePlot::SampleRate = 48000.0f;
+
+	void SMetaSoundFilterFrequencyResponsePlot::Construct()
+	{
+		ChildSlot
+			[
+				SAssignNew(FrequencyResponsePlot, SAudioSpectrumPlot)
+					.Clipping(EWidgetClipping::ClipToBounds)
+					.ViewMinSoundLevel(-24.0f)
+					.TiltExponent_Lambda([] { return 0.0f; }) // Binding this property has the effect of hiding its context menu entry (Tilting the spectrum is not desired here).
+					.FrequencyAxisPixelBucketMode_Lambda([]() { return EAudioSpectrumPlotFrequencyAxisPixelBucketMode::Sample; }) // Binding this property has the effect of hiding its context menu entry (PixelBucketMode is not much use here).
+					.OnGetAudioSpectrumData(this, &SMetaSoundFilterFrequencyResponsePlot::GetAudioSpectrumData)
+			];
+
+		ContextMenuExtension = FrequencyResponsePlot->AddContextMenuExtension(EExtensionHook::Before, nullptr, FMenuExtensionDelegate::CreateSP(this, &SMetaSoundFilterFrequencyResponsePlot::ExtendSpectrumPlotContextMenu));
+	}
+
+	SMetaSoundFilterFrequencyResponsePlot::~SMetaSoundFilterFrequencyResponsePlot()
+	{
+		if (ContextMenuExtension.IsValid())
+		{
+			FrequencyResponsePlot->RemoveContextMenuExtension(ContextMenuExtension.ToSharedRef());
+		}
+	}
+
+	void SMetaSoundFilterFrequencyResponsePlot::Tick(const FGeometry& AllottedGeometry, const double InCurrentTime, const float InDeltaTime)
+	{
+		bHasFilterParams = UpdateFilterParams();
+	}
+
+	FAudioPowerSpectrumData SMetaSoundFilterFrequencyResponsePlot::GetAudioSpectrumData()
+	{
+		if (!bHasFilterParams || !FrequencyResponsePlot.IsValid())
+		{
+			return FAudioPowerSpectrumData();
+		}
+
+		// Fill array of frequencies (in Hz) to be plotted (these may be log spaced or linear, as defined by the FrequencyResponsePlot ScaleInfo transform):
+		CenterFrequencies.Reset();
+		float LocalX = 0.0f;
+		const float LocalXEnd = FrequencyResponsePlot->GetPaintSpaceGeometry().GetLocalSize().X;
+		const FAudioSpectrumPlotScaleInfo ScaleInfo = FrequencyResponsePlot->GetScaleInfo();
+		while (LocalX <= LocalXEnd)
+		{
+			const float Frequency = ScaleInfo.LocalXToFrequency(LocalX);
+			CenterFrequencies.Add(Frequency);
+			LocalX += 1.0f;
+		}
+
+		// Create array of complex z values for all desired frequencies (interleaved real and imaginary):
+		const int32 NumFrequencies = CenterFrequencies.Num();
+		TArray<float> ComplexValues;
+		ComplexValues.SetNumUninitialized(2 * NumFrequencies);
+		const float HzToOmega = UE_TWO_PI / SampleRate;
+		for (int32 Index = 0; Index < NumFrequencies; Index++)
+		{
+			const float Omega = HzToOmega * CenterFrequencies[Index];
+			FMath::SinCos(&ComplexValues[2 * Index + 1], &ComplexValues[2 * Index + 0], Omega);
+		}
+
+		// Get the frequency response for all of the frequencies:
+		ArrayCalculateFilterResponseInPlace(ComplexValues);
+
+		// Store frequency response as squared magnitudes:
+		SquaredMagnitudes.SetNumUninitialized(NumFrequencies);
+		Audio::ArrayComplexToPower(ComplexValues, SquaredMagnitudes);
+
+		return FAudioPowerSpectrumData
+		{
+			.CenterFrequencies = CenterFrequencies,
+			.SquaredMagnitudes = SquaredMagnitudes,
+		};
+	}
+
+
+
+	void SMetaSoundBiquadFilterFrequencyResponsePlot::Construct(const FArguments& InArgs)
+	{
+		SMetaSoundFilterFrequencyResponsePlot::Construct();
+
+		FilterType = InArgs._FilterType;
+		CutoffFrequency = InArgs._CutoffFrequency;
+		Bandwidth = InArgs._Bandwidth;
+		GainDb = InArgs._GainDb;
+	}
+
+	bool SMetaSoundBiquadFilterFrequencyResponsePlot::UpdateFilterParams()
+	{
+		// Try and get all required filter params:
+		const TOptional<Audio::EBiquadFilter::Type> FilterTypeValue = FilterType.Get();
+		const TOptional<float> CutoffFrequencyValue = CutoffFrequency.Get();
+		const TOptional<float> BandwidthValue = Bandwidth.Get();
+		const TOptional<float> GainDbValue = GainDb.Get();
+		if (!FilterTypeValue.IsSet() || !CutoffFrequencyValue.IsSet() || !BandwidthValue.IsSet() || !GainDbValue.IsSet())
+		{
+			return false;
+		}
+
+		const float MaxCutoffFrequency = 0.5f * SampleRate;
+		const float CurrentFrequency = FMath::Clamp(*CutoffFrequencyValue, 0.f, MaxCutoffFrequency);
+		const float CurrentBandwidth = FMath::Max(*BandwidthValue, 0.f);
+		const float CurrentFilterGainDb = FMath::Clamp(*GainDbValue, -90.0f, 20.0f);
+
+		// Update the filter with the filter params:
+		if (Filter.GetNumChannels() == 0)
+		{
+			constexpr int32 NumChannels = 1;
+			Filter.Init(SampleRate, NumChannels, static_cast<Audio::EBiquadFilter::Type>(*FilterTypeValue), CurrentFrequency, CurrentBandwidth, CurrentFilterGainDb);
+		}
+		else
+		{
+			Filter.SetParams(static_cast<Audio::EBiquadFilter::Type>(*FilterTypeValue), CurrentFrequency, CurrentBandwidth, CurrentFilterGainDb);
+		}
+
+		return true;
+	}
+
+	void SMetaSoundBiquadFilterFrequencyResponsePlot::ArrayCalculateFilterResponseInPlace(TArrayView<float> InOutComplexValues) const
+	{
+		Filter.ArrayCalculateResponseInPlace(InOutComplexValues);
+	}
+
+
+
+	void SMetaSoundLadderFilterFrequencyResponsePlot::Construct(const FArguments& InArgs)
+	{
+		SMetaSoundFilterFrequencyResponsePlot::Construct();
+
+		constexpr int32 NumChannels = 1;
+		Filter.Init(SampleRate, NumChannels);
+
+		CutoffFrequency = InArgs._CutoffFrequency;
+		Resonance = InArgs._Resonance;
+	}
+
+	bool SMetaSoundLadderFilterFrequencyResponsePlot::UpdateFilterParams()
+	{
+		// Try and get all required filter params:
+		const TOptional<float> CutoffFrequencyValue = CutoffFrequency.Get();
+		const TOptional<float> ResonanceValue = Resonance.Get();
+		if (!CutoffFrequencyValue.IsSet() || !ResonanceValue.IsSet())
+		{
+			return false;
+		}
+
+		const float MaxCutoffFrequency = 0.5f * SampleRate;
+		const float CurrentFrequency = FMath::Clamp(*CutoffFrequencyValue, 0.f, MaxCutoffFrequency);
+		const float CurrentResonance = FMath::Clamp(*ResonanceValue, 1.0f, 10.0f);
+
+		// Update the filter with the filter params:
+		Filter.SetQ(CurrentResonance);
+		Filter.SetFrequency(CurrentFrequency);
+
+		Filter.Update();
+
+		return true;
+	}
+
+	void SMetaSoundLadderFilterFrequencyResponsePlot::ArrayCalculateFilterResponseInPlace(TArrayView<float> InOutComplexValues) const
+	{
+		Filter.ArrayCalculateResponseInPlace(InOutComplexValues);
+	}
+
+
+
+	void SMetaSoundOnePoleHighPassFilterFrequencyResponsePlot::Construct(const FArguments& InArgs)
+	{
+		SMetaSoundFilterFrequencyResponsePlot::Construct();
+
+		constexpr int32 NumChannels = 1;
+		Filter.Init(SampleRate, NumChannels);
+
+		CutoffFrequency = InArgs._CutoffFrequency;
+	}
+
+	bool SMetaSoundOnePoleHighPassFilterFrequencyResponsePlot::UpdateFilterParams()
+	{
+		// Try and get all required filter params:
+		const TOptional<float> Frequency = CutoffFrequency.Get();
+		if (!Frequency.IsSet())
+		{
+			return false;
+		}
+
+		const float ClampedFreq = FMath::Clamp(0.0f, *Frequency, SampleRate);
+
+		// Update the filter with the filter params:
+		Filter.StartFrequencyInterpolation(ClampedFreq);
+
+		return true;
+	}
+
+	void SMetaSoundOnePoleHighPassFilterFrequencyResponsePlot::ArrayCalculateFilterResponseInPlace(TArrayView<float> InOutComplexValues) const
+	{
+		Filter.ArrayCalculateResponseInPlace(InOutComplexValues);
+	}
+
+
+
+	void SMetaSoundOnePoleLowPassFilterFrequencyResponsePlot::Construct(const FArguments& InArgs)
+	{
+		SMetaSoundFilterFrequencyResponsePlot::Construct();
+
+		constexpr int32 NumChannels = 1;
+		Filter.Init(SampleRate, NumChannels);
+
+		CutoffFrequency = InArgs._CutoffFrequency;
+	}
+
+	bool SMetaSoundOnePoleLowPassFilterFrequencyResponsePlot::UpdateFilterParams()
+	{
+		// Try and get all required filter params:
+		const TOptional<float> Frequency = CutoffFrequency.Get();
+		if (!Frequency.IsSet())
+		{
+			return false;
+		}
+
+		const float ClampedFreq = FMath::Clamp(0.0f, *Frequency, SampleRate);
+
+		// Update the filter with the filter params:
+		Filter.StartFrequencyInterpolation(ClampedFreq);
+
+		return true;
+	}
+
+	void SMetaSoundOnePoleLowPassFilterFrequencyResponsePlot::ArrayCalculateFilterResponseInPlace(TArrayView<float> InOutComplexValues) const
+	{
+		Filter.ArrayCalculateResponseInPlace(InOutComplexValues);
+	}
+
+
+
+	const FName SMetaSoundStateVariableFilterFrequencyResponsePlot::LowPassFilter(TEXT("Low Pass Filter"));
+	const FName SMetaSoundStateVariableFilterFrequencyResponsePlot::HighPassFilter(TEXT("High Pass Filter"));
+	const FName SMetaSoundStateVariableFilterFrequencyResponsePlot::BandPass(TEXT("Band Pass"));
+	const FName SMetaSoundStateVariableFilterFrequencyResponsePlot::BandStop(TEXT("Band Stop"));
+
+	SMetaSoundStateVariableFilterFrequencyResponsePlot::SMetaSoundStateVariableFilterFrequencyResponsePlot()
+		: DisplayedFilterResponse(LowPassFilter)
+	{
+		constexpr int32 NumChannels = 1;
+		Filter.Init(SampleRate, NumChannels);
+	}
+
+	void SMetaSoundStateVariableFilterFrequencyResponsePlot::Construct(const FArguments& InArgs, UMetasoundEditorGraphNode* InMetaSoundNode)
+	{
+		SMetaSoundFilterFrequencyResponsePlot::Construct();
+
+		CutoffFrequency = InArgs._CutoffFrequency;
+		Resonance = InArgs._Resonance;
+		BandStopControl = InArgs._BandStopControl;
+
+		MetaSoundNode = InMetaSoundNode;
+	}
+
+	bool SMetaSoundStateVariableFilterFrequencyResponsePlot::UpdateFilterParams()
+	{
+		using namespace SMetasoundFilterFrequencyResponsePlotsPrivate;
+
+		const TStrongObjectPtr<UMetasoundEditorGraphNode> GraphNode = MetaSoundNode.Pin();
+		if (!GraphNode.IsValid())
+		{
+			return false;
+		}
+
+		// If there are connected audio outputs but the current selection is not connected, auto-select a new filter type:
+		UEdGraphPin** FirstConnectedAudioOutputPin = GraphNode->Pins.FindByPredicate(IsConnectedAudioOutputPin);
+		if (FirstConnectedAudioOutputPin != nullptr)
+		{
+			for (UEdGraphPin* Pin : GraphNode->Pins)
+			{
+				const bool bIsPinConnected = !Pin->LinkedTo.IsEmpty();
+				if (Pin->PinName == DisplayedFilterResponse && !bIsPinConnected)
+				{
+					DisplayedFilterResponse = (*FirstConnectedAudioOutputPin)->PinName;
+				}
+			}
+		}
+
+		// Set the filter type on the filter:
+		if (DisplayedFilterResponse == LowPassFilter)
+		{
+			Filter.SetFilterType(Audio::EFilter::LowPass);
+		}
+		else if (DisplayedFilterResponse == HighPassFilter)
+		{
+			Filter.SetFilterType(Audio::EFilter::HighPass);
+		}
+		else if (DisplayedFilterResponse == BandPass)
+		{
+			Filter.SetFilterType(Audio::EFilter::BandPass);
+		}
+		else if (DisplayedFilterResponse == BandStop)
+		{
+			Filter.SetFilterType(Audio::EFilter::BandStop);
+		}
+
+		// Try and get all required filter params:
+		const TOptional<float> CutoffFrequencyValue = CutoffFrequency.Get();
+		const TOptional<float> ResonanceValue = Resonance.Get();
+		const TOptional<float> BandStopControlValue = BandStopControl.Get();
+		if (!CutoffFrequencyValue.IsSet() || !ResonanceValue.IsSet() || !BandStopControl.IsSet())
+		{
+			return false;
+		}
+
+		const float MaxCutoffFrequency = 0.5f * SampleRate;
+		const float CurrentFrequency = FMath::Clamp(*CutoffFrequencyValue, 0.f, MaxCutoffFrequency);
+		const float CurrentResonance = FMath::Clamp(*ResonanceValue, 0.f, 10.f);
+		const float CurrentBandStopControl = FMath::Clamp(*BandStopControlValue, 0.f, 1.f);
+
+		// Update the filter with the filter params:
+		Filter.SetQ(CurrentResonance);
+		Filter.SetFrequency(CurrentFrequency);
+		Filter.SetBandStopControl(CurrentBandStopControl);
+
+		Filter.Update();
+
+		return true;
+	}
+
+	void SMetaSoundStateVariableFilterFrequencyResponsePlot::ArrayCalculateFilterResponseInPlace(TArrayView<float> InOutComplexValues) const
+	{
+		Filter.ArrayCalculateResponseInPlace(InOutComplexValues);
+	}
+
+	void SMetaSoundStateVariableFilterFrequencyResponsePlot::ExtendSpectrumPlotContextMenu(FMenuBuilder& MenuBuilder)
+	{
+		using namespace SMetasoundFilterFrequencyResponsePlotsPrivate;
+
+		const TStrongObjectPtr<UMetasoundEditorGraphNode> GraphNode = MetaSoundNode.Pin();
+		if (GraphNode.IsValid())
+		{
+			// Display filter response selection submenu if no outputs are connected, or more than one output is connected:
+			const int NumConnectedAudioOutputPins = Algo::CountIf(GraphNode->Pins, IsConnectedAudioOutputPin);
+			if (NumConnectedAudioOutputPins == 0 || NumConnectedAudioOutputPins > 1)
+			{
+				MenuBuilder.AddSubMenu(
+					LOCTEXT("DisplayedFilterResponse", "Displayed Filter Response"),
+					FText(),
+					FNewMenuDelegate::CreateSP(this, &SMetaSoundStateVariableFilterFrequencyResponsePlot::BuildFilterOutputSubMenu));
+			}
+		}
+	}
+
+	void SMetaSoundStateVariableFilterFrequencyResponsePlot::BuildFilterOutputSubMenu(FMenuBuilder& SubMenu)
+	{
+		using namespace SMetasoundFilterFrequencyResponsePlotsPrivate;
+
+		const TStrongObjectPtr<UMetasoundEditorGraphNode> GraphNode = MetaSoundNode.Pin();
+		if (GraphNode.IsValid())
+		{
+			// Add menu entries for all connected audio output pins. If there are no connected audio output pins then add entries for all audio output pins:
+			const bool bHasConnectedAudioOutputPins = GraphNode->Pins.ContainsByPredicate(IsConnectedAudioOutputPin);
+			for (UEdGraphPin* Pin : GraphNode->Pins)
+			{
+				const bool bIsAudioOutputPin = Pin->Direction == EGPD_Output && Pin->PinType.PinCategory == FGraphBuilder::PinCategoryAudio;
+				if (bIsAudioOutputPin)
+				{
+					const bool bIsPinConnected = !Pin->LinkedTo.IsEmpty();
+					if (bIsPinConnected || !bHasConnectedAudioOutputPins)
+					{
+						SubMenu.AddMenuEntry(
+							FText::FromName(Pin->PinName),
+							FText::FromString(Pin->PinToolTip),
+							FSlateIcon(),
+							FUIAction(
+								FExecuteAction::CreateSPLambda(this, [this, PinName = Pin->PinName]() { DisplayedFilterResponse = PinName; }),
+								FCanExecuteAction(),
+								FIsActionChecked::CreateSPLambda(this, [this, PinName = Pin->PinName]() { return DisplayedFilterResponse == PinName; })
+							),
+							NAME_None,
+							EUserInterfaceActionType::ToggleButton);
+					}
+				}
+			}
+		}
+	}
+
+}
+
+#undef LOCTEXT_NAMESPACE
diff --git a/Engine/Plugins/Runtime/Metasound/Source/MetasoundEditor/Private/SMetasoundFilterFrequencyResponsePlots.h b/Engine/Plugins/Runtime/Metasound/Source/MetasoundEditor/Private/SMetasoundFilterFrequencyResponsePlots.h
new file mode 100644
--- /dev/null
+++ b/Engine/Plugins/Runtime/Metasound/Source/MetasoundEditor/Private/SMetasoundFilterFrequencyResponsePlots.h
@@ -0,0 +1,180 @@
+// Copyright Epic Games, Inc. All Rights Reserved.
+#pragma once
+
+#include "DSP/Filter.h"
+#include "DSP/InterpolatedOnePole.h"
+#include "SAudioSpectrumPlot.h"
+#include "SMetasoundGraphNode.h"
+
+namespace Metasound::Editor
+{
+	struct FCreateGraphNodeVisualizationWidgetParams;
+
+	TSharedRef<SWidget> CreateMetaSoundBiquadFilterGraphNodeVisualizationWidget(const FCreateGraphNodeVisualizationWidgetParams& InParams);
+	TSharedRef<SWidget> CreateMetaSoundLadderFilterGraphNodeVisualizationWidget(const FCreateGraphNodeVisualizationWidgetParams& InParams);
+	TSharedRef<SWidget> CreateMetaSoundOnePoleHighPassFilterGraphNodeVisualizationWidget(const FCreateGraphNodeVisualizationWidgetParams& InParams);
+	TSharedRef<SWidget> CreateMetaSoundOnePoleLowPassFilterGraphNodeVisualizationWidget(const FCreateGraphNodeVisualizationWidgetParams& InParams);
+	TSharedRef<SWidget> CreateMetaSoundStateVariableFilterGraphNodeVisualizationWidget(const FCreateGraphNodeVisualizationWidgetParams& InParams);
+
+	/**
+	* Abstract class implementing shared functionality for any Metasound basic filter. Allows for displaying a frequency response plot of the filter.
+	*/
+	class SMetaSoundFilterFrequencyResponsePlot : public SCompoundWidget
+	{
+	public:
+		void Construct();
+		virtual ~SMetaSoundFilterFrequencyResponsePlot();
+
+	protected:
+		// Update function for derived classes to update their state each frame. Implementations should return true if they have the required information to plot a frequency response, or false otherwise.
+		virtual bool UpdateFilterParams() = 0;
+
+		// Derived classes should apply the filter transfer function to each z-domain value in the given array (complex numbers given as interleaved floats).
+		virtual void ArrayCalculateFilterResponseInPlace(TArrayView<float> InOutComplexValues) const = 0;
+
+		// Derived classes can optionally add items to the spectrum plot menu by overriding this member function.
+		virtual void ExtendSpectrumPlotContextMenu(FMenuBuilder& MenuBuilder) {};
+
+		// Begin SWidget overrides.
+		virtual void Tick(const FGeometry& AllottedGeometry, const double InCurrentTime, const float InDeltaTime) override;
+		// End SWidget overrides.
+
+		static const float SampleRate;
+		
+	private:
+		FAudioPowerSpectrumData GetAudioSpectrumData();
+
+		TSharedPtr<SAudioSpectrumPlot> FrequencyResponsePlot;
+		TSharedPtr<const FExtensionBase> ContextMenuExtension;
+		TArray<float> CenterFrequencies;
+		TArray<float> SquaredMagnitudes;
+		bool bHasFilterParams = false;
+	};
+
+
+	class SMetaSoundBiquadFilterFrequencyResponsePlot : public SMetaSoundFilterFrequencyResponsePlot
+	{
+	public:
+		SLATE_BEGIN_ARGS(SMetaSoundBiquadFilterFrequencyResponsePlot)
+			{}
+			SLATE_ATTRIBUTE(TOptional<Audio::EBiquadFilter::Type>, FilterType)
+			SLATE_ATTRIBUTE(TOptional<float>, CutoffFrequency)
+			SLATE_ATTRIBUTE(TOptional<float>, Bandwidth)
+			SLATE_ATTRIBUTE(TOptional<float>, GainDb)
+		SLATE_END_ARGS()
+
+		void Construct(const FArguments& InArgs);
+
+	protected:
+		virtual bool UpdateFilterParams() override;
+		virtual void ArrayCalculateFilterResponseInPlace(TArrayView<float> InOutComplexValues) const override;
+
+	private:
+		Audio::FBiquadFilter Filter;
+
+		TAttribute<TOptional<Audio::EBiquadFilter::Type>> FilterType;
+		TAttribute<TOptional<float>> CutoffFrequency;
+		TAttribute<TOptional<float>> Bandwidth;
+		TAttribute<TOptional<float>> GainDb;
+	};
+
+
+	class SMetaSoundLadderFilterFrequencyResponsePlot : public SMetaSoundFilterFrequencyResponsePlot
+	{
+	public:
+		SLATE_BEGIN_ARGS(SMetaSoundLadderFilterFrequencyResponsePlot)
+			{}
+			SLATE_ATTRIBUTE(TOptional<float>, CutoffFrequency)
+			SLATE_ATTRIBUTE(TOptional<float>, Resonance)
+		SLATE_END_ARGS()
+
+		void Construct(const FArguments& InArgs);
+
+	protected:
+		virtual bool UpdateFilterParams() override;
+		virtual void ArrayCalculateFilterResponseInPlace(TArrayView<float> InOutComplexValues) const override;
+
+	private:
+		Audio::FLadderFilter Filter;
+
+		TAttribute<TOptional<float>> CutoffFrequency;
+		TAttribute<TOptional<float>> Resonance;
+	};
+
+
+	class SMetaSoundOnePoleHighPassFilterFrequencyResponsePlot : public SMetaSoundFilterFrequencyResponsePlot
+	{
+	public:
+		SLATE_BEGIN_ARGS(SMetaSoundOnePoleHighPassFilterFrequencyResponsePlot)
+			{}
+			SLATE_ATTRIBUTE(TOptional<float>, CutoffFrequency)
+		SLATE_END_ARGS()
+
+		void Construct(const FArguments& InArgs);
+
+	protected:
+		virtual bool UpdateFilterParams() override;
+		virtual void ArrayCalculateFilterResponseInPlace(TArrayView<float> InOutComplexValues) const override;
+
+	private:
+		Audio::FInterpolatedHPF Filter;
+		TAttribute<TOptional<float>> CutoffFrequency;
+	};
+
+
+	class SMetaSoundOnePoleLowPassFilterFrequencyResponsePlot : public SMetaSoundFilterFrequencyResponsePlot
+	{
+	public:
+		SLATE_BEGIN_ARGS(SMetaSoundOnePoleLowPassFilterFrequencyResponsePlot)
+			{}
+			SLATE_ATTRIBUTE(TOptional<float>, CutoffFrequency)
+		SLATE_END_ARGS()
+
+		void Construct(const FArguments& InArgs);
+
+	protected:
+		virtual bool UpdateFilterParams() override;
+		virtual void ArrayCalculateFilterResponseInPlace(TArrayView<float> InOutComplexValues) const override;
+
+	private:
+		Audio::FInterpolatedLPF Filter;
+		TAttribute<TOptional<float>> CutoffFrequency;
+	};
+
+
+	class SMetaSoundStateVariableFilterFrequencyResponsePlot : public SMetaSoundFilterFrequencyResponsePlot
+	{
+	public:
+		SMetaSoundStateVariableFilterFrequencyResponsePlot();
+
+		SLATE_BEGIN_ARGS(SMetaSoundStateVariableFilterFrequencyResponsePlot)
+			{}
+			SLATE_ATTRIBUTE(TOptional<float>, CutoffFrequency)
+			SLATE_ATTRIBUTE(TOptional<float>, Resonance)
+			SLATE_ATTRIBUTE(TOptional<float>, BandStopControl)
+		SLATE_END_ARGS()
+
+		void Construct(const FArguments& InArgs, UMetasoundEditorGraphNode* InMetaSoundNode);
+
+	protected:
+		virtual bool UpdateFilterParams() override;
+		virtual void ArrayCalculateFilterResponseInPlace(TArrayView<float> InOutComplexValues) const override;
+		virtual void ExtendSpectrumPlotContextMenu(FMenuBuilder& MenuBuilder) override;
+
+	private:
+		void BuildFilterOutputSubMenu(FMenuBuilder& SubMenu);
+
+		static const FName LowPassFilter;
+		static const FName HighPassFilter;
+		static const FName BandPass;
+		static const FName BandStop;
+
+		TWeakObjectPtr<UMetasoundEditorGraphNode> MetaSoundNode;
+		Audio::FStateVariableFilter Filter;
+		FName DisplayedFilterResponse;
+
+		TAttribute<TOptional<float>> CutoffFrequency;
+		TAttribute<TOptional<float>> Resonance;
+		TAttribute<TOptional<float>> BandStopControl;
+	};
+}
diff --git a/Engine/Plugins/Runtime/Metasound/Source/MetasoundEditor/Private/SMetasoundFilterGraphNodes.cpp b/Engine/Plugins/Runtime/Metasound/Source/MetasoundEditor/Private/SMetasoundFilterGraphNodes.cpp
deleted file mode 100644
diff --git a/Engine/Plugins/Runtime/Metasound/Source/MetasoundEditor/Private/SMetasoundFilterGraphNodes.h b/Engine/Plugins/Runtime/Metasound/Source/MetasoundEditor/Private/SMetasoundFilterGraphNodes.h
deleted file mode 100644
diff --git a/Engine/Plugins/Runtime/Metasound/Source/MetasoundEditor/Private/SMetasoundGraphNode.cpp b/Engine/Plugins/Runtime/Metasound/Source/MetasoundEditor/Private/SMetasoundGraphNode.cpp
--- a/Engine/Plugins/Runtime/Metasound/Source/MetasoundEditor/Private/SMetasoundGraphNode.cpp
+++ b/Engine/Plugins/Runtime/Metasound/Source/MetasoundEditor/Private/SMetasoundGraphNode.cpp
@@ -23,6 +23,7 @@
 #include "MetasoundEditorGraphInputNode.h"
 #include "MetasoundEditorGraphMemberDefaults.h"
 #include "MetasoundEditorGraphNode.h"
+#include "MetasoundEditorGraphNodeVisualization.h"
 #include "MetasoundEditorGraphSchema.h"
 #include "MetasoundEditorModule.h"
 #include "MetasoundEditorSettings.h"
@@ -400,6 +401,25 @@ namespace Metasound
 			}
 		}
 
+		void SMetaSoundGraphNode::CreateBelowPinControls(TSharedPtr<SVerticalBox> MainBox)
+		{
+			if (MainBox.IsValid())
+			{
+				UMetasoundEditorGraphNode& MetaSoundNode = GetMetaSoundNode();
+
+				const FName NodeClassName = MetaSoundNode.GetBreadcrumb().ClassName.GetFullName();
+				const FCreateGraphNodeVisualizationWidgetParams CreateParams{ .MetaSoundNode = &MetaSoundNode };
+				if (TSharedPtr<SWidget> VisualizationWidget = FGraphNodeVisualizationRegistry::Get().CreateVisualizationWidget(NodeClassName, CreateParams))
+				{
+					MainBox->AddSlot()
+						.Padding(1.0f, 0.0f)
+						[
+							VisualizationWidget.ToSharedRef()
+						];
+				}
+			}
+		}
+
 		TSharedRef<SWidget> SMetaSoundGraphNode::CreateTitleWidget(TSharedPtr<SNodeTitle> NodeTitle)
 		{
 			Frontend::FConstNodeHandle NodeHandle = GetMetaSoundNode().GetConstNodeHandle();
diff --git a/Engine/Plugins/Runtime/Metasound/Source/MetasoundEditor/Private/SMetasoundGraphNode.h b/Engine/Plugins/Runtime/Metasound/Source/MetasoundEditor/Private/SMetasoundGraphNode.h
--- a/Engine/Plugins/Runtime/Metasound/Source/MetasoundEditor/Private/SMetasoundGraphNode.h
+++ b/Engine/Plugins/Runtime/Metasound/Source/MetasoundEditor/Private/SMetasoundGraphNode.h
@@ -47,6 +47,7 @@ namespace Metasound::Editor
 		virtual void CreateOutputSideAddButton(TSharedPtr<SVerticalBox> OutputBox) override;
 		virtual TSharedPtr<SGraphPin> CreatePinWidget(UEdGraphPin* InPin) const override;
 		virtual void CreateStandardPinWidget(UEdGraphPin* InPin) override;
+		virtual void CreateBelowPinControls(TSharedPtr<SVerticalBox> MainBox) override;
 		virtual TSharedRef<SWidget> CreateNodeContentArea() override;
 		virtual TSharedRef<SWidget> CreateTitleWidget(TSharedPtr<SNodeTitle> NodeTitle) override;
 		virtual const FSlateBrush* GetNodeBodyBrush() const override;
diff --git a/Engine/Plugins/Runtime/Metasound/Source/MetasoundEditor/Public/MetasoundEditorGraphNode.h b/Engine/Plugins/Runtime/Metasound/Source/MetasoundEditor/Public/MetasoundEditorGraphNode.h
--- a/Engine/Plugins/Runtime/Metasound/Source/MetasoundEditor/Public/MetasoundEditorGraphNode.h
+++ b/Engine/Plugins/Runtime/Metasound/Source/MetasoundEditor/Public/MetasoundEditorGraphNode.h
@@ -181,6 +181,34 @@ public:
 	// Returns whether or not the class interface, metadata, or style has been changed since the last node refresh
 	bool ContainsClassChange() const;
 
+	// Graph node visualization widgets can attempt to get the current value of the given named input pin. For connected input pins, a value may not be returned unless sound preview is active.
+	METASOUNDEDITOR_API bool TryGetPinVisualizationValue(FName InPinName, bool& OutValue) const;
+	METASOUNDEDITOR_API bool TryGetPinVisualizationValue(FName InPinName, int32& OutValue) const;
+	METASOUNDEDITOR_API bool TryGetPinVisualizationValue(FName InPinName, float& OutValue) const;
+
+	template<class T>
+	inline TOptional<T> GetPinVisualizationValue(FName InPinName) const
+	{
+		if constexpr (std::is_enum_v<T>)
+		{
+			int32 Result;
+			if (TryGetPinVisualizationValue(InPinName, Result))
+			{
+				return static_cast<T>(Result);
+			}
+			return NullOpt;
+		}
+		else
+		{
+			T Result;
+			if (TryGetPinVisualizationValue(InPinName, Result))
+			{
+				return Result;
+			}
+			return NullOpt;
+		}
+	}
+	
 protected:
 	FGuid InterfaceChangeID;
 	FGuid MetadataChangeID;
diff --git a/Engine/Plugins/Runtime/Metasound/Source/MetasoundEditor/Public/MetasoundEditorModule.h b/Engine/Plugins/Runtime/Metasound/Source/MetasoundEditor/Public/MetasoundEditorModule.h
--- a/Engine/Plugins/Runtime/Metasound/Source/MetasoundEditor/Public/MetasoundEditorModule.h
+++ b/Engine/Plugins/Runtime/Metasound/Source/MetasoundEditor/Public/MetasoundEditorModule.h
@@ -18,6 +18,7 @@ class IDetailPropertyRow;
 class UEdGraphPin;
 class UMetasoundEditorGraph;
 class UMetasoundEditorGraphMemberDefaultLiteral;
+class UMetasoundEditorGraphNode;
 namespace Metasound::Engine 
 {
 	enum class EAssetScanStatus : uint8;
@@ -75,6 +76,15 @@ namespace Metasound
 			virtual TUniquePtr<FMetasoundDefaultLiteralCustomizationBase> CreateLiteralCustomization(IDetailCategoryBuilder& DefaultCategoryBuilder) const = 0;
 		};
 
+
+		struct FCreateGraphNodeVisualizationWidgetParams
+		{
+			UMetasoundEditorGraphNode* MetaSoundNode = nullptr;
+		};
+
+		DECLARE_DELEGATE_RetVal_OneParam(TSharedRef<SWidget>, FOnCreateGraphNodeVisualizationWidget, const FCreateGraphNodeVisualizationWidgetParams&);
+
+
 		class METASOUNDEDITOR_API IMetasoundEditorModule : public IModuleInterface
 		{
 		public:
@@ -111,6 +121,9 @@ namespace Metasound
 				const FSlateBrush* InPinConnectedIcon = nullptr, const FSlateBrush* InPinDisconnectedIcon = nullptr) = 0;
 			
 			virtual void RegisterCustomPinType(FName InDataTypeName, const FGraphPinParams& Params) = 0;
+
+			// For the given node class, register a delegate that can be used for creating in-graph node visualizations.
+			virtual void RegisterGraphNodeVisualization(FName InNodeClassName, FOnCreateGraphNodeVisualizationWidget OnCreateGraphNodeVisualizationWidget) = 0;
 		};
 	} // namespace Editor
 } // namespace Metasound
