UE-171619 - Right click for MetaSound pin options should have 'promote to variable'

diff --git a/Engine/Plugins/Runtime/Metasound/Source/MetasoundEditor/Private/MetasoundEditor.cpp b/Engine/Plugins/Runtime/Metasound/Source/MetasoundEditor/Private/MetasoundEditor.cpp
--- a/Engine/Plugins/Runtime/Metasound/Source/MetasoundEditor/Private/MetasoundEditor.cpp
+++ b/Engine/Plugins/Runtime/Metasound/Source/MetasoundEditor/Private/MetasoundEditor.cpp
@@ -2744,6 +2744,22 @@ namespace Metasound
 					FExecuteAction::CreateLambda([this] { RenameSelectedNode(); }),
 					FCanExecuteAction::CreateLambda([this]() { return CanRenameSelectedNodes(); }));
 
+				GraphEditorCommands->MapAction(FEditorCommands::Get().PromoteToInput,
+					FExecuteAction::CreateLambda([this] { PromoteToInput(); }),
+					FCanExecuteAction::CreateLambda([this]() { return CanPromoteToInput(); }));
+
+				GraphEditorCommands->MapAction(FEditorCommands::Get().PromoteToOutput,
+					FExecuteAction::CreateLambda([this] { PromoteToOutput(); }),
+					FCanExecuteAction::CreateLambda([this]() { return CanPromoteToOutput(); }));
+
+				GraphEditorCommands->MapAction(FEditorCommands::Get().PromoteToVariable,
+					FExecuteAction::CreateLambda([this] { PromoteToVariable(); }),
+					FCanExecuteAction::CreateLambda([this]() { return CanPromoteToVariable(); }));
+
+				GraphEditorCommands->MapAction(FEditorCommands::Get().PromoteToDeferredVariable,
+					FExecuteAction::CreateLambda([this] { PromoteToDeferredVariable(); }),
+					FCanExecuteAction::CreateLambda([this]() { return CanPromoteToDeferredVariable(); }));
+
 				GraphEditorCommands->MapAction(FGraphEditorCommands::Get().HideNoConnectionPins,
 					FExecuteAction::CreateSP(this, &FEditor::HideUnconnectedPins));				
 				
@@ -4599,7 +4615,139 @@ namespace Metasound
 				ParentPtr->bRefreshGraph = true;
 			}
 		}
+
+		bool FEditor::CanPromoteToInput()
+		{
+			if (MetasoundGraphEditor.IsValid())
+			{
+				UMetasoundEditorGraph& Graph = GetMetaSoundGraphChecked();
+
+				UEdGraphPin* TargetPin = MetasoundGraphEditor->GetGraphPinForMenu();
+				check(TargetPin);
+
+				if (TargetPin->Direction == EGPD_Input)
+				{
+					return true;
+				}
+			}
+			return false;
+		}
+
+		void FEditor::PromoteToInput()
+		{
+			using namespace Metasound::Frontend;
+
+			if (MetasoundGraphEditor.IsValid())
+			{
+				UMetasoundEditorGraph& Graph = GetMetaSoundGraphChecked();
+
+				UEdGraphPin* TargetPin = MetasoundGraphEditor->GetGraphPinForMenu();
+				check(TargetPin);
+
+				UEdGraphNode* OwningNode = TargetPin->GetOwningNode();
+				FVector2D Location = FVector2D(OwningNode->NodePosX, OwningNode->NodePosY);
+				Metasound::SchemaUtils::PromoteToInput(&Graph, TargetPin, Location - DisplayStyle::NodeLayout::DefaultOffsetX, false);
+			}
+		}
+
+		bool FEditor::CanPromoteToOutput()
+		{
+			if (MetasoundGraphEditor.IsValid())
+			{
+				UMetasoundEditorGraph& Graph = GetMetaSoundGraphChecked();
+
+				UEdGraphPin* TargetPin = MetasoundGraphEditor->GetGraphPinForMenu();
+				check(TargetPin);
+
+				if (TargetPin->Direction == EGPD_Output)
+				{
+					return true;
+				}
+			}
+			return false;
+		}
+
+		void FEditor::PromoteToOutput()
+		{
+			using namespace Metasound::Frontend;
+
+			if (MetasoundGraphEditor.IsValid())
+			{
+				UMetasoundEditorGraph& Graph = GetMetaSoundGraphChecked();
+
+				UEdGraphPin* TargetPin = MetasoundGraphEditor->GetGraphPinForMenu();
+				check(TargetPin);
+
+				UEdGraphNode* OwningNode = TargetPin->GetOwningNode();
+				FVector2D Location = FVector2D(OwningNode->NodePosX, OwningNode->NodePosY);
+				Metasound::SchemaUtils::PromoteToOutput(&Graph, TargetPin, Location + DisplayStyle::NodeLayout::DefaultOffsetX * 2.f, false);
+			}
+		}
+
+		bool FEditor::CanPromoteToVariable()
+		{
+			return true;
+		}
+
+		void FEditor::PromoteToVariable()
+		{
+			using namespace Metasound::Frontend;
+
+			if (MetasoundGraphEditor.IsValid())
+			{
+				UMetasoundEditorGraph& Graph = GetMetaSoundGraphChecked();
+
+				UEdGraphPin* TargetPin = MetasoundGraphEditor->GetGraphPinForMenu();
+				check(TargetPin);
+
+				UEdGraphNode* OwningNode = TargetPin->GetOwningNode();
+				FVector2D Location = FVector2D(OwningNode->NodePosX, OwningNode->NodePosY);
+				if (TargetPin->Direction == EGPD_Input)
+				{
+					Metasound::SchemaUtils::PromoteToVariable(&Graph, TargetPin, Location - DisplayStyle::NodeLayout::DefaultOffsetX, false);
+				}
+				else
+				{
+					Metasound::SchemaUtils::PromoteToMutatorVariable(&Graph, TargetPin, Location + DisplayStyle::NodeLayout::DefaultOffsetX * 2.f, false);
+				}
+			}
+		}
+
+		bool FEditor::CanPromoteToDeferredVariable()
+		{
+			if (MetasoundGraphEditor.IsValid())
+			{
+				UMetasoundEditorGraph& Graph = GetMetaSoundGraphChecked();
+
+				UEdGraphPin* TargetPin = MetasoundGraphEditor->GetGraphPinForMenu();
+				check(TargetPin);
+
+				if (TargetPin->Direction == EGPD_Input)
+				{
+					return true;
+				}
+			}
+			return false;
+		}
+
+		void FEditor::PromoteToDeferredVariable()
+		{
+			using namespace Metasound::Frontend;
+
+			if (MetasoundGraphEditor.IsValid())
+			{
+				UMetasoundEditorGraph& Graph = GetMetaSoundGraphChecked();
+
+				UEdGraphPin* TargetPin = MetasoundGraphEditor->GetGraphPinForMenu();
+				check(TargetPin);
+
+				UEdGraphNode* OwningNode = TargetPin->GetOwningNode();
+				FVector2D Location = FVector2D(OwningNode->NodePosX, OwningNode->NodePosY);
+				Metasound::SchemaUtils::PromoteToDeferredVariable(&Graph, TargetPin, Location - DisplayStyle::NodeLayout::DefaultOffsetX, false);
+			}
+		}
 	}
 }
+
 #undef LOCTEXT_NAMESPACE
 
diff --git a/Engine/Plugins/Runtime/Metasound/Source/MetasoundEditor/Private/MetasoundEditor.h b/Engine/Plugins/Runtime/Metasound/Source/MetasoundEditor/Private/MetasoundEditor.h
--- a/Engine/Plugins/Runtime/Metasound/Source/MetasoundEditor/Private/MetasoundEditor.h
+++ b/Engine/Plugins/Runtime/Metasound/Source/MetasoundEditor/Private/MetasoundEditor.h
@@ -369,6 +369,30 @@ namespace Metasound
 			/**Show pins without connection*/
 			void ShowUnconnectedPins();
 
+			/** Checks if pin can be promoted */
+			bool CanPromoteToInput();
+			
+			/** Promotes pin to graph input */
+			void PromoteToInput();
+
+			/** Checks if pin can be promoted */
+			bool CanPromoteToOutput();
+
+			/** Promotes pin to graph output */
+			void PromoteToOutput();
+
+			/** Checks if pin can be promoted */
+			bool CanPromoteToVariable();
+
+			/** Promotes pin to graph variable */
+			void PromoteToVariable();
+
+			/** Checks if pin can be promoted */
+			bool CanPromoteToDeferredVariable();
+
+			/** Promotes pin to graph deferred variable */
+			void PromoteToDeferredVariable();
+
 		private:
 			/** Forces refresh of pages view. */
 			void RefreshPagesView();
diff --git a/Engine/Plugins/Runtime/Metasound/Source/MetasoundEditor/Private/MetasoundEditorCommands.cpp b/Engine/Plugins/Runtime/Metasound/Source/MetasoundEditor/Private/MetasoundEditorCommands.cpp
--- a/Engine/Plugins/Runtime/Metasound/Source/MetasoundEditor/Private/MetasoundEditorCommands.cpp
+++ b/Engine/Plugins/Runtime/Metasound/Source/MetasoundEditor/Private/MetasoundEditorCommands.cpp
@@ -27,6 +27,11 @@ namespace Metasound
 			UI_COMMAND(ConvertFromPreset, "Convert From Preset", "Converts this preset to a fully accessible MetaSound.", EUserInterfaceActionType::Button, FInputChord());
 			UI_COMMAND(Delete, "Delete Selected", "Delete selected items.", EUserInterfaceActionType::None, FInputChord(EKeys::Delete));
 			UI_COMMAND(FindInMetaSound, "Find in MetaSound", "Find a node or pin within this MetaSound.", EUserInterfaceActionType::Button, FInputChord(EModifierKey::Control, EKeys::F));
+		
+			UI_COMMAND(PromoteToInput, "Promote To Graph Input", "Promotes pin to graph input.", EUserInterfaceActionType::Button, FInputChord());
+			UI_COMMAND(PromoteToOutput, "Promote To Graph Output", "Promotes pin to graph output.", EUserInterfaceActionType::Button, FInputChord());
+			UI_COMMAND(PromoteToVariable, "Promote To Graph Variable", "Promotes pin to graph variable.", EUserInterfaceActionType::Button, FInputChord());
+			UI_COMMAND(PromoteToDeferredVariable, "Promote To Deferred Graph Variable", "Promotes pin to deferred graph variable.", EUserInterfaceActionType::Button, FInputChord());
 		}
 	} // namespace Editor
 } // namespace Metasound
diff --git a/Engine/Plugins/Runtime/Metasound/Source/MetasoundEditor/Private/MetasoundEditorGraphSchema.cpp b/Engine/Plugins/Runtime/Metasound/Source/MetasoundEditor/Private/MetasoundEditorGraphSchema.cpp
--- a/Engine/Plugins/Runtime/Metasound/Source/MetasoundEditor/Private/MetasoundEditorGraphSchema.cpp
+++ b/Engine/Plugins/Runtime/Metasound/Source/MetasoundEditor/Private/MetasoundEditorGraphSchema.cpp
@@ -392,6 +392,211 @@ PRAGMA_ENABLE_DEPRECATION_WARNINGS
 			}
 		} // namespace SchemaPrivate
 	} // namespace Editor
+
+	namespace SchemaUtils
+	{
+		UEdGraphNode* PromoteToInput(UEdGraph* ParentGraph, UEdGraphPin* FromPin, const FVector2D InLocation, bool bSelectNewNode)
+		{
+			using namespace Metasound::Editor;
+			using namespace Metasound::Frontend;
+
+			UMetasoundEditorGraph* MetasoundGraph = CastChecked<UMetasoundEditorGraph>(ParentGraph);
+			UObject& ParentMetasound = MetasoundGraph->GetMetasoundChecked();
+			FMetaSoundFrontendDocumentBuilder& Builder = IDocumentBuilderRegistry::GetChecked().FindOrBeginBuilding(&ParentMetasound);
+			FMetasoundFrontendVertexHandle InputVertexHandle = FGraphBuilder::GetPinVertexHandle(Builder, FromPin);
+			if (!ensure(InputVertexHandle.IsSet()))
+			{
+				return nullptr;
+			}
+
+			const FMetasoundFrontendVertex* InputVertex = Builder.FindNodeInput(InputVertexHandle.NodeID, InputVertexHandle.VertexID);;
+			if (!ensure(InputVertex))
+			{
+				return nullptr;
+			}
+
+			const FScopedTransaction Transaction(LOCTEXT("PromoteNodeInputToGraphInput", "Promote MetaSound Node Input to Graph Input"));
+			ParentMetasound.Modify();
+			MetasoundGraph->Modify();
+
+			FMetasoundFrontendLiteral DefaultValue;
+			FGraphBuilder::GetPinLiteral(*FromPin, DefaultValue);
+			const FCreateNodeVertexParams VertexParams =
+			{
+				InputVertex->TypeName,
+				Builder.GetNodeInputAccessType(InputVertexHandle.NodeID, InputVertex->VertexID)
+			};
+			FMetasoundFrontendClassInput ClassInput = FGraphBuilder::CreateUniqueClassInput(ParentMetasound, VertexParams, &DefaultValue, &InputVertex->Name);
+			if (const FMetasoundFrontendNode* NewNode = Builder.AddGraphInput(ClassInput))
+			{
+				UMetasoundEditorGraphInput* Input = MetasoundGraph->FindOrAddInput(NewNode->GetID());
+				if (ensure(Input))
+				{
+					if (const FMetasoundFrontendNode* NewTemplateNode = FInputNodeTemplate::CreateNode(Builder, ClassInput.Name))
+					{
+						if (UMetasoundEditorGraphNode* NewGraphNode = FGraphBuilder::AddInputNode(ParentMetasound, NewTemplateNode->GetID()))
+						{
+							NewGraphNode->UpdateFrontendNodeLocation(InLocation);
+							NewGraphNode->SyncLocationFromFrontendNode();
+							UEdGraphNode* EdGraphNode = CastChecked<UEdGraphNode>(NewGraphNode);
+
+							if (ensure(SchemaPrivate::TryConnectNewNodeToMatchingDataTypePin(*EdGraphNode, FromPin)))
+							{
+								FGraphBuilder::RegisterGraphWithFrontend(ParentMetasound);
+								if (bSelectNewNode)
+								{
+									SchemaPrivate::SelectNodeInEditorForRename(*MetasoundGraph, *NewGraphNode);
+								}
+								else
+								{
+									SchemaPrivate::SelectNodeInEditor(*MetasoundGraph, *NewGraphNode);
+								}
+
+								return EdGraphNode;
+							}
+						}
+					}
+				}
+			}
+
+			return nullptr;
+		}
+	
+		UEdGraphNode* PromoteToOutput(UEdGraph* ParentGraph, UEdGraphPin* FromPin, const FVector2D InLocation, bool bSelectNewNode)
+		{
+			using namespace Metasound::Editor;
+			using namespace Metasound::Frontend;
+
+			UMetasoundEditorGraph* MetasoundGraph = CastChecked<UMetasoundEditorGraph>(ParentGraph);
+			UObject& ParentMetasound = MetasoundGraph->GetMetasoundChecked();
+			FMetaSoundFrontendDocumentBuilder& Builder = IDocumentBuilderRegistry::GetChecked().FindOrBeginBuilding(&ParentMetasound);
+			FMetasoundFrontendVertexHandle OutputVertexHandle = FGraphBuilder::GetPinVertexHandle(Builder, FromPin);
+			if (!ensure(OutputVertexHandle.IsSet()))
+			{
+				return nullptr;
+			}
+
+			const FMetasoundFrontendVertex* OutputVertex = Builder.FindNodeOutput(OutputVertexHandle.NodeID, OutputVertexHandle.VertexID);;
+			if (!ensure(OutputVertex))
+			{
+				return nullptr;
+			}
+
+			const FScopedTransaction Transaction(LOCTEXT("PromoteNodeOutputToGraphOutput", "Promote MetaSound Node Output to Graph Output"));
+			ParentMetasound.Modify();
+			MetasoundGraph->Modify();
+
+			const FCreateNodeVertexParams VertexParams =
+			{
+				OutputVertex->TypeName,
+				Builder.GetNodeOutputAccessType(OutputVertexHandle.NodeID, OutputVertex->VertexID)
+			};
+			const FMetasoundFrontendClassOutput ClassOutput = FGraphBuilder::CreateUniqueClassOutput(ParentMetasound, VertexParams, &OutputVertex->Name);
+			const FMetasoundFrontendNode* OutputNode = Builder.AddGraphOutput(ClassOutput);
+			if (ensure(OutputNode))
+			{
+				const FGuid OutputNodeID = OutputNode->GetID();
+				UMetasoundEditorGraphOutput* Output = MetasoundGraph->FindOrAddOutput(OutputNodeID);
+				if (ensure(Output))
+				{
+					if (UMetasoundEditorGraphOutputNode* NewGraphNode = FGraphBuilder::AddOutputNode(ParentMetasound, OutputNodeID))
+					{
+						NewGraphNode->UpdateFrontendNodeLocation(InLocation);
+						NewGraphNode->SyncLocationFromFrontendNode();
+
+						UEdGraphNode* EdGraphNode = CastChecked<UEdGraphNode>(NewGraphNode);
+
+						if (ensure(SchemaPrivate::TryConnectNewNodeToMatchingDataTypePin(*EdGraphNode, FromPin)))
+						{
+							FGraphBuilder::RegisterGraphWithFrontend(ParentMetasound);
+							if (bSelectNewNode)
+							{
+								SchemaPrivate::SelectNodeInEditorForRename(*MetasoundGraph, *NewGraphNode);
+							}
+							else
+							{
+								SchemaPrivate::SelectNodeInEditor(*MetasoundGraph, *NewGraphNode);
+							}
+
+							return EdGraphNode;
+						}
+					}
+				}
+			}
+
+			return nullptr;
+		}
+
+		UEdGraphNode* PromoteToVariable(UEdGraph* ParentGraph, UEdGraphPin* FromPin, const FVector2D InLocation, bool bSelectNewNode)
+		{
+			using namespace Metasound;
+			using namespace Metasound::Editor;
+			using namespace Metasound::Frontend;
+
+			FConstInputHandle InputHandle = FGraphBuilder::GetConstInputHandleFromPin(FromPin);
+			if (!ensure(InputHandle->IsValid()))
+			{
+				return nullptr;
+			}
+
+			const FName NodeName = InputHandle->GetName();
+			const FName DataType = InputHandle->GetDataType();
+			FMetasoundFrontendClass VariableClass;
+			if (ensure(IDataTypeRegistry::Get().GetFrontendVariableAccessorClass(DataType, VariableClass)))
+			{
+				return SchemaPrivate::PromoteToVariable(NodeName, *FromPin, DataType, VariableClass, InLocation, bSelectNewNode);
+			}
+
+			return nullptr;
+		}
+
+		UEdGraphNode* PromoteToDeferredVariable(UEdGraph* ParentGraph, UEdGraphPin* FromPin, const FVector2D InLocation, bool bSelectNewNode)
+		{
+			using namespace Metasound;
+			using namespace Metasound::Editor;
+			using namespace Metasound::Frontend;
+
+			FConstInputHandle InputHandle = FGraphBuilder::GetConstInputHandleFromPin(FromPin);
+			if (!ensure(InputHandle->IsValid()))
+			{
+				return nullptr;
+			}
+
+			const FName NodeName = InputHandle->GetName();
+			const FName DataType = InputHandle->GetDataType();
+			FMetasoundFrontendClass VariableClass;
+			if (ensure(IDataTypeRegistry::Get().GetFrontendVariableDeferredAccessorClass(DataType, VariableClass)))
+			{
+				return SchemaPrivate::PromoteToVariable(NodeName, *FromPin, DataType, VariableClass, InLocation, bSelectNewNode);
+			}
+
+			return nullptr;
+		}
+
+
+		UEdGraphNode* PromoteToMutatorVariable(UEdGraph* ParentGraph, UEdGraphPin* FromPin, const FVector2D InLocation, bool bSelectNewNode)
+		{
+			using namespace Metasound;
+			using namespace Metasound::Editor;
+			using namespace Metasound::Frontend;
+
+			FConstOutputHandle OutputHandle = FGraphBuilder::GetConstOutputHandleFromPin(FromPin);
+			if (!ensure(OutputHandle->IsValid()))
+			{
+				return nullptr;
+			}
+
+			const FName NodeName = OutputHandle->GetName();
+			const FName DataType = OutputHandle->GetDataType();
+			FMetasoundFrontendClass VariableClass;
+			if (ensure(IDataTypeRegistry::Get().GetFrontendVariableMutatorClass(DataType, VariableClass)))
+			{
+				return SchemaPrivate::PromoteToVariable(NodeName, *FromPin, DataType, VariableClass, InLocation, bSelectNewNode);
+			}
+
+			return nullptr;
+		}
+	} // namespace SchemaUtils
 } // namespace Metasound
 
 UEdGraphNode* FMetasoundGraphSchemaAction_NodeWithMultipleOutputs::PerformAction(class UEdGraph* ParentGraph, TArray<UEdGraphPin*>& FromPins, const FVector2D Location, bool bSelectNewNode/* = true*/)
@@ -557,61 +762,8 @@ FMetasoundGraphSchemaAction_PromoteToInput::FMetasoundGraphSchemaAction_PromoteT
 
 UEdGraphNode* FMetasoundGraphSchemaAction_PromoteToInput::PerformAction(UEdGraph* ParentGraph, UEdGraphPin* FromPin, const FVector2D InLocation, bool bSelectNewNode /* = true */)
 {
-	using namespace Metasound::Editor;
-	using namespace Metasound::Frontend;
-
-	UMetasoundEditorGraph* MetasoundGraph = CastChecked<UMetasoundEditorGraph>(ParentGraph);
-	UObject& ParentMetasound = MetasoundGraph->GetMetasoundChecked();
-	FMetaSoundFrontendDocumentBuilder& Builder = IDocumentBuilderRegistry::GetChecked().FindOrBeginBuilding(&ParentMetasound);
-	FMetasoundFrontendVertexHandle InputVertexHandle = FGraphBuilder::GetPinVertexHandle(Builder, FromPin);
-	if (!ensure(InputVertexHandle.IsSet()))
-	{
-		return nullptr;
-	}
-
-	const FMetasoundFrontendVertex* InputVertex = Builder.FindNodeInput(InputVertexHandle.NodeID, InputVertexHandle.VertexID);;
-	if (!ensure(InputVertex))
-	{
-		return nullptr;
-	}
-
-	const FScopedTransaction Transaction(LOCTEXT("PromoteNodeInputToGraphInput", "Promote MetaSound Node Input to Graph Input"));
-	ParentMetasound.Modify();
-	MetasoundGraph->Modify();
-
-	FMetasoundFrontendLiteral DefaultValue;
-	FGraphBuilder::GetPinLiteral(*FromPin, DefaultValue);
-	const FCreateNodeVertexParams VertexParams =
-	{
-		InputVertex->TypeName,
-		Builder.GetNodeInputAccessType(InputVertexHandle.NodeID, InputVertex->VertexID)
-	};
-	FMetasoundFrontendClassInput ClassInput = FGraphBuilder::CreateUniqueClassInput(ParentMetasound, VertexParams, &DefaultValue, &InputVertex->Name);
-	if (const FMetasoundFrontendNode* NewNode = Builder.AddGraphInput(ClassInput))
-	{
-		UMetasoundEditorGraphInput* Input = MetasoundGraph->FindOrAddInput(NewNode->GetID());
-		if (ensure(Input))
-		{
-			if (const FMetasoundFrontendNode* NewTemplateNode = FInputNodeTemplate::CreateNode(Builder, ClassInput.Name))
-			{
-				if (UMetasoundEditorGraphNode* NewGraphNode = FGraphBuilder::AddInputNode(ParentMetasound, NewTemplateNode->GetID()))
-				{
-					NewGraphNode->UpdateFrontendNodeLocation(InLocation);
-					NewGraphNode->SyncLocationFromFrontendNode();
-					UEdGraphNode* EdGraphNode = CastChecked<UEdGraphNode>(NewGraphNode);
-
-					if (ensure(SchemaPrivate::TryConnectNewNodeToMatchingDataTypePin(*EdGraphNode, FromPin)))
-					{
-						FGraphBuilder::RegisterGraphWithFrontend(ParentMetasound);
-						SchemaPrivate::SelectNodeInEditorForRename(*MetasoundGraph, *NewGraphNode);
-						return EdGraphNode;
-					}
-				}
-			}
-		}
-	}
-
-	return nullptr;
+	using namespace Metasound;
+	return SchemaUtils::PromoteToInput(ParentGraph, FromPin, InLocation, bSelectNewNode);
 }
 
 FMetasoundGraphSchemaAction_PromoteToVariable_AccessorNode::FMetasoundGraphSchemaAction_PromoteToVariable_AccessorNode()
@@ -626,24 +778,7 @@ FMetasoundGraphSchemaAction_PromoteToVariable_AccessorNode::FMetasoundGraphSchem
 UEdGraphNode* FMetasoundGraphSchemaAction_PromoteToVariable_AccessorNode::PerformAction(UEdGraph* ParentGraph, UEdGraphPin* FromPin, const FVector2D InLocation, bool bSelectNewNode /* = true */)
 {
 	using namespace Metasound;
-	using namespace Metasound::Editor;
-	using namespace Metasound::Frontend;
-
-	FConstInputHandle InputHandle = FGraphBuilder::GetConstInputHandleFromPin(FromPin);
-	if (!ensure(InputHandle->IsValid()))
-	{
-		return nullptr;
-	}
-
-	const FName NodeName = InputHandle->GetName();
-	const FName DataType = InputHandle->GetDataType();
-	FMetasoundFrontendClass VariableClass;
-	if (ensure(IDataTypeRegistry::Get().GetFrontendVariableAccessorClass(DataType, VariableClass)))
-	{
-		return SchemaPrivate::PromoteToVariable(NodeName, *FromPin, DataType, VariableClass, InLocation, bSelectNewNode);
-	}
-
-	return nullptr;
+	return SchemaUtils::PromoteToVariable(ParentGraph, FromPin, InLocation, bSelectNewNode);
 }
 
 FMetasoundGraphSchemaAction_PromoteToVariable_DeferredAccessorNode::FMetasoundGraphSchemaAction_PromoteToVariable_DeferredAccessorNode()
@@ -658,24 +793,7 @@ FMetasoundGraphSchemaAction_PromoteToVariable_DeferredAccessorNode::FMetasoundGr
 UEdGraphNode* FMetasoundGraphSchemaAction_PromoteToVariable_DeferredAccessorNode::PerformAction(UEdGraph* ParentGraph, UEdGraphPin* FromPin, const FVector2D InLocation, bool bSelectNewNode /* = true */)
 {
 	using namespace Metasound;
-	using namespace Metasound::Editor;
-	using namespace Metasound::Frontend;
-
-	FConstInputHandle InputHandle = FGraphBuilder::GetConstInputHandleFromPin(FromPin);
-	if (!ensure(InputHandle->IsValid()))
-	{
-		return nullptr;
-	}
-
-	const FName NodeName = InputHandle->GetName();
-	const FName DataType = InputHandle->GetDataType();
-	FMetasoundFrontendClass VariableClass;
-	if (ensure(IDataTypeRegistry::Get().GetFrontendVariableDeferredAccessorClass(DataType, VariableClass)))
-	{
-		return SchemaPrivate::PromoteToVariable(NodeName, *FromPin, DataType, VariableClass, InLocation, bSelectNewNode);
-	}
-
-	return nullptr;
+	return SchemaUtils::PromoteToDeferredVariable(ParentGraph, FromPin, InLocation, bSelectNewNode);
 }
 
 FMetasoundGraphSchemaAction_PromoteToVariable_MutatorNode::FMetasoundGraphSchemaAction_PromoteToVariable_MutatorNode()
@@ -690,24 +808,7 @@ FMetasoundGraphSchemaAction_PromoteToVariable_MutatorNode::FMetasoundGraphSchema
 UEdGraphNode* FMetasoundGraphSchemaAction_PromoteToVariable_MutatorNode::PerformAction(UEdGraph* ParentGraph, UEdGraphPin* FromPin, const FVector2D InLocation, bool bSelectNewNode /* = true */)
 {
 	using namespace Metasound;
-	using namespace Metasound::Editor;
-	using namespace Metasound::Frontend;
-
-	FConstOutputHandle OutputHandle = FGraphBuilder::GetConstOutputHandleFromPin(FromPin);
-	if (!ensure(OutputHandle->IsValid()))
-	{
-		return nullptr;
-	}
-
-	const FName NodeName = OutputHandle->GetName();
-	const FName DataType = OutputHandle->GetDataType();
-	FMetasoundFrontendClass VariableClass;
-	if (ensure(IDataTypeRegistry::Get().GetFrontendVariableMutatorClass(DataType, VariableClass)))
-	{
-		return SchemaPrivate::PromoteToVariable(NodeName, *FromPin, DataType, VariableClass, InLocation, bSelectNewNode);
-	}
-
-	return nullptr;
+	return SchemaUtils::PromoteToMutatorVariable(ParentGraph, FromPin, InLocation, bSelectNewNode);
 }
 
 FMetasoundGraphSchemaAction_NewOutput::FMetasoundGraphSchemaAction_NewOutput(FText InNodeCategory, FText InDisplayName, FGuid InOutputNodeID, FText InToolTip, Metasound::Editor::EPrimaryContextGroup InGroup)
@@ -772,59 +873,8 @@ FMetasoundGraphSchemaAction_PromoteToOutput::FMetasoundGraphSchemaAction_Promote
 
 UEdGraphNode* FMetasoundGraphSchemaAction_PromoteToOutput::PerformAction(UEdGraph* ParentGraph, UEdGraphPin* FromPin, const FVector2D InLocation, bool bSelectNewNode /* = true */)
 {
-	using namespace Metasound::Editor;
-	using namespace Metasound::Frontend;
-
-	UMetasoundEditorGraph* MetasoundGraph = CastChecked<UMetasoundEditorGraph>(ParentGraph);
-	UObject& ParentMetasound = MetasoundGraph->GetMetasoundChecked();
-	FMetaSoundFrontendDocumentBuilder& Builder = IDocumentBuilderRegistry::GetChecked().FindOrBeginBuilding(&ParentMetasound);
-	FMetasoundFrontendVertexHandle OutputVertexHandle = FGraphBuilder::GetPinVertexHandle(Builder, FromPin);
-	if (!ensure(OutputVertexHandle.IsSet()))
-	{
-		return nullptr;
-	}
-
-	const FMetasoundFrontendVertex* OutputVertex = Builder.FindNodeOutput(OutputVertexHandle.NodeID, OutputVertexHandle.VertexID);;
-	if (!ensure(OutputVertex))
-	{
-		return nullptr;
-	}
-
-	const FScopedTransaction Transaction(LOCTEXT("PromoteNodeOutputToGraphOutput", "Promote MetaSound Node Output to Graph Output"));
-	ParentMetasound.Modify();
-	MetasoundGraph->Modify();
-
-	const FCreateNodeVertexParams VertexParams =
-	{
-		OutputVertex->TypeName,
-		Builder.GetNodeOutputAccessType(OutputVertexHandle.NodeID, OutputVertex->VertexID)
-	};
-	const FMetasoundFrontendClassOutput ClassOutput = FGraphBuilder::CreateUniqueClassOutput(ParentMetasound, VertexParams, &OutputVertex->Name);
-	const FMetasoundFrontendNode* OutputNode = Builder.AddGraphOutput(ClassOutput);
-	if (ensure(OutputNode))
-	{
-		const FGuid OutputNodeID = OutputNode->GetID();
-		UMetasoundEditorGraphOutput* Output = MetasoundGraph->FindOrAddOutput(OutputNodeID);
-		if (ensure(Output))
-		{
-			if (UMetasoundEditorGraphOutputNode* NewGraphNode = FGraphBuilder::AddOutputNode(ParentMetasound, OutputNodeID))
-			{ 
-				NewGraphNode->UpdateFrontendNodeLocation(InLocation);
-				NewGraphNode->SyncLocationFromFrontendNode();
-
-				UEdGraphNode* EdGraphNode = CastChecked<UEdGraphNode>(NewGraphNode);
-
-				if (ensure(SchemaPrivate::TryConnectNewNodeToMatchingDataTypePin(*EdGraphNode, FromPin)))
-				{
-					FGraphBuilder::RegisterGraphWithFrontend(ParentMetasound);
-					SchemaPrivate::SelectNodeInEditorForRename(*MetasoundGraph, *NewGraphNode);
-					return EdGraphNode;
-				}
-			}
-		}
-	}
-
-	return nullptr;
+	using namespace Metasound;
+	return SchemaUtils::PromoteToOutput(ParentGraph, FromPin, InLocation, bSelectNewNode);
 }
 
 FMetasoundGraphSchemaAction_NewVariableNode::FMetasoundGraphSchemaAction_NewVariableNode(FText InNodeCategory, FText InDisplayName, FGuid InVariableID, FText InToolTip)
@@ -1370,6 +1420,22 @@ void UMetasoundEditorGraphSchema::GetContextMenuActions(class UToolMenu* Menu, c
 			}
 		}
 	}
+	else if (Context->Pin && Context->Node && Context->Node->IsA<UMetasoundEditorGraphNode>())
+	{
+		FToolMenuSection& Section = Menu->FindOrAddSection("Pin Actions");
+
+		if (Context->Pin->Direction == EGPD_Input)
+		{
+			Section.AddMenuEntry(FEditorCommands::Get().PromoteToInput);
+			Section.AddMenuEntry(FEditorCommands::Get().PromoteToVariable);
+			Section.AddMenuEntry(FEditorCommands::Get().PromoteToDeferredVariable);
+		}
+		else
+		{
+			Section.AddMenuEntry(FEditorCommands::Get().PromoteToOutput);
+			Section.AddMenuEntry(FEditorCommands::Get().PromoteToVariable);
+		}
+	}
 
 	Super::GetContextMenuActions(Menu, Context);
 }
diff --git a/Engine/Plugins/Runtime/Metasound/Source/MetasoundEditor/Public/MetasoundEditorCommands.h b/Engine/Plugins/Runtime/Metasound/Source/MetasoundEditor/Public/MetasoundEditorCommands.h
--- a/Engine/Plugins/Runtime/Metasound/Source/MetasoundEditor/Public/MetasoundEditorCommands.h
+++ b/Engine/Plugins/Runtime/Metasound/Source/MetasoundEditor/Public/MetasoundEditorCommands.h
@@ -61,6 +61,18 @@ namespace Metasound
 			// Show the Find in MetaSound window. 
 			TSharedPtr<FUICommandInfo> FindInMetaSound;
 
+			/** Promotes pin to graph input */
+			TSharedPtr<FUICommandInfo> PromoteToInput;
+
+			/** Promotes pin to graph output */
+			TSharedPtr<FUICommandInfo> PromoteToOutput;
+
+			/** Promotes pin to graph variable */
+			TSharedPtr<FUICommandInfo> PromoteToVariable;
+
+			/** Promotes pin to deferred graph variable */
+			TSharedPtr<FUICommandInfo> PromoteToDeferredVariable;
+
 			/** Initialize commands */
 			virtual void RegisterCommands() override;
 		};
diff --git a/Engine/Plugins/Runtime/Metasound/Source/MetasoundEditor/Public/MetasoundEditorGraphSchema.h b/Engine/Plugins/Runtime/Metasound/Source/MetasoundEditor/Public/MetasoundEditorGraphSchema.h
--- a/Engine/Plugins/Runtime/Metasound/Source/MetasoundEditor/Public/MetasoundEditorGraphSchema.h
+++ b/Engine/Plugins/Runtime/Metasound/Source/MetasoundEditor/Public/MetasoundEditorGraphSchema.h
@@ -50,6 +50,15 @@ namespace Metasound
 			FOutputFilterFunction OutputFilterFunction;
 		};
 	} // namespace Editor
+
+	namespace SchemaUtils
+	{
+		UEdGraphNode* PromoteToInput(UEdGraph* ParentGraph, UEdGraphPin* FromPin, const FVector2D InLocation, bool bSelectNewNode);
+		UEdGraphNode* PromoteToOutput(UEdGraph* ParentGraph, UEdGraphPin* FromPin, const FVector2D InLocation, bool bSelectNewNode);
+		UEdGraphNode* PromoteToVariable(UEdGraph* ParentGraph, UEdGraphPin* FromPin, const FVector2D InLocation, bool bSelectNewNode);
+		UEdGraphNode* PromoteToDeferredVariable(UEdGraph* ParentGraph, UEdGraphPin* FromPin, const FVector2D InLocation, bool bSelectNewNode);
+		UEdGraphNode* PromoteToMutatorVariable(UEdGraph* ParentGraph, UEdGraphPin* FromPin, const FVector2D InLocation, bool bSelectNewNode);
+	} // namespace SchemaUtils
 } // namespace Metasound
 
 
