Fix ensures when copying and pasting MetaSound interface input and reroute nodes

diff --git a/Engine/Plugins/Runtime/Metasound/Source/MetasoundEditor/Private/MetasoundEditorDocumentClipboardUtils.cpp b/Engine/Plugins/Runtime/Metasound/Source/MetasoundEditor/Private/MetasoundEditorDocumentClipboardUtils.cpp
--- a/Engine/Plugins/Runtime/Metasound/Source/MetasoundEditor/Private/MetasoundEditorDocumentClipboardUtils.cpp
+++ b/Engine/Plugins/Runtime/Metasound/Source/MetasoundEditor/Private/MetasoundEditorDocumentClipboardUtils.cpp
@@ -83,7 +83,17 @@ namespace Metasound::Editor
 				const FMetasoundFrontendNode* InputTemplateNode = FInputNodeTemplate::CreateNode(Builder, Input->GetMemberName());
 				if (ensure(InputTemplateNode))
 				{
-					InputNode->NodeID = InputTemplateNode->GetID();
+					FGuid TemplateNodeID = InputTemplateNode->GetID();
+					InputNode->NodeID = TemplateNodeID;
+
+					// Remove default node location from input node. 
+					// Correct node location from the ed graph node will be set subsequently in ProcessPastedNodePositions
+					TArray<FGuid> NodeLocationGuids;
+					InputTemplateNode->Style.Display.Locations.GetKeys(NodeLocationGuids);
+					if (!NodeLocationGuids.IsEmpty())
+					{
+						Builder.RemoveNodeLocation(TemplateNodeID);
+					}
 				}
 			}
 			else
@@ -218,6 +228,7 @@ namespace Metasound::Editor
 
 	void FDocumentClipboardUtils::ProcessPastedExternalNodes(FMetasoundAssetBase& OutAsset, TArray<UMetasoundEditorGraphNode*>& OutPastedNodes, FDocumentPasteNotifications& OutNotifications)
 	{
+		using namespace Engine;
 		using namespace Frontend;
 
 		OutNotifications.bPastedNodesCreateLoop = false;
@@ -268,12 +279,16 @@ namespace Metasound::Editor
 			}
 			else
 			{
-				FMetasoundFrontendClass FrontendClass;
 				if (const INodeTemplate* Template = INodeTemplateRegistry::Get().FindTemplate(Breadcrumb.ClassName))
 				{
-					FNodeHandle NodeHandle = FGraphBuilder::AddExternalNodeHandle(MetaSound, Breadcrumb.ClassName);
-					ExternalNode->NodeID = NodeHandle->GetID();
-					if (!NodeHandle->IsValid())
+					FMetaSoundFrontendDocumentBuilder& Builder = IDocumentBuilderRegistry::GetChecked().FindOrBeginBuilding(OutAsset.GetOwningAsset());
+					const FNodeTemplateGenerateInterfaceParams TemplateParams = Breadcrumb.TemplateParams.IsSet() ? *Breadcrumb.TemplateParams : FNodeTemplateGenerateInterfaceParams();
+					const FMetasoundFrontendNode* TemplateNode = Builder.AddNodeByTemplate(*Template, TemplateParams);
+					if (TemplateNode)
+					{
+						ExternalNode->NodeID = TemplateNode->GetID();
+					}
+					else
 					{
 						OutPastedNodes.RemoveAtSwap(Index, 1, EAllowShrinking::No);
 						Graph.RemoveNode(ExternalNode);
diff --git a/Engine/Plugins/Runtime/Metasound/Source/MetasoundEditor/Private/MetasoundEditorGraphNode.cpp b/Engine/Plugins/Runtime/Metasound/Source/MetasoundEditor/Private/MetasoundEditorGraphNode.cpp
--- a/Engine/Plugins/Runtime/Metasound/Source/MetasoundEditor/Private/MetasoundEditorGraphNode.cpp
+++ b/Engine/Plugins/Runtime/Metasound/Source/MetasoundEditor/Private/MetasoundEditorGraphNode.cpp
@@ -1082,10 +1082,41 @@ bool UMetasoundEditorGraphExternalNode::CanAutoUpdate() const
 void UMetasoundEditorGraphExternalNode::CacheBreadcrumb()
 {
 	using namespace Metasound::Frontend;
+	using namespace Metasound::Editor;
 	FConstNodeHandle NodeHandle = GetConstNodeHandle();
 	const FNodeRegistryKey RegistryKey(NodeHandle->GetClassMetadata());
 	Breadcrumb.bIsClassNative = FMetasoundFrontendRegistryContainer::Get()->IsNodeNative(RegistryKey);
 	Breadcrumb.ClassName = NodeHandle->GetClassMetadata().GetClassName();
+
+	// Cache template node generation parameters
+	if (const INodeTemplate* Template = INodeTemplateRegistry::Get().FindTemplate(Breadcrumb.ClassName))
+	{
+		if (!Breadcrumb.TemplateParams.IsSet())
+		{
+			Breadcrumb.TemplateParams = FNodeTemplateGenerateInterfaceParams();
+		}
+		Breadcrumb.TemplateParams->InputsToConnect.Reset();
+		Breadcrumb.TemplateParams->OutputsToConnect.Reset();
+
+		for (UEdGraphPin* Pin : Pins)
+		{
+			if (Pin)
+			{
+				const FName DataType = FGraphBuilder::GetPinDataType(Pin);
+				if (!DataType.IsNone())
+				{
+					if (Pin->Direction == EGPD_Input)
+					{
+						Breadcrumb.TemplateParams->InputsToConnect.Add(DataType);
+					}
+					else
+					{
+						Breadcrumb.TemplateParams->OutputsToConnect.Add(DataType);
+					}
+				}
+			}
+		}
+	}
 }
 
 void UMetasoundEditorGraphExternalNode::CacheTitle()
diff --git a/Engine/Plugins/Runtime/Metasound/Source/MetasoundEditor/Public/MetasoundEditorGraphNode.h b/Engine/Plugins/Runtime/Metasound/Source/MetasoundEditor/Public/MetasoundEditorGraphNode.h
--- a/Engine/Plugins/Runtime/Metasound/Source/MetasoundEditor/Public/MetasoundEditorGraphNode.h
+++ b/Engine/Plugins/Runtime/Metasound/Source/MetasoundEditor/Public/MetasoundEditorGraphNode.h
@@ -7,6 +7,7 @@
 #include "MetasoundFrontendController.h"
 #include "MetasoundFrontendDocument.h"
 #include "MetasoundFrontendLiteral.h"
+#include "MetasoundFrontendNodeTemplateRegistry.h"
 #include "MetasoundFrontendRegistries.h"
 #include "Misc/Guid.h"
 #include "Sound/SoundWave.h"
@@ -48,6 +49,10 @@ struct FMetasoundEditorGraphNodeBreadcrumb
 
 	UPROPERTY()
 	bool bIsClassNative = true;
+
+	// For use with template nodes only
+	UPROPERTY()
+	TOptional<FNodeTemplateGenerateInterfaceParams> TemplateParams;
 };
 
 USTRUCT()
diff --git a/Engine/Plugins/Runtime/Metasound/Source/MetasoundFrontend/Private/MetasoundFrontendDocumentBuilder.cpp b/Engine/Plugins/Runtime/Metasound/Source/MetasoundFrontend/Private/MetasoundFrontendDocumentBuilder.cpp
--- a/Engine/Plugins/Runtime/Metasound/Source/MetasoundFrontend/Private/MetasoundFrontendDocumentBuilder.cpp
+++ b/Engine/Plugins/Runtime/Metasound/Source/MetasoundFrontend/Private/MetasoundFrontendDocumentBuilder.cpp
@@ -1243,7 +1243,7 @@ const FMetasoundFrontendNode* FMetaSoundFrontendDocumentBuilder::AddNodeByClassN
 	return nullptr;
 }
 
-const FMetasoundFrontendNode* FMetaSoundFrontendDocumentBuilder::AddNodeByTemplate(const Metasound::Frontend::INodeTemplate& InTemplate, Metasound::Frontend::FNodeTemplateGenerateInterfaceParams Params, FGuid InNodeID, const FGuid* InPageID)
+const FMetasoundFrontendNode* FMetaSoundFrontendDocumentBuilder::AddNodeByTemplate(const Metasound::Frontend::INodeTemplate& InTemplate, FNodeTemplateGenerateInterfaceParams Params, FGuid InNodeID, const FGuid* InPageID)
 {
 	using namespace Metasound;
 	using namespace Metasound::Frontend;
diff --git a/Engine/Plugins/Runtime/Metasound/Source/MetasoundFrontend/Public/MetasoundFrontendDocumentBuilder.h b/Engine/Plugins/Runtime/Metasound/Source/MetasoundFrontend/Public/MetasoundFrontendDocumentBuilder.h
--- a/Engine/Plugins/Runtime/Metasound/Source/MetasoundFrontend/Public/MetasoundFrontendDocumentBuilder.h
+++ b/Engine/Plugins/Runtime/Metasound/Source/MetasoundFrontend/Public/MetasoundFrontendDocumentBuilder.h
@@ -176,7 +176,8 @@ public:
 	bool AddInterface(FName InterfaceName);
 
 	const FMetasoundFrontendNode* AddNodeByClassName(const FMetasoundFrontendClassName& InClassName, int32 InMajorVersion = 1, FGuid InNodeID = FGuid::NewGuid(), const FGuid* InPageID = nullptr);
-	const FMetasoundFrontendNode* AddNodeByTemplate(const Metasound::Frontend::INodeTemplate& InTemplate, Metasound::Frontend::FNodeTemplateGenerateInterfaceParams Params, FGuid InNodeID = FGuid::NewGuid(), const FGuid* InPageID = nullptr);
+	
+	const FMetasoundFrontendNode* AddNodeByTemplate(const Metasound::Frontend::INodeTemplate& InTemplate, FNodeTemplateGenerateInterfaceParams Params, FGuid InNodeID = FGuid::NewGuid(), const FGuid* InPageID = nullptr);
 
 #if WITH_EDITORONLY_DATA
 	// Adds a graph page to the given builder's document
diff --git a/Engine/Plugins/Runtime/Metasound/Source/MetasoundFrontend/Public/MetasoundFrontendNodeTemplateRegistry.h b/Engine/Plugins/Runtime/Metasound/Source/MetasoundFrontend/Public/MetasoundFrontendNodeTemplateRegistry.h
--- a/Engine/Plugins/Runtime/Metasound/Source/MetasoundFrontend/Public/MetasoundFrontendNodeTemplateRegistry.h
+++ b/Engine/Plugins/Runtime/Metasound/Source/MetasoundFrontend/Public/MetasoundFrontendNodeTemplateRegistry.h
@@ -10,10 +10,26 @@
 #include "UObject/NoExportTypes.h"
 #include "UObject/ScriptInterface.h"
 
+#include "MetasoundFrontendNodeTemplateRegistry.generated.h"
+
 // Forward Declarations
 class IMetaSoundDocumentInterface;
 struct FMetaSoundFrontendDocumentBuilder;
 
+USTRUCT()
+struct METASOUNDFRONTEND_API FNodeTemplateGenerateInterfaceParams
+{
+	GENERATED_BODY()
+
+	// TODO: Currently just DataType FName. Subsequent change will replace this with a VertexHandle
+	// and will add builder reference to once builder supports template nodes and controllers
+	// are no longer used to add template nodes from editor code.
+	UPROPERTY()
+	TArray<FName> InputsToConnect;
+	
+	UPROPERTY()
+	TArray<FName> OutputsToConnect;
+};
 
 namespace Metasound::Frontend
 {
@@ -27,15 +43,6 @@ namespace Metasound::Frontend
 	using FConstNodeHandle = TSharedRef<const INodeController>;
 #endif // WITH_EDITOR
 
-	struct METASOUNDFRONTEND_API FNodeTemplateGenerateInterfaceParams
-	{
-		// TODO: Currently just DataType FName. Subsequent change will replace this with a VertexHandle
-		// and will add builder reference to once builder supports template nodes and controllers
-		// are no longer used to add template nodes from editor code.
-		TArray<FName> InputsToConnect;
-		TArray<FName> OutputsToConnect;
-	};
-
 	class METASOUNDFRONTEND_API INodeTemplateTransform
 	{
 	public:
