Focus details panel in case pages (or other tab) is covering when selecting node, member, etc. in MetaSounds asset editor
#rnx

diff --git a/Engine/Plugins/Runtime/Metasound/Source/MetasoundEditor/Private/MetasoundEditor.cpp b/Engine/Plugins/Runtime/Metasound/Source/MetasoundEditor/Private/MetasoundEditor.cpp
--- a/Engine/Plugins/Runtime/Metasound/Source/MetasoundEditor/Private/MetasoundEditor.cpp
+++ b/Engine/Plugins/Runtime/Metasound/Source/MetasoundEditor/Private/MetasoundEditor.cpp
@@ -1243,6 +1243,10 @@ namespace Metasound
 			{
 				MetasoundDetails->SetObjects(SelectedObjects);
 				MetasoundDetails->HideFilterArea(false);
+				if (TabManager.IsValid())
+				{
+					TabManager->TryInvokeTab(TabFactory::Names::Details);
+				}
 			}
 		}
 
