Export MetasoundEngineNodesNames.h constants

#rnx

diff --git a/Engine/Plugins/Runtime/Metasound/Source/MetasoundEngine/Public/MetasoundEngineNodesNames.h b/Engine/Plugins/Runtime/Metasound/Source/MetasoundEngine/Public/MetasoundEngineNodesNames.h
--- a/Engine/Plugins/Runtime/Metasound/Source/MetasoundEngine/Public/MetasoundEngineNodesNames.h
+++ b/Engine/Plugins/Runtime/Metasound/Source/MetasoundEngine/Public/MetasoundEngineNodesNames.h
@@ -8,12 +8,12 @@ namespace Metasound
 {
 	namespace EngineNodes
 	{
-		extern const FName Namespace;
+		METASOUNDENGINE_API extern const FName Namespace;
 
-		extern const FName MonoVariant;
-		extern const FName StereoVariant;
-		extern const FName QuadVariant;
-		extern const FName FiveDotOneVariant;
-		extern const FName SevenDotOneVariant;
+		METASOUNDENGINE_API extern const FName MonoVariant;
+		METASOUNDENGINE_API extern const FName StereoVariant;
+		METASOUNDENGINE_API extern const FName QuadVariant;
+		METASOUNDENGINE_API extern const FName FiveDotOneVariant;
+		METASOUNDENGINE_API extern const FName SevenDotOneVariant;
 	}
 }
