[MetaSounds][Widgets] - Default Horizontal Widget not rendering correctly
#rnx

diff --git a/Engine/Plugins/Runtime/Metasound/Source/MetasoundEditor/Private/SMetasoundGraphNode.cpp b/Engine/Plugins/Runtime/Metasound/Source/MetasoundEditor/Private/SMetasoundGraphNode.cpp
--- a/Engine/Plugins/Runtime/Metasound/Source/MetasoundEditor/Private/SMetasoundGraphNode.cpp
+++ b/Engine/Plugins/Runtime/Metasound/Source/MetasoundEditor/Private/SMetasoundGraphNode.cpp
@@ -914,7 +914,7 @@ namespace Metasound
 
 									if (!bUseAudioMaterialWidgets)
 									{
-										InputWidget->SetDesiredSizeOverride(SliderDesiredSizeVertical);
+										InputWidget->SetDesiredSizeOverride(FVector2D(SliderDesiredSizeVertical.Y, SliderDesiredSizeVertical.X));
 									}
 								}
 
