Fix for missing visual MetaSound editor template input nodes in corrupted document migration (submit on behalf of rob.gay)

#rnx

diff --git a/Engine/Plugins/Runtime/Metasound/Source/MetasoundEditor/Private/MetasoundEditor.cpp b/Engine/Plugins/Runtime/Metasound/Source/MetasoundEditor/Private/MetasoundEditor.cpp
--- a/Engine/Plugins/Runtime/Metasound/Source/MetasoundEditor/Private/MetasoundEditor.cpp
+++ b/Engine/Plugins/Runtime/Metasound/Source/MetasoundEditor/Private/MetasoundEditor.cpp
@@ -1076,11 +1076,15 @@ namespace Metasound
 
 			if (FMetasoundAssetBase* MetaSoundAsset = IMetasoundUObjectRegistry::Get().GetObjectAsAssetBase(ObjectToEdit))
 			{
-				if (MetaSoundAsset->VersionAsset(Builder->GetBuilder()))
+				FMetaSoundFrontendDocumentBuilder& DocBuilder = Builder->GetBuilder();
+				if (MetaSoundAsset->VersionAsset(DocBuilder))
 				{
 					MetaSoundAsset->SetVersionedOnLoad();
 				}
 
+				constexpr bool bForceNodeCreation = false;
+				FInputNodeTemplate::GetChecked().Inject(DocBuilder, bForceNodeCreation);
+
 				// Ensures validation is re-run on re-opening of the editor.
 				// This is needed to refresh errors potentially caused by unloading of
 				// references (ex. if a referenced asset is force deleted in the editor).
diff --git a/Engine/Plugins/Runtime/Metasound/Source/MetasoundFrontend/Private/NodeTemplates/MetasoundFrontendNodeTemplateInput.cpp b/Engine/Plugins/Runtime/Metasound/Source/MetasoundFrontend/Private/NodeTemplates/MetasoundFrontendNodeTemplateInput.cpp
--- a/Engine/Plugins/Runtime/Metasound/Source/MetasoundFrontend/Private/NodeTemplates/MetasoundFrontendNodeTemplateInput.cpp
+++ b/Engine/Plugins/Runtime/Metasound/Source/MetasoundFrontend/Private/NodeTemplates/MetasoundFrontendNodeTemplateInput.cpp
@@ -5,6 +5,7 @@
 #include "Algo/AnyOf.h"
 #include "Internationalization/Text.h"
 #include "MetasoundFrontendDataTypeRegistry.h"
+#include "MetasoundFrontendDocument.h"
 #include "MetasoundFrontendDocumentBuilder.h"
 #include "MetasoundFrontendNodeTemplateRegistry.h"
 #include "MetasoundFrontendRegistries.h"
@@ -22,7 +23,12 @@ namespace Metasound::Frontend
 	namespace InputNodeTemplatePrivate
 	{
 		// Creates an input template node, sets node position (should only ever be one in style location) from and connects it to the associated input with the given name.
-		const FMetasoundFrontendNode* InitTemplateNode(const INodeTemplate& InTemplate, FName InputName, FMetaSoundFrontendDocumentBuilder& InOutBuilder, const FGuid* InPageID = nullptr)
+		void InitTemplateNode(
+			const INodeTemplate& InTemplate,
+			FName InputName,
+			FMetaSoundFrontendDocumentBuilder& InOutBuilder,
+			const FMetasoundFrontendVertexHandle& InputNodeVertex,
+			const TArray<FMetasoundFrontendVertexHandle>& ConnectedVertices, const FGuid* InPageID = nullptr)
 		{
 			FMetasoundFrontendEdge NewEdge;
 			FName TypeName;
@@ -45,22 +51,74 @@ namespace Metasound::Frontend
 			}
 
 			FNodeTemplateGenerateInterfaceParams Params { { }, { TypeName } };
-			const FMetasoundFrontendNode* NewNode = InOutBuilder.AddNodeByTemplate(InTemplate, MoveTemp(Params), FGuid::NewGuid(), InPageID);
-			check(NewNode);
-			NewEdge.ToNodeID = NewNode->GetID();
-			NewEdge.ToVertexID = NewNode->Interface.Inputs.Last().VertexID;
+			
+			
+			const FMetasoundFrontendNode* TemplateNode = InOutBuilder.AddNodeByTemplate(InTemplate, MoveTemp(Params), FGuid::NewGuid(), InPageID);
+			check(TemplateNode);
+			NewEdge.ToNodeID = TemplateNode->GetID();
+			NewEdge.ToVertexID = TemplateNode->Interface.Inputs.Last().VertexID;
+
 
 #if WITH_EDITORONLY_DATA
-			for (const TPair<FGuid, FVector2D>& Pair : Locations)
+			if (Locations.IsEmpty())
 			{
-				InOutBuilder.SetNodeLocation(NewEdge.ToNodeID, Pair.Value, nullptr, InPageID);
+				// If connections are present, add a location for safety.  Try adding near existing node.
+				if (!ConnectedVertices.IsEmpty())
+				{
+					UE_LOG(LogMetaSound, Warning, TEXT("Template node being generated for input '%s' had no editor location set.  Procedurally placing near connected node."), *InputName.ToString());
+					FVector2D NewLocation;
+					if (const FMetasoundFrontendNode* Node = InOutBuilder.FindNode(ConnectedVertices.Last().NodeID))
+					{
+						if (!Node->Style.Display.Locations.IsEmpty())
+						{
+							for (const TPair<FGuid, FVector2D>& Pair : Node->Style.Display.Locations)
+							{
+								const FGuid ConnectedVertexID = ConnectedVertices.Last().VertexID;
+								const uint32 Index = Node->Interface.Inputs.IndexOfByPredicate([&ConnectedVertexID](const FMetasoundFrontendVertex& Input)
+								{
+									return Input.VertexID == ConnectedVertexID;
+								});
+								NewLocation = Pair.Value;
+								NewLocation -= DisplayStyle::NodeLayout::DefaultOffsetX;
+								// Offset Y position based on connected input index to avoid overlapping nodes
+								NewLocation += Index * DisplayStyle::NodeLayout::DefaultOffsetY;
+								break;
+							}
+						}
+					}
+					
+					InOutBuilder.SetNodeLocation(NewEdge.ToNodeID, NewLocation, nullptr, InPageID);
+				}
+			}
+			else
+			{
+				for (const TPair<FGuid, FVector2D>& Pair : Locations)
+				{
+					
+					InOutBuilder.SetNodeLocation(NewEdge.ToNodeID, Pair.Value, nullptr, InPageID);
+				}
 			}
 #endif // WITH_EDITORONLY_DATA
 
 			// Add edge between input node and new template node
 			InOutBuilder.AddEdge(MoveTemp(NewEdge), InPageID);
 
-			return NewNode;
+			FMetasoundFrontendEdge EdgeToRemove { InputNodeVertex.NodeID, InputNodeVertex.VertexID };
+			for (const FMetasoundFrontendVertexHandle& ConnectedVertex : ConnectedVertices)
+			{
+				// Swap connections from input node to connected node to now be from template node to connected node
+				EdgeToRemove.ToNodeID = ConnectedVertex.NodeID;
+				EdgeToRemove.ToVertexID = ConnectedVertex.VertexID;
+				InOutBuilder.RemoveEdge(EdgeToRemove);
+
+				InOutBuilder.AddEdge(FMetasoundFrontendEdge
+				{
+					TemplateNode->GetID(),
+					TemplateNode->Interface.Outputs.Last().VertexID,
+					ConnectedVertex.NodeID,
+					ConnectedVertex.VertexID
+				});
+			}
 		};
 	} // namespace InputNodeTemplatePrivate
 
@@ -75,7 +133,35 @@ namespace Metasound::Frontend
 		{
 			const INodeTemplate* ThisTemplate = INodeTemplateRegistry::Get().FindTemplate(ClassName);
 			check(ThisTemplate);
-			return InputNodeTemplatePrivate::InitTemplateNode(*ThisTemplate, InputName, InOutBuilder, InPageID);
+					
+			const FMetasoundFrontendNode* InputNode = InOutBuilder.FindGraphInputNode(InputName);
+			if (!InputNode)
+			{
+				return nullptr;
+			}
+
+			const FGuid& InputNodeOutputVertexID = InputNode->Interface.Outputs.Last().VertexID;
+
+			TArray<const FMetasoundFrontendNode*> ConnectedInputNodes;
+			TArray<const FMetasoundFrontendVertex*> ConnectedInputVertices = InOutBuilder.FindNodeInputsConnectedToNodeOutput(InputNode->GetID(), InputNodeOutputVertexID, &ConnectedInputNodes);
+
+			FMetasoundFrontendVertexHandle InputNodeVertex { InputNode->GetID(), InputNodeOutputVertexID };
+
+			TArray<FMetasoundFrontendVertexHandle> ConnectedVertices;
+			for (int32 Index = 0; Index < ConnectedInputVertices.Num(); ++Index)
+			{
+				// Ignore edges already connected to input template nodes & cache connected vertex pair
+				// as adding a template node in the subsequent step may invalidate these connected node/vertex
+				// pointers.
+				const FMetasoundFrontendVertex* ConnectedVertex = ConnectedInputVertices[Index];
+				const FMetasoundFrontendNode* ConnectedNode = ConnectedInputNodes[Index];
+				const FMetasoundFrontendClass* Class = InOutBuilder.FindDependency(ConnectedNode->ClassID);
+				if (Class->Metadata.GetClassName() != FInputNodeTemplate::ClassName)
+				{
+					ConnectedVertices.Add(FMetasoundFrontendVertexHandle { ConnectedNode->GetID(), ConnectedVertex->VertexID });
+				}
+			}
+			InputNodeTemplatePrivate::InitTemplateNode(*ThisTemplate, InputName, InOutBuilder, InputNodeVertex, ConnectedVertices, InPageID);
 		}
 
 		return nullptr;
@@ -230,7 +316,7 @@ namespace Metasound::Frontend
 			TArray<const FMetasoundFrontendNode*> ConnectedInputNodes;
 			TArray<const FMetasoundFrontendVertex*> ConnectedInputVertices = InOutBuilder.FindNodeInputsConnectedToNodeOutput(InputNode->GetID(), InputNodeOutputVertexID, &ConnectedInputNodes);
 
-			FMetasoundFrontendEdge EdgeToRemove { InputNode->GetID(), InputNodeOutputVertexID };
+			FMetasoundFrontendVertexHandle InputNodeVertex { InputNode->GetID(), InputNodeOutputVertexID };
 
 			bool bHasTemplateConnection = false;
 
@@ -258,30 +344,13 @@ namespace Metasound::Frontend
 				if (bForceNodeCreation && !bHasTemplateConnection)
 				{
 					bInjectedNodes = true;
-					InputNodeTemplatePrivate::InitTemplateNode(*this, Input.Name, InOutBuilder);
+					InputNodeTemplatePrivate::InitTemplateNode(*this, Input.Name, InOutBuilder, InputNodeVertex, ConnectedVertices);
 				}
 			}
 			else
 			{
 				bInjectedNodes = true;
-				const FMetasoundFrontendNode* TemplateNode = InputNodeTemplatePrivate::InitTemplateNode(*this, Input.Name, InOutBuilder);
-				check(TemplateNode);
-
-				for (const FMetasoundFrontendVertexHandle& ConnectedVertex : ConnectedVertices)
-				{
-					// Swap connections from input node to connected node to now be from template node to connected node
-					EdgeToRemove.ToNodeID = ConnectedVertex.NodeID;
-					EdgeToRemove.ToVertexID = ConnectedVertex.VertexID;
-					InOutBuilder.RemoveEdge(EdgeToRemove);
-
-					InOutBuilder.AddEdge(FMetasoundFrontendEdge
-					{
-						TemplateNode->GetID(),
-						TemplateNode->Interface.Outputs.Last().VertexID,
-						ConnectedVertex.NodeID,
-						ConnectedVertex.VertexID
-					});
-				}
+				InputNodeTemplatePrivate::InitTemplateNode(*this, Input.Name, InOutBuilder, InputNodeVertex, ConnectedVertices);
 			}
 		}
 
