Fix bad merge in MetaSounds input template nodes (CL 35379521)

#rnx

diff --git a/Engine/Plugins/Runtime/Metasound/Source/MetasoundFrontend/Private/NodeTemplates/MetasoundFrontendNodeTemplateInput.cpp b/Engine/Plugins/Runtime/Metasound/Source/MetasoundFrontend/Private/NodeTemplates/MetasoundFrontendNodeTemplateInput.cpp
--- a/Engine/Plugins/Runtime/Metasound/Source/MetasoundFrontend/Private/NodeTemplates/MetasoundFrontendNodeTemplateInput.cpp
+++ b/Engine/Plugins/Runtime/Metasound/Source/MetasoundFrontend/Private/NodeTemplates/MetasoundFrontendNodeTemplateInput.cpp
@@ -23,7 +23,7 @@ namespace Metasound::Frontend
 	namespace InputNodeTemplatePrivate
 	{
 		// Creates an input template node, sets node position (should only ever be one in style location) from and connects it to the associated input with the given name.
-		void InitTemplateNode(
+		const FMetasoundFrontendNode* InitTemplateNode(
 			const INodeTemplate& InTemplate,
 			FName InputName,
 			FMetaSoundFrontendDocumentBuilder& InOutBuilder,
@@ -52,13 +52,11 @@ namespace Metasound::Frontend
 
 			FNodeTemplateGenerateInterfaceParams Params { { }, { TypeName } };
 			
-			
 			const FMetasoundFrontendNode* TemplateNode = InOutBuilder.AddNodeByTemplate(InTemplate, MoveTemp(Params), FGuid::NewGuid(), InPageID);
 			check(TemplateNode);
 			NewEdge.ToNodeID = TemplateNode->GetID();
 			NewEdge.ToVertexID = TemplateNode->Interface.Inputs.Last().VertexID;
 
-
 #if WITH_EDITORONLY_DATA
 			if (Locations.IsEmpty())
 			{
@@ -94,7 +92,6 @@ namespace Metasound::Frontend
 			{
 				for (const TPair<FGuid, FVector2D>& Pair : Locations)
 				{
-					
 					InOutBuilder.SetNodeLocation(NewEdge.ToNodeID, Pair.Value, nullptr, InPageID);
 				}
 			}
@@ -119,6 +116,7 @@ namespace Metasound::Frontend
 					ConnectedVertex.VertexID
 				});
 			}
+			return TemplateNode;
 		};
 	} // namespace InputNodeTemplatePrivate
 
@@ -141,27 +139,9 @@ namespace Metasound::Frontend
 			}
 
 			const FGuid& InputNodeOutputVertexID = InputNode->Interface.Outputs.Last().VertexID;
-
-			TArray<const FMetasoundFrontendNode*> ConnectedInputNodes;
-			TArray<const FMetasoundFrontendVertex*> ConnectedInputVertices = InOutBuilder.FindNodeInputsConnectedToNodeOutput(InputNode->GetID(), InputNodeOutputVertexID, &ConnectedInputNodes);
-
 			FMetasoundFrontendVertexHandle InputNodeVertex { InputNode->GetID(), InputNodeOutputVertexID };
 
-			TArray<FMetasoundFrontendVertexHandle> ConnectedVertices;
-			for (int32 Index = 0; Index < ConnectedInputVertices.Num(); ++Index)
-			{
-				// Ignore edges already connected to input template nodes & cache connected vertex pair
-				// as adding a template node in the subsequent step may invalidate these connected node/vertex
-				// pointers.
-				const FMetasoundFrontendVertex* ConnectedVertex = ConnectedInputVertices[Index];
-				const FMetasoundFrontendNode* ConnectedNode = ConnectedInputNodes[Index];
-				const FMetasoundFrontendClass* Class = InOutBuilder.FindDependency(ConnectedNode->ClassID);
-				if (Class->Metadata.GetClassName() != FInputNodeTemplate::ClassName)
-				{
-					ConnectedVertices.Add(FMetasoundFrontendVertexHandle { ConnectedNode->GetID(), ConnectedVertex->VertexID });
-				}
-			}
-			InputNodeTemplatePrivate::InitTemplateNode(*ThisTemplate, InputName, InOutBuilder, InputNodeVertex, ConnectedVertices, InPageID);
+			return InputNodeTemplatePrivate::InitTemplateNode(*ThisTemplate, InputName, InOutBuilder, InputNodeVertex, /*ConnectedVertices*/{}, InPageID);
 		}
 
 		return nullptr;
