UE-170993 - [Workflow Improvement] - Fast way to promote all items in node to inputs

diff --git a/Engine/Plugins/Runtime/Metasound/Source/MetasoundEditor/Private/MetasoundEditor.cpp b/Engine/Plugins/Runtime/Metasound/Source/MetasoundEditor/Private/MetasoundEditor.cpp
--- a/Engine/Plugins/Runtime/Metasound/Source/MetasoundEditor/Private/MetasoundEditor.cpp
+++ b/Engine/Plugins/Runtime/Metasound/Source/MetasoundEditor/Private/MetasoundEditor.cpp
@@ -2725,6 +2725,14 @@ namespace Metasound
 					FExecuteAction::CreateSP(this, &FEditor::AddInput),
 					FCanExecuteAction::CreateSP(this, &FEditor::CanAddInput));
 
+				GraphEditorCommands->MapAction(FEditorCommands::Get().PromoteAllToCommonInputs,
+					FExecuteAction::CreateSP(this, &FEditor::PromoteAllToCommonInputs),
+					FCanExecuteAction::CreateSP(this, &FEditor::CanPromoteAllToCommonInputs));
+
+				GraphEditorCommands->MapAction(FEditorCommands::Get().PromoteAllToInput,
+					FExecuteAction::CreateSP(this, &FEditor::PromoteAllToInputs),
+					FCanExecuteAction::CreateSP(this, &FEditor::CanPromoteAllToInputs));
+
 				// Editing Commands
 				GraphEditorCommands->MapAction(FGenericCommands::Get().SelectAll,
 					FExecuteAction::CreateLambda([this]() { MetasoundGraphEditor->SelectAllNodes(); }));
@@ -4754,8 +4762,207 @@ namespace Metasound
 				Metasound::SchemaUtils::PromoteToDeferredVariable(&Graph, TargetPin, Location - DisplayStyle::NodeLayout::DefaultOffsetX, false);
 			}
 		}
+
+		int32 FEditor::PromotableSelectedNodes()
+		{
+			int32 Counter = 0;
+
+			FGraphPanelSelectionSet SelectedNodes = MetasoundGraphEditor->GetSelectedNodes();
+			for (FGraphPanelSelectionSet::TConstIterator NodeIt(SelectedNodes); NodeIt; ++NodeIt)
+			{
+				UEdGraphNode* Node = Cast<UEdGraphNode>(*NodeIt);
+				for (const UEdGraphPin* Pin : Node->Pins)
+				{
+					if (Pin->Direction == EGPD_Input && !Pin->HasAnyConnections())
+					{
+						Counter += 1;
+						break;
+					}
+				}
+			}
+			return Counter;
+		}
+
+		bool FEditor::CanPromoteAllToInputs()
+		{
+			return PromotableSelectedNodes() > 0;
+		}
+
+		void FEditor::PromoteAllToInputs()
+		{
+			using namespace Frontend;
+
+			UObject& ParentMetasound = *GetMetasoundObject();
+			UMetasoundEditorGraph& MetasoundGraph = GetMetaSoundGraphChecked();
+
+			const FScopedTransaction Transaction(LOCTEXT("PromoteNodeInputsToGraphInputs", "Promote MetaSound Node Inputs to Graph Inputs"));
+			ParentMetasound.Modify();
+			MetasoundGraph.Modify();
+
+			FGraphPanelSelectionSet SelectedNodes = MetasoundGraphEditor->GetSelectedNodes();
+			for (FGraphPanelSelectionSet::TConstIterator NodeIt(SelectedNodes); NodeIt; ++NodeIt)
+			{
+				UMetasoundEditorGraphNode* EdGraphNode = Cast<UMetasoundEditorGraphNode>(*NodeIt);
+				FVector2D NodeOffset = FVector2D(0.f);
+
+				for (UEdGraphPin* Pin : EdGraphNode->Pins)
+				{
+					if (Pin->Direction != EGPD_Input || Pin->HasAnyConnections())
+					{
+						continue;
+					}
+
+					FMetaSoundFrontendDocumentBuilder& DocBuilder = Builder->GetBuilder();
+
+					FMetasoundFrontendVertexHandle InputVertexHandle = FGraphBuilder::GetPinVertexHandle(DocBuilder, Pin);
+					check(InputVertexHandle.IsSet());
+					const FMetasoundFrontendVertex* InputVertex = DocBuilder.FindNodeInput(InputVertexHandle.NodeID, InputVertexHandle.VertexID);
+					check(InputVertex);
+
+					FName Name = FGraphBuilder::GenerateUniqueNameByClassType(ParentMetasound, EMetasoundFrontendClassType::Input, Pin->GetName());
+
+					EMetaSoundBuilderResult Result = EMetaSoundBuilderResult::Failed;
+					FMetasoundFrontendLiteral Literal;
+					FGraphBuilder::GetPinLiteral(*Pin, Literal);
+					FMetaSoundBuilderNodeOutputHandle OutputHandle = Builder->AddGraphInputNode(Name, InputVertex->TypeName, Literal, Result, false);
+					check(Result == EMetaSoundBuilderResult::Succeeded);
+
+					FVector2D Location = FVector2D(EdGraphNode->NodePosX, EdGraphNode->NodePosY);
+					Location -= DisplayStyle::NodeLayout::DefaultOffsetX;
+					Location += NodeOffset;
+					NodeOffset += DisplayStyle::NodeLayout::DefaultOffsetY * 0.5f;
+
+					Builder->SetNodeLocation(OutputHandle.NodeID, Location, Result);
+					check(Result == EMetaSoundBuilderResult::Succeeded);
+
+					if (const FMetasoundFrontendNode* NewTemplateNode = FInputNodeTemplate::CreateNode(DocBuilder, Name))
+					{
+						if (UMetasoundEditorGraphNode* NewGraphNode = FGraphBuilder::AddInputNode(ParentMetasound, NewTemplateNode->GetID()))
+						{
+							FMetaSoundNodeHandle NewNodeHandle(NewGraphNode->GetFrontendNode()->GetID());
+							FName OutputName = NewGraphNode->GetFrontendNode()->Interface.Outputs[0].Name;
+							OutputHandle = Builder->FindNodeOutputByName(NewNodeHandle, OutputName, Result);
+							check(Result == EMetaSoundBuilderResult::Succeeded);
+
+							FMetaSoundNodeHandle SourceNodeHandle(EdGraphNode->GetFrontendNode()->GetID());
+							FMetaSoundBuilderNodeInputHandle InputHandle = Builder->FindNodeInputByName(SourceNodeHandle, InputVertex->Name, Result);
+							check(Result == EMetaSoundBuilderResult::Succeeded);
+
+							Builder->ConnectNodes(OutputHandle, InputHandle, Result);
+							check(Result == EMetaSoundBuilderResult::Succeeded);
+						}
+					}
+				}
+			}
+		}
+
+		bool FEditor::CanPromoteAllToCommonInputs()
+		{
+			return PromotableSelectedNodes() > 1;
+		}
+
+		void FEditor::PromoteAllToCommonInputs()
+		{
+			using namespace Frontend;
+
+			UObject& ParentMetasound = *GetMetasoundObject();
+			UMetasoundEditorGraph& MetasoundGraph = GetMetaSoundGraphChecked();
+
+			const FScopedTransaction Transaction(LOCTEXT("PromoteNodeInputsToCommonGraphInputs", "Promote MetaSound Node Inputs to Shared Graph Inputs"));
+			ParentMetasound.Modify();
+			MetasoundGraph.Modify();
+
+			FMetaSoundFrontendDocumentBuilder& DocBuilder = Builder->GetBuilder();
+
+			// PinsMap.Key { pin name, pin data type }
+			TMap<TPair<FName, FName>, TArray<UEdGraphPin*>> PinsMap;
+			TMap<FGuid, FVector2D> NodeOffsets;
+
+			// Find common pins and save them for processing at later stage
+			FGraphPanelSelectionSet SelectedNodes = MetasoundGraphEditor->GetSelectedNodes();
+			for (FGraphPanelSelectionSet::TConstIterator NodeIt(SelectedNodes); NodeIt; ++NodeIt)
+			{
+				UEdGraphNode* EdGraphNode = Cast<UEdGraphNode>(*NodeIt);
+				for (UEdGraphPin* Pin : EdGraphNode->Pins)
+				{
+					if (Pin->Direction == EGPD_Input && !Pin->HasAnyConnections())
+					{
+						// Get type name from pin
+						FMetasoundFrontendVertexHandle InputVertexHandle = FGraphBuilder::GetPinVertexHandle(DocBuilder, Pin);
+						check(InputVertexHandle.IsSet());
+						const FMetasoundFrontendVertex* InputVertex = DocBuilder.FindNodeInput(InputVertexHandle.NodeID, InputVertexHandle.VertexID);
+						check(InputVertex);
+
+						TPair<FName, FName> Pair = { Pin->GetFName(), InputVertex->TypeName };
+						
+						if (TArray<UEdGraphPin*>* PtrArr = PinsMap.Find(Pair))
+						{
+							PtrArr->Add(Pin);
+						}
+						else
+						{
+							TArray<UEdGraphPin*> Arr;
+							Arr.Add(Pin);
+							PinsMap.Add(Pair, Arr);
+						}
+					}
+
+					NodeOffsets.Add(EdGraphNode->NodeGuid, FVector2D(0.f));
+				}
+			}
+
+			for (const TPair<TPair<FName, FName>, TArray<UEdGraphPin*>>& Pair : PinsMap)
+			{
+				check(Pair.Value.Num() != 0);
+				
+				FName PinName = Pair.Key.Key;
+				FName TypeName = Pair.Key.Value;
+				UEdGraphPin* SourcePin = Pair.Value[0];
+				
+				FName InputName = FGraphBuilder::GenerateUniqueNameByClassType(ParentMetasound, EMetasoundFrontendClassType::Input, PinName.ToString());
+				
+				EMetaSoundBuilderResult Result = EMetaSoundBuilderResult::Failed;
+				FMetasoundFrontendLiteral Literal;
+				FGraphBuilder::GetPinLiteral(*SourcePin, Literal);
+				FMetaSoundBuilderNodeOutputHandle OutputHandle = Builder->AddGraphInputNode(InputName, TypeName, Literal, Result, false);
+				check(Result == EMetaSoundBuilderResult::Succeeded);
+
+				FVector2D* NodeOffset = NodeOffsets.Find(SourcePin->GetOwningNode()->NodeGuid);
+				check(NodeOffset);
+
+				FVector2D Location = FVector2D(SourcePin->GetOwningNode()->NodePosX, SourcePin->GetOwningNode()->NodePosY);
+				Location -= DisplayStyle::NodeLayout::DefaultOffsetX;
+				Location += *NodeOffset;
+				*NodeOffset += DisplayStyle::NodeLayout::DefaultOffsetY * 0.5f;
+				
+				Builder->SetNodeLocation(OutputHandle.NodeID, Location, Result);
+				check(Result == EMetaSoundBuilderResult::Succeeded);
+
+				if (const FMetasoundFrontendNode* NewTemplateNode = FInputNodeTemplate::CreateNode(DocBuilder, InputName))
+				{
+					if (UMetasoundEditorGraphNode* NewGraphNode = FGraphBuilder::AddInputNode(ParentMetasound, NewTemplateNode->GetID()))
+					{
+						FMetaSoundNodeHandle NewNodeHandle(NewGraphNode->GetFrontendNode()->GetID());
+						FName OutputName = NewGraphNode->GetFrontendNode()->Interface.Outputs[0].Name;
+						OutputHandle = Builder->FindNodeOutputByName(NewNodeHandle, OutputName, Result);
+						check(Result == EMetaSoundBuilderResult::Succeeded);
+
+						for (const UEdGraphPin* Pin : Pair.Value)
+						{
+							UMetasoundEditorGraphNode* EdGraphNode = Cast<UMetasoundEditorGraphNode>(Pin->GetOwningNode());
+
+							FMetaSoundNodeHandle SourceNodeHandle(EdGraphNode->GetFrontendNode()->GetID());
+							FMetaSoundBuilderNodeInputHandle InputHandle = Builder->FindNodeInputByName(SourceNodeHandle, PinName, Result);
+							check(Result == EMetaSoundBuilderResult::Succeeded);
+							
+							Builder->ConnectNodes(OutputHandle, InputHandle, Result);
+							check(Result == EMetaSoundBuilderResult::Succeeded);
+						}
+					}
+				}
+			}
+		}
 	}
 }
 
 #undef LOCTEXT_NAMESPACE
-
diff --git a/Engine/Plugins/Runtime/Metasound/Source/MetasoundEditor/Private/MetasoundEditor.h b/Engine/Plugins/Runtime/Metasound/Source/MetasoundEditor/Private/MetasoundEditor.h
--- a/Engine/Plugins/Runtime/Metasound/Source/MetasoundEditor/Private/MetasoundEditor.h
+++ b/Engine/Plugins/Runtime/Metasound/Source/MetasoundEditor/Private/MetasoundEditor.h
@@ -393,7 +393,21 @@ namespace Metasound
 			/** Promotes pin to graph deferred variable */
 			void PromoteToDeferredVariable();
 
+			/** Checks if node's inputs can be promoted */
+			bool CanPromoteAllToInputs();
+
+			/** Promotes node's inputs to unique graph inputs */
+			void PromoteAllToInputs();
+
+			/** Checks if node's inputs can be promoted */
+			bool CanPromoteAllToCommonInputs();
+
+			/** Promotes node's inputs to shared graph inputs */
+			void PromoteAllToCommonInputs();
+
 		private:
+			int32 PromotableSelectedNodes();
+
 			/** Forces refresh of pages view. */
 			void RefreshPagesView();
 
diff --git a/Engine/Plugins/Runtime/Metasound/Source/MetasoundEditor/Private/MetasoundEditorCommands.cpp b/Engine/Plugins/Runtime/Metasound/Source/MetasoundEditor/Private/MetasoundEditorCommands.cpp
--- a/Engine/Plugins/Runtime/Metasound/Source/MetasoundEditor/Private/MetasoundEditorCommands.cpp
+++ b/Engine/Plugins/Runtime/Metasound/Source/MetasoundEditor/Private/MetasoundEditorCommands.cpp
@@ -32,6 +32,9 @@ namespace Metasound
 			UI_COMMAND(PromoteToOutput, "Promote To Graph Output", "Promotes pin to graph output.", EUserInterfaceActionType::Button, FInputChord());
 			UI_COMMAND(PromoteToVariable, "Promote To Graph Variable", "Promotes pin to graph variable.", EUserInterfaceActionType::Button, FInputChord());
 			UI_COMMAND(PromoteToDeferredVariable, "Promote To Deferred Graph Variable", "Promotes pin to deferred graph variable.", EUserInterfaceActionType::Button, FInputChord());
+
+			UI_COMMAND(PromoteAllToInput, "Promote All To Input", "Promotes unconnected node inputs to graph inputs.", EUserInterfaceActionType::Button, FInputChord());
+			UI_COMMAND(PromoteAllToCommonInputs, "Promote All To Common Inputs", "Promotes unconnected node input pins to graph inputs, sharing inputs if possible.", EUserInterfaceActionType::Button, FInputChord());
 		}
 	} // namespace Editor
 } // namespace Metasound
diff --git a/Engine/Plugins/Runtime/Metasound/Source/MetasoundEditor/Private/MetasoundEditorGraphSchema.cpp b/Engine/Plugins/Runtime/Metasound/Source/MetasoundEditor/Private/MetasoundEditorGraphSchema.cpp
--- a/Engine/Plugins/Runtime/Metasound/Source/MetasoundEditor/Private/MetasoundEditorGraphSchema.cpp
+++ b/Engine/Plugins/Runtime/Metasound/Source/MetasoundEditor/Private/MetasoundEditorGraphSchema.cpp
@@ -1387,6 +1387,8 @@ void UMetasoundEditorGraphSchema::GetContextMenuActions(class UToolMenu* Menu, c
 		Section.AddMenuEntry(FGraphEditorCommands::Get().FindReferences, LOCTEXT("FindReferencesInGraph", "Find References In Graph"),
 			LOCTEXT("FindReferencesInGraph_Tooltip", "Find References to the selected Node in the current Graph"),
 			FSlateIcon());
+		Section.AddMenuEntry(FEditorCommands::Get().PromoteAllToInput);
+		Section.AddMenuEntry(FEditorCommands::Get().PromoteAllToCommonInputs);
 
 		FToolMenuSection& OrganizationSection = Menu->FindOrAddSection("MetasoundGraphNodeActionsOrganization", LOCTEXT("NodeActionsOrganizationMenuHeader", "Organization"));
 
diff --git a/Engine/Plugins/Runtime/Metasound/Source/MetasoundEditor/Public/MetasoundEditorCommands.h b/Engine/Plugins/Runtime/Metasound/Source/MetasoundEditor/Public/MetasoundEditorCommands.h
--- a/Engine/Plugins/Runtime/Metasound/Source/MetasoundEditor/Public/MetasoundEditorCommands.h
+++ b/Engine/Plugins/Runtime/Metasound/Source/MetasoundEditor/Public/MetasoundEditorCommands.h
@@ -73,6 +73,12 @@ namespace Metasound
 			/** Promotes pin to deferred graph variable */
 			TSharedPtr<FUICommandInfo> PromoteToDeferredVariable;
 
+			/** Promotes all unconnected node inputs to graph inputs. */
+			TSharedPtr<FUICommandInfo> PromoteAllToInput;
+
+			/** Promotes all unconnected node inputs to shared graph inputs. */
+			TSharedPtr<FUICommandInfo> PromoteAllToCommonInputs;
+
 			/** Initialize commands */
 			virtual void RegisterCommands() override;
 		};
