[MetaSounds][Presets] - Users can realign Preset connections, but alignment is not retained on Play. Disables "Alignment", "Promote To" and "Delete Unused" Actions for presets
#rnx

diff --git a/Engine/Plugins/Runtime/Metasound/Source/MetasoundEditor/Private/MetasoundEditor.cpp b/Engine/Plugins/Runtime/Metasound/Source/MetasoundEditor/Private/MetasoundEditor.cpp
--- a/Engine/Plugins/Runtime/Metasound/Source/MetasoundEditor/Private/MetasoundEditor.cpp
+++ b/Engine/Plugins/Runtime/Metasound/Source/MetasoundEditor/Private/MetasoundEditor.cpp
@@ -4131,15 +4131,18 @@ namespace Metasound
 
 			if (Actions.IsEmpty())//Section is selected
 			{
-				MenuBuilder.BeginSection("GraphActionMenuSectionActions", LOCTEXT("SectionActionsMenuHeader", "Section Actions"));
-				MenuBuilder.AddMenuEntry(
-					LOCTEXT("DeleteAllUnusedInSection", "Delete Unused Members"),
-					LOCTEXT("DeleteAllUnusedInSectionTooltip", "Delete all Unused Members under this Section"),
-					FSlateIcon(),
-					FUIAction(
-						FExecuteAction::CreateSP(this, &FEditor::DeleteAllUnusedInSection),
-						FCanExecuteAction::CreateSP(this, &FEditor::CanDeleteUnusedMembers)));
-				MenuBuilder.EndSection();
+				if (!Builder->IsPreset())
+				{
+					MenuBuilder.BeginSection("GraphActionMenuSectionActions", LOCTEXT("SectionActionsMenuHeader", "Section Actions"));
+					MenuBuilder.AddMenuEntry(
+						LOCTEXT("DeleteAllUnusedInSection", "Delete Unused Members"),
+						LOCTEXT("DeleteAllUnusedInSectionTooltip", "Delete all Unused Members under this Section"),
+						FSlateIcon(),
+						FUIAction(
+							FExecuteAction::CreateSP(this, &FEditor::DeleteAllUnusedInSection),
+							FCanExecuteAction::CreateSP(this, &FEditor::CanDeleteUnusedMembers)));
+					MenuBuilder.EndSection();
+				}
 			}
 			else //Member is selected
 			{
diff --git a/Engine/Plugins/Runtime/Metasound/Source/MetasoundEditor/Private/MetasoundEditorGraphNode.cpp b/Engine/Plugins/Runtime/Metasound/Source/MetasoundEditor/Private/MetasoundEditorGraphNode.cpp
--- a/Engine/Plugins/Runtime/Metasound/Source/MetasoundEditor/Private/MetasoundEditorGraphNode.cpp
+++ b/Engine/Plugins/Runtime/Metasound/Source/MetasoundEditor/Private/MetasoundEditorGraphNode.cpp
@@ -710,27 +710,30 @@ void UMetasoundEditorGraphNode::GetNodeContextMenuActions(UToolMenu* Menu, UGrap
 
 	if (Context->Node)
 	{
+		if (!GetBuilderChecked().IsPreset())
 		{
-			FToolMenuSection& Section = Menu->FindOrAddSection("MetasoundGraphNodeActionsOrganization", LOCTEXT("NodeActionsOrganizationMenuHeader", "Organization"));
-			Section.AddSubMenu("Alignment", LOCTEXT("AlignmentHeader", "Alignment"), FText(), FNewToolMenuDelegate::CreateLambda([](UToolMenu* SubMenu)
 			{
+				FToolMenuSection& Section = Menu->FindOrAddSection("MetasoundGraphNodeActionsOrganization", LOCTEXT("NodeActionsOrganizationMenuHeader", "Organization"));
+				Section.AddSubMenu("Alignment", LOCTEXT("AlignmentHeader", "Alignment"), FText(), FNewToolMenuDelegate::CreateLambda([](UToolMenu* SubMenu)
 				{
-					FToolMenuSection& SubMenuSection = SubMenu->AddSection("EdGraphSchemaAlignment", LOCTEXT("AlignHeader", "Align"));
-					SubMenuSection.AddMenuEntry(FGraphEditorCommands::Get().AlignNodesTop);
-					SubMenuSection.AddMenuEntry(FGraphEditorCommands::Get().AlignNodesMiddle);
-					SubMenuSection.AddMenuEntry(FGraphEditorCommands::Get().AlignNodesBottom);
-					SubMenuSection.AddMenuEntry(FGraphEditorCommands::Get().AlignNodesLeft);
-					SubMenuSection.AddMenuEntry(FGraphEditorCommands::Get().AlignNodesCenter);
-					SubMenuSection.AddMenuEntry(FGraphEditorCommands::Get().AlignNodesRight);
-					SubMenuSection.AddMenuEntry(FGraphEditorCommands::Get().StraightenConnections);
-				}
+					{
+						FToolMenuSection& SubMenuSection = SubMenu->AddSection("EdGraphSchemaAlignment", LOCTEXT("AlignHeader", "Align"));
+						SubMenuSection.AddMenuEntry(FGraphEditorCommands::Get().AlignNodesTop);
+						SubMenuSection.AddMenuEntry(FGraphEditorCommands::Get().AlignNodesMiddle);
+						SubMenuSection.AddMenuEntry(FGraphEditorCommands::Get().AlignNodesBottom);
+						SubMenuSection.AddMenuEntry(FGraphEditorCommands::Get().AlignNodesLeft);
+						SubMenuSection.AddMenuEntry(FGraphEditorCommands::Get().AlignNodesCenter);
+						SubMenuSection.AddMenuEntry(FGraphEditorCommands::Get().AlignNodesRight);
+						SubMenuSection.AddMenuEntry(FGraphEditorCommands::Get().StraightenConnections);
+					}
 
-				{
-					FToolMenuSection& SubMenuSection = SubMenu->AddSection("EdGraphSchemaDistribution", LOCTEXT("DistributionHeader", "Distribution"));
-					SubMenuSection.AddMenuEntry(FGraphEditorCommands::Get().DistributeNodesHorizontally);
-					SubMenuSection.AddMenuEntry(FGraphEditorCommands::Get().DistributeNodesVertically);
-				}
-			}));
+					{
+						FToolMenuSection& SubMenuSection = SubMenu->AddSection("EdGraphSchemaDistribution", LOCTEXT("DistributionHeader", "Distribution"));
+						SubMenuSection.AddMenuEntry(FGraphEditorCommands::Get().DistributeNodesHorizontally);
+						SubMenuSection.AddMenuEntry(FGraphEditorCommands::Get().DistributeNodesVertically);
+					}
+				}));
+			}			
 		}
 	}
 }
diff --git a/Engine/Plugins/Runtime/Metasound/Source/MetasoundEditor/Private/MetasoundEditorGraphSchema.cpp b/Engine/Plugins/Runtime/Metasound/Source/MetasoundEditor/Private/MetasoundEditorGraphSchema.cpp
--- a/Engine/Plugins/Runtime/Metasound/Source/MetasoundEditor/Private/MetasoundEditorGraphSchema.cpp
+++ b/Engine/Plugins/Runtime/Metasound/Source/MetasoundEditor/Private/MetasoundEditorGraphSchema.cpp
@@ -1424,18 +1424,24 @@ void UMetasoundEditorGraphSchema::GetContextMenuActions(class UToolMenu* Menu, c
 	}
 	else if (Context->Pin && Context->Node && Context->Node->IsA<UMetasoundEditorGraphNode>())
 	{
-		FToolMenuSection& Section = Menu->FindOrAddSection("Pin Actions");
+		const UMetasoundEditorGraph* EdGraph = CastChecked<UMetasoundEditorGraph>(Context->Graph);
+		UMetaSoundBuilderBase& Builder = Metasound::Engine::FDocumentBuilderRegistry::GetChecked().FindOrBeginBuilding(EdGraph->GetMetasoundChecked());
 
-		if (Context->Pin->Direction == EGPD_Input)
+		if (!Builder.IsPreset())
 		{
-			Section.AddMenuEntry(FEditorCommands::Get().PromoteToInput);
-			Section.AddMenuEntry(FEditorCommands::Get().PromoteToVariable);
-			Section.AddMenuEntry(FEditorCommands::Get().PromoteToDeferredVariable);
-		}
-		else
-		{
-			Section.AddMenuEntry(FEditorCommands::Get().PromoteToOutput);
-			Section.AddMenuEntry(FEditorCommands::Get().PromoteToVariable);
+			FToolMenuSection& Section = Menu->FindOrAddSection("Pin Actions");
+
+			if (Context->Pin->Direction == EGPD_Input)
+			{
+				Section.AddMenuEntry(FEditorCommands::Get().PromoteToInput);
+				Section.AddMenuEntry(FEditorCommands::Get().PromoteToVariable);
+				Section.AddMenuEntry(FEditorCommands::Get().PromoteToDeferredVariable);
+			}
+			else
+			{
+				Section.AddMenuEntry(FEditorCommands::Get().PromoteToOutput);
+				Section.AddMenuEntry(FEditorCommands::Get().PromoteToVariable);
+			}
 		}
 	}
 
