Checkpoint on Resolve Page logic rework to support input default literals

#rnx

diff --git a/Engine/Plugins/Runtime/Metasound/Source/MetasoundEditor/Private/MetasoundEditorModule.cpp b/Engine/Plugins/Runtime/Metasound/Source/MetasoundEditor/Private/MetasoundEditorModule.cpp
--- a/Engine/Plugins/Runtime/Metasound/Source/MetasoundEditor/Private/MetasoundEditorModule.cpp
+++ b/Engine/Plugins/Runtime/Metasound/Source/MetasoundEditor/Private/MetasoundEditorModule.cpp
@@ -423,7 +423,7 @@ namespace Metasound
 				}
 
 				const UMetasoundEditorSettings* EditorSettings = GetDefault<UMetasoundEditorSettings>();
-				FDocumentBuilderRegistry::GetChecked().GetOnResolveAuditionPageInfoDelegate().BindUObject(EditorSettings, &UMetasoundEditorSettings::GetAuditionPageInfo);
+				FDocumentBuilderRegistry::GetChecked().GetOnResolveAuditionPageInfoDelegate().BindUObject(EditorSettings, &UMetasoundEditorSettings::ResolveAuditionPageInfo);
 
 				FEditorDelegates::PreBeginPIE.AddWeakLambda(EditorSettings, [](const bool /* bSimulating */)
 				{
@@ -443,7 +443,7 @@ namespace Metasound
 						FOnResolveAuditionPageInfo& OnResolveAuditionPageInfo = FDocumentBuilderRegistry::GetChecked().GetOnResolveAuditionPageInfoDelegate();
 						if (!OnResolveAuditionPageInfo.IsBoundToObject(EdSettings))
 						{
-							OnResolveAuditionPageInfo.BindUObject(EdSettings, &UMetasoundEditorSettings::GetAuditionPageInfo);
+							OnResolveAuditionPageInfo.BindUObject(EdSettings, &UMetasoundEditorSettings::ResolveAuditionPageInfo);
 						}
 					}
 				});
diff --git a/Engine/Plugins/Runtime/Metasound/Source/MetasoundEditor/Private/MetasoundEditorSettings.cpp b/Engine/Plugins/Runtime/Metasound/Source/MetasoundEditor/Private/MetasoundEditorSettings.cpp
--- a/Engine/Plugins/Runtime/Metasound/Source/MetasoundEditor/Private/MetasoundEditorSettings.cpp
+++ b/Engine/Plugins/Runtime/Metasound/Source/MetasoundEditor/Private/MetasoundEditorSettings.cpp
@@ -109,21 +109,14 @@ const FAudioMaterialMeterStyle* UMetasoundEditorSettings::GetMeterStyle() const
 	return &FAudioWidgetsStyle::Get().GetWidgetStyle<FAudioMaterialMeterStyle>("AudioMaterialMeter.Style");
 }
 
-Metasound::Engine::FAuditionPageInfo UMetasoundEditorSettings::GetAuditionPageInfo(const FMetasoundFrontendDocument& InDocument) const
+Metasound::Engine::FAuditionPageInfo UMetasoundEditorSettings::ResolveAuditionPageInfo(const TSet<FGuid>& InPageIDs) const
 {
 	using namespace Metasound::Engine;
 
 	FAuditionPageInfo PreviewInfo { .PlatformName = AuditionPlatform };
-
-	TSet<FGuid> DocPageIds;
-	InDocument.RootGraph.IterateGraphPages([&DocPageIds](const FMetasoundFrontendGraph& PageGraph)
-	{
-		DocPageIds.Add(PageGraph.PageID);
-	});
-
-	auto PageIsCooked = [&DocPageIds, &PreviewInfo](const FMetaSoundPageSettings& PageSettings)
+	auto PageIsCooked = [&InPageIDs, &PreviewInfo](const FMetaSoundPageSettings& PageSettings)
 	{
-		if (DocPageIds.Contains(PageSettings.UniqueId))
+		if (InPageIDs.Contains(PageSettings.UniqueId))
 		{
 			return PageSettings.IsCooked.GetValueForPlatform(PreviewInfo.PlatformName);
 		}
diff --git a/Engine/Plugins/Runtime/Metasound/Source/MetasoundEditor/Private/MetasoundEditorSettings.h b/Engine/Plugins/Runtime/Metasound/Source/MetasoundEditor/Private/MetasoundEditorSettings.h
--- a/Engine/Plugins/Runtime/Metasound/Source/MetasoundEditor/Private/MetasoundEditorSettings.h
+++ b/Engine/Plugins/Runtime/Metasound/Source/MetasoundEditor/Private/MetasoundEditorSettings.h
@@ -285,5 +285,5 @@ public:
 	/** Get the AudioMaterialMeter Style. If MeterStyleOverride is not set, returns default style.*/
 	const FAudioMaterialMeterStyle* GetMeterStyle() const;
 
-	Metasound::Engine::FAuditionPageInfo GetAuditionPageInfo(const FMetasoundFrontendDocument& InDocument) const;
+	Metasound::Engine::FAuditionPageInfo ResolveAuditionPageInfo(const TSet<FGuid>& InPageIDs) const;
 };
diff --git a/Engine/Plugins/Runtime/Metasound/Source/MetasoundEngine/Private/MetasoundBuilderSubsystem.cpp b/Engine/Plugins/Runtime/Metasound/Source/MetasoundEngine/Private/MetasoundBuilderSubsystem.cpp
--- a/Engine/Plugins/Runtime/Metasound/Source/MetasoundEngine/Private/MetasoundBuilderSubsystem.cpp
+++ b/Engine/Plugins/Runtime/Metasound/Source/MetasoundEngine/Private/MetasoundBuilderSubsystem.cpp
@@ -58,6 +58,29 @@ namespace Metasound::Engine
 			FLiteralNodeConstructorParams Params { "Literal", FGuid::NewGuid(), MoveTemp(ValueLiteral) };
 			return IDataTypeRegistry::Get().CreateLiteralNode(DataType, MoveTemp(Params));
 		}
+
+		TOptional<FMetasoundFrontendLiteral> TryResolveNodeInputDefault(const FMetaSoundFrontendDocumentBuilder& Builder, const FGuid& InNodeID, FName VertexName)
+		{
+			TOptional<FMetasoundFrontendLiteral> DefaultLiteral;
+			if (const FMetasoundFrontendVertexLiteral* InputDefault = Builder.FindNodeInputDefault(InNodeID, VertexName))
+			{
+				DefaultLiteral = InputDefault->Value;
+			}
+			else if (const TArray<FMetasoundFrontendClassInputDefault>* ClassDefaults = Builder.FindNodeClassInputDefaults(InNodeID, VertexName))
+			{
+				FGuid PageID;
+				if (FDocumentBuilderRegistry::GetChecked().TryResolveTargetPageID(*ClassDefaults, PageID))
+				{
+					auto MatchesPageID = [&PageID](const FMetasoundFrontendClassInputDefault& InputDefault) { return InputDefault.PageID == PageID; };
+					if (const FMetasoundFrontendClassInputDefault* ClassDefault = ClassDefaults->FindByPredicate(MatchesPageID))
+					{
+						DefaultLiteral = ClassDefault->Literal;
+					}
+				}
+			}
+
+			return DefaultLiteral;
+		}
 	} // namespace BuilderSubsystemPrivate
 } // namespace Metasound::Engine
 
@@ -306,17 +329,19 @@ void UMetaSoundSourceBuilder::InitTargetPageDelegates(Metasound::Frontend::FDocu
 	});
 
 	const IMetaSoundDocumentInterface& DocInterface = GetConstBuilder().GetConstDocumentInterfaceChecked();
-	const FGuid PageID = FDocumentBuilderRegistry::GetChecked().ResolveTargetPageID(DocInterface.GetConstDocument(), DocInterface.GetAssetPathChecked());
-
-	FEdgeModifyDelegates& EdgeDelegates = OutDocumentDelegates.FindEdgeDelegatesChecked(PageID);
-	EdgeDelegates.OnEdgeAdded.AddUObject(this, &UMetaSoundSourceBuilder::OnEdgeAdded);
-	EdgeDelegates.OnRemoveSwappingEdge.AddUObject(this, &UMetaSoundSourceBuilder::OnRemoveSwappingEdge);
-
-	FNodeModifyDelegates& NodeDelegates = OutDocumentDelegates.FindNodeDelegatesChecked(PageID);
-	NodeDelegates.OnNodeAdded.AddUObject(this, &UMetaSoundSourceBuilder::OnNodeAdded);
-	NodeDelegates.OnNodeInputLiteralSet.AddUObject(this, &UMetaSoundSourceBuilder::OnNodeInputLiteralSet);
-	NodeDelegates.OnRemoveSwappingNode.AddUObject(this, &UMetaSoundSourceBuilder::OnRemoveSwappingNode);
-	NodeDelegates.OnRemovingNodeInputLiteral.AddUObject(this, &UMetaSoundSourceBuilder::OnRemovingNodeInputLiteral);
+	FGuid PageID;
+	if (FDocumentBuilderRegistry::GetChecked().TryResolveTargetPageID(DocInterface.GetConstDocument().RootGraph, PageID))
+	{
+		FEdgeModifyDelegates& EdgeDelegates = OutDocumentDelegates.FindEdgeDelegatesChecked(PageID);
+		EdgeDelegates.OnEdgeAdded.AddUObject(this, &UMetaSoundSourceBuilder::OnEdgeAdded);
+		EdgeDelegates.OnRemoveSwappingEdge.AddUObject(this, &UMetaSoundSourceBuilder::OnRemoveSwappingEdge);
+
+		FNodeModifyDelegates& NodeDelegates = OutDocumentDelegates.FindNodeDelegatesChecked(PageID);
+		NodeDelegates.OnNodeAdded.AddUObject(this, &UMetaSoundSourceBuilder::OnNodeAdded);
+		NodeDelegates.OnNodeInputLiteralSet.AddUObject(this, &UMetaSoundSourceBuilder::OnNodeInputLiteralSet);
+		NodeDelegates.OnRemoveSwappingNode.AddUObject(this, &UMetaSoundSourceBuilder::OnRemoveSwappingNode);
+		NodeDelegates.OnRemovingNodeInputLiteral.AddUObject(this, &UMetaSoundSourceBuilder::OnRemovingNodeInputLiteral);
+	}
 }
 
 void UMetaSoundSourceBuilder::OnAssetReferenceAdded(TScriptInterface<IMetaSoundDocumentInterface> DocInterface)
@@ -585,13 +610,9 @@ void UMetaSoundSourceBuilder::OnRemoveSwappingEdge(int32 SwapIndex, int32 LastIn
 		const FMetasoundFrontendVertex* ToNodeInput = Builder.FindNodeInput(EdgeBeingRemoved.ToNodeID, EdgeBeingRemoved.ToVertexID);
 		if (FromNodeOutput && ToNodeInput)
 		{
-			const FMetasoundFrontendLiteral* InputDefault = Builder.GetNodeInputDefault(EdgeBeingRemoved.ToNodeID, EdgeBeingRemoved.ToVertexID);
-			if (!InputDefault)
-			{
-				InputDefault = Builder.GetNodeInputClassDefault(EdgeBeingRemoved.ToNodeID, EdgeBeingRemoved.ToVertexID);
-			}
+			TOptional<FMetasoundFrontendLiteral> InputDefault = BuilderSubsystemPrivate::TryResolveNodeInputDefault(Builder, EdgeBeingRemoved.ToNodeID, ToNodeInput->Name);
 
-			if (ensureAlwaysMsgf(InputDefault, TEXT("Could not dynamically assign default literal upon removing edge: literal should be assigned by either the frontend document's input or the class definition")))
+			if (ensureAlwaysMsgf(InputDefault.IsSet(), TEXT("Could not dynamically assign default literal upon removing edge: literal should be assigned by either the frontend document's input or the class definition")))
 			{
 				TUniquePtr<INode> LiteralNode = BuilderSubsystemPrivate::CreateDynamicNodeFromFrontendLiteral(ToNodeInput->TypeName, *InputDefault);
 				Transactor.RemoveDataEdge(EdgeBeingRemoved.FromNodeID, FromNodeOutput->Name, EdgeBeingRemoved.ToNodeID, ToNodeInput->Name, MoveTemp(LiteralNode));
@@ -689,8 +710,8 @@ void UMetaSoundSourceBuilder::OnRemovingNodeInputLiteral(int32 NodeIndex, int32
 			const FMetasoundFrontendNode& Node = Nodes[NodeIndex];
 			const FMetasoundFrontendVertex& Input = Node.Interface.Inputs[VertexIndex];
 
-			const FMetasoundFrontendLiteral* InputDefault = Builder.GetNodeInputClassDefault(Node.GetID(), Input.VertexID);
-			if (ensureAlwaysMsgf(InputDefault, TEXT("Could not dynamically assign default literal from class definition upon removing input '%s' literal: document's dependency entry invalid and has no default assigned"), *Input.Name.ToString()))
+			TOptional<FMetasoundFrontendLiteral> InputDefault = BuilderSubsystemPrivate::TryResolveNodeInputDefault(Builder, Node.GetID(), Input.Name);
+			if (ensureAlwaysMsgf(InputDefault.IsSet(), TEXT("Could not dynamically assign default literal from class definition upon removing input '%s' literal: document's dependency entry invalid and has no default assigned"), *Input.Name.ToString()))
 			{
 				TUniquePtr<INode> LiteralNode = BuilderSubsystemPrivate::CreateDynamicNodeFromFrontendLiteral(Input.TypeName, *InputDefault);
 				Transactor.SetValue(Node.GetID(), Input.Name, MoveTemp(LiteralNode));
diff --git a/Engine/Plugins/Runtime/Metasound/Source/MetasoundEngine/Private/MetasoundDocumentBuilderRegistry.cpp b/Engine/Plugins/Runtime/Metasound/Source/MetasoundEngine/Private/MetasoundDocumentBuilderRegistry.cpp
--- a/Engine/Plugins/Runtime/Metasound/Source/MetasoundEngine/Private/MetasoundDocumentBuilderRegistry.cpp
+++ b/Engine/Plugins/Runtime/Metasound/Source/MetasoundEngine/Private/MetasoundDocumentBuilderRegistry.cpp
@@ -324,88 +324,101 @@ namespace Metasound::Engine
 		return bReloaded;
 	}
 
-	FGuid FDocumentBuilderRegistry::ResolveTargetPageID(const FMetasoundFrontendDocument& Document, const FTopLevelAssetPath& AssetPath) const
+	bool FDocumentBuilderRegistry::TryResolveTargetPageID(const FMetasoundFrontendGraphClass& InGraphClass, FGuid& OutResolvedPageID) const
+	{
+		TSet<FGuid> DocPageIds;
+		InGraphClass.IterateGraphPages([&DocPageIds](const FMetasoundFrontendGraph& PageGraph)
+		{
+			DocPageIds.Add(PageGraph.PageID);
+		});
+
+		return TryResolveTargetPageID(DocPageIds, OutResolvedPageID);
+	}
+
+	bool FDocumentBuilderRegistry::TryResolveTargetPageID(const FMetasoundFrontendClassInput& InClassInput, FGuid& OutResolvedPageID) const
+	{
+		TSet<FGuid> DocPageIds;
+		InClassInput.IterateDefaults([&DocPageIds](const FGuid& PageID, const FMetasoundFrontendLiteral&)
+		{
+			DocPageIds.Add(PageID);
+		});
+
+		return TryResolveTargetPageID(DocPageIds, OutResolvedPageID);
+	}
+
+	bool FDocumentBuilderRegistry::TryResolveTargetPageID(const TArray<FMetasoundFrontendClassInputDefault>& InClassDefaults, FGuid& OutResolvedPageID) const
+	{
+		TSet<FGuid> DocPageIds;
+		for (const FMetasoundFrontendClassInputDefault& ClassDefault : InClassDefaults)
+		{
+			DocPageIds.Add(ClassDefault.PageID);
+		}
+
+		return TryResolveTargetPageID(DocPageIds, OutResolvedPageID);
+	}
+
+	bool FDocumentBuilderRegistry::TryResolveTargetPageID(const TSet<FGuid>& InPageIDs, FGuid& OutResolvedPageID) const
 	{
 		FName PlatformName = FPlatformProperties::IniPlatformName();
 
 #if WITH_EDITOR
 		if (OnResolveAuditionPageInfo.IsBound())
 		{
-			FAuditionPageInfo PreviewInfo = OnResolveAuditionPageInfo.Execute(Document);
+			FAuditionPageInfo PreviewInfo = OnResolveAuditionPageInfo.Execute(InPageIDs);
 			if (PreviewInfo.PageID.IsSet())
 			{
-				return PreviewInfo.PageID.GetValue();
+				OutResolvedPageID = PreviewInfo.PageID.GetValue();
+				return InPageIDs.Contains(OutResolvedPageID);
 			}
 
 			PlatformName = PreviewInfo.PlatformName;
 		}
 #endif // WITH_EDITOR
 
-		TSet<FGuid> DocPageIds;
-		Document.RootGraph.IterateGraphPages([&DocPageIds](const FMetasoundFrontendGraph& PageGraph)
-		{
-			DocPageIds.Add(PageGraph.PageID);
-		});
-
-		bool bImplementsPages = false;
 		const UMetaSoundSettings* Settings = GetDefault<UMetaSoundSettings>();
 		if (Settings)
 		{
 			const FGuid& TargetPageID = Settings->GetTargetPageID();
 			const TArray<FMetaSoundPageSettings>& PageSettingsArray = Settings->GetPageSettings();
-			bImplementsPages = !Settings->GetPageSettings().IsEmpty();
 
 			bool bFoundMatch = false;
 			for (int32 Index = PageSettingsArray.Num() - 1; Index >= 0; --Index)
 			{
 				const FMetaSoundPageSettings& PageSettings = PageSettingsArray[Index];
 				bFoundMatch |= PageSettings.UniqueId == TargetPageID;
-				const bool bAssetImplementsPage = DocPageIds.Contains(PageSettings.UniqueId);
+				const bool bAssetImplementsPage = InPageIDs.Contains(PageSettings.UniqueId);
 				if (bFoundMatch && bAssetImplementsPage)
 				{
 #if WITH_EDITOR
 					const bool bIsCooked = PageSettings.IsCooked.GetValueForPlatform(PlatformName);
 					if (bIsCooked)
 					{
-						return PageSettings.UniqueId;
+						OutResolvedPageID = PageSettings.UniqueId;
+						return true;
 					}
 #else // !WITH_EDITOR
-					return PageSettings.UniqueId;
+					OutResolvedPageID = PageSettings.UniqueId;
+					return true;
 #endif // !WITH_EDITOR
 				}
 			}
 		}
 
-		// Set to arbitrary ID so the default is prioritized when iterating implemented
-		// asset pages. All documents are guaranteed to implement at least one page.
-		FGuid PageID = FGuid::NewGuid();
-		Document.RootGraph.IterateGraphPages([&PageID](const FMetasoundFrontendGraph& Graph)
+		// In the case that no id was able to be resolved, returned value should use
+		// default graph. If no pages were implemented, error and attempt to return
+		// a valid pageID, but there is no guarantee the page ID represents an implemented
+		// page.
+		auto IsValidPageID = [](const FGuid& InPageID) { return InPageID != Frontend::DefaultPageID; };
+		if (const FGuid* ImplPageID = Algo::FindByPredicate(InPageIDs, IsValidPageID))
 		{
-			if (PageID.IsValid())
-			{
-				PageID = Graph.PageID;
-			}
-		});
-
-#if !NO_LOGGING
-		// Log error if page settings exist, but could not resolve ID with the given document
-		if (bImplementsPages)
+			OutResolvedPageID = *ImplPageID;
+		}
+		else
 		{
-			FString PageIdentifier = PageID.ToString();
-			if (const FMetaSoundPageSettings* PageSettings = Settings->FindPageSettings(PageID))
-			{
-				PageIdentifier = PageSettings->Name.ToString();
-			}
-
-			UE_LOG(LogMetaSound, Error,
-				TEXT("'%s' failed to resolve executable page ID:  \nMetaSound 'Page Settings' does not provide a valid fallback page for execution"
-					" on the desired platform, which can result in undefined behavior. Registering asset's page with ID '%s'."),
-				*AssetPath.ToString(),
-				*PageIdentifier);
+			OutResolvedPageID = Frontend::DefaultPageID;
 		}
-#endif // !NO_LOGGING
 
-		return PageID;
+		return InPageIDs.Contains(OutResolvedPageID);
 	}
 
 	void FDocumentBuilderRegistry::SetEventLogVerbosity(ELogEvent Event, ELogVerbosity::Type Verbosity)
diff --git a/Engine/Plugins/Runtime/Metasound/Source/MetasoundEngine/Public/MetasoundDocumentBuilderRegistry.h b/Engine/Plugins/Runtime/Metasound/Source/MetasoundEngine/Public/MetasoundDocumentBuilderRegistry.h
--- a/Engine/Plugins/Runtime/Metasound/Source/MetasoundEngine/Public/MetasoundDocumentBuilderRegistry.h
+++ b/Engine/Plugins/Runtime/Metasound/Source/MetasoundEngine/Public/MetasoundDocumentBuilderRegistry.h
@@ -21,7 +21,7 @@ namespace Metasound::Engine
 		TOptional<FGuid> PageID;
 	};
 
-	DECLARE_DELEGATE_RetVal_OneParam(FAuditionPageInfo, FOnResolveAuditionPageInfo, const FMetasoundFrontendDocument& /* Document */);
+	DECLARE_DELEGATE_RetVal_OneParam(FAuditionPageInfo, FOnResolveAuditionPageInfo, const TSet<FGuid>& /* InPageIDs */);
 #endif // WITH_EDITOR
 
 	class METASOUNDENGINE_API FDocumentBuilderRegistry : public Frontend::IDocumentBuilderRegistry
@@ -137,7 +137,9 @@ namespace Metasound::Engine
 		bool ReloadBuilder(const FMetasoundFrontendClassName& InClassName) const override;
 
 		// Given the provided document and its respective pages, returns the PageID to be used for runtime IGraph and proxy generation.
-		virtual FGuid ResolveTargetPageID(const FMetasoundFrontendDocument& Document, const FTopLevelAssetPath& AssetPath) const override;
+		virtual bool TryResolveTargetPageID(const FMetasoundFrontendGraphClass& InGraphClass, FGuid& OutResolvedPageID) const override;
+		virtual bool TryResolveTargetPageID(const FMetasoundFrontendClassInput& InClassInput, FGuid& OutResolvedPageID) const override;
+		virtual bool TryResolveTargetPageID(const TArray<FMetasoundFrontendClassInputDefault>& Defaults, FGuid& OutResolvedPageID) const override;
 
 		void SetEventLogVerbosity(ELogEvent Event, ELogVerbosity::Type Verbosity);
 
@@ -145,6 +147,7 @@ namespace Metasound::Engine
 		void AddBuilderInternal(const FMetasoundFrontendClassName& InClassName, UMetaSoundBuilderBase* NewBuilder) const;
 		bool CanPostEventLog(ELogEvent Event, ELogVerbosity::Type Verbosity) const;
 		void FinishBuildingInternal(UMetaSoundBuilderBase& Builder, bool bForceUnregisterNodeClass) const;
+		bool TryResolveTargetPageID(const TSet<FGuid>& InPageIDs, FGuid& OutResolvedPageID) const;
 
 #if WITH_EDITOR
 		FOnResolveAuditionPageInfo OnResolveAuditionPageInfo;
diff --git a/Engine/Plugins/Runtime/Metasound/Source/MetasoundFrontend/Private/MetasoundFrontendDocument.cpp b/Engine/Plugins/Runtime/Metasound/Source/MetasoundFrontend/Private/MetasoundFrontendDocument.cpp
--- a/Engine/Plugins/Runtime/Metasound/Source/MetasoundFrontend/Private/MetasoundFrontendDocument.cpp
+++ b/Engine/Plugins/Runtime/Metasound/Source/MetasoundFrontend/Private/MetasoundFrontendDocument.cpp
@@ -7,6 +7,7 @@
 #include "Algo/Transform.h"
 #include "IAudioParameterInterfaceRegistry.h"
 #include "Logging/LogMacros.h"
+#include "MetasoundDocumentInterface.h"
 #include "MetasoundFrontend.h"
 #include "MetasoundFrontendDocumentIdGenerator.h"
 #include "MetasoundFrontendDocumentVersioning.h"
@@ -27,6 +28,25 @@ namespace Metasound
 
 	namespace Frontend
 	{
+		namespace DocumentPrivate
+		{
+			template <typename ResolveType>
+			bool TryResolveTargetPageID(const ResolveType& InToResolve, FGuid& OutPageID)
+			{
+				OutPageID = Frontend::DefaultPageID;
+
+				// Registry is not available in tests, so for now resolution is considered successful at this level
+				// if registry is not initialized and providing a resolved page ID. TODO: Add a test implementation
+				// that returns the default page (or whatever page behavior is desired for testing).
+				if (IDocumentBuilderRegistry* BuilderRegistry = IDocumentBuilderRegistry::Get())
+				{
+					return IDocumentBuilderRegistry::GetChecked().TryResolveTargetPageID(InToResolve, OutPageID);
+				}
+
+				return true;
+			}
+		} // namespace DocumentPrivate
+
 #if WITH_EDITORONLY_DATA
 		const FText DefaultGraphPageDisplayName = NSLOCTEXT("MetasoundFrontend", "DefaultGraphPageDisplayName", "Default");
 #endif // WITH_EDITORONLY_DATA
@@ -887,6 +907,11 @@ FMetasoundFrontendLiteral& FMetasoundFrontendClassInput::FindDefaultChecked(cons
 	return *Literal;
 }
 
+const TArray<FMetasoundFrontendClassInputDefault>& FMetasoundFrontendClassInput::GetDefaults() const
+{
+	return Defaults;
+}
+
 FMetasoundFrontendLiteral& FMetasoundFrontendClassInput::InitDefault()
 {
 	checkf(Defaults.IsEmpty(), TEXT("Default(s) already initialized"));
@@ -926,6 +951,12 @@ bool FMetasoundFrontendClassInput::RemoveDefault(const FGuid& InPageID)
 	auto IsPage = [&InPageID](const FMetasoundFrontendClassInputDefault& Default) { return Default.PageID == InPageID; };
 	return Defaults.RemoveAllSwap(IsPage) > 0;
 }
+
+void FMetasoundFrontendClassInput::SetDefaults(TArray<FMetasoundFrontendClassInputDefault> InputDefaults)
+{
+	Defaults = MoveTemp(InputDefaults);
+}
+
 #endif // WITH_EDITORONLY_DATA
 
 FMetasoundFrontendClassVariable::FMetasoundFrontendClassVariable(const FMetasoundFrontendClassVertex& InOther)
@@ -1207,27 +1238,15 @@ namespace Metasound::Frontend
 		}
 	}
 
-	void ForEachLiteral(const FMetasoundFrontendDocument& InDoc, FForEachLiteralFunctionRef OnLiteral, const FGuid& InPageID)
-	{
-		ForEachLiteral(InDoc.RootGraph, OnLiteral, InPageID);
-
-		for (const FMetasoundFrontendGraphClass& GraphClass : InDoc.Subgraphs)
-		{
-			ForEachLiteral(GraphClass, OnLiteral, InPageID);
-		}
-
-		for (const FMetasoundFrontendClass& Dependency : InDoc.Dependencies)
-		{
-			ForEachLiteral(Dependency, OnLiteral);
-		}
-	}
-
 	void ForEachLiteral(const FMetasoundFrontendGraphClass& InGraphClass, FForEachLiteralFunctionRef OnLiteral)
 	{
 		ForEachLiteral(static_cast<const FMetasoundFrontendClass&>(InGraphClass), OnLiteral);
 
-		InGraphClass.IterateGraphPages([&OnLiteral](const FMetasoundFrontendGraph& Graph)
+		FGuid PageID = Frontend::DefaultPageID;
+		const bool bPageIDResolved = DocumentPrivate::TryResolveTargetPageID(InGraphClass, PageID);
+		if (bPageIDResolved)
 		{
+			const FMetasoundFrontendGraph& Graph = InGraphClass.FindConstGraphChecked(PageID);
 			for (const FMetasoundFrontendNode& Node : Graph.Nodes)
 			{
 				ForEachLiteral(Node, OnLiteral);
@@ -1237,22 +1256,6 @@ namespace Metasound::Frontend
 			{
 				OnLiteral(Variable.TypeName, Variable.Literal);
 			}
-		});
-	}
-
-	void ForEachLiteral(const FMetasoundFrontendGraphClass& InGraphClass, FForEachLiteralFunctionRef OnLiteral, const FGuid& InPageID)
-	{
-		ForEachLiteral(static_cast<const FMetasoundFrontendClass&>(InGraphClass), OnLiteral);
-
-		const FMetasoundFrontendGraph& Graph = InGraphClass.FindConstGraphChecked(InPageID);
-		for (const FMetasoundFrontendNode& Node : Graph.Nodes)
-		{
-			ForEachLiteral(Node, OnLiteral);
-		}
-
-		for (const FMetasoundFrontendVariable& Variable : Graph.Variables)
-		{
-			OnLiteral(Variable.TypeName, Variable.Literal);
 		}
 	}
 
@@ -1260,8 +1263,13 @@ namespace Metasound::Frontend
 	{
 		for (const FMetasoundFrontendClassInput& ClassInput : InClass.Interface.Inputs)
 		{
-			const FMetasoundFrontendLiteral& DefaultLiteral = ClassInput.FindConstDefaultChecked(Frontend::DefaultPageID);
-			OnLiteral(ClassInput.TypeName, DefaultLiteral);
+			FGuid PageID = Frontend::DefaultPageID;
+			const bool bPageIDResolved = DocumentPrivate::TryResolveTargetPageID(ClassInput, PageID);
+			if (bPageIDResolved)
+			{
+				const FMetasoundFrontendLiteral& DefaultLiteral = ClassInput.FindConstDefaultChecked(PageID);
+				OnLiteral(ClassInput.TypeName, DefaultLiteral);
+			}
 		}
 	}
 
diff --git a/Engine/Plugins/Runtime/Metasound/Source/MetasoundFrontend/Private/MetasoundFrontendDocumentBuilder.cpp b/Engine/Plugins/Runtime/Metasound/Source/MetasoundFrontend/Private/MetasoundFrontendDocumentBuilder.cpp
--- a/Engine/Plugins/Runtime/Metasound/Source/MetasoundFrontend/Private/MetasoundFrontendDocumentBuilder.cpp
+++ b/Engine/Plugins/Runtime/Metasound/Source/MetasoundFrontend/Private/MetasoundFrontendDocumentBuilder.cpp
@@ -1,6 +1,7 @@
 // Copyright Epic Games, Inc. All Rights Reserved.
 #include "MetasoundFrontendDocumentBuilder.h"
 
+#include "Algo/AllOf.h"
 #include "Algo/AnyOf.h"
 #include "Algo/Find.h"
 #include "Algo/ForEach.h"
@@ -1865,6 +1866,98 @@ const FMetasoundFrontendVertex* FMetaSoundFrontendDocumentBuilder::FindNodeInput
 	return NodeCache.FindInputVertex(InNodeID, InVertexName);
 }
 
+const TArray<FMetasoundFrontendClassInputDefault>* FMetaSoundFrontendDocumentBuilder::FindNodeClassInputDefaults(const FGuid& InNodeID, FName InVertexName, const FGuid* InPageID) const
+{
+	using namespace Metasound;
+
+	if (const FMetasoundFrontendNode* Node = FindNode(InNodeID, InPageID))
+	{
+		if (const FMetasoundFrontendClass* Class = FindDependency(Node->ClassID))
+		{
+			const EMetasoundFrontendClassType ClassType = Class->Metadata.GetType();
+			switch (ClassType)
+			{
+				case EMetasoundFrontendClassType::External:
+				{
+					auto MatchesName = [&InVertexName](const FMetasoundFrontendClassInput& Input) { return Input.Name == InVertexName; };
+					if (const FMetasoundFrontendClassInput* Input = Class->Interface.Inputs.FindByPredicate(MatchesName))
+					{
+						return &Input->GetDefaults();
+					}
+				}
+				break;
+
+				case EMetasoundFrontendClassType::Input:
+				case EMetasoundFrontendClassType::Output:
+				case EMetasoundFrontendClassType::Literal:
+				{
+					return &Class->Interface.Inputs.Last().GetDefaults();
+				}
+				break;
+
+				case EMetasoundFrontendClassType::Variable:
+				case EMetasoundFrontendClassType::VariableDeferredAccessor:
+				case EMetasoundFrontendClassType::VariableAccessor:
+				case EMetasoundFrontendClassType::VariableMutator:
+				{
+					using namespace VariableNames;
+					auto IsDataInput = [](const FMetasoundFrontendClassInput& Input) { return Input.Name == METASOUND_GET_PARAM_NAME(InputData); };
+					if (const FMetasoundFrontendClassInput* Input = Class->Interface.Inputs.FindByPredicate(IsDataInput))
+					{
+						return &Input->GetDefaults();
+					}
+				}
+				break;
+
+				case EMetasoundFrontendClassType::Template:
+				{
+					const Frontend::FNodeRegistryKey Key = Frontend::FNodeRegistryKey(Class->Metadata);
+					const Frontend::INodeTemplate* Template = Frontend::INodeTemplateRegistry::Get().FindTemplate(Key);
+					check(Template);
+					const FGuid PageID = InPageID ? *InPageID : Frontend::DefaultPageID;
+					return Template->FindNodeClassInputDefaults(*this, PageID, InNodeID, InVertexName);
+				}
+				break;
+
+				case EMetasoundFrontendClassType::Graph:
+				case EMetasoundFrontendClassType::Invalid:
+				default:
+				{
+					checkNoEntry();
+				}
+				break;
+			}
+		}
+	}
+
+	return nullptr;
+}
+
+const FMetasoundFrontendVertexLiteral* FMetaSoundFrontendDocumentBuilder::FindNodeInputDefault(const FGuid& InNodeID, const FGuid& InVertexID, const FGuid* InPageID) const
+{
+	if (const FMetasoundFrontendNode* Node = FindNode(InNodeID, InPageID))
+	{
+		auto VertexLiteralMatchesID = [&InVertexID](const FMetasoundFrontendVertexLiteral& VertexLiteral)
+		{
+			return VertexLiteral.VertexID == InVertexID;
+		};
+		return Node->InputLiterals.FindByPredicate(VertexLiteralMatchesID);
+	}
+
+	return nullptr;
+}
+
+const FMetasoundFrontendVertexLiteral* FMetaSoundFrontendDocumentBuilder::FindNodeInputDefault(const FGuid& InNodeID, FName InVertexName, const FGuid* InPageID) const
+{
+	using namespace Metasound::Frontend;
+	if (const FMetasoundFrontendVertex* Vertex = FindNodeInput(InNodeID, InVertexName, InPageID))
+	{
+		return FindNodeInputDefault(InNodeID, Vertex->VertexID, InPageID);
+	}
+
+	return nullptr;
+}
+
 TArray<const FMetasoundFrontendVertex*> FMetaSoundFrontendDocumentBuilder::FindNodeInputs(const FGuid& InNodeID, FName TypeName, const FGuid* InPageID) const
 {
 	const FGuid& PageID = InPageID ? *InPageID : BuildPageID;
@@ -2216,7 +2309,6 @@ const FMetasoundFrontendLiteral* FMetaSoundFrontendDocumentBuilder::GetNodeInput
 		{
 			const FMetasoundFrontendVertex& NodeInput = Node.Interface.Inputs[VertexIndex];
 
-			// If default not found on node, check class definition
 			auto IsLiteral = [&InVertexID](const FMetasoundFrontendVertexLiteral& Literal) { return Literal.VertexID == InVertexID; };
 			const int32 LiteralIndex = Node.InputLiterals.IndexOfByPredicate(IsLiteral);
 			if (LiteralIndex != INDEX_NONE)
@@ -2689,7 +2781,10 @@ void FMetaSoundFrontendDocumentBuilder::RemoveAllGraphPages()
 	FMetasoundFrontendGraphClass& RootGraph = GetDocumentChecked().RootGraph;
 	RootGraph.IterateGraphPages([this](FMetasoundFrontendGraph& Graph)
 	{
-		DocumentDelegates->PageDelegates.OnRemovingPage.Broadcast(Frontend::FDocumentMutatePageArgs { Graph.PageID });
+		if (Graph.PageID != Frontend::DefaultPageID)
+		{
+			DocumentDelegates->PageDelegates.OnRemovingPage.Broadcast(Frontend::FDocumentMutatePageArgs { Graph.PageID });
+		}
 	});
 	RootGraph.RemoveAllGraphPages();
 	SetBuildPageID(Frontend::DefaultPageID);
diff --git a/Engine/Plugins/Runtime/Metasound/Source/MetasoundFrontend/Private/MetasoundFrontendGraph.cpp b/Engine/Plugins/Runtime/Metasound/Source/MetasoundFrontend/Private/MetasoundFrontendGraph.cpp
--- a/Engine/Plugins/Runtime/Metasound/Source/MetasoundFrontend/Private/MetasoundFrontendGraph.cpp
+++ b/Engine/Plugins/Runtime/Metasound/Source/MetasoundFrontend/Private/MetasoundFrontendGraph.cpp
@@ -8,6 +8,7 @@
 #include "Algo/Transform.h"
 #include "CoreMinimal.h"
 #include "CoreGlobals.h"
+#include "MetasoundDocumentInterface.h"
 #include "MetasoundFrontend.h"
 #include "MetasoundFrontendDataTypeRegistry.h"
 #include "MetasoundFrontendNodeTemplateRegistry.h"
@@ -21,18 +22,37 @@
 
 namespace Metasound
 {
-	namespace FrontendGraphPrivate
+	namespace Frontend
 	{
-		FNodeInitData CreateNodeInitData(const FMetasoundFrontendNode& InNode)
+		namespace GraphPrivate
 		{
-			FNodeInitData InitData;
+			FNodeInitData CreateNodeInitData(const FMetasoundFrontendNode& InNode)
+			{
+				FNodeInitData InitData;
 
-			InitData.InstanceName = InNode.Name;
-			InitData.InstanceID = InNode.GetID();
+				InitData.InstanceName = InNode.Name;
+				InitData.InstanceID = InNode.GetID();
 
-			return InitData;
-		}
-	}
+				return InitData;
+			}
+
+			template <typename ResolveType>
+			bool TryResolveTargetPageID(const ResolveType& InToResolve, FGuid& OutPageID)
+			{
+				OutPageID = Frontend::DefaultPageID;
+
+				// Registry is not available in tests, so for now resolution is considered successful at this level
+				// if registry is not initialized and providing a resolved page ID. TODO: Add a test implementation
+				// that returns the default page (or whatever page behavior is desired for testing).
+				if (IDocumentBuilderRegistry* BuilderRegistry = IDocumentBuilderRegistry::Get())
+				{
+					return IDocumentBuilderRegistry::GetChecked().TryResolveTargetPageID(InToResolve, OutPageID);
+				}
+
+				return true;
+			}
+		} // namespace GraphPrivate
+	} // namespace Frontend
 
 	FFrontendGraph::FFrontendGraph(const FString& InInstanceName, const FGuid& InInstanceID)
 	:	FGraph(InInstanceName, InInstanceID)
@@ -204,7 +224,7 @@ namespace Metasound
 			};
 
 			{
-				const FNodeInitData InitData = FrontendGraphPrivate::CreateNodeInitData(InNode);
+				const FNodeInitData InitData = GraphPrivate::CreateNodeInitData(InNode);
 				TArray<FDefaultLiteralData> DefaultLiteralData = GetInputDefaultLiteralData(InContext, InNode, InitData, InEdgeDestinations);
 				for (FDefaultLiteralData& Data : DefaultLiteralData)
 				{
@@ -231,7 +251,7 @@ namespace Metasound
 
 		check(InNode.ClassID == InClass.ID);
 
-		const FNodeInitData InitData = FrontendGraphPrivate::CreateNodeInitData(InNode);
+		const FNodeInitData InitData = GraphPrivate::CreateNodeInitData(InNode);
 		{
 			TArray<FDefaultLiteralData> DefaultLiteralData = GetInputDefaultLiteralData(InContext, InNode, InitData, InEdgeDestinations);
 			for (FDefaultLiteralData& Data : DefaultLiteralData)
@@ -301,6 +321,8 @@ namespace Metasound
 
 	const FMetasoundFrontendLiteral* FFrontendGraphBuilder::FindInputLiteralForInputNode(const FMetasoundFrontendNode& InInputNode, const FMetasoundFrontendClass& InInputNodeClass, const FMetasoundFrontendClassInput& InOwningGraphClassInput)
 	{
+		using namespace Frontend;
+
 		// Default value priority is:
 		// 1. A value set directly on the node
 		// 2. A default value of the owning graph
@@ -327,26 +349,34 @@ namespace Metasound
 			}
 		}
 
+		FGuid PageID = Frontend::DefaultPageID;
+		const bool bPageIDResolved = GraphPrivate::TryResolveTargetPageID(InOwningGraphClassInput, PageID);
+
 		// Check for default value on owning graph.
 		if (nullptr == Literal)
 		{
 			// Find Class Default that is not invalid
-			const FMetasoundFrontendLiteral& DefaultLiteral = InOwningGraphClassInput.FindConstDefaultChecked(Frontend::DefaultPageID);
-			if (DefaultLiteral.IsValid())
+			if (bPageIDResolved)
 			{
-				Literal = &DefaultLiteral;
+				const FMetasoundFrontendLiteral& DefaultLiteral = InOwningGraphClassInput.FindConstDefaultChecked(PageID);
+				if (DefaultLiteral.IsValid())
+				{
+					Literal = &DefaultLiteral;
+				}
 			}
 		}
 
 		// Check for default value on input node class
 		if (nullptr == Literal && ensure(InInputNodeClass.Interface.Inputs.Num() == 1))
 		{
-			const FMetasoundFrontendClassInput& InputNodeClassInput = InInputNodeClass.Interface.Inputs[0];
-
-			const FMetasoundFrontendLiteral& DefaultLiteral = InputNodeClassInput.FindConstDefaultChecked(Frontend::DefaultPageID);
-			if (DefaultLiteral.IsValid())
+			const FMetasoundFrontendClassInput& InputNodeClassInput = InInputNodeClass.Interface.Inputs.Last();
+			if (bPageIDResolved)
 			{
-				Literal = &DefaultLiteral;
+				const FMetasoundFrontendLiteral& DefaultLiteral = InputNodeClassInput.FindConstDefaultChecked(PageID);
+				if (DefaultLiteral.IsValid())
+				{
+					Literal = &DefaultLiteral;
+				}
 			}
 		}
 
@@ -743,14 +773,20 @@ PRAGMA_ENABLE_DEPRECATION_WARNINGS
 		return bSuccess;
 	}
 
-	TUniquePtr<FFrontendGraph> FFrontendGraphBuilder::CreateGraph(FBuildContext& InContext, const FMetasoundFrontendGraphClass& InGraphClass, const FGuid* InPageId)
+	TUniquePtr<FFrontendGraph> FFrontendGraphBuilder::CreateGraph(FBuildContext& InContext, const FMetasoundFrontendGraphClass& InGraphClass)
 	{
+		using namespace Frontend;
+
 		const FString GraphName = InContext.DebugAssetName;
 
-		const FMetasoundFrontendGraph& PageGraph = InPageId
-			? InGraphClass.FindConstGraphChecked(*InPageId)
-			: InGraphClass.GetConstDefaultGraph();
+		FGuid PageID = Frontend::DefaultPageID;
+		const bool bPageIDResolved = GraphPrivate::TryResolveTargetPageID(InGraphClass, PageID);
+		if (!bPageIDResolved)
+		{
+			return { };
+		}
 
+		const FMetasoundFrontendGraph& PageGraph = InGraphClass.FindConstGraphChecked(PageID);
 		FBuildGraphContext BuildGraphContext
 		{
 			MakeUnique<FFrontendGraph>(GraphName, InContext.GraphId),
@@ -786,7 +822,7 @@ PRAGMA_ENABLE_DEPRECATION_WARNINGS
 		return CreateGraph(InGraph, InSubgraphs, InDependencies, InDebugAssetName);
 	}
 
-	TUniquePtr<FFrontendGraph> FFrontendGraphBuilder::CreateGraph(const FMetasoundFrontendGraphClass& InGraph, const TArray<FMetasoundFrontendGraphClass>& InSubgraphs, const TArray<FMetasoundFrontendClass>& InDependencies, const Frontend::FProxyDataCache& InProxyDataCache, const FString& InDebugAssetName, const FGuid InGraphId, const FGuid* InPageId)
+	TUniquePtr<FFrontendGraph> FFrontendGraphBuilder::CreateGraph(const FMetasoundFrontendGraphClass& InGraph, const TArray<FMetasoundFrontendGraphClass>& InSubgraphs, const TArray<FMetasoundFrontendClass>& InDependencies, const Frontend::FProxyDataCache& InProxyDataCache, const FString& InDebugAssetName, const FGuid InGraphId)
 	{
 		FBuildContext Context
 		{
@@ -820,10 +856,9 @@ PRAGMA_ENABLE_DEPRECATION_WARNINGS
 		}
 
 		// Create each subgraph.
-		const FGuid PageId = InPageId ? *InPageId : FGuid();
 		for (const FMetasoundFrontendGraphClass* FrontendSubgraphPtr : FrontendSubgraphPtrs)
 		{
-			TSharedPtr<const INode> Subgraph(CreateGraph(Context, *FrontendSubgraphPtr, &PageId).Release());
+			TSharedPtr<const INode> Subgraph(CreateGraph(Context, *FrontendSubgraphPtr).Release());
 			if (!Subgraph.IsValid())
 			{
 				UE_LOG(LogMetaSound, Warning, TEXT("Failed to create subgraph [SubgraphName: %s] in asset '%s'"), *FrontendSubgraphPtr->Metadata.GetClassName().ToString(), *InDebugAssetName);
@@ -836,36 +871,35 @@ PRAGMA_ENABLE_DEPRECATION_WARNINGS
 		}
 
 		// Create parent graph.
-		return CreateGraph(Context, InGraph, &PageId);
+		return CreateGraph(Context, InGraph);
 	}
 
-	TUniquePtr<FFrontendGraph> FFrontendGraphBuilder::CreateGraph(const FMetasoundFrontendDocument& InDocument, const FString& InDebugAssetName, const FGuid* InPageId)
+	TUniquePtr<FFrontendGraph> FFrontendGraphBuilder::CreateGraph(const FMetasoundFrontendDocument& InDocument, const FString& InDebugAssetName)
 	{
 		// Create proxies before creating graph.
-		const FGuid PageId = InPageId ? *InPageId : FGuid();
 		Frontend::FProxyDataCache ProxyDataCache;
-		ProxyDataCache.CreateAndCacheProxies(InDocument, PageId);
+		ProxyDataCache.CreateAndCacheProxies(InDocument);
 		
-		return CreateGraph(InDocument, ProxyDataCache, InDebugAssetName, Frontend::CreateLocallyUniqueId(), &PageId);
+		return CreateGraph(InDocument, ProxyDataCache, InDebugAssetName, Frontend::CreateLocallyUniqueId());
 	}
 
-	TUniquePtr<FFrontendGraph> FFrontendGraphBuilder::CreateGraph(const FMetasoundFrontendGraphClass& InGraph, const TArray<FMetasoundFrontendGraphClass>& InSubgraphs, const TArray<FMetasoundFrontendClass>& InDependencies, const FString& InDebugAssetName, const FGuid* InPageId)
+	TUniquePtr<FFrontendGraph> FFrontendGraphBuilder::CreateGraph(const FMetasoundFrontendGraphClass& InGraph, const TArray<FMetasoundFrontendGraphClass>& InSubgraphs, const TArray<FMetasoundFrontendClass>& InDependencies, const FString& InDebugAssetName)
 	{
 		// Create proxies before building graph
-		const FGuid PageId = InPageId ? *InPageId : FGuid();
 		Frontend::FProxyDataCache ProxyDataCache;
-		ProxyDataCache.CreateAndCacheProxies(InGraph, PageId);
+		ProxyDataCache.CreateAndCacheProxies(InGraph);
+
 		for (const FMetasoundFrontendGraphClass& SubgraphClass : InSubgraphs)
 		{
-			ProxyDataCache.CreateAndCacheProxies(SubgraphClass, PageId);
+			ProxyDataCache.CreateAndCacheProxies(SubgraphClass);
 		}
+
 		for (const FMetasoundFrontendClass& DependencyClass : InDependencies)
 		{
 			ProxyDataCache.CreateAndCacheProxies(DependencyClass);
 		}
 
-
-		return CreateGraph(InGraph, InSubgraphs, InDependencies, ProxyDataCache, InDebugAssetName, Frontend::CreateLocallyUniqueId(), &PageId);
+		return CreateGraph(InGraph, InSubgraphs, InDependencies, ProxyDataCache, InDebugAssetName, Frontend::CreateLocallyUniqueId());
 	}
 	
 	TUniquePtr<FFrontendGraph> FFrontendGraphBuilder::CreateGraph(const FMetasoundFrontendDocument& InDocument, const TSet<FName>& InTransmittableInputNames, const FString& InDebugAssetName)
@@ -873,8 +907,8 @@ PRAGMA_ENABLE_DEPRECATION_WARNINGS
 		return CreateGraph(InDocument, InDebugAssetName);
 	}
 
-	TUniquePtr<FFrontendGraph> FFrontendGraphBuilder::CreateGraph(const FMetasoundFrontendDocument& InDocument, const Frontend::FProxyDataCache& InProxyDataCache, const FString& InDebugAssetName, const FGuid InGraphId, const FGuid* InPageId)
+	TUniquePtr<FFrontendGraph> FFrontendGraphBuilder::CreateGraph(const FMetasoundFrontendDocument& InDocument, const Frontend::FProxyDataCache& InProxyDataCache, const FString& InDebugAssetName, const FGuid InGraphId)
 	{
-		return CreateGraph(InDocument.RootGraph, InDocument.Subgraphs, InDocument.Dependencies, InProxyDataCache, InDebugAssetName, InGraphId, InPageId);
+		return CreateGraph(InDocument.RootGraph, InDocument.Subgraphs, InDocument.Dependencies, InProxyDataCache, InDebugAssetName, InGraphId);
 	}
 }
diff --git a/Engine/Plugins/Runtime/Metasound/Source/MetasoundFrontend/Private/MetasoundFrontendGraphController.cpp b/Engine/Plugins/Runtime/Metasound/Source/MetasoundFrontend/Private/MetasoundFrontendGraphController.cpp
--- a/Engine/Plugins/Runtime/Metasound/Source/MetasoundFrontend/Private/MetasoundFrontendGraphController.cpp
+++ b/Engine/Plugins/Runtime/Metasound/Source/MetasoundFrontend/Private/MetasoundFrontendGraphController.cpp
@@ -1635,9 +1635,9 @@ namespace Metasound
 				FString UnknownAsset = TEXT("UnknownAsset");
 				FProxyDataCache ProxyCache;
 				const FMetasoundFrontendDocument& Doc = *(OwningDocument->GetDocumentPtr().Get());
-				ProxyCache.CreateAndCacheProxies(Doc, Frontend::DefaultPageID);
+				ProxyCache.CreateAndCacheProxies(Doc);
 
-				TUniquePtr<FFrontendGraph> Graph = FFrontendGraphBuilder::CreateGraph(*GraphClass, Subgraphs, Dependencies, ProxyCache, UnknownAsset, Frontend::CreateLocallyUniqueId(), { });
+				TUniquePtr<FFrontendGraph> Graph = FFrontendGraphBuilder::CreateGraph(*GraphClass, Subgraphs, Dependencies, ProxyCache, UnknownAsset, Frontend::CreateLocallyUniqueId());
 
 				if (!Graph.IsValid())
 				{
diff --git a/Engine/Plugins/Runtime/Metasound/Source/MetasoundFrontend/Private/MetasoundFrontendNodeTemplateRegistry.cpp b/Engine/Plugins/Runtime/Metasound/Source/MetasoundFrontend/Private/MetasoundFrontendNodeTemplateRegistry.cpp
--- a/Engine/Plugins/Runtime/Metasound/Source/MetasoundFrontend/Private/MetasoundFrontendNodeTemplateRegistry.cpp
+++ b/Engine/Plugins/Runtime/Metasound/Source/MetasoundFrontend/Private/MetasoundFrontendNodeTemplateRegistry.cpp
@@ -3,6 +3,7 @@
 #include "MetasoundFrontendNodeTemplateRegistry.h"
 
 #include "MetasoundFrontendDocument.h"
+#include "MetasoundFrontendDocumentBuilder.h"
 #include "MetasoundFrontendRegistryContainerImpl.h"
 #include "MetasoundFrontendRegistryTransaction.h"
 #include "MetasoundFrontendTransform.h"
@@ -107,6 +108,26 @@ namespace Metasound::Frontend
 		return nullptr;
 	}
 
+	const TArray<FMetasoundFrontendClassInputDefault>* FNodeTemplateBase::FindNodeClassInputDefaults(const FMetaSoundFrontendDocumentBuilder& InBuilder, const FGuid& InPageID, const FGuid& InNodeID, FName VertexName) const
+	{
+		if (const FMetasoundFrontendVertex* Vertex = InBuilder.FindNodeInput(InNodeID, VertexName, &InPageID))
+		{
+			const FMetasoundFrontendNode* Node = InBuilder.FindNode(InNodeID, &InPageID);
+			check(Node);
+			if (const FMetasoundFrontendClass* Class = InBuilder.FindDependency(Node->ClassID))
+			{
+				const EMetasoundFrontendClassType ClassType = Class->Metadata.GetType();
+				auto MatchesName = [&Vertex](const FMetasoundFrontendClassInput& Input) { return Input.Name == Vertex->Name; };
+				if (const FMetasoundFrontendClassInput* Input = Class->Interface.Inputs.FindByPredicate(MatchesName))
+				{
+					return &Input->GetDefaults();
+				}
+			}
+		}
+
+		return nullptr;
+	}
+
 #if WITH_EDITOR
 	FText FNodeTemplateBase::GetNodeDisplayName(const IMetaSoundDocumentInterface& Interface, const FGuid& InPageID, const FGuid& InNodeID) const
 	{
diff --git a/Engine/Plugins/Runtime/Metasound/Source/MetasoundFrontend/Private/MetasoundFrontendProxyDataCache.cpp b/Engine/Plugins/Runtime/Metasound/Source/MetasoundFrontend/Private/MetasoundFrontendProxyDataCache.cpp
--- a/Engine/Plugins/Runtime/Metasound/Source/MetasoundFrontend/Private/MetasoundFrontendProxyDataCache.cpp
+++ b/Engine/Plugins/Runtime/Metasound/Source/MetasoundFrontend/Private/MetasoundFrontendProxyDataCache.cpp
@@ -22,8 +22,8 @@ namespace Metasound
 	{
 		namespace MetasoundFrontendProxyDataCachePrivate
 		{
-			template<typename DocElementType, typename ...Args>
-			void CreateAndCacheProxies(FProxyDataCache& InCache, const DocElementType& InDocElement, Args&&... InArgs)
+			template<typename DocElementType>
+			void CreateAndCacheProxies(FProxyDataCache& InCache, const DocElementType& InDocElement)
 			{
 				checkf(IsInGameThread() || IsInAudioThread(), TEXT("Proxies should only be created in the game thread or audio thread"));
 				IDataTypeRegistry& DataRegistry = IDataTypeRegistry::Get();
@@ -61,20 +61,20 @@ namespace Metasound
 						}
 					}
 				};
-				ForEachLiteral(InDocElement, CreateProxies, Forward<Args>(InArgs)...);
+				ForEachLiteral(InDocElement, CreateProxies);
 			}
 		}
 
-		void FProxyDataCache::CreateAndCacheProxies(const FMetasoundFrontendDocument& InDocument, const FGuid& InPageID)
+		void FProxyDataCache::CreateAndCacheProxies(const FMetasoundFrontendDocument& InDocument)
 		{
 			METASOUND_TRACE_CPUPROFILER_EVENT_SCOPE(Metasound::FProxyDataCache::CreateAndCacheProxies_Document)
-			MetasoundFrontendProxyDataCachePrivate::CreateAndCacheProxies(*this, InDocument, InPageID);
+			MetasoundFrontendProxyDataCachePrivate::CreateAndCacheProxies(*this, InDocument);
 		}
 
-		void FProxyDataCache::CreateAndCacheProxies(const FMetasoundFrontendGraphClass& InGraphClass, const FGuid& InPageID)
+		void FProxyDataCache::CreateAndCacheProxies(const FMetasoundFrontendGraphClass& InGraphClass)
 		{
 			METASOUND_TRACE_CPUPROFILER_EVENT_SCOPE(Metasound::FProxyDataCache::CreateAndCacheProxies_GraphClass)
-			MetasoundFrontendProxyDataCachePrivate::CreateAndCacheProxies(*this, InGraphClass, InPageID);
+			MetasoundFrontendProxyDataCachePrivate::CreateAndCacheProxies(*this, InGraphClass);
 		}
 
 		void FProxyDataCache::CreateAndCacheProxies(const FMetasoundFrontendClass& InClass)
diff --git a/Engine/Plugins/Runtime/Metasound/Source/MetasoundFrontend/Private/MetasoundFrontendProxyDataCache.h b/Engine/Plugins/Runtime/Metasound/Source/MetasoundFrontend/Private/MetasoundFrontendProxyDataCache.h
--- a/Engine/Plugins/Runtime/Metasound/Source/MetasoundFrontend/Private/MetasoundFrontendProxyDataCache.h
+++ b/Engine/Plugins/Runtime/Metasound/Source/MetasoundFrontend/Private/MetasoundFrontendProxyDataCache.h
@@ -32,10 +32,10 @@ namespace Metasound
 		{
 		public:
 			/** Create and cache all proxies found in a document. */
-			void CreateAndCacheProxies(const FMetasoundFrontendDocument& InDocument, const FGuid& InPageID);
+			void CreateAndCacheProxies(const FMetasoundFrontendDocument& InDocument);
 
 			/** Create and cache all proxies found in a graph class. */
-			void CreateAndCacheProxies(const FMetasoundFrontendGraphClass& InGraphClass, const FGuid& InPageID);
+			void CreateAndCacheProxies(const FMetasoundFrontendGraphClass& InGraphClass);
 
 			/** Create and cache all proxies found in a frontend class class. */
 			void CreateAndCacheProxies(const FMetasoundFrontendClass& InClass);
diff --git a/Engine/Plugins/Runtime/Metasound/Source/MetasoundFrontend/Private/MetasoundFrontendRegistryContainerImpl.cpp b/Engine/Plugins/Runtime/Metasound/Source/MetasoundFrontend/Private/MetasoundFrontendRegistryContainerImpl.cpp
--- a/Engine/Plugins/Runtime/Metasound/Source/MetasoundFrontend/Private/MetasoundFrontendRegistryContainerImpl.cpp
+++ b/Engine/Plugins/Runtime/Metasound/Source/MetasoundFrontend/Private/MetasoundFrontendRegistryContainerImpl.cpp
@@ -48,14 +48,13 @@ namespace Metasound::Frontend
 
 	namespace RegistryPrivate
 	{
-		TScriptInterface<IMetaSoundDocumentInterface> BuildRegistryDocument(TScriptInterface<IMetaSoundDocumentInterface> DocumentInterface, bool bAsync, FGuid& OutPageID)
+		TScriptInterface<IMetaSoundDocumentInterface> BuildRegistryDocument(TScriptInterface<IMetaSoundDocumentInterface> DocumentInterface, bool bAsync)
 		{
 			METASOUND_TRACE_CPUPROFILER_EVENT_SCOPE(Metasound::Frontend::BuildRegistryDocument);
 
 			UObject* DocObject = DocumentInterface.GetObject();
 			check(DocObject);
 			const FMetasoundFrontendDocument& Document = DocumentInterface->GetConstDocument();
-			OutPageID = IDocumentBuilderRegistry::GetChecked().ResolveTargetPageID(Document, DocumentInterface->GetAssetPathChecked());
 
 #if WITH_EDITOR
 			// Only assets require template node processing
@@ -274,7 +273,7 @@ namespace Metasound::Frontend
 		};
 	} // namespace RegistryPrivate
 
-	void FRegistryContainerImpl::BuildAndRegisterGraphFromDocument(const FMetasoundFrontendDocument& InDocument, const FProxyDataCache& InProxyDataCache, FNodeClassInfo&& InNodeClassInfo, const FGuid& InPageID)
+	void FRegistryContainerImpl::BuildAndRegisterGraphFromDocument(const FMetasoundFrontendDocument& InDocument, const FProxyDataCache& InProxyDataCache, FNodeClassInfo&& InNodeClassInfo)
 	{
 		using namespace RegistryPrivate;
 
@@ -282,7 +281,7 @@ namespace Metasound::Frontend
 		METASOUND_TRACE_CPUPROFILER_EVENT_SCOPE_TEXT(*FString::Printf(TEXT("Metasound::FRegistryContainerImpl::BuildAndRegisterGraphFromDocument asset %s"), *InNodeClassInfo.AssetPath.ToString()));
 
 		// Use the asset class id for the graph id because it should be locally unique. Unlike other ids, it is regenerated on asset duplicate. 
-		TUniquePtr<FFrontendGraph> FrontendGraph = FFrontendGraphBuilder::CreateGraph(InDocument, InProxyDataCache, InNodeClassInfo.AssetPath.ToString(), /*GraphId=*/InNodeClassInfo.AssetClassID, &InPageID);
+		TUniquePtr<FFrontendGraph> FrontendGraph = FFrontendGraphBuilder::CreateGraph(InDocument, InProxyDataCache, InNodeClassInfo.AssetPath.ToString(), /*GraphId=*/InNodeClassInfo.AssetClassID);
 		if (!FrontendGraph.IsValid())
 		{
 			UE_LOG(LogMetaSound, Error, TEXT("Failed to build MetaSound graph in asset '%s'"), *InNodeClassInfo.AssetPath.ToString());
@@ -480,8 +479,7 @@ namespace Metasound::Frontend
 
 		// Use the asset path of the provided document interface object for identification, *NOT* the
 		// built version as the build process may in fact create a new object with a transient path.
-		FGuid PageID = Frontend::DefaultPageID;
-		const TScriptInterface<IMetaSoundDocumentInterface> RegistryDocInterface = RegistryPrivate::BuildRegistryDocument(InDocumentInterface, bAsync, PageID);
+		const TScriptInterface<IMetaSoundDocumentInterface> RegistryDocInterface = RegistryPrivate::BuildRegistryDocument(InDocumentInterface, bAsync);
 
 		UObject* OwningObject = RegistryDocInterface.GetObject();
 		check(OwningObject);
@@ -489,7 +487,7 @@ namespace Metasound::Frontend
 		// Proxies are created synchronously to avoid creating proxies in async tasks. Proxies
 		// are created from UObjects which need to be protected from GC and non-GT access.
 		FProxyDataCache ProxyDataCache;
-		ProxyDataCache.CreateAndCacheProxies(Document, PageID);
+		ProxyDataCache.CreateAndCacheProxies(Document);
 
 		// Store update to newly registered node in history so nodes
 		// can be queried by transaction ID
@@ -503,12 +501,12 @@ namespace Metasound::Frontend
 		{
 			Tasks::FTask BuildAndRegisterTask = AsyncRegistrationPipe.Launch(
 				UE_SOURCE_LOCATION,
-				[RegistryKey, PageID, ClassInfo = MoveTemp(NodeClassInfo), RegistryDocInterface, ProxyDataCache = MoveTemp(ProxyDataCache)]() mutable
+				[RegistryKey, ClassInfo = MoveTemp(NodeClassInfo), RegistryDocInterface, ProxyDataCache = MoveTemp(ProxyDataCache)]() mutable
 				{
 					FRegistryContainerImpl& Registry = FRegistryContainerImpl::Get();
 					// Unregister the graph before re-registering
 					Registry.UnregisterGraphInternal(RegistryKey);
-					Registry.BuildAndRegisterGraphFromDocument(RegistryDocInterface->GetConstDocument(), ProxyDataCache, MoveTemp(ClassInfo), PageID);
+					Registry.BuildAndRegisterGraphFromDocument(RegistryDocInterface->GetConstDocument(), ProxyDataCache, MoveTemp(ClassInfo));
 					Registry.RemoveRegistrationTask(RegistryKey, FNodeRegistryTransaction::ETransactionType::NodeRegistration);
 					Registry.RemoveDocumentReference(RegistryDocInterface);
 				}
@@ -527,7 +525,7 @@ namespace Metasound::Frontend
 			UnregisterGraphInternal(RegistryKey);
 
 			// Build and register graph synchronously
-			BuildAndRegisterGraphFromDocument(RegistryDocInterface->GetConstDocument(), ProxyDataCache, MoveTemp(NodeClassInfo), PageID);
+			BuildAndRegisterGraphFromDocument(RegistryDocInterface->GetConstDocument(), ProxyDataCache, MoveTemp(NodeClassInfo));
 		}
 
 		return RegistryKey;
diff --git a/Engine/Plugins/Runtime/Metasound/Source/MetasoundFrontend/Private/MetasoundFrontendRegistryContainerImpl.h b/Engine/Plugins/Runtime/Metasound/Source/MetasoundFrontend/Private/MetasoundFrontendRegistryContainerImpl.h
--- a/Engine/Plugins/Runtime/Metasound/Source/MetasoundFrontend/Private/MetasoundFrontendRegistryContainerImpl.h
+++ b/Engine/Plugins/Runtime/Metasound/Source/MetasoundFrontend/Private/MetasoundFrontendRegistryContainerImpl.h
@@ -148,7 +148,7 @@ namespace Metasound
 				FTopLevelAssetPath AssetPath;
 			};
 
-			void BuildAndRegisterGraphFromDocument(const FMetasoundFrontendDocument& InDocument, const FProxyDataCache& InProxyDataCache, FNodeClassInfo&& InNodeClassInfo, const FGuid& InPageID);
+			void BuildAndRegisterGraphFromDocument(const FMetasoundFrontendDocument& InDocument, const FProxyDataCache& InProxyDataCache, FNodeClassInfo&& InNodeClassInfo);
 
 			void AddRegistrationTask(const FGraphRegistryKey& InKey, FActiveRegistrationTaskInfo&& TaskInfo);
 			void RemoveRegistrationTask(const FGraphRegistryKey& InKey, FNodeRegistryTransaction::ETransactionType TransactionType);
diff --git a/Engine/Plugins/Runtime/Metasound/Source/MetasoundFrontend/Private/NodeTemplates/MetasoundFrontendNodeTemplateInput.cpp b/Engine/Plugins/Runtime/Metasound/Source/MetasoundFrontend/Private/NodeTemplates/MetasoundFrontendNodeTemplateInput.cpp
--- a/Engine/Plugins/Runtime/Metasound/Source/MetasoundFrontend/Private/NodeTemplates/MetasoundFrontendNodeTemplateInput.cpp
+++ b/Engine/Plugins/Runtime/Metasound/Source/MetasoundFrontend/Private/NodeTemplates/MetasoundFrontendNodeTemplateInput.cpp
@@ -63,7 +63,7 @@ namespace Metasound::Frontend
 				// If connections are present, add a location for safety.  Try adding near existing node.
 				if (!ConnectedVertices.IsEmpty())
 				{
-					UE_LOG(LogMetaSound, Warning, TEXT("Template node being generated for input '%s' had no editor location set.  Procedurally placing near connected node."), *InputName.ToString());
+					UE_LOG(LogMetaSound, Display, TEXT("Template node being generated for input '%s' had no editor location set.  Procedurally placing near connected node."), *InputName.ToString());
 					FVector2D NewLocation;
 					if (const FMetasoundFrontendNode* Node = InOutBuilder.FindNode(ConnectedVertices.Last().NodeID))
 					{
@@ -148,6 +148,12 @@ namespace Metasound::Frontend
 	}
 #endif // WITH_EDITOR
 
+	const TArray<FMetasoundFrontendClassInputDefault>* FInputNodeTemplate::FindNodeClassInputDefaults(const FMetaSoundFrontendDocumentBuilder& InBuilder, const FGuid& InPageID, const FGuid& InNodeID, FName VertexName) const
+	{
+		// Just returns the default of the given node's class input and not walk to values provided by connected input like reroutes.
+		return FNodeTemplateBase::FindNodeClassInputDefaults(InBuilder, InPageID, InNodeID, VertexName);
+	}
+
 	const FMetasoundFrontendClassName& FInputNodeTemplate::GetClassName() const
 	{
 		return FInputNodeTemplate::ClassName;
diff --git a/Engine/Plugins/Runtime/Metasound/Source/MetasoundFrontend/Private/NodeTemplates/MetasoundFrontendNodeTemplateReroute.cpp b/Engine/Plugins/Runtime/Metasound/Source/MetasoundFrontend/Private/NodeTemplates/MetasoundFrontendNodeTemplateReroute.cpp
--- a/Engine/Plugins/Runtime/Metasound/Source/MetasoundFrontend/Private/NodeTemplates/MetasoundFrontendNodeTemplateReroute.cpp
+++ b/Engine/Plugins/Runtime/Metasound/Source/MetasoundFrontend/Private/NodeTemplates/MetasoundFrontendNodeTemplateReroute.cpp
@@ -245,6 +245,35 @@ namespace Metasound::Frontend
 		return FrontendClass;
 	}
 
+	const TArray<FMetasoundFrontendClassInputDefault>* FRerouteNodeTemplate::FindNodeClassInputDefaults(const FMetaSoundFrontendDocumentBuilder& InBuilder, const FGuid& InPageID, const FGuid& InNodeID, FName VertexName) const
+	{
+		// Recursive search up DAG for first connected non-reroute node's input class input
+		if (const FMetasoundFrontendNode* Node = InBuilder.FindNode(InNodeID, &InPageID))
+		{
+			// Should only ever be one
+			const FMetasoundFrontendVertex& RerouteOutput = Node->Interface.Outputs.Last();
+
+			TArray<const FMetasoundFrontendNode*> ConnectedNodes;
+			TArray<const FMetasoundFrontendVertex*> ConnectedInputs = InBuilder.FindNodeInputsConnectedToNodeOutput(InNodeID, RerouteOutput.VertexID, &ConnectedNodes, &InPageID);
+			for (int32 Index = 0; Index < ConnectedNodes.Num(); ++Index)
+			{
+				const FMetasoundFrontendNode* ConnectedNode = ConnectedNodes[Index];
+				if (const FMetasoundFrontendClass* ConnectedNodeClass = InBuilder.FindDependency(ConnectedNode->ClassID))
+				{
+					const FMetasoundFrontendVertex* ConnectedInput = ConnectedInputs[Index];
+					if (ConnectedNodeClass->Metadata.GetClassName() == GetClassName())
+					{
+						return this->FindNodeClassInputDefaults(InBuilder, InPageID, ConnectedNode->GetID(), ConnectedInput->Name);
+					}
+
+					return InBuilder.FindNodeClassInputDefaults(ConnectedNode->GetID(), ConnectedInput->Name, &InPageID);
+				}
+			}
+		}
+
+		return nullptr;
+	}
+
 	EMetasoundFrontendVertexAccessType FRerouteNodeTemplate::GetNodeInputAccessType(const FMetaSoundFrontendDocumentBuilder& InBuilder, const FGuid& InPageID, const FGuid& InNodeID, const FGuid& InVertexID) const
 	{
 		// Recursive search up DAG for first connected non-reroute node's input access type
diff --git a/Engine/Plugins/Runtime/Metasound/Source/MetasoundFrontend/Public/MetasoundDocumentInterface.h b/Engine/Plugins/Runtime/Metasound/Source/MetasoundFrontend/Public/MetasoundDocumentInterface.h
--- a/Engine/Plugins/Runtime/Metasound/Source/MetasoundFrontend/Public/MetasoundDocumentInterface.h
+++ b/Engine/Plugins/Runtime/Metasound/Source/MetasoundFrontend/Public/MetasoundDocumentInterface.h
@@ -106,8 +106,17 @@ namespace Metasound::Frontend
 		// false if not found.
 		virtual bool ReloadBuilder(const FMetasoundFrontendClassName& InClassName) const = 0;
 
-		// Given the provided document and its respective pages, returns the PageID to be used for runtime IGraph and proxy generation.
-		virtual FGuid ResolveTargetPageID(const FMetasoundFrontendDocument& Document, const FTopLevelAssetPath& AssetPath) const = 0;
+		// Given the provided document and its respective pages, provides the PageID to be used for runtime IGraph and proxy generation.
+		// Returns true if valid page was found, false if not.
+		virtual bool TryResolveTargetPageID(const FMetasoundFrontendGraphClass& InGraphClass, FGuid& OutResolvedPageID) const = 0;
+
+		// Given the provided input and its respective paged default values, provides the default PageID to be used for runtime IGraph and proxy generation.
+		// Returns true if valid page was found, false if not.
+		virtual bool TryResolveTargetPageID(const FMetasoundFrontendClassInput& InClassInput, FGuid& OutResolvedPageID) const = 0;
+
+		// Given the provided array of default values, provides the default PageID to be used for runtime IGraph and proxy generation.
+		// Returns true if valid page was found, false if not.
+		virtual bool TryResolveTargetPageID(const TArray<FMetasoundFrontendClassInputDefault>& Defaults, FGuid& OutResolvedPageID) const = 0;
 
 		static IDocumentBuilderRegistry* Get();
 		static IDocumentBuilderRegistry& GetChecked();
diff --git a/Engine/Plugins/Runtime/Metasound/Source/MetasoundFrontend/Public/MetasoundFrontendDocument.h b/Engine/Plugins/Runtime/Metasound/Source/MetasoundFrontend/Public/MetasoundFrontendDocument.h
--- a/Engine/Plugins/Runtime/Metasound/Source/MetasoundFrontend/Public/MetasoundFrontendDocument.h
+++ b/Engine/Plugins/Runtime/Metasound/Source/MetasoundFrontend/Public/MetasoundFrontendDocument.h
@@ -894,6 +894,7 @@ struct METASOUNDFRONTEND_API FMetasoundFrontendClassVertex : public FMetasoundFr
 	void SplitName(FName& OutNamespace, FName& OutParameterName) const;
 
 	static bool IsFunctionalEquivalent(const FMetasoundFrontendClassVertex& InLHS, const FMetasoundFrontendClassVertex& InRHS);
+
 	// Whether vertex access types are compatible when connecting from an output to an input 
 	static bool CanConnectVertexAccessTypes(EMetasoundFrontendVertexAccessType InFromType, EMetasoundFrontendVertexAccessType InToType);
 };
@@ -942,14 +943,15 @@ struct METASOUNDFRONTEND_API FMetasoundFrontendClassInputDefault
 
 	FMetasoundFrontendClassInputDefault() = default;
 
-	FMetasoundFrontendClassInputDefault(const FMetasoundFrontendLiteral& Literal)
-		: Literal(Literal)
+	FMetasoundFrontendClassInputDefault(const FMetasoundFrontendLiteral& InLiteral)
+		: Literal(InLiteral)
 		, PageID(Metasound::Frontend::DefaultPageID)
 	{
 	}
 
-	FMetasoundFrontendClassInputDefault(const FGuid& InPageID)
-		: PageID(InPageID)
+	FMetasoundFrontendClassInputDefault(const FGuid& InPageID, FMetasoundFrontendLiteral InLiteral = { })
+		: Literal(MoveTemp(InLiteral))
+		, PageID(InPageID)
 	{
 	}
 
@@ -994,6 +996,7 @@ public:
 	const FMetasoundFrontendLiteral& FindConstDefaultChecked(const FGuid& InPageID) const;
 	FMetasoundFrontendLiteral* FindDefault(const FGuid& InPageID);
 	FMetasoundFrontendLiteral& FindDefaultChecked(const FGuid& InPageID);
+	const TArray<FMetasoundFrontendClassInputDefault>& GetDefaults() const;
 	FMetasoundFrontendLiteral& InitDefault();
 	void InitDefault(FMetasoundFrontendLiteral InitLiteral);
 	void IterateDefaults(TFunctionRef<void(const FGuid&, FMetasoundFrontendLiteral&)> IterFunc);
@@ -1002,6 +1005,7 @@ public:
 #if WITH_EDITORONLY_DATA
 	void RemoveAllDefaults();
 	bool RemoveDefault(const FGuid& InPageID);
+	void SetDefaults(TArray<FMetasoundFrontendClassInputDefault> InputDefaults);
 #endif // WITH_EDITORONLY_DATA
 };
 
@@ -1012,7 +1016,6 @@ struct METASOUNDFRONTEND_API FMetasoundFrontendClassVariable : public FMetasound
 	GENERATED_BODY()
 
 	FMetasoundFrontendClassVariable() = default;
-
 	FMetasoundFrontendClassVariable(const FMetasoundFrontendClassVertex& InOther);
 
 	// Default value for this variable.
diff --git a/Engine/Plugins/Runtime/Metasound/Source/MetasoundFrontend/Public/MetasoundFrontendDocumentBuilder.h b/Engine/Plugins/Runtime/Metasound/Source/MetasoundFrontend/Public/MetasoundFrontendDocumentBuilder.h
--- a/Engine/Plugins/Runtime/Metasound/Source/MetasoundFrontend/Public/MetasoundFrontendDocumentBuilder.h
+++ b/Engine/Plugins/Runtime/Metasound/Source/MetasoundFrontend/Public/MetasoundFrontendDocumentBuilder.h
@@ -239,6 +239,16 @@ public:
 
 	const FMetasoundFrontendVertex* FindNodeInput(const FGuid& InNodeID, const FGuid& InVertexID, const FGuid* InPageID = nullptr) const;
 	const FMetasoundFrontendVertex* FindNodeInput(const FGuid& InNodeID, FName InVertexName, const FGuid* InPageID = nullptr) const;
+
+	// Returns class defaults associated with the given node input (as defined in the associated node's dependency)
+	const TArray<FMetasoundFrontendClassInputDefault>* FindNodeClassInputDefaults(const FGuid& InNodeID, FName InVertexName, const FGuid* InPageID = nullptr) const;
+
+	// Returns node input's vertex default if valid and assigned.
+	const FMetasoundFrontendVertexLiteral* FindNodeInputDefault(const FGuid& InNodeID, const FGuid& InVertexID, const FGuid* InPageID = nullptr) const;
+
+	// Returns node input's vertex default if valid and assigned.
+	const FMetasoundFrontendVertexLiteral* FindNodeInputDefault(const FGuid& InNodeID, FName InVertexName, const FGuid* InPageID = nullptr) const;
+
 	TArray<const FMetasoundFrontendVertex*> FindNodeInputs(const FGuid& InNodeID, FName TypeName = FName(), const FGuid* InPageID = nullptr) const;
 	TArray<const FMetasoundFrontendVertex*> FindNodeInputsConnectedToNodeOutput(const FGuid& InOutputNodeID, const FGuid& InOutputVertexID, TArray<const FMetasoundFrontendNode*>* ConnectedInputNodes = nullptr, const FGuid* InPageID = nullptr) const;
 
diff --git a/Engine/Plugins/Runtime/Metasound/Source/MetasoundFrontend/Public/MetasoundFrontendGraph.h b/Engine/Plugins/Runtime/Metasound/Source/MetasoundFrontend/Public/MetasoundFrontendGraph.h
--- a/Engine/Plugins/Runtime/Metasound/Source/MetasoundFrontend/Public/MetasoundFrontendGraph.h
+++ b/Engine/Plugins/Runtime/Metasound/Source/MetasoundFrontend/Public/MetasoundFrontendGraph.h
@@ -104,13 +104,13 @@ namespace Metasound
 		static TUniquePtr<FFrontendGraph> CreateGraph(const FMetasoundFrontendGraphClass& InGraphClass, const TArray<FMetasoundFrontendGraphClass>& InSubgraphs, const TArray<FMetasoundFrontendClass>& InDependencies, const TSet<FName>& TransmittableInputNames, const FString& InDebugAssetName);
 
 		/* Create a FFrontendGraph from a FMetasoundFrontendDocument.*/
-		static TUniquePtr<FFrontendGraph> CreateGraph(const FMetasoundFrontendDocument& InDocument, const FString& InDebugAssetName, const FGuid* InPageId = nullptr);
+		static TUniquePtr<FFrontendGraph> CreateGraph(const FMetasoundFrontendDocument& InDocument, const FString& InDebugAssetName);
 
 		/* Create a FFrontendGraph from a FMetasoundFrontendDocument retrieving proxies from a FProxyDataCache.*/
-		static TUniquePtr<FFrontendGraph> CreateGraph(const FMetasoundFrontendDocument& InDocument, const Frontend::FProxyDataCache& InProxies, const FString& InDebugAssetName, const FGuid InGraphId = Frontend::CreateLocallyUniqueId(), const FGuid* InPageId = nullptr);
+		static TUniquePtr<FFrontendGraph> CreateGraph(const FMetasoundFrontendDocument& InDocument, const Frontend::FProxyDataCache& InProxies, const FString& InDebugAssetName, const FGuid InGraphId = Frontend::CreateLocallyUniqueId());
 
 		/* Create a FFrontendGraph from a FMetasoundFrontendDocument and subobjects.*/
-		static TUniquePtr<FFrontendGraph> CreateGraph(const FMetasoundFrontendGraphClass& InGraph, const TArray<FMetasoundFrontendGraphClass>& InSubgraphs, const TArray<FMetasoundFrontendClass>& InDependencies, const FString& InDebugAssetName, const FGuid* InPageId = nullptr);
+		static TUniquePtr<FFrontendGraph> CreateGraph(const FMetasoundFrontendGraphClass& InGraph, const TArray<FMetasoundFrontendGraphClass>& InSubgraphs, const TArray<FMetasoundFrontendClass>& InDependencies, const FString& InDebugAssetName);
 
 		/* Create a FFrontendGraph from a FMetasoundFrontendDocument and subobjects, retrieving proxies from a FProxyDataCache.*/
 		static TUniquePtr<FFrontendGraph> CreateGraph(
@@ -119,8 +119,7 @@ namespace Metasound
 			const TArray<FMetasoundFrontendClass>& InDependencies,
 			const Frontend::FProxyDataCache& InProxyDataCache,
 			const FString& InDebugAssetName,
-			const FGuid InGraphId = Frontend::CreateLocallyUniqueId(),
-			const FGuid* InPageId = nullptr);
+			const FGuid InGraphId = Frontend::CreateLocallyUniqueId());
 
 	private:
 		struct FDefaultLiteralData
@@ -165,7 +164,7 @@ namespace Metasound
 
 		static bool SortSubgraphDependencies(TArray<const FMetasoundFrontendGraphClass*>& Subgraphs);
 
-		static TUniquePtr<FFrontendGraph> CreateGraph(FBuildContext& InContext, const FMetasoundFrontendGraphClass& InSubgraph, const FGuid* InPageId = nullptr);
+		static TUniquePtr<FFrontendGraph> CreateGraph(FBuildContext& InContext, const FMetasoundFrontendGraphClass& InSubgraph);
 
 		static const FMetasoundFrontendClassInput* FindClassInputForInputNode(const FMetasoundFrontendGraphClass& InOwningGraph, const FMetasoundFrontendNode& InInputNode, int32& OutClassInputIndex);
 		static const FMetasoundFrontendClassOutput* FindClassOutputForOutputNode(const FMetasoundFrontendGraphClass& InOwningGraph, const FMetasoundFrontendNode& InOutputNode, int32& OutClassOutputIndex);
diff --git a/Engine/Plugins/Runtime/Metasound/Source/MetasoundFrontend/Public/MetasoundFrontendNodeTemplateRegistry.h b/Engine/Plugins/Runtime/Metasound/Source/MetasoundFrontend/Public/MetasoundFrontendNodeTemplateRegistry.h
--- a/Engine/Plugins/Runtime/Metasound/Source/MetasoundFrontend/Public/MetasoundFrontendNodeTemplateRegistry.h
+++ b/Engine/Plugins/Runtime/Metasound/Source/MetasoundFrontend/Public/MetasoundFrontendNodeTemplateRegistry.h
@@ -64,6 +64,12 @@ namespace Metasound::Frontend
 
 		virtual ~INodeTemplate() = default;
 
+		virtual const TArray<FMetasoundFrontendClassInputDefault>* FindNodeClassInputDefaults(
+			const FMetaSoundFrontendDocumentBuilder& InBuilder,
+			const FGuid& InPageID,
+			const FGuid& InNodeID,
+			FName VertexName) const = 0;
+
 		// Returns note template class name.
 		virtual const FMetasoundFrontendClassName& GetClassName() const = 0;
 
@@ -148,6 +154,8 @@ namespace Metasound::Frontend
 	class METASOUNDFRONTEND_API FNodeTemplateBase : public INodeTemplate
 	{
 	public:
+		virtual const TArray<FMetasoundFrontendClassInputDefault>* FindNodeClassInputDefaults(const FMetaSoundFrontendDocumentBuilder& InBuilder, const FGuid& InPageID, const FGuid& InNodeID, FName VertexName) const override;
+
 #if WITH_EDITOR
 		virtual FText GetNodeDisplayName(const IMetaSoundDocumentInterface& Interface, const FGuid& InPageID, const FGuid& InNodeID) const override;
 		virtual FText GetInputVertexDisplayName(const FMetaSoundFrontendDocumentBuilder& InBuilder, const FGuid& InPageID, const FGuid& InNodeID, FName InputName) const override;
diff --git a/Engine/Plugins/Runtime/Metasound/Source/MetasoundFrontend/Public/NodeTemplates/MetasoundFrontendNodeTemplateInput.h b/Engine/Plugins/Runtime/Metasound/Source/MetasoundFrontend/Public/NodeTemplates/MetasoundFrontendNodeTemplateInput.h
--- a/Engine/Plugins/Runtime/Metasound/Source/MetasoundFrontend/Public/NodeTemplates/MetasoundFrontendNodeTemplateInput.h
+++ b/Engine/Plugins/Runtime/Metasound/Source/MetasoundFrontend/Public/NodeTemplates/MetasoundFrontendNodeTemplateInput.h
@@ -37,6 +37,7 @@ namespace Metasound::Frontend
 		static const FMetasoundFrontendNode* CreateNode(FMetaSoundFrontendDocumentBuilder& InOutBuilder, FName InputName, const FGuid* InPageID = nullptr);
 #endif // WITH_EDITOR
 
+		virtual const TArray<FMetasoundFrontendClassInputDefault>* FindNodeClassInputDefaults(const FMetaSoundFrontendDocumentBuilder& InBuilder, const FGuid& InPageID, const FGuid& InNodeID, FName VertexName) const override;
 		virtual const FMetasoundFrontendClassName& GetClassName() const override;
 
 #if WITH_EDITOR
diff --git a/Engine/Plugins/Runtime/Metasound/Source/MetasoundFrontend/Public/NodeTemplates/MetasoundFrontendNodeTemplateReroute.h b/Engine/Plugins/Runtime/Metasound/Source/MetasoundFrontend/Public/NodeTemplates/MetasoundFrontendNodeTemplateReroute.h
--- a/Engine/Plugins/Runtime/Metasound/Source/MetasoundFrontend/Public/NodeTemplates/MetasoundFrontendNodeTemplateReroute.h
+++ b/Engine/Plugins/Runtime/Metasound/Source/MetasoundFrontend/Public/NodeTemplates/MetasoundFrontendNodeTemplateReroute.h
@@ -38,6 +38,7 @@ namespace Metasound::Frontend
 		virtual FMetasoundFrontendNodeInterface GenerateNodeInterface(FNodeTemplateGenerateInterfaceParams InParams) const override;
 		virtual TUniquePtr<INodeTemplateTransform> GenerateNodeTransform() const override;
 		virtual const FMetasoundFrontendClass& GetFrontendClass() const override;
+		virtual const TArray<FMetasoundFrontendClassInputDefault>* FindNodeClassInputDefaults(const FMetaSoundFrontendDocumentBuilder& InBuilder, const FGuid& InPageID, const FGuid& InNodeID, FName VertexName) const override;
 		virtual EMetasoundFrontendVertexAccessType GetNodeInputAccessType(const FMetaSoundFrontendDocumentBuilder& InBuilder, const FGuid& InPageID, const FGuid& InNodeID, const FGuid& InVertexID) const override;
 		virtual EMetasoundFrontendVertexAccessType GetNodeOutputAccessType(const FMetaSoundFrontendDocumentBuilder& InBuilder, const FGuid& InPageID, const FGuid& InNodeID, const FGuid& InVertexID) const override;
 		virtual const FMetasoundFrontendVersionNumber& GetVersionNumber() const override;
