- Add page info to render stats in MetaSound Editor
- Implement GetOptions feedback for simpler audition page UX
- Tidy connection manager init in asset editor & only visualize connections if viewed page is audition page
#rnx

diff --git a/Engine/Plugins/Runtime/Metasound/Source/MetasoundEditor/Private/MetasoundDetailCustomization.cpp b/Engine/Plugins/Runtime/Metasound/Source/MetasoundEditor/Private/MetasoundDetailCustomization.cpp
--- a/Engine/Plugins/Runtime/Metasound/Source/MetasoundEditor/Private/MetasoundDetailCustomization.cpp
+++ b/Engine/Plugins/Runtime/Metasound/Source/MetasoundEditor/Private/MetasoundDetailCustomization.cpp
@@ -442,15 +442,12 @@ namespace Metasound::Editor
 
 		Utilities->SetEnabled(TAttribute<bool>::Create(TAttribute<bool>::FGetter::CreateSP(this, &FMetaSoundDetailCustomizationBase::IsGraphEditable)));
 
-		const FText HeaderName = LOCTEXT("PagesGroupDisplayName", "Pages");
-		IDetailCategoryBuilder& Category = DetailLayout.EditCategory(FName(ItemName), HeaderName);
-		Category.AddCustomRow(HeaderName) [ Utilities ];
-
-		EntryWidgets = SNew(SVerticalBox);
-		Category.AddCustomRow(LOCTEXT("ImplementedPagesLabel", "Implemented Pages"))
-		[
-			EntryWidgets->AsShared()
-		];
+		{
+			const FText HeaderName = LOCTEXT("PageGraphsDisplayName", "Pages");
+			IDetailCategoryBuilder& Category = DetailLayout.EditCategory("Pages", HeaderName);
+			Category.AddCustomRow(HeaderName)[Utilities];
+			Category.AddCustomRow(LOCTEXT("ImplementedPagesLabel", "Graphs")) [ SAssignNew(EntryWidgets, SVerticalBox) ];
+		}
 
 		// Registration of page listener instance calls OnReload which in turn causes RefreshView, so no need to call directly
 		if (UMetasoundEditorViewBase* View = CastChecked<UMetasoundEditorViewBase>(Objects.Last()))
diff --git a/Engine/Plugins/Runtime/Metasound/Source/MetasoundEditor/Private/MetasoundEditor.cpp b/Engine/Plugins/Runtime/Metasound/Source/MetasoundEditor/Private/MetasoundEditor.cpp
--- a/Engine/Plugins/Runtime/Metasound/Source/MetasoundEditor/Private/MetasoundEditor.cpp
+++ b/Engine/Plugins/Runtime/Metasound/Source/MetasoundEditor/Private/MetasoundEditor.cpp
@@ -76,7 +76,7 @@
 #include "ScopedTransaction.h"
 #include "SMetasoundActionMenu.h"
 #include "SMetasoundPalette.h"
-#include "SMetasoundRenderStats.h"
+#include "SMetasoundStats.h"
 #include "SNodePanel.h"
 #include "Stats/Stats.h"
 #include "Styling/AppStyle.h"
@@ -848,12 +848,34 @@ namespace Metasound
 
 			InTabManager->RegisterTabSpawner(TabNamesPrivate::GraphCanvas, FOnSpawnTab::CreateLambda(
 				[
+					InFocusPageWidget = FocusPageWidget,
 					InPlayTimeWidget = PlayTimeWidget,
 					InMetasoundGraphEditor = MetasoundGraphEditor,
 					InRenderStatsWidget = RenderStatsWidget
 				](const FSpawnTabArgs& Args)
 			{
 				TSharedRef<SDockTab> SpawnedTab = SNew(SDockTab).Label(LOCTEXT("MetasoundGraphCanvasTitle", "MetaSound Graph"));
+
+				TSharedPtr<SVerticalBox> StatsWidget = SNew(SVerticalBox)
+					+ SVerticalBox::Slot()
+					.HAlign(HAlign_Left)
+					.VAlign(VAlign_Center)
+					.Padding(2.f, 1.f)
+					[
+						InRenderStatsWidget.ToSharedRef()
+					];
+
+				if (AssetEditorPrivate::EnablePageEditor)
+				{
+					StatsWidget->AddSlot()
+						.HAlign(HAlign_Left)
+						.VAlign(VAlign_Center)
+						.AutoHeight()
+						[
+							InFocusPageWidget.ToSharedRef()
+						];
+				}
+
 				SpawnedTab->SetContent(SNew(SOverlay)
 					+ SOverlay::Slot()
 					[
@@ -866,7 +888,7 @@ namespace Metasound
 					+ SOverlay::Slot()
 					.VAlign(VAlign_Bottom)
 					[
-						InRenderStatsWidget.ToSharedRef()
+						StatsWidget.ToSharedRef()
 					]
 					.Padding(5.0f, 5.0f)
 				);
@@ -975,6 +997,20 @@ namespace Metasound
 				BackgroundColor = MetaSoundStyle->GetColor("MetasoundEditor.Analyzers.BackgroundColor");
 			}
 
+			TSharedRef<SVerticalBox> PageWidget = SNew(SVerticalBox)
+				.Visibility(TAttribute<EVisibility>::Create([this]()
+				{
+					const bool bMakeVisible = IsPlaying() && AssetEditorPrivate::EnablePageEditor;
+					return bMakeVisible ? EVisibility::Visible : EVisibility::Collapsed;
+				}))
+				+ SVerticalBox::Slot()
+				.HAlign(HAlign_Left)
+				.VAlign(VAlign_Center)
+				.AutoHeight()
+				[
+					RenderPageWidget.ToSharedRef()
+				];
+
 			return SNew(SOverlay)
 			+ SOverlay::Slot()
 			[
@@ -986,7 +1022,12 @@ namespace Metasound
 				SNew(SSplitter)
 				.Orientation(Orient_Vertical)
 				+ SSplitter::Slot()
-				.Value(0.6f)
+					.Value(0.05f)
+					[
+						PageWidget
+					]
+				+ SSplitter::Slot()
+				.Value(0.5f)
 				[
 					SNew(SVerticalBox)
 					+ SVerticalBox::Slot()
@@ -997,17 +1038,17 @@ namespace Metasound
 					]
 				]
 				+ SSplitter::Slot()
-				.Value(0.2f)
+				.Value(0.15f)
 				[
 					OutputOscilloscope->GetPanelWidget()
 				]
 				+ SSplitter::Slot()
-				.Value(0.2f)
+				.Value(0.15f)
 				[
 					OutputVectorscope->GetPanelWidget()
 				]
 				+ SSplitter::Slot()
-				.Value(0.2f)
+				.Value(0.15f)
 				[
 					OutputSpectrumAnalyzer->GetWidget()
 				]
@@ -1202,6 +1243,10 @@ namespace Metasound
 			{
 				Play();
 			}
+			else
+			{
+				UpdatePageInfo(false);
+			}
 
 			FSlateApplication::Get().SetUserFocus(0, MetasoundGraphEditor);
 		}
@@ -1740,7 +1785,7 @@ namespace Metasound
 			DestroyAnalyzers();
 		}
 
-		TSharedRef<SWidget> FEditor::CreateAuditionMenuOptions() const
+		TSharedRef<SWidget> FEditor::CreateAuditionMenuOptions()
 		{
 			TSharedPtr<FUICommandList> Commands = MakeShared<FUICommandList>();
 			FMenuBuilder MenuBuilder(true, Commands);
@@ -1761,7 +1806,7 @@ namespace Metasound
 			return MenuBuilder.MakeWidget();
 		}
 
-		void FEditor::CreateAuditionPageSubMenuOptions(FMenuBuilder& MenuBuilder) const
+		void FEditor::CreateAuditionPageSubMenuOptions(FMenuBuilder& MenuBuilder)
 		{
 			TWeakObjectPtr<const UMetaSoundSettings> SettingsPtr = GetDefault<UMetaSoundSettings>();
 			if (!SettingsPtr.IsValid())
@@ -1787,6 +1832,7 @@ namespace Metasound
 						if (UMetasoundEditorSettings* EditorSettings = GetMutableDefault<UMetasoundEditorSettings>())
 						{
 							EditorSettings->AuditionPlatform = PlatformName;
+							Stop();
 						}
 					});
 
@@ -1873,6 +1919,7 @@ namespace Metasound
 						if (UMetasoundEditorSettings* EditorSettings = GetMutableDefault<UMetasoundEditorSettings>())
 						{
 							EditorSettings->AuditionTargetPage = AuditionTargetPage;
+							Stop();
 						}
 					});
 
@@ -2305,6 +2352,8 @@ namespace Metasound
 
 		void FEditor::Play()
 		{
+			using namespace Engine;
+
 			if (USoundBase* MetaSoundToPlay = Cast<USoundBase>(GetMetasoundObject()))
 			{
 				HighestMessageSeverity = GetMetaSoundGraphChecked().GetHighestMessageSeverity();
@@ -2320,83 +2369,55 @@ namespace Metasound
 
 				// Set the send to the audio bus that is used for analyzing the metasound output
 				check(GEditor);
-				if (UAudioComponent* PreviewComp = GEditor->PlayPreviewSound(MetaSoundToPlay))
-				{
-					PlayTime = 0.0;
-
-					UObject* ParamInterfaceObject = PreviewComp;
-					if (ensure(ParamInterfaceObject))
-					{
-						SetPreviewID(ParamInterfaceObject->GetUniqueID());
-					}
-
-					if (UAudioBus* AudioBus = OutputMeter->GetAudioBus())
-					{
-						PreviewComp->SetAudioBusSendPostEffect(AudioBus, 1.0f);
-					}
 
-					if (UAudioBus* AudioBus = OutputOscilloscope->GetAudioBus())
-					{
-						PreviewComp->SetAudioBusSendPostEffect(AudioBus, 1.0f);
-					}
+				PlayTime = 0.0;
+				UpdatePageInfo(true);
 
-					if (UAudioBus* AudioBus = OutputVectorscope->GetAudioBus())
+				if (UMetaSoundSource* Source = Cast<UMetaSoundSource>(GetMetasoundObject()))
+				{
+					if (UAudioComponent* PreviewComp = GEditor->PlayPreviewSound(Source))
 					{
-						PreviewComp->SetAudioBusSendPostEffect(AudioBus, 1.0f);
-					}
+						SetPreviewID(PreviewComp->GetUniqueID());
 
-					if (UAudioBus* AudioBus = OutputSpectrumAnalyzer->GetAudioBus())
-					{
-						PreviewComp->SetAudioBusSendPostEffect(AudioBus, 1.0f);
-					}
+						if (UAudioBus* AudioBus = OutputMeter->GetAudioBus())
+						{
+							PreviewComp->SetAudioBusSendPostEffect(AudioBus, 1.0f);
+						}
 
-					FMetasoundAssetBase* MetaSoundAsset = IMetasoundUObjectRegistry::Get().GetObjectAsAssetBase(MetaSoundToPlay);
-					check(MetaSoundAsset);
+						if (UAudioBus* AudioBus = OutputOscilloscope->GetAudioBus())
+						{
+							PreviewComp->SetAudioBusSendPostEffect(AudioBus, 1.0f);
+						}
 
-					FAudioDevice* AudioDevice = PreviewComp->GetAudioDevice();
-					check(AudioDevice);
+						if (UAudioBus* AudioBus = OutputVectorscope->GetAudioBus())
+						{
+							PreviewComp->SetAudioBusSendPostEffect(AudioBus, 1.0f);
+						}
 
-					const FName& AudioBufferTypeName = GetMetasoundDataTypeName<FAudioBuffer>();
-					const FSampleRate DeviceSampleRate = static_cast<FSampleRate>(AudioDevice->GetSampleRate());
-					const uint32 PlayOrder = PreviewComp->GetLastPlayOrder();
-					const uint64 TransmitterID = Audio::GetTransmitterID(PreviewComp->GetAudioComponentID(), 0, PlayOrder);
+						if (UAudioBus* AudioBus = OutputSpectrumAnalyzer->GetAudioBus())
+						{
+							PreviewComp->SetAudioBusSendPostEffect(AudioBus, 1.0f);
+						}
 
-					UMetaSoundSource* Source = CastChecked<UMetaSoundSource>(MetaSoundToPlay);
-					GraphConnectionManager = MakeUnique<FGraphConnectionManager>(
-						*MetaSoundAsset, *PreviewComp, TransmitterID, Source->GetOperatorSettings(DeviceSampleRate));
+						GraphConnectionManager = RebuildConnectionManager(PreviewComp);
+					}
 				}
 
 				MetasoundGraphEditor->RegisterActiveTimer(0.0f,
 					FWidgetActiveTimerDelegate::CreateLambda([this](double InCurrentTime, float InDeltaTime)
 					{
-						const bool bIsPlaying = IsPlaying();
-						if (RenderStatsWidget.IsValid())
-						{
-							check(IsInGameThread());
-							RenderStatsWidget->Update(bIsPlaying, Cast<const UMetaSoundSource>(GetMetasoundObject()));
-						}
+						const bool bIsPlaying = UpdatePlayTime(InDeltaTime);
+						UpdateRenderInfo();
 
 						if (bIsPlaying)
 						{
-							if (PlayTimeWidget.IsValid())
-							{
-								PlayTime += InDeltaTime;
-								FString PlayTimeString = FTimespan::FromSeconds(PlayTime).ToString();
-
-								// Remove leading '+'
-								PlayTimeString.ReplaceInline(TEXT("+"), TEXT(""));
-								PlayTimeWidget->SetText(FText::FromString(PlayTimeString));
-							}
-
 							return EActiveTimerReturnType::Continue;
 						}
 						else
 						{
-							SetPreviewID(INDEX_NONE);
 							PlayTime = 0.0;
-							PlayTimeWidget->SetText(FText::GetEmpty());
-							GraphConnectionManager = MakeUnique<FGraphConnectionManager>();
-
+							UpdatePageInfo(false);
+							GraphConnectionManager = RebuildConnectionManager();
 							return EActiveTimerReturnType::Stop;
 						}
 					})
@@ -2492,12 +2513,11 @@ namespace Metasound
 		{
 			if (const UMetaSoundSettings* Settings = GetDefault<UMetaSoundSettings>())
 			{
-				const FMetaSoundFrontendDocumentBuilder& DocBuilder = Builder->GetConstBuilder();
-				if (const FMetaSoundPageSettings* PageSettings = Settings->FindPageSettings(DocBuilder.GetBuildPageID()))
+				if (Builder.IsValid())
 				{
-					constexpr bool bFocusPageEditor = false; // Already Focused
-					EMetaSoundBuilderResult Result = EMetaSoundBuilderResult::Failed;
-					UMetaSoundEditorSubsystem::GetChecked().SetFocusedPage(Builder.Get(), PageSettings->Name, bFocusPageEditor, Result);
+					constexpr bool bOpenEditor = false; // Already Focused
+					const FMetaSoundFrontendDocumentBuilder& DocBuilder = Builder->GetConstBuilder();
+					UMetaSoundEditorSubsystem::GetChecked().SetFocusedPage(*Builder.Get(), DocBuilder.GetBuildPageID(), bOpenEditor);
 				}
 			}
 		}
@@ -2845,7 +2865,13 @@ namespace Metasound
 				.TextStyle(FAppStyle::Get(), "Graph.ZoomText")
 				.ColorAndOpacity(FLinearColor(1, 1, 1, 0.30f));
 
-			SAssignNew(RenderStatsWidget, SMetaSoundRenderStats)
+			SAssignNew(FocusPageWidget, SPageStats)
+				.Visibility(EVisibility::HitTestInvisible);
+
+			SAssignNew(RenderPageWidget, SPageStats)
+				.Visibility(EVisibility::HitTestInvisible);
+
+			SAssignNew(RenderStatsWidget, SRenderStats)
 				.Visibility(EVisibility::HitTestInvisible);
 		}
 
@@ -4622,6 +4648,107 @@ namespace Metasound
 			bMemberRenameRequested = true;
 		}
 
+		TUniquePtr<FGraphConnectionManager> FEditor::RebuildConnectionManager(UAudioComponent* PreviewComp) const
+		{
+			if (!PreviewComp)
+			{
+				return MakeUnique<FGraphConnectionManager>();
+			}
+
+			UMetaSoundSource* Source = Cast<UMetaSoundSource>(GetMetasoundObject());
+			if (!Source)
+			{
+				return MakeUnique<FGraphConnectionManager>();
+			}
+
+			const UMetasoundEditorSettings* EdSettings = GetDefault<UMetasoundEditorSettings>();
+			check(EdSettings);
+			const UMetaSoundSettings* Settings = GetDefault<UMetaSoundSettings>();
+			check(Settings);
+
+			const FMetaSoundPageSettings* TargetPageSettings = Settings->FindPageSettings(EdSettings->AuditionTargetPage);
+			const FMetaSoundPageSettings* BuilderPageSettings = Settings->FindPageSettings(Builder->GetConstBuilder().GetBuildPageID());
+			if (TargetPageSettings != BuilderPageSettings)
+			{
+				return MakeUnique<FGraphConnectionManager>();
+			}
+
+			FAudioDevice* AudioDevice = PreviewComp->GetAudioDevice();
+			check(AudioDevice);
+			const FName& AudioBufferTypeName = GetMetasoundDataTypeName<FAudioBuffer>();
+			const FSampleRate DeviceSampleRate = static_cast<FSampleRate>(AudioDevice->GetSampleRate());
+			const uint32 PlayOrder = PreviewComp->GetLastPlayOrder();
+			const uint64 TransmitterID = Audio::GetTransmitterID(PreviewComp->GetAudioComponentID(), 0, PlayOrder);
+
+			return MakeUnique<FGraphConnectionManager>(*Source, *PreviewComp, TransmitterID, Source->GetOperatorSettings(DeviceSampleRate));
+		}
+
+		void FEditor::UpdatePageInfo(bool bIsPlaying)
+		{
+			using namespace Engine;
+
+			const UMetaSoundSettings* Settings = GetDefault<UMetaSoundSettings>();
+			check(Settings)
+
+			if (FocusPageWidget.IsValid())
+			{
+				// For now, this is one and the same, but may change in the future to accommodate conflating
+				// the static nature of the active build graph ID with what the user is visualizing.
+				const FMetaSoundPageSettings* PageSettings = nullptr;
+				if (Builder.IsValid())
+				{
+					const FGuid& PageID = Builder->GetConstBuilder().GetBuildPageID();
+					PageSettings = Settings->FindPageSettings(PageID);
+				}
+				FocusPageWidget->Update(PageSettings);
+			}
+
+			if (RenderPageWidget.IsValid())
+			{
+				const FMetaSoundPageSettings* PageSettings = nullptr;
+				if (bIsPlaying && Builder.IsValid())
+				{
+					const IMetaSoundDocumentInterface& DocInterface = Builder->GetConstBuilder().GetConstDocumentInterfaceChecked();
+					FGuid PageID = Frontend::DefaultPageID;
+					FDocumentBuilderRegistry::GetChecked().TryResolveTargetPageID(DocInterface.GetConstDocument().RootGraph, PageID);
+					PageSettings = Settings->FindPageSettings(PageID);
+				}
+				RenderPageWidget->Update(PageSettings);
+			}
+		}
+
+		bool FEditor::UpdatePlayTime(float InDeltaTime)
+		{
+			const bool bIsPlaying = IsPlaying();
+			if (bIsPlaying)
+			{
+				if (PlayTimeWidget.IsValid())
+				{
+					PlayTime += InDeltaTime;
+
+					FString PlayTimeString = FTimespan::FromSeconds(PlayTime).ToString();
+					PlayTimeString.ReplaceInline(TEXT("+"), TEXT(""));
+					PlayTimeWidget->SetText(FText::FromString(MoveTemp(PlayTimeString)));
+				}
+
+				return true;
+			}
+
+			SetPreviewID(INDEX_NONE);
+			PlayTime = 0.0;
+			PlayTimeWidget->SetText(FText::GetEmpty());
+			return false;
+		}
+
+		void FEditor::UpdateRenderInfo()
+		{
+			if (RenderStatsWidget.IsValid())
+			{
+				const bool bIsPlaying = IsPlaying();
+				RenderStatsWidget->Update(bIsPlaying, Cast<const UMetaSoundSource>(GetMetasoundObject()));
+			}
+		}
+
 		void FEditor::FDocumentListener::OnBuilderReloaded(Frontend::FDocumentModifyDelegates& OutDelegates)
 		{
 			OutDelegates.PageDelegates.OnPageSet.AddSP(this, &FDocumentListener::OnPageSet);
@@ -4631,6 +4758,8 @@ namespace Metasound
 		{
 			if (TSharedPtr<FEditor> ParentPtr = Parent.Pin())
 			{
+				ParentPtr->Stop();
+				ParentPtr->UpdatePageInfo(false);
 				ParentPtr->bRefreshGraph = true;
 			}
 		}
diff --git a/Engine/Plugins/Runtime/Metasound/Source/MetasoundEditor/Private/MetasoundEditor.h b/Engine/Plugins/Runtime/Metasound/Source/MetasoundEditor/Private/MetasoundEditor.h
--- a/Engine/Plugins/Runtime/Metasound/Source/MetasoundEditor/Private/MetasoundEditor.h
+++ b/Engine/Plugins/Runtime/Metasound/Source/MetasoundEditor/Private/MetasoundEditor.h
@@ -22,6 +22,7 @@
 #include "SGraphActionMenu.h"
 #include "SGraphPanel.h"
 #include "SMetasoundPalette.h"
+#include "SMetasoundStats.h"
 #include "Sound/AudioBus.h"
 #include "Textures/SlateIcon.h"
 #include "TickableEditorObject.h"
@@ -48,7 +49,6 @@ class FTabManager;
 class SDockableTab;
 class SGraphEditor;
 class SMetasoundPalette;
-class SMetaSoundRenderStats;
 class FSlateRect;
 class IDetailsView;
 class IToolkitHost;
@@ -56,6 +56,7 @@ class IToolkitHost;
 namespace Metasound::Editor
 {
 	class SFindInMetasound;
+	class SRenderStats;
 }
 class SVerticalBox;
 class UAudioComponent;
@@ -99,495 +100,509 @@ class UMetasoundPagesView : public UMetasoundEditorViewBase
 	GENERATED_BODY()
 };
 
-namespace Metasound
+namespace Metasound::Editor
 {
-	namespace Editor
-	{
-		// Forward Declarations
-		class FMetasoundGraphMemberSchemaAction;
+	// Forward Declarations
+	class FMetasoundGraphMemberSchemaAction;
 
-		/* Enums to use when grouping the members in the list panel. Enum order dictates visible order. */
-		enum class ENodeSection : uint8
-		{
-			None,
-			Inputs,
-			Outputs,
-			Variables,
+	/* Enums to use when grouping the members in the list panel. Enum order dictates visible order. */
+	enum class ENodeSection : uint8
+	{
+		None,
+		Inputs,
+		Outputs,
+		Variables,
 
-			COUNT
-		};
+		COUNT
+	};
 
 
-		class FEditor : public IMetasoundEditor, public FNotifyHook, public FEditorUndoClient, public FTickableEditorObject
-		{
-		public:
-			static const FName EditorName;
+	class FEditor : public IMetasoundEditor, public FNotifyHook, public FEditorUndoClient, public FTickableEditorObject
+	{
+	public:
+		static const FName EditorName;
 
-			FEditor();
-			virtual ~FEditor();
+		FEditor();
+		virtual ~FEditor();
 
-			virtual void RegisterTabSpawners(const TSharedRef<FTabManager>& TabManager) override;
-			virtual void UnregisterTabSpawners(const TSharedRef<FTabManager>& TabManager) override;
+		virtual void RegisterTabSpawners(const TSharedRef<FTabManager>& TabManager) override;
+		virtual void UnregisterTabSpawners(const TSharedRef<FTabManager>& TabManager) override;
 
-			double GetPlayTime() const;
-			TSharedPtr<SGraphEditor> GetGraphEditor() const;
+		double GetPlayTime() const;
+		TSharedPtr<SGraphEditor> GetGraphEditor() const;
 
-			/** Edits the specified Metasound object */
-			void InitMetasoundEditor(const EToolkitMode::Type Mode, const TSharedPtr<IToolkitHost>& InitToolkitHost, UObject* ObjectToEdit);
+		/** Edits the specified Metasound object */
+		void InitMetasoundEditor(const EToolkitMode::Type Mode, const TSharedPtr<IToolkitHost>& InitToolkitHost, UObject* ObjectToEdit);
 
-			UAudioComponent* GetAudioComponent() const;
+		UAudioComponent* GetAudioComponent() const;
 
-			/** IMetasoundEditor interface */
-			virtual UObject* GetMetasoundObject() const override;
-			virtual void SetSelection(const TArray<UObject*>& SelectedObjects) override;
-			virtual bool GetBoundsForSelectedNodes(FSlateRect& Rect, float Padding) override;
+		/** IMetasoundEditor interface */
+		virtual UObject* GetMetasoundObject() const override;
+		virtual void SetSelection(const TArray<UObject*>& SelectedObjects) override;
+		virtual bool GetBoundsForSelectedNodes(FSlateRect& Rect, float Padding) override;
 
-			/** IToolkit interface */
-			virtual FName GetToolkitFName() const override;
-			virtual FText GetBaseToolkitName() const override;
-			virtual FString GetWorldCentricTabPrefix() const override;
-			virtual FLinearColor GetWorldCentricTabColorScale() const override;
-			virtual const FSlateBrush* GetDefaultTabIcon() const override;
-			virtual FLinearColor GetDefaultTabColor() const override;
+		/** IToolkit interface */
+		virtual FName GetToolkitFName() const override;
+		virtual FText GetBaseToolkitName() const override;
+		virtual FString GetWorldCentricTabPrefix() const override;
+		virtual FLinearColor GetWorldCentricTabColorScale() const override;
+		virtual const FSlateBrush* GetDefaultTabIcon() const override;
+		virtual FLinearColor GetDefaultTabColor() const override;
 
-			/** IAssetEditorInstance interface */
-			virtual FName GetEditorName() const override;
+		/** IAssetEditorInstance interface */
+		virtual FName GetEditorName() const override;
 
-			virtual FString GetDocumentationLink() const override
-			{
-				return FString(TEXT("working-with-audio/sound-sources/meta-sounds"));
-			}
+		virtual FString GetDocumentationLink() const override
+		{
+			return FString(TEXT("working-with-audio/sound-sources/meta-sounds"));
+		}
 
-			/** FEditorUndoClient Interface */
-			virtual void PostUndo(bool bSuccess) override;
-			virtual void PostRedo(bool bSuccess) override { PostUndo(bSuccess); }
+		/** FEditorUndoClient Interface */
+		virtual void PostUndo(bool bSuccess) override;
+		virtual void PostRedo(bool bSuccess) override { PostUndo(bSuccess); }
 
-			/** FTickableEditorObject Interface */
-			virtual void Tick(float DeltaTime) override;
-			virtual TStatId GetStatId() const override;
-			virtual ETickableTickType GetTickableTickType() const override { return ETickableTickType::Always; }
+		/** FTickableEditorObject Interface */
+		virtual void Tick(float DeltaTime) override;
+		virtual TStatId GetStatId() const override;
+		virtual ETickableTickType GetTickableTickType() const override { return ETickableTickType::Always; }
 
-			virtual void Play() override;
-			virtual void Stop() override;
+		virtual void Play() override;
+		virtual void Stop() override;
 
-			virtual bool IsPlaying() const override;
+		virtual bool IsPlaying() const override;
 
-			/** Whether pasting the currently selected nodes is permissible */
-			bool CanPasteNodes();
+		/** Whether pasting the currently selected nodes is permissible */
+		bool CanPasteNodes();
 
-			/** Duplicates the selected node(s) in the graph */
-			void DuplicateNodes();
+		/** Duplicates the selected node(s) in the graph */
+		void DuplicateNodes();
 
-			/** Forces all UX pertaining to the root graph's details panel to be refreshed. */
-			void RefreshDetails();
+		/** Forces all UX pertaining to the root graph's details panel to be refreshed. */
+		void RefreshDetails();
 			
-			/** Pastes node(s) from the clipboard to the graph */
-			void PasteNodes(const FVector2D* InLocation = nullptr);
-			void PasteNodes(const FVector2D* InLocation, const FText& InTransactionText);
+		/** Pastes node(s) from the clipboard to the graph */
+		void PasteNodes(const FVector2D* InLocation = nullptr);
+		void PasteNodes(const FVector2D* InLocation, const FText& InTransactionText);
 
-			/** Returns Graph Connection Manager associated with this editor */
-			FGraphConnectionManager& GetConnectionManager();
-			const FGraphConnectionManager& GetConnectionManager() const;
+		/** Returns Graph Connection Manager associated with this editor */
+		FGraphConnectionManager& GetConnectionManager();
+		const FGraphConnectionManager& GetConnectionManager() const;
 
-			/** Forces all UX pertaining to the root graph's interface to be refreshed, returning the first selected member. */
-			UMetasoundEditorGraphMember* RefreshGraphMemberMenu();
+		/** Forces all UX pertaining to the root graph's interface to be refreshed, returning the first selected member. */
+		UMetasoundEditorGraphMember* RefreshGraphMemberMenu();
 
-			/** Updates selected node classes to highest class found in the MetaSound Class Registry. */
-			void UpdateSelectedNodeClasses();
+		/** Updates selected node classes to highest class found in the MetaSound Class Registry. */
+		void UpdateSelectedNodeClasses();
 
-			/** Whether or not MetaSound can be auditioned */
-			bool IsAuditionable() const;
+		/** Whether or not MetaSound can be auditioned */
+		bool IsAuditionable() const;
 
-			/* Whether the displayed graph is marked as editable */
-			bool IsGraphEditable() const;
+		/* Whether the displayed graph is marked as editable */
+		bool IsGraphEditable() const;
 
-			void ClearSelectionAndSelectNode(UEdGraphNode* Node);
+		void ClearSelectionAndSelectNode(UEdGraphNode* Node);
 
-			int32 GetNumNodesSelected() const
-			{
-				return MetasoundGraphEditor->GetSelectedNodes().Num();
-			}
+		int32 GetNumNodesSelected() const
+		{
+			return MetasoundGraphEditor->GetSelectedNodes().Num();
+		}
 
-			/** Creates analyzers */
-			void CreateAnalyzers(UMetaSoundSource& MetaSoundSource);
+		/** Creates analyzers */
+		void CreateAnalyzers(UMetaSoundSource& MetaSoundSource);
 
-			/** Destroys analyzers */
-			void DestroyAnalyzers();
+		/** Destroys analyzers */
+		void DestroyAnalyzers();
 
-			/** Jumps to the given nodes on the graph (templated to support arrays of various MetaSound graph node types) */
-			template<typename TNodeType = UMetasoundEditorGraphNode>
-			void JumpToNodes(const TArray<TNodeType*>& InNodes)
+		/** Jumps to the given nodes on the graph (templated to support arrays of various MetaSound graph node types) */
+		template<typename TNodeType = UMetasoundEditorGraphNode>
+		void JumpToNodes(const TArray<TNodeType*>& InNodes)
+		{
+			if (!MetasoundGraphEditor.IsValid())
 			{
-				if (!MetasoundGraphEditor.IsValid())
-				{
-					return;
-				}
+				return;
+			}
 
-				MetasoundGraphEditor->ClearSelectionSet();
-				const UMetasoundEditorGraph& Graph = GetMetaSoundGraphChecked();
-				if (!InNodes.IsEmpty())
+			MetasoundGraphEditor->ClearSelectionSet();
+			const UMetasoundEditorGraph& Graph = GetMetaSoundGraphChecked();
+			if (!InNodes.IsEmpty())
+			{
+				if (SGraphPanel* GraphPanel = MetasoundGraphEditor->GetGraphPanel())
 				{
-					if (SGraphPanel* GraphPanel = MetasoundGraphEditor->GetGraphPanel())
+					FVector2D BottomLeft = { TNumericLimits<float>::Max(), TNumericLimits<float>::Max() };
+					FVector2D TopRight = { TNumericLimits<float>::Min(), TNumericLimits<float>::Min() };
+					for (TNodeType* Node : InNodes)
 					{
-						FVector2D BottomLeft = { TNumericLimits<float>::Max(), TNumericLimits<float>::Max() };
-						FVector2D TopRight = { TNumericLimits<float>::Min(), TNumericLimits<float>::Min() };
-						for (TNodeType* Node : InNodes)
+						if (!Node || Node->GetGraph() != &Graph)
 						{
-							if (!Node || Node->GetGraph() != &Graph)
-							{
-								continue;
-							}
-
-							constexpr bool bSelected = true;
-							MetasoundGraphEditor->SetNodeSelection(Node, bSelected);
-							BottomLeft.X = FMath::Min(BottomLeft.X, Node->NodePosX);
-							BottomLeft.Y = FMath::Min(BottomLeft.Y, Node->NodePosY);
-							TopRight.X = FMath::Max(TopRight.X, Node->NodePosX + Node->EstimateNodeWidth());
-							TopRight.Y = FMath::Max(TopRight.Y, Node->NodePosY);
+							continue;
 						}
 
-						GraphPanel->JumpToRect(BottomLeft, TopRight);
+						constexpr bool bSelected = true;
+						MetasoundGraphEditor->SetNodeSelection(Node, bSelected);
+						BottomLeft.X = FMath::Min(BottomLeft.X, Node->NodePosX);
+						BottomLeft.Y = FMath::Min(BottomLeft.Y, Node->NodePosY);
+						TopRight.X = FMath::Max(TopRight.X, Node->NodePosX + Node->EstimateNodeWidth());
+						TopRight.Y = FMath::Max(TopRight.Y, Node->NodePosY);
 					}
+
+					GraphPanel->JumpToRect(BottomLeft, TopRight);
 				}
 			}
+		}
 
-			void SetDelayedRename();
+		void SetDelayedRename();
 
-		protected:
-			// Callbacks for action tree
-			bool CanRenameOnActionNode(TWeakPtr<FGraphActionNode> InSelectedNode) const;
-			bool CanAddNewElementToSection(int32 InSectionID) const;
-			void CollectAllActions(FGraphActionListBuilderBase& OutAllActions);
-			void CollectStaticSections(TArray<int32>& StaticSectionIDs);
-			TSharedRef<SWidget> CreateAddButton(int32 InSectionID, FText AddNewText, FName MetaDataTag);
-			FText GetFilterText() const;
-			bool HandleActionMatchesName(FEdGraphSchemaAction* InAction, const FName& InName) const;
-			FReply OnActionDragged(const TArray<TSharedPtr<FEdGraphSchemaAction>>& InActions, const FPointerEvent& MouseEvent);
-			void OnMemberActionDoubleClicked(const TArray< TSharedPtr<FEdGraphSchemaAction> >& InActions);
-			FActionMenuContent OnCreateGraphActionMenu(UEdGraph* InGraph, const FVector2D& InNodePosition, const TArray<UEdGraphPin*>& InDraggedPins, bool bAutoExpand, SGraphEditor::FActionMenuClosed InOnMenuClosed);
-			void OnActionSelected(const TArray<TSharedPtr<FEdGraphSchemaAction>>& InActions, ESelectInfo::Type InSelectionType);
-			FReply OnAddButtonClickedOnSection(int32 InSectionID);
-			TSharedRef<SWidget> OnGetMenuSectionWidget(TSharedRef<SWidget> RowWidget, int32 InSectionID);
-			FText GetSectionTitle(ENodeSection InSection) const;
-			FText OnGetSectionTitle(int32 InSectionID);
-			TSharedRef<SWidget> OnCreateWidgetForAction(struct FCreateWidgetForActionData* const InCreateData);
-			TSharedPtr<SWidget> OnContextMenuOpening();
+	protected:
+		// Callbacks for action tree
+		bool CanRenameOnActionNode(TWeakPtr<FGraphActionNode> InSelectedNode) const;
+		bool CanAddNewElementToSection(int32 InSectionID) const;
+		void CollectAllActions(FGraphActionListBuilderBase& OutAllActions);
+		void CollectStaticSections(TArray<int32>& StaticSectionIDs);
+		TSharedRef<SWidget> CreateAddButton(int32 InSectionID, FText AddNewText, FName MetaDataTag);
+		FText GetFilterText() const;
+		bool HandleActionMatchesName(FEdGraphSchemaAction* InAction, const FName& InName) const;
+		FReply OnActionDragged(const TArray<TSharedPtr<FEdGraphSchemaAction>>& InActions, const FPointerEvent& MouseEvent);
+		void OnMemberActionDoubleClicked(const TArray< TSharedPtr<FEdGraphSchemaAction> >& InActions);
+		FActionMenuContent OnCreateGraphActionMenu(UEdGraph* InGraph, const FVector2D& InNodePosition, const TArray<UEdGraphPin*>& InDraggedPins, bool bAutoExpand, SGraphEditor::FActionMenuClosed InOnMenuClosed);
+		void OnActionSelected(const TArray<TSharedPtr<FEdGraphSchemaAction>>& InActions, ESelectInfo::Type InSelectionType);
+		FReply OnAddButtonClickedOnSection(int32 InSectionID);
+		TSharedRef<SWidget> OnGetMenuSectionWidget(TSharedRef<SWidget> RowWidget, int32 InSectionID);
+		FText GetSectionTitle(ENodeSection InSection) const;
+		FText OnGetSectionTitle(int32 InSectionID);
+		TSharedRef<SWidget> OnCreateWidgetForAction(struct FCreateWidgetForActionData* const InCreateData);
+		TSharedPtr<SWidget> OnContextMenuOpening();
 
-			/** Called when the selection changes in the GraphEditor */
-			void OnSelectedNodesChanged(const TSet<UObject*>& NewSelection);
+		/** Called when the selection changes in the GraphEditor */
+		void OnSelectedNodesChanged(const TSet<UObject*>& NewSelection);
 
-			FGraphAppearanceInfo GetGraphAppearance() const;
+		FGraphAppearanceInfo GetGraphAppearance() const;
 
-			UMetasoundEditorGraph& GetMetaSoundGraphChecked();
+		UMetasoundEditorGraph& GetMetaSoundGraphChecked();
 
-			FText GetGraphStatusDescription() const;
-			const FSlateIcon& GetPlayIcon() const;
-			const FSlateIcon& GetStopIcon() const;
+		FText GetGraphStatusDescription() const;
+		const FSlateIcon& GetPlayIcon() const;
+		const FSlateIcon& GetStopIcon() const;
 
-			/**
-			 * Called when a node's title is committed for a rename
-			 *
-			 * @param	NewText				New title text
-			 * @param	CommitInfo			How text was committed
-			 * @param	NodeBeingChanged	The node being changed
-			 */
-			void OnNodeTitleCommitted(const FText& NewText, ETextCommit::Type CommitInfo, UEdGraphNode* NodeBeingChanged);
+		/**
+			* Called when a node's title is committed for a rename
+			*
+			* @param	NewText				New title text
+			* @param	CommitInfo			How text was committed
+			* @param	NodeBeingChanged	The node being changed
+			*/
+		void OnNodeTitleCommitted(const FText& NewText, ETextCommit::Type CommitInfo, UEdGraphNode* NodeBeingChanged);
 
-			/** Deletes from the Metasound Menu (i.e. input or output) if in focus, or the currently selected nodes if the graph editor is in focus. */
-			void DeleteSelected();
+		/** Deletes from the Metasound Menu (i.e. input or output) if in focus, or the currently selected nodes if the graph editor is in focus. */
+		void DeleteSelected();
 
-			void DeleteInterfaceItem(TSharedPtr<FMetasoundGraphMemberSchemaAction> ActionToDelete);
+		void DeleteInterfaceItem(TSharedPtr<FMetasoundGraphMemberSchemaAction> ActionToDelete);
 
-			/** Delete the currently selected nodes */
-			void DeleteSelectedNodes();
+		/** Delete the currently selected nodes */
+		void DeleteSelectedNodes();
 
-			/** Delete the currently selected interface items */
-			void DeleteSelectedInterfaceItems();
+		/** Delete the currently selected interface items */
+		void DeleteSelectedInterfaceItems();
 
-			/** Cut the currently selected nodes */
-			void CutSelectedNodes();
+		/** Cut the currently selected nodes */
+		void CutSelectedNodes();
 
-			/** Copy the currently selected nodes */
-			void CopySelectedNodes() const;
+		/** Copy the currently selected nodes */
+		void CopySelectedNodes() const;
 
-			/** Whether or not the currently selected node(s) can be duplicated */
-			bool CanDuplicateNodes() const;
+		/** Whether or not the currently selected node(s) can be duplicated */
+		bool CanDuplicateNodes() const;
 
-			/** Whether the currently selected node(s) can be deleted */
-			bool CanDeleteNodes() const;
+		/** Whether the currently selected node(s) can be deleted */
+		bool CanDeleteNodes() const;
 
-			/** Whether the currently selected interface item(s) can be deleted */
-			bool CanDeleteInterfaceItems() const;
+		/** Whether the currently selected interface item(s) can be deleted */
+		bool CanDeleteInterfaceItems() const;
 
-			/** Whether at least one of the currently selected node(s) can be renamed. */
-			bool CanRenameSelectedNodes() const;
+		/** Whether at least one of the currently selected node(s) can be renamed. */
+		bool CanRenameSelectedNodes() const;
 
-			/** Rename selected node (currently applies to comments and member nodes). */
-			void RenameSelectedNode();
+		/** Rename selected node (currently applies to comments and member nodes). */
+		void RenameSelectedNode();
 
-			/** Whether at least one of the currently selected interface item(s) can be renamed. */
-			bool CanRenameSelectedInterfaceItems() const;
+		/** Whether at least one of the currently selected interface item(s) can be renamed. */
+		bool CanRenameSelectedInterfaceItems() const;
 
-			/** Rename selected interface item. */
-			void RenameSelectedInterfaceItem();
+		/** Rename selected interface item. */
+		void RenameSelectedInterfaceItem();
 
-			/** Whether the currently selected Member item(s) can be duplicated. */
-			bool CanDuplicateSelectedMemberItems() const;
+		/** Whether the currently selected Member item(s) can be duplicated. */
+		bool CanDuplicateSelectedMemberItems() const;
 
-			/** Duplicate selected Member items. */
-			void DuplicateSelectedMemberItems();
+		/** Duplicate selected Member items. */
+		void DuplicateSelectedMemberItems();
 
-			/** Whether there are nodes to jump to for the currently selected interface item. */
-			bool CanJumpToNodesForSelectedInterfaceItem() const;
+		/** Whether there are nodes to jump to for the currently selected interface item. */
+		bool CanJumpToNodesForSelectedInterfaceItem() const;
 
-			/** Jumps to the nodes corresponding to the first valid currently selected interface item. */
-			void JumpToNodesForSelectedInterfaceItem();
+		/** Jumps to the nodes corresponding to the first valid currently selected interface item. */
+		void JumpToNodesForSelectedInterfaceItem();
 			
-			/** Delete all unused members from the selected section*/
-			void DeleteAllUnusedInSection();
+		/** Delete all unused members from the selected section*/
+		void DeleteAllUnusedInSection();
 
-			/** Whether the selection is not to a valid member*/
-			bool CanDeleteUnusedMembers() const;
+		/** Whether the selection is not to a valid member*/
+		bool CanDeleteUnusedMembers() const;
 
-			/** Called to undo the last action */
-			void UndoGraphAction();
+		/** Called to undo the last action */
+		void UndoGraphAction();
 
-			/** Called to redo the last undone action */
-			void RedoGraphAction();
+		/** Called to redo the last undone action */
+		void RedoGraphAction();
 
-			void RefreshEditorContext(UObject& MetaSound);
+		void RefreshEditorContext(UObject& MetaSound);
 
-			/** Show and focus the Find in MetaSound tab. */
-			void ShowFindInMetaSound();
+		/** Show and focus the Find in MetaSound tab. */
+		void ShowFindInMetaSound();
 
-			/** Find selected node from Graph*/
-			void FindSelectedNodeInGraph();
+		/** Find selected node from Graph*/
+		void FindSelectedNodeInGraph();
 
-			/**Hide pins without connection*/
-			void HideUnconnectedPins();
+		/**Hide pins without connection*/
+		void HideUnconnectedPins();
 			
-			/**Show pins without connection*/
-			void ShowUnconnectedPins();
+		/**Show pins without connection*/
+		void ShowUnconnectedPins();
 
-			/** Checks if pin can be promoted */
-			bool CanPromoteToInput();
+		/** Checks if pin can be promoted */
+		bool CanPromoteToInput();
 			
-			/** Promotes pin to graph input */
-			void PromoteToInput();
+		/** Promotes pin to graph input */
+		void PromoteToInput();
 
-			/** Checks if pin can be promoted */
-			bool CanPromoteToOutput();
+		/** Checks if pin can be promoted */
+		bool CanPromoteToOutput();
 
-			/** Promotes pin to graph output */
-			void PromoteToOutput();
+		/** Promotes pin to graph output */
+		void PromoteToOutput();
 
-			/** Checks if pin can be promoted */
-			bool CanPromoteToVariable();
+		/** Checks if pin can be promoted */
+		bool CanPromoteToVariable();
 
-			/** Promotes pin to graph variable */
-			void PromoteToVariable();
+		/** Promotes pin to graph variable */
+		void PromoteToVariable();
 
-			/** Checks if pin can be promoted */
-			bool CanPromoteToDeferredVariable();
+		/** Checks if pin can be promoted */
+		bool CanPromoteToDeferredVariable();
 
-			/** Promotes pin to graph deferred variable */
-			void PromoteToDeferredVariable();
+		/** Promotes pin to graph deferred variable */
+		void PromoteToDeferredVariable();
 
-			/** Checks if node's inputs can be promoted */
-			bool CanPromoteAllToInputs();
+		/** Checks if node's inputs can be promoted */
+		bool CanPromoteAllToInputs();
 
-			/** Promotes node's inputs to unique graph inputs */
-			void PromoteAllToInputs();
+		/** Promotes node's inputs to unique graph inputs */
+		void PromoteAllToInputs();
 
-			/** Checks if node's inputs can be promoted */
-			bool CanPromoteAllToCommonInputs();
+		/** Checks if node's inputs can be promoted */
+		bool CanPromoteAllToCommonInputs();
 
-			/** Promotes node's inputs to shared graph inputs */
-			void PromoteAllToCommonInputs();
+		/** Promotes node's inputs to shared graph inputs */
+		void PromoteAllToCommonInputs();
 
-		private:
-			int32 PromotableSelectedNodes();
+	private:
+		int32 PromotableSelectedNodes();
 
-			/** Forces refresh of pages view. */
-			void RefreshPagesView();
+		/** Forces refresh of pages view. */
+		void RefreshPagesView();
 
-			/** Forces refresh of interfaces view. */
-			void RefreshInterfaceView();
+		/** Forces refresh of interfaces view. */
+		void RefreshInterfaceView();
 
-			void RemoveInvalidSelection();
+		void RemoveInvalidSelection();
 
-			void SetPreviewID(uint32 InPreviewID);
+		void SetPreviewID(uint32 InPreviewID);
 
-			void ExportNodesToText(FString& OutText) const;
+		void ExportNodesToText(FString& OutText) const;
 
-			// Sets the globally targeted audition page for execution to this editor's currently focused page
-			void SyncFocusedPage() const;
+		// Sets the globally targeted audition page for execution to this editor's currently focused page
+		void SyncFocusedPage() const;
 
-			/** FNotifyHook interface */
-			virtual void NotifyPostChange(const FPropertyChangedEvent& PropertyChangedEvent, FProperty* PropertyThatChanged) override;
+		/** FNotifyHook interface */
+		virtual void NotifyPostChange(const FPropertyChangedEvent& PropertyChangedEvent, FProperty* PropertyThatChanged) override;
 
-			/** Creates all internal widgets for the tabs to point at */
-			void CreateInternalWidgets(UObject& MetaSound);
+		/** Creates all internal widgets for the tabs to point at */
+		void CreateInternalWidgets(UObject& MetaSound);
 
-			/** Builds the toolbar widget for the Metasound editor */
-			void ExtendToolbarInternal();
+		/** Builds the toolbar widget for the Metasound editor */
+		void ExtendToolbarInternal();
 
-			/** Binds new graph commands to delegates */
-			void BindGraphCommands();
+		/** Binds new graph commands to delegates */
+		void BindGraphCommands();
 
-			FSlateIcon GetImportStatusImage() const;
+		FSlateIcon GetImportStatusImage() const;
 
-			FSlateIcon GetExportStatusImage() const;
+		FSlateIcon GetExportStatusImage() const;
 
-			// TODO: Move import/export out of editor and into import/export asset actions
-			void Import();
-			void Export();
+		// TODO: Move import/export out of editor and into import/export asset actions
+		void Import();
+		void Export();
 
-			/** Toolbar command methods */
-			void ExecuteNode();
+		/** Toolbar command methods */
+		void ExecuteNode();
 
-			/** Whether we can play the current selection of nodes */
-			bool CanExecuteNode() const;
+		/** Whether we can play the current selection of nodes */
+		bool CanExecuteNode() const;
 
-			/** Either play the Metasound or stop currently playing sound */
-			void TogglePlayback();
+		/** Either play the Metasound or stop currently playing sound */
+		void TogglePlayback();
 
-			/** Executes specified node (If supported) */
-			void ExecuteNode(UEdGraphNode* Node);
+		/** Executes specified node (If supported) */
+		void ExecuteNode(UEdGraphNode* Node);
 
-			/** Sync the content browser to the current selection of nodes */
-			void SyncInBrowser();
+		/** Sync the content browser to the current selection of nodes */
+		void SyncInBrowser();
 
-			/** Converts the MetaSound from a preset to a fully modifiable MetaSound. */
-			void ConvertFromPreset();
+		/** Converts the MetaSound from a preset to a fully modifiable MetaSound. */
+		void ConvertFromPreset();
 
-			/** Creates audition menu options */
-			TSharedRef<SWidget> CreateAuditionMenuOptions() const;
+		/** Creates audition menu options */
+		TSharedRef<SWidget> CreateAuditionMenuOptions();
 
-			/** Creates page menu options */
-			void CreateAuditionPageSubMenuOptions(FMenuBuilder& MenuBuilder) const;
+		/** Creates page menu options */
+		void CreateAuditionPageSubMenuOptions(FMenuBuilder& MenuBuilder);
 
-			/** Show the Metasound object's Source settings in the Details panel */
-			void EditSourceSettings();
+		/** Show the Metasound object's Source settings in the Details panel */
+		void EditSourceSettings();
 
-			/** Show the Metasound object's settings in the Details panel */
-			void EditMetasoundSettings();
+		/** Show the Metasound object's settings in the Details panel */
+		void EditMetasoundSettings();
 
-			/** Add an input to the currently selected node */
-			void AddInput();
+		/** Add an input to the currently selected node */
+		void AddInput();
 
-			/** Whether we can add an input to the currently selected node */
-			bool CanAddInput() const;
+		/** Whether we can add an input to the currently selected node */
+		bool CanAddInput() const;
 
-			/* Create comment node on graph */
-			void OnCreateComment();
+		/* Create comment node on graph */
+		void OnCreateComment();
 
-			/** Create new graph editor widget */
-			void CreateGraphEditorWidget(UObject& MetaSound);
+		/** Create new graph editor widget */
+		void CreateGraphEditorWidget(UObject& MetaSound);
 
-			TSharedPtr<SWidget> BuildAnalyzerWidget() const;
+		TSharedPtr<SWidget> BuildAnalyzerWidget() const;
 
-			void EditObjectSettings();
+		void EditObjectSettings();
 
-			void NotifyAssetPrimeComplete();
-			void NotifyAssetPrimeInProgress();
-			void NotifyDocumentVersioned();
-			void NotifyNodePasteFailure_MultipleVariableSetters();
-			void NotifyNodePasteFailure_ReferenceLoop();
+		void NotifyAssetPrimeComplete();
+		void NotifyAssetPrimeInProgress();
+		void NotifyDocumentVersioned();
+		void NotifyNodePasteFailure_MultipleVariableSetters();
+		void NotifyNodePasteFailure_ReferenceLoop();
 
-			/** List of open tool panels; used to ensure only one exists at any one time */
-			TMap<FName, TWeakPtr<SDockableTab>> SpawnedToolPanels;
+		TUniquePtr<FGraphConnectionManager> RebuildConnectionManager(UAudioComponent* PreviewComponent = nullptr) const;
 
-			/** New Graph Editor */
-			TSharedPtr<SGraphEditor> MetasoundGraphEditor;
+		/** Updates the page info widget. */
+		void UpdatePageInfo(bool bIsPlaying);
 
-			/** Details tab */
-			TSharedPtr<IDetailsView> MetasoundDetails;
+		/** Updates the play time widget.  Returns true if currently playing, false if not. */
+		bool UpdatePlayTime(float InDeltaTime);
 
-			/** Pages tab */
-			TSharedPtr<IDetailsView> PagesDetails;
-			TStrongObjectPtr<UMetasoundPagesView> PagesView;
+		/** Updates the render info widget. */
+		void UpdateRenderInfo();
 
-			/** Interfaces tab */
-			TSharedPtr<IDetailsView> InterfacesDetails;
-			TStrongObjectPtr<UMetasoundInterfacesView> InterfacesView;
+		/** List of open tool panels; used to ensure only one exists at any one time */
+		TMap<FName, TWeakPtr<SDockableTab>> SpawnedToolPanels;
 
-			/** Metasound graph members menu */
-			TSharedPtr<SGraphActionMenu> GraphMembersMenu;
+		/** New Graph Editor */
+		TSharedPtr<SGraphEditor> MetasoundGraphEditor;
 
-			/** Meter used in the analyzer tab for auditioning preview output. */
-			TSharedPtr<AudioWidgets::FAudioMeter> OutputMeter;
+		/** Details tab */
+		TSharedPtr<IDetailsView> MetasoundDetails;
 
-			/** Oscilloscope used in the analyzer tab for visualizing preview output. */
-			TSharedPtr<AudioWidgets::FAudioOscilloscope> OutputOscilloscope;
+		/** Pages tab */
+		TSharedPtr<IDetailsView> PagesDetails;
+		TStrongObjectPtr<UMetasoundPagesView> PagesView;
 
-			/** Vectorscope used in the analyzer tab for visualizing preview output. */
-			TSharedPtr<AudioWidgets::FAudioVectorscope> OutputVectorscope;
+		/** Interfaces tab */
+		TSharedPtr<IDetailsView> InterfacesDetails;
+		TStrongObjectPtr<UMetasoundInterfacesView> InterfacesView;
 
-			/** Spectrum analyzer used in the analyzer tab for auditioning preview output. */
-			TSharedPtr<AudioWidgets::FAudioSpectrumAnalyzer> OutputSpectrumAnalyzer;
+		/** Metasound graph members menu */
+		TSharedPtr<SGraphActionMenu> GraphMembersMenu;
 
-			/** Find in MetaSound widget. */
-			TSharedPtr<SFindInMetasound> FindWidget;
+		/** Meter used in the analyzer tab for auditioning preview output. */
+		TSharedPtr<AudioWidgets::FAudioMeter> OutputMeter;
 
-			/** Palette of Node types */
-			TSharedPtr<SMetasoundPalette> Palette;
+		/** Oscilloscope used in the analyzer tab for visualizing preview output. */
+		TSharedPtr<AudioWidgets::FAudioOscilloscope> OutputOscilloscope;
 
-			/** Widget showing playtime that overlays the graph when previewing */
-			TSharedPtr<STextBlock> PlayTimeWidget;
-			double PlayTime = 0.0;
+		/** Vectorscope used in the analyzer tab for visualizing preview output. */
+		TSharedPtr<AudioWidgets::FAudioVectorscope> OutputVectorscope;
 
-			/** Widget showing render performance information that overlays the graph when previewing */
-			TSharedPtr<SMetaSoundRenderStats> RenderStatsWidget;
+		/** Spectrum analyzer used in the analyzer tab for auditioning preview output. */
+		TSharedPtr<AudioWidgets::FAudioSpectrumAnalyzer> OutputSpectrumAnalyzer;
 
-			TUniquePtr<FGraphConnectionManager> GraphConnectionManager;
+		/** Find in MetaSound widget. */
+		TSharedPtr<SFindInMetasound> FindWidget;
 
-			/** Command list for this editor */
-			TSharedPtr<FUICommandList> GraphEditorCommands;
+		/** Palette of Node types */
+		TSharedPtr<SMetasoundPalette> Palette;
 
-			/** Pointer to builder being actively used to mutate MetaSound asset */
-			TStrongObjectPtr<UMetaSoundBuilderBase> Builder;
+		/** Widget showing playtime that overlays the graph when previewing */
+		TSharedPtr<STextBlock> PlayTimeWidget;
+		double PlayTime = 0.0;
 
-			/** Whether or not metasound being edited is valid */
-			bool bPassedValidation = true;
+		/** Widget showing page info regarding the currently focused page that overlays the graph tab content */
+		TSharedPtr<SPageStats> FocusPageWidget;
 
-			/** Text content used when either duplicating or pasting from clipboard (avoids deserializing twice) */
-			FString NodeTextToPaste;
+		/** Widget showing page info regarding the currently rendering page that overlays the analyzer tab content */
+		TSharedPtr<SPageStats> RenderPageWidget;
 
-			/** Boolean state for when selection change handle should not respond due to selection state
-			  * being manually applied in code */
-			bool bManuallyClearingGraphSelection = false;
+		/** Widget showing render performance information that overlays the graph */
+		TSharedPtr<SRenderStats> RenderStatsWidget;
 
-			/** Whether or not registry is currently being primed due to assets still loading */
-			bool bPrimingRegistry = false;
+		TUniquePtr<FGraphConnectionManager> GraphConnectionManager;
 
-			int32 HighestMessageSeverity = EMessageSeverity::Info;
+		/** Command list for this editor */
+		TSharedPtr<FUICommandList> GraphEditorCommands;
 
-			TSharedPtr<SNotificationItem> NotificationPtr;
+		/** Pointer to builder being actively used to mutate MetaSound asset */
+		TStrongObjectPtr<UMetaSoundBuilderBase> Builder;
 
-			bool bMemberRenameRequested = false;
+		/** Whether or not metasound being edited is valid */
+		bool bPassedValidation = true;
 
-			bool bRefreshGraph = false;
+		/** Text content used when either duplicating or pasting from clipboard (avoids deserializing twice) */
+		FString NodeTextToPaste;
 
-			class FDocumentListener : public Frontend::IDocumentBuilderTransactionListener
-			{
-				TWeakPtr<FEditor> Parent;
+		/** Boolean state for when selection change handle should not respond due to selection state
+			* being manually applied in code */
+		bool bManuallyClearingGraphSelection = false;
 
-			public:
-				FDocumentListener() = default;
-				FDocumentListener(TSharedRef<FEditor> InParent)
-					: Parent(InParent)
-				{
-				}
+		/** Whether or not registry is currently being primed due to assets still loading */
+		bool bPrimingRegistry = false;
+
+		int32 HighestMessageSeverity = EMessageSeverity::Info;
 
-				virtual ~FDocumentListener() = default;
+		TSharedPtr<SNotificationItem> NotificationPtr;
 
-			private:
-				virtual void OnBuilderReloaded(Frontend::FDocumentModifyDelegates& OutDelegates) override;
-				void OnPageSet(const Frontend::FDocumentMutatePageArgs& Args);
-			};
+		bool bMemberRenameRequested = false;
 
-			TSharedPtr<FDocumentListener> DocListener;
+		bool bRefreshGraph = false;
+
+		class FDocumentListener : public Frontend::IDocumentBuilderTransactionListener
+		{
+			TWeakPtr<FEditor> Parent;
+
+		public:
+			FDocumentListener() = default;
+			FDocumentListener(TSharedRef<FEditor> InParent)
+				: Parent(InParent)
+			{
+			}
+
+			virtual ~FDocumentListener() = default;
+
+		private:
+			virtual void OnBuilderReloaded(Frontend::FDocumentModifyDelegates& OutDelegates) override;
+			void OnPageSet(const Frontend::FDocumentMutatePageArgs& Args);
 		};
-	} // namespace Editor
-} // namespace Metasound
+
+		TSharedPtr<FDocumentListener> DocListener;
+	};
+} // namespace Metasound::Editor
diff --git a/Engine/Plugins/Runtime/Metasound/Source/MetasoundEditor/Private/MetasoundEditorSettings.cpp b/Engine/Plugins/Runtime/Metasound/Source/MetasoundEditor/Private/MetasoundEditorSettings.cpp
--- a/Engine/Plugins/Runtime/Metasound/Source/MetasoundEditor/Private/MetasoundEditorSettings.cpp
+++ b/Engine/Plugins/Runtime/Metasound/Source/MetasoundEditor/Private/MetasoundEditorSettings.cpp
@@ -8,6 +8,7 @@
 #include "MetasoundDocumentBuilderRegistry.h"
 #include "MetasoundFrontendDocument.h"
 #include "MetasoundSettings.h"
+#include "PlatformInfo.h"
 #include "Styling/AppStyle.h"
 #include "Styling/SlateWidgetStyleAsset.h"
 #include "Widgets/Notifications/SNotificationList.h"
@@ -150,4 +151,21 @@ Metasound::Engine::FAuditionPageInfo UMetasoundEditorSettings::ResolveAuditionPa
 	return PreviewInfo;
 
 }
+
+TArray<FName> UMetasoundEditorSettings::GetAuditionPlatformNames()
+{
+	if (const UMetaSoundSettings* Settings = GetDefault<UMetaSoundSettings>())
+	{
+		TSet<FName> PlatformNames { "Default" };
+		for (const FMetaSoundPageSettings& PageSettings : Settings->GetPageSettings())
+		{
+			TSet<FName> PagePlatforms;
+			PageSettings.IsCooked.PerPlatform.GetKeys(PagePlatforms);
+			PlatformNames.Append(PagePlatforms);
+		}
+		return PlatformNames.Array();
+	}
+
+	return { };
+}
 #undef LOCTEXT_NAMESPACE // "MetaSoundEditor"
diff --git a/Engine/Plugins/Runtime/Metasound/Source/MetasoundEditor/Private/MetasoundEditorSettings.h b/Engine/Plugins/Runtime/Metasound/Source/MetasoundEditor/Private/MetasoundEditorSettings.h
--- a/Engine/Plugins/Runtime/Metasound/Source/MetasoundEditor/Private/MetasoundEditorSettings.h
+++ b/Engine/Plugins/Runtime/Metasound/Source/MetasoundEditor/Private/MetasoundEditorSettings.h
@@ -56,7 +56,10 @@ enum class EMetasoundActiveDetailView : uint8
 UENUM()
 enum class EAuditionPageMode : uint8
 {
+	// Audition Target Page set automatically set to page focused in asset editor
 	Focused,
+
+	// Audition Target Page is specified by user (does not automatically change when new page is visualized)
 	User
 };
 
@@ -152,20 +155,26 @@ public:
 	FString DefaultAuthor;
 
 	/* Currently set page audition mode. Set by the MetaSound Asset Editor. */
-	UPROPERTY(VisibleAnywhere, config, Category = Pages)
+	UPROPERTY(EditAnywhere, config, Category = Audition, meta = (DisplayName = "Audition Mode"))
 	EAuditionPageMode AuditionPageMode = EAuditionPageMode::Focused;
 
 	/** Name of platform to mock when previewing playback. This will limit playback to only pages that are cooked for the given platform.
 	  * Set in the MetaSound Asset Editor.
 	  */
-	UPROPERTY(VisibleAnywhere, config, Category = Pages)
+	UPROPERTY(EditAnywhere, config, Category = Audition, meta = (DisplayName = "Platform", GetOptions = "MetasoundEditor.MetasoundEditorSettings.GetAuditionPlatformNames"))
 	FName AuditionPlatform;
 
 	/** Name of the page to target when previewing playback in editor. If target page is not implemented for the set audition platform,
 	  *  uses order of cooked pages (see 'project MetaSound Settings --> Page Settings' for order) falling back to lower index-ordered page
 	  * implemented in MetaSound asset. Set in the MetaSound Asset Editor.
 	  */
-	UPROPERTY(VisibleAnywhere, config, Category = Pages)
+	UPROPERTY(EditAnywhere, config, Category = Audition, meta =
+	(
+		DisplayName = "Audition Page",
+		EditCondition = "AuditionPageMode == EAuditionPageMode::User",
+		EditConditionHides = true,
+		GetOptions = "MetasoundEngine.MetaSoundSettings.GetPageNames")
+	)
 	FName AuditionTargetPage;
 
 	/** Maps Pin Category To Pin Color */
@@ -286,4 +295,7 @@ public:
 	const FAudioMaterialMeterStyle* GetMeterStyle() const;
 
 	Metasound::Engine::FAuditionPageInfo ResolveAuditionPageInfo(const TSet<FGuid>& InPageIDs) const;
+
+	UFUNCTION()
+	static TArray<FName> GetAuditionPlatformNames();
 };
diff --git a/Engine/Plugins/Runtime/Metasound/Source/MetasoundEditor/Private/MetasoundEditorSubsystem.cpp b/Engine/Plugins/Runtime/Metasound/Source/MetasoundEditor/Private/MetasoundEditorSubsystem.cpp
--- a/Engine/Plugins/Runtime/Metasound/Source/MetasoundEditor/Private/MetasoundEditorSubsystem.cpp
+++ b/Engine/Plugins/Runtime/Metasound/Source/MetasoundEditor/Private/MetasoundEditorSubsystem.cpp
@@ -249,7 +249,7 @@ void UMetaSoundEditorSubsystem::RegisterToolbarExtender(TSharedRef<FExtender> In
 	EditorToolbarExtenders.AddUnique(InExtender);
 }
 
-void UMetaSoundEditorSubsystem::SetFocusedPage(UMetaSoundBuilderBase* Builder, FName PageName, bool bFocusPageEditor, EMetaSoundBuilderResult& OutResult) const
+void UMetaSoundEditorSubsystem::SetFocusedPage(UMetaSoundBuilderBase* Builder, FName PageName, bool bOpenEditor, EMetaSoundBuilderResult& OutResult) const
 {
 	using namespace Metasound::Frontend;
 	if (!Builder)
@@ -262,7 +262,7 @@ void UMetaSoundEditorSubsystem::SetFocusedPage(UMetaSoundBuilderBase* Builder, F
 	check(Settings);
 	if (const FMetaSoundPageSettings* PageSettings = Settings->FindPageSettings(PageName))
 	{
-		const bool bFocusedPage = SetFocusedPageInternal(*PageSettings, *Builder, bFocusPageEditor);
+		const bool bFocusedPage = SetFocusedPageInternal(*PageSettings, *Builder, bOpenEditor);
 		if (bFocusedPage)
 		{
 			OutResult = EMetaSoundBuilderResult::Succeeded;
@@ -273,7 +273,7 @@ void UMetaSoundEditorSubsystem::SetFocusedPage(UMetaSoundBuilderBase* Builder, F
 	OutResult = EMetaSoundBuilderResult::Failed;
 }
 
-bool UMetaSoundEditorSubsystem::SetFocusedPage(UMetaSoundBuilderBase& Builder, const FGuid& InPageID, bool bFocusPageEditor) const
+bool UMetaSoundEditorSubsystem::SetFocusedPage(UMetaSoundBuilderBase& Builder, const FGuid& InPageID, bool bOpenEditor) const
 {
 	using namespace Metasound::Frontend;
 
@@ -281,32 +281,37 @@ bool UMetaSoundEditorSubsystem::SetFocusedPage(UMetaSoundBuilderBase& Builder, c
 	check(Settings);
 	if (const FMetaSoundPageSettings* PageSettings = Settings->FindPageSettings(InPageID))
 	{
-		return SetFocusedPageInternal(*PageSettings, Builder, bFocusPageEditor);
+		return SetFocusedPageInternal(*PageSettings, Builder, bOpenEditor);
 	}
 
 	return false;
 }
 
-bool UMetaSoundEditorSubsystem::SetFocusedPageInternal(const FMetaSoundPageSettings& InPageSettings, UMetaSoundBuilderBase& Builder, bool bFocusPageEditor) const
+bool UMetaSoundEditorSubsystem::SetFocusedPageInternal(const FMetaSoundPageSettings& InPageSettings, UMetaSoundBuilderBase& Builder, bool bOpenEditor) const
 {
 	using namespace Metasound::Frontend;
 
 	const FScopedTransaction Transaction(FText::Format(LOCTEXT("SetFocusedPageTransactionFormat", "Set Focused Page '{0}'"), FText::FromName(InPageSettings.Name)));
 	Builder.Modify();
-	if (Builder.GetBuilder().SetBuildPageID(InPageSettings.UniqueId))
+
+	if (UMetasoundEditorSettings* EditorSettings = GetMutableDefault<UMetasoundEditorSettings>())
 	{
-		if (UMetasoundEditorSettings* EditorSettings = GetMutableDefault<UMetasoundEditorSettings>())
+		if (EditorSettings->AuditionPageMode == EAuditionPageMode::Focused)
 		{
-			if (EditorSettings->AuditionPageMode == EAuditionPageMode::Focused)
-			{
-				EditorSettings->AuditionTargetPage = InPageSettings.Name;
+			// Must set audition target page before setting build page ID as listeners
+			// to build page ID changes need to reliably be able to adjust to newly assigned
+			// audition target page.
+			EditorSettings->AuditionTargetPage = InPageSettings.Name;
+		}
 
-				// Reregister to ensure all future audible instances are using the new page implementation.
-				RegisterGraphWithFrontend(Builder.GetBuilder().CastDocumentObjectChecked<UObject>());
-			}
+		if (Builder.GetBuilder().SetBuildPageID(InPageSettings.UniqueId))
+		{
+
+			// Reregister to ensure all future audible instances are using the new page implementation.
+			RegisterGraphWithFrontend(Builder.GetBuilder().CastDocumentObjectChecked<UObject>());
 		}
 
-		if (GEditor && bFocusPageEditor)
+		if (GEditor && bOpenEditor)
 		{
 			if (UAssetEditorSubsystem* AssetEditorSubsystem = GEditor->GetEditorSubsystem<UAssetEditorSubsystem>())
 			{
diff --git a/Engine/Plugins/Runtime/Metasound/Source/MetasoundEditor/Private/SMetasoundRenderStats.cpp b/Engine/Plugins/Runtime/Metasound/Source/MetasoundEditor/Private/SMetasoundRenderStats.cpp
deleted file mode 100644
diff --git a/Engine/Plugins/Runtime/Metasound/Source/MetasoundEditor/Private/SMetasoundRenderStats.h b/Engine/Plugins/Runtime/Metasound/Source/MetasoundEditor/Private/SMetasoundRenderStats.h
deleted file mode 100644
diff --git a/Engine/Plugins/Runtime/Metasound/Source/MetasoundEditor/Private/SMetasoundStats.cpp b/Engine/Plugins/Runtime/Metasound/Source/MetasoundEditor/Private/SMetasoundStats.cpp
new file mode 100644
--- /dev/null
+++ b/Engine/Plugins/Runtime/Metasound/Source/MetasoundEditor/Private/SMetasoundStats.cpp
@@ -0,0 +1,118 @@
+// Copyright Epic Games, Inc. All Rights Reserved.
+#include "SMetasoundStats.h"
+
+#include "Components/AudioComponent.h"
+#include "Containers/UnrealString.h"
+#include "Editor.h"
+#include "Internationalization/Text.h"
+#include "MetasoundDocumentBuilderRegistry.h"
+#include "MetasoundEditorSettings.h"
+#include "MetasoundFrontendDocumentBuilder.h"
+#include "MetasoundSettings.h"
+#include "MetasoundSource.h"
+#include "MetasoundGenerator.h"
+#include "Templates/SharedPointer.h"
+#include "Widgets/DeclarativeSyntaxSupport.h"
+#include "Widgets/SBoxPanel.h"
+#include "Widgets/Text/STextBlock.h"
+
+
+namespace Metasound::Editor
+{
+	void SPageStats::Construct(const FArguments& InArgs)
+	{
+		SVerticalBox::Construct(SVerticalBox::FArguments());
+		AddSlot()
+			.HAlign(HAlign_Left)
+			[
+				SAssignNew(PageTextWidget, STextBlock)
+				.Visibility(EVisibility::HitTestInvisible)
+				.TextStyle(FAppStyle::Get(), "Graph.ZoomText")
+				.ColorAndOpacity(FLinearColor(1, 1, 1, 0.30f))
+			];
+	}
+
+	void SPageStats::Update(const FMetaSoundPageSettings* PageSettings)
+	{
+		using namespace Engine;
+
+		FString PageInfo;
+		EVisibility Visibility = EVisibility::Hidden;
+		if (PageTextWidget.IsValid() && PageSettings)
+		{
+			PageInfo = FString::Printf(TEXT("Page: %s"), *PageSettings->Name.ToString());
+			Visibility = EVisibility::Visible;
+		}
+
+		PageTextWidget->SetText(FText::FromString(MoveTemp(PageInfo)));
+		PageTextWidget->SetVisibility(Visibility);
+	}
+
+	void SRenderStats::Construct(const FArguments& InArgs)
+	{
+		SVerticalBox::Construct(SVerticalBox::FArguments());
+		AddSlot()
+		.HAlign(HAlign_Left)
+		.AutoHeight()
+		[
+			SAssignNew(RenderStatsCostWidget, STextBlock)
+			.Visibility(EVisibility::HitTestInvisible)
+			.TextStyle(FAppStyle::Get(), "Graph.ZoomText")
+			.ColorAndOpacity(FLinearColor(1, 1, 1, 0.30f))
+		];
+
+		AddSlot()
+		.HAlign(HAlign_Left)
+		.AutoHeight()
+		[
+			SAssignNew(RenderStatsCPUWidget, STextBlock)
+			.Visibility(EVisibility::HitTestInvisible)
+			.TextStyle(FAppStyle::Get(), "Graph.ZoomText")
+			.ColorAndOpacity(FLinearColor(1, 1, 1, 0.30f))
+		];
+	}
+
+	void SRenderStats::Update(bool bIsPlaying, const UMetaSoundSource* InSource)
+	{
+		using namespace Metasound;
+
+		// Reset maximum values when play restarts
+		if (bIsPlaying && !bPreviousIsPlaying)
+		{
+			MaxRelativeRenderCost = 0.f;
+			MaxCPUCoreUtilization = 0;
+		}
+		bPreviousIsPlaying = bIsPlaying;
+
+		if (RenderStatsCPUWidget.IsValid() && RenderStatsCostWidget.IsValid())
+		{
+			double CPUCoreUtilization = 0;
+			float RelativeRenderCost = 0.f;
+
+			// Find generator for playing preview component. 
+			if (bIsPlaying && InSource)
+			{
+				if (const UAudioComponent* PreviewComponent = GEditor->GetPreviewAudioComponent())
+				{
+					TSharedPtr<FMetasoundGenerator> Generator = InSource->GetGeneratorForAudioComponent(PreviewComponent->GetAudioComponentID()).Pin();
+					if (Generator.IsValid())
+					{
+						// Update render stats
+						CPUCoreUtilization = Generator->GetCPUCoreUtilization();
+						MaxCPUCoreUtilization = FMath::Max(MaxCPUCoreUtilization, CPUCoreUtilization);
+
+						RelativeRenderCost = Generator->GetRelativeRenderCost();
+						MaxRelativeRenderCost = FMath::Max(MaxRelativeRenderCost, RelativeRenderCost);
+					}
+				}
+			}
+
+			// Display updated render stats. 
+			FString CPUCoreUtilizationString = FString::Printf(TEXT("%3.2f%% (%3.2f%% Max) CPU Core"), 100. * CPUCoreUtilization, 100. * MaxCPUCoreUtilization);
+			RenderStatsCPUWidget->SetText(FText::FromString(CPUCoreUtilizationString));
+
+			FString RenderCostString = FString::Printf(TEXT("%3.2f (%3.2f Max) Relative Render Cost"), RelativeRenderCost, MaxRelativeRenderCost);
+			RenderStatsCostWidget->SetText(FText::FromString(RenderCostString));
+		}
+	}
+} // namespace Metasound::Editor
diff --git a/Engine/Plugins/Runtime/Metasound/Source/MetasoundEditor/Private/SMetasoundStats.h b/Engine/Plugins/Runtime/Metasound/Source/MetasoundEditor/Private/SMetasoundStats.h
new file mode 100644
--- /dev/null
+++ b/Engine/Plugins/Runtime/Metasound/Source/MetasoundEditor/Private/SMetasoundStats.h
@@ -0,0 +1,53 @@
+// Copyright Epic Games, Inc. All Rights Reserved.
+#pragma once
+
+#include "Templates/SharedPointer.h"
+#include "Widgets/DeclarativeSyntaxSupport.h"
+#include "Widgets/SBoxPanel.h"
+
+
+// Forward Declarations
+class IMetaSoundDocumentInterface;
+class STextBlock;
+class UMetaSoundBuilderBase;
+class UMetaSoundSource;
+
+struct FMetaSoundPageSettings;
+
+
+namespace Metasound::Editor
+{
+	/** Widget for displaying page stats of a previewing MetaSound. */
+	class SPageStats : public SVerticalBox
+	{
+	public:
+		SLATE_BEGIN_ARGS(SPageStats) { };
+		SLATE_END_ARGS()
+
+		void Construct(const FArguments& InArgs);
+		void Update(const FMetaSoundPageSettings* PageSettings);
+
+	private:
+		TSharedPtr<STextBlock> PageTextWidget;
+	};
+
+
+	/** Widget for displaying render stats of a previewing MetaSound. */
+	class SRenderStats : public SVerticalBox
+	{
+	public:
+		SLATE_BEGIN_ARGS( SRenderStats) { };
+		SLATE_END_ARGS()
+
+		void Construct(const FArguments& InArgs);
+		void Update(bool bIsPlaying, const UMetaSoundSource* InSource);
+
+	private:
+		TSharedPtr<STextBlock> RenderStatsCostWidget;
+		TSharedPtr<STextBlock> RenderStatsCPUWidget;
+
+		bool bPreviousIsPlaying = false;
+		double MaxCPUCoreUtilization = 0;
+		float MaxRelativeRenderCost = 0.f;
+	};
+} // namespace Metasound::Editor
diff --git a/Engine/Plugins/Runtime/Metasound/Source/MetasoundEditor/Public/MetasoundEditorSubsystem.h b/Engine/Plugins/Runtime/Metasound/Source/MetasoundEditor/Public/MetasoundEditorSubsystem.h
--- a/Engine/Plugins/Runtime/Metasound/Source/MetasoundEditor/Public/MetasoundEditorSubsystem.h
+++ b/Engine/Plugins/Runtime/Metasound/Source/MetasoundEditor/Public/MetasoundEditorSubsystem.h
@@ -83,13 +83,13 @@ public:
 	// the provided name. If the given builder has an asset editor open, optionally opens or brings
 	// that editor's associated page into user focus.
 	UFUNCTION(BlueprintCallable, Category = "Audio|MetaSound|Builder|Editor", meta = (ExpandEnumAsExecs = "OutResult"))
-	void SetFocusedPage(UMetaSoundBuilderBase* Builder, FName PageName, bool bFocusPageEditor, EMetaSoundBuilderResult& OutResult) const;
+	void SetFocusedPage(UMetaSoundBuilderBase* Builder, FName PageName, bool bOpenEditor, EMetaSoundBuilderResult& OutResult) const;
 
 	// If the given PageID is implemented on the provided builder, sets the focused page of
 	// the provided builder to the given PageID if and sets the audition target page to
 	// the provided ID. If the given builder has an asset editor open, optionally opens or brings
 	// that editor's associated PageID into user focus.
-	bool SetFocusedPage(UMetaSoundBuilderBase& Builder, const FGuid& InPageID, bool bFocusPageEditor) const;
+	bool SetFocusedPage(UMetaSoundBuilderBase& Builder, const FGuid& InPageID, bool bOpenEditor) const;
 
 	// Unregisters toolbar extender that is displayed in the MetaSound Asset Editor.
 	bool UnregisterToolbarExtender(TSharedRef<FExtender> InExtender);
@@ -104,7 +104,7 @@ public:
 	static const UMetaSoundEditorSubsystem& GetConstChecked();
 
 private:
-	bool SetFocusedPageInternal(const FMetaSoundPageSettings& InPageSettings, UMetaSoundBuilderBase& Builder, bool bFocusPageEditor) const;
+	bool SetFocusedPageInternal(const FMetaSoundPageSettings& InPageSettings, UMetaSoundBuilderBase& Builder, bool bOpenEditor) const;
 
 	// Copy over sound wave settings such as attenuation, modulation, and sound class from the template sound wave to the MetaSound
 	void SetSoundWaveSettingsFromTemplate(USoundWave& NewMetasound, const USoundWave& TemplateSoundWave) const;
diff --git a/Engine/Plugins/Runtime/Metasound/Source/MetasoundEngine/Private/MetasoundSettings.cpp b/Engine/Plugins/Runtime/Metasound/Source/MetasoundEngine/Private/MetasoundSettings.cpp
--- a/Engine/Plugins/Runtime/Metasound/Source/MetasoundEngine/Private/MetasoundSettings.cpp
+++ b/Engine/Plugins/Runtime/Metasound/Source/MetasoundEngine/Private/MetasoundSettings.cpp
@@ -299,21 +299,30 @@ FName UMetaSoundSettings::GetQualitySettingPropertyName()
 #endif // WITH_EDITORONLY_DATA
 
 #if WITH_EDITOR
-TArray<FName> UMetaSoundSettings::GetQualityList()
+TArray<FName> UMetaSoundSettings::GetPageNames()
 {
-	TArray<FName> Names;
+	if (const UMetaSoundSettings* Settings = GetDefault<UMetaSoundSettings>())
+	{
+		TArray<FName> Names;
+		auto GetName = [](const FMetaSoundPageSettings& Page) { return Page.Name; };
+		Algo::Transform(Settings->GetPageSettings(), Names, GetName);
+		return Names;
+	}
 
-#if WITH_EDITORONLY_DATA
+	return { };
+}
+
+TArray<FName> UMetaSoundSettings::GetQualityNames()
+{
 	if (const UMetaSoundSettings* Settings = GetDefault<UMetaSoundSettings>())
 	{
-		Algo::Transform(Settings->GetQualitySettings(), Names, [](const FMetaSoundQualitySettings& Quality) -> FName
-		{
-			return Quality.Name;
-		});
+		TArray<FName> Names;
+		auto GetName = [](const FMetaSoundQualitySettings& Quality) { return Quality.Name; };
+		Algo::Transform(Settings->GetQualitySettings(), Names, GetName);
+		return Names;
 	}
-#endif //WITH_EDITORONLY_DATA
 
-	return Names;
+	return { };
 }
 #endif // WITH_EDITOR
 
diff --git a/Engine/Plugins/Runtime/Metasound/Source/MetasoundEngine/Public/MetasoundSettings.h b/Engine/Plugins/Runtime/Metasound/Source/MetasoundEngine/Public/MetasoundSettings.h
--- a/Engine/Plugins/Runtime/Metasound/Source/MetasoundEngine/Public/MetasoundSettings.h
+++ b/Engine/Plugins/Runtime/Metasound/Source/MetasoundEngine/Public/MetasoundSettings.h
@@ -51,8 +51,8 @@ public:
 	/**
 	* Returns a list of quality settings to present to a combobox
 	* */
-	UFUNCTION(meta = (DeprecatedFunction, DeprecationMessage = "Use UMetaSoundSettings::GetQualityList instead"))
-	static TArray<FName> GetQualityList() { return { }; };
+	UFUNCTION(meta = (DeprecatedFunction, DeprecationMessage = "Use UMetaSoundSettings::GetQualityNames instead"))
+	static TArray<FName> GetQualityNames() { return { }; };
 };
 
 USTRUCT()
@@ -194,13 +194,19 @@ public:
 
 #if WITH_EDITOR
 public:
-	/* Returns a list of quality settings to present to a combobox. Ex:
-	 * UPROPERTY(EditAnywhere, BlueprintReadWrite, meta=(GetOptions="MetasoundEngine.MetaSoundSettings.GetQualityList"))
-	 * FName QualitySetting;
+	/* Returns an array of page names. Can be used to present to a combobox. Ex:
+	 * UPROPERTY(... meta=(GetOptions="MetasoundEngine.MetaSoundSettings.GetPageNames"))
+	 * FName Page;
 	*/
 	UFUNCTION()
-	static TArray<FName> GetQualityList();
+	static TArray<FName> GetPageNames();
 
+	/* Returns an array of quality setting names. Can be used to present to a combobox. Ex:
+	 * UPROPERTY(... meta=(GetOptions="MetasoundEngine.MetaSoundSettings.GetQualityNames"))
+	 * FName QualitySetting;
+	*/
+	UFUNCTION()
+	static TArray<FName> GetQualityNames();
 private:
 	void ConformPageSettingsDefault(bool bNotifyDefaultConformed);
 
diff --git a/Engine/Plugins/Runtime/Metasound/Source/MetasoundEngine/Public/MetasoundSource.h b/Engine/Plugins/Runtime/Metasound/Source/MetasoundEngine/Public/MetasoundSource.h
--- a/Engine/Plugins/Runtime/Metasound/Source/MetasoundEngine/Public/MetasoundSource.h
+++ b/Engine/Plugins/Runtime/Metasound/Source/MetasoundEngine/Public/MetasoundSource.h
@@ -129,8 +129,8 @@ public:
 	EMetaSoundOutputAudioFormat OutputFormat;
 
 #if WITH_EDITORONLY_DATA
-	// The Quality this Metasound will use. These are defined in the MetaSounds project settings.
-	UPROPERTY(EditAnywhere, BlueprintReadWrite, meta = (GetOptions="MetasoundEngine.MetaSoundSettings.GetQualityList"), Category = "Metasound")
+	// The QualitySetting MetaSound will use, as defined in 'MetaSound' Settings.
+	UPROPERTY(EditAnywhere, BlueprintReadWrite, meta = (GetOptions="MetasoundEngine.MetaSoundSettings.GetQualityNames"), Category = "Metasound")
 	FName QualitySetting;
 
 	// This a editor only look up for the Quality Setting above. Preventing orphaning of the original name.
