Trivial: Expose controller API for querying class vertices to better enable piecemeal deprecation

#rnx

diff --git a/Engine/Plugins/Runtime/Metasound/Source/MetasoundFrontend/Private/MetasoundFrontendGraphController.h b/Engine/Plugins/Runtime/Metasound/Source/MetasoundFrontend/Private/MetasoundFrontendGraphController.h
--- a/Engine/Plugins/Runtime/Metasound/Source/MetasoundFrontend/Private/MetasoundFrontendGraphController.h
+++ b/Engine/Plugins/Runtime/Metasound/Source/MetasoundFrontend/Private/MetasoundFrontendGraphController.h
@@ -205,6 +205,12 @@ namespace Metasound
 			FDocumentHandle GetOwningDocument() override;
 			FConstDocumentHandle GetOwningDocument() const override;
 
+			// Exposed to aid in transition of controller API to Document Builder API
+			virtual const FMetasoundFrontendClassInput* FindInputDescriptionWithName(const FVertexName& InName) const;
+			virtual const FMetasoundFrontendClassInput* FindInputDescriptionWithVertexID(const FGuid& InVertexID) const;
+			virtual const FMetasoundFrontendClassOutput* FindOutputDescriptionWithName(const FVertexName& InName) const;
+			virtual const FMetasoundFrontendClassOutput* FindOutputDescriptionWithVertexID(const FGuid& InVertexID) const;
+
 		protected:
 
 			FDocumentAccess ShareAccess() override;
@@ -262,16 +268,12 @@ namespace Metasound
 			FConstNodeHandle GetNodeHandle(const FConstNodeAndClass& InNodeAndClass) const;
 
 			FMetasoundFrontendClassInput* FindInputDescriptionWithName(const FVertexName& InName);
-			const FMetasoundFrontendClassInput* FindInputDescriptionWithName(const FVertexName& InName) const;
 
 			FMetasoundFrontendClassInput* FindInputDescriptionWithVertexID(const FGuid& InVertexID);
-			const FMetasoundFrontendClassInput* FindInputDescriptionWithVertexID(const FGuid& InVertexID) const;
 
 			FMetasoundFrontendClassOutput* FindOutputDescriptionWithName(const FVertexName& InName);
-			const FMetasoundFrontendClassOutput* FindOutputDescriptionWithName(const FVertexName& InName) const;
 
 			FMetasoundFrontendClassOutput* FindOutputDescriptionWithVertexID(const FGuid& InVertexID);
-			const FMetasoundFrontendClassOutput* FindOutputDescriptionWithVertexID(const FGuid& InVertexID) const;
 
 			FClassInputAccessPtr FindInputDescriptionWithNodeID(FGuid InNodeID);
 			FConstClassInputAccessPtr FindInputDescriptionWithNodeID(FGuid InNodeID) const;
diff --git a/Engine/Plugins/Runtime/Metasound/Source/MetasoundFrontend/Private/MetasoundFrontendInvalidController.h b/Engine/Plugins/Runtime/Metasound/Source/MetasoundFrontend/Private/MetasoundFrontendInvalidController.h
--- a/Engine/Plugins/Runtime/Metasound/Source/MetasoundFrontend/Private/MetasoundFrontendInvalidController.h
+++ b/Engine/Plugins/Runtime/Metasound/Source/MetasoundFrontend/Private/MetasoundFrontendInvalidController.h
@@ -435,6 +435,11 @@ namespace Metasound
 
 			virtual void UpdateInterfaceChangeID() override { }
 
+			virtual const FMetasoundFrontendClassInput* FindInputDescriptionWithName(const FVertexName& InName) const override { return nullptr; }
+			virtual const FMetasoundFrontendClassInput* FindInputDescriptionWithVertexID(const FGuid& InVertexID) const override { return nullptr; }
+			virtual const FMetasoundFrontendClassOutput* FindOutputDescriptionWithName(const FVertexName& InName) const override { return nullptr; }
+			virtual const FMetasoundFrontendClassOutput* FindOutputDescriptionWithVertexID(const FGuid& InVertexID) const override { return nullptr; }
+
 		protected:
 			virtual FDocumentAccess ShareAccess() override { return FDocumentAccess(); }
 			virtual FConstDocumentAccess ShareAccess() const override { return FConstDocumentAccess(); }
diff --git a/Engine/Plugins/Runtime/Metasound/Source/MetasoundFrontend/Public/MetasoundFrontendController.h b/Engine/Plugins/Runtime/Metasound/Source/MetasoundFrontend/Public/MetasoundFrontendController.h
--- a/Engine/Plugins/Runtime/Metasound/Source/MetasoundFrontend/Public/MetasoundFrontendController.h
+++ b/Engine/Plugins/Runtime/Metasound/Source/MetasoundFrontend/Public/MetasoundFrontendController.h
@@ -996,6 +996,11 @@ namespace Metasound
 
 			/** Returns a handle to the document owning this graph. */
 			virtual FConstDocumentHandle GetOwningDocument() const = 0;
+
+			virtual const FMetasoundFrontendClassInput* FindInputDescriptionWithName(const FVertexName& InName) const = 0;
+			virtual const FMetasoundFrontendClassInput* FindInputDescriptionWithVertexID(const FGuid& InVertexID) const = 0;
+			virtual const FMetasoundFrontendClassOutput* FindOutputDescriptionWithName(const FVertexName& InName) const = 0;
+			virtual const FMetasoundFrontendClassOutput* FindOutputDescriptionWithVertexID(const FGuid& InVertexID) const = 0;
 		};
 
 		/* An IDocumentController provides methods for querying and manipulating a Metasound document. */
