Fix resolution being applied to incorrect class input

#rnx

diff --git a/Engine/Plugins/Runtime/Metasound/Source/MetasoundFrontend/Private/MetasoundFrontendGraph.cpp b/Engine/Plugins/Runtime/Metasound/Source/MetasoundFrontend/Private/MetasoundFrontendGraph.cpp
--- a/Engine/Plugins/Runtime/Metasound/Source/MetasoundFrontend/Private/MetasoundFrontendGraph.cpp
+++ b/Engine/Plugins/Runtime/Metasound/Source/MetasoundFrontend/Private/MetasoundFrontendGraph.cpp
@@ -335,7 +335,7 @@ namespace Metasound
 		{
 			const FMetasoundFrontendVertex& InputVertex = InInputNode.Interface.Inputs[0];
 
-			// Find input literal matching VerteXID
+			// Find input literal matching VertexID
 			const FMetasoundFrontendVertexLiteral* VertexLiteral = InInputNode.InputLiterals.FindByPredicate(
 				[&](const FMetasoundFrontendVertexLiteral& InVertexLiteral)
 				{
@@ -349,13 +349,12 @@ namespace Metasound
 			}
 		}
 
-		FGuid PageID = Frontend::DefaultPageID;
-		const bool bPageIDResolved = GraphPrivate::TryResolveTargetPageID(InOwningGraphClassInput, PageID);
-
 		// Check for default value on owning graph.
 		if (nullptr == Literal)
 		{
 			// Find Class Default that is not invalid
+			FGuid PageID = Frontend::DefaultPageID;
+			const bool bPageIDResolved = GraphPrivate::TryResolveTargetPageID(InOwningGraphClassInput, PageID);
 			if (bPageIDResolved)
 			{
 				const FMetasoundFrontendLiteral& DefaultLiteral = InOwningGraphClassInput.FindConstDefaultChecked(PageID);
@@ -370,6 +369,9 @@ namespace Metasound
 		if (nullptr == Literal && ensure(InInputNodeClass.Interface.Inputs.Num() == 1))
 		{
 			const FMetasoundFrontendClassInput& InputNodeClassInput = InInputNodeClass.Interface.Inputs.Last();
+
+			FGuid PageID = Frontend::DefaultPageID;
+			const bool bPageIDResolved = GraphPrivate::TryResolveTargetPageID(InputNodeClassInput, PageID);
 			if (bPageIDResolved)
 			{
 				const FMetasoundFrontendLiteral& DefaultLiteral = InputNodeClassInput.FindConstDefaultChecked(PageID);
