[Metasound] Expose class and pin names for bus writer node

#rnx

diff --git a/Engine/Plugins/Runtime/Metasound/Source/MetasoundEngine/Private/MetasoundAudioBusWriterNode.cpp b/Engine/Plugins/Runtime/Metasound/Source/MetasoundEngine/Private/MetasoundAudioBusWriterNode.cpp
--- a/Engine/Plugins/Runtime/Metasound/Source/MetasoundEngine/Private/MetasoundAudioBusWriterNode.cpp
+++ b/Engine/Plugins/Runtime/Metasound/Source/MetasoundEngine/Private/MetasoundAudioBusWriterNode.cpp
@@ -1,5 +1,7 @@
 // Copyright Epic Games, Inc. All Rights Reserved.
 
+#include "MetasoundAudioBusWriterNode.h"
+
 #include "AudioMixerDevice.h"
 #include "AudioBusSubsystem.h"
 #include "AudioDevice.h"
@@ -7,7 +9,6 @@
 #include "MediaPacket.h"
 #include "MetasoundAudioBuffer.h"
 #include "MetasoundAudioBus.h"
-#include "MetasoundEngineNodesNames.h"
 #include "MetasoundExecutableOperator.h"
 #include "MetasoundFacade.h"
 #include "MetasoundNodeRegistrationMacro.h"
@@ -20,8 +21,16 @@ namespace Metasound
 {
 	namespace AudioBusWriterNode
 	{
-		METASOUND_PARAM(InParamAudioBusOutput, "Audio Bus", "Audio Bus Asset.");
-		METASOUND_PARAM(InParamAudio, "In {0}", "Audio input for channel {0}.");
+		namespace Inputs
+		{
+			DEFINE_METASOUND_PARAM(AudioBus, "Audio Bus", "Audio Bus Asset.");
+			DEFINE_METASOUND_PARAM(Audio, "In {0}", "Audio input for channel {0}.");
+		}
+
+		int32 GetCurrentMajorVersion()
+		{
+			return 1;
+		}
 	}
 
 	int32 AudioBusWriterNodeInitialNumBlocks(int32 BlockSizeFrames, int32 AudioMixerOutputFrames)
@@ -39,12 +48,11 @@ namespace Metasound
 		{
 			auto InitNodeInfo = []() -> FNodeClassMetadata
 			{
-				FName OperatorName = *FString::Printf(TEXT("Audio Bus Writer (%d)"), NumChannels);
 				FText NodeDisplayName = METASOUND_LOCTEXT_FORMAT("AudioBusWriterDisplayNamePattern", "Audio Bus Writer ({0})", NumChannels);
 
 				FNodeClassMetadata Info;
-				Info.ClassName = { EngineNodes::Namespace, OperatorName, TEXT("") };
-				Info.MajorVersion = 1;
+				Info.ClassName = AudioBusWriterNode::GetClassName<NumChannels>();
+				Info.MajorVersion = AudioBusWriterNode::GetCurrentMajorVersion();
 				Info.MinorVersion = 0;
 				Info.DisplayName = NodeDisplayName;
 				Info.Description = METASOUND_LOCTEXT("AudioBusWriter_Description", "Sends audio data to the audio bus asset.");
@@ -67,10 +75,10 @@ namespace Metasound
 			auto CreateVertexInterface = []() -> FVertexInterface
 			{
 				FInputVertexInterface InputInterface;
-				InputInterface.Add(TInputDataVertex<FAudioBusAsset>(METASOUND_GET_PARAM_NAME_AND_METADATA(InParamAudioBusOutput)));
+				InputInterface.Add(TInputDataVertex<FAudioBusAsset>(METASOUND_GET_PARAM_NAME_AND_METADATA(Inputs::AudioBus)));
 				for (uint32 i = 0; i < NumChannels; ++i)
 				{
-					InputInterface.Add(TInputDataVertex<FAudioBuffer>(METASOUND_GET_PARAM_NAME_WITH_INDEX_AND_METADATA(InParamAudio, i)));
+					InputInterface.Add(TInputDataVertex<FAudioBuffer>(METASOUND_GET_PARAM_NAME_WITH_INDEX_AND_METADATA(Inputs::Audio, i)));
 				}
 
 				FOutputVertexInterface OutputInterface;
@@ -95,12 +103,12 @@ namespace Metasound
 
 			if (bHasEnvironmentVars)
 			{
-				FAudioBusAssetReadRef AudioBusIn = InputData.GetOrConstructDataReadReference<FAudioBusAsset>(METASOUND_GET_PARAM_NAME(InParamAudioBusOutput));
+				FAudioBusAssetReadRef AudioBusIn = InputData.GetOrConstructDataReadReference<FAudioBusAsset>(METASOUND_GET_PARAM_NAME(Inputs::AudioBus));
 
 				TArray<FAudioBufferReadRef> AudioInputs;
 				for (int32 ChannelIndex = 0; ChannelIndex < NumChannels; ++ChannelIndex)
 				{
-					AudioInputs.Add(InputData.GetOrConstructDataReadReference<FAudioBuffer>(METASOUND_GET_PARAM_NAME_WITH_INDEX(InParamAudio, ChannelIndex), InParams.OperatorSettings));
+					AudioInputs.Add(InputData.GetOrConstructDataReadReference<FAudioBuffer>(METASOUND_GET_PARAM_NAME_WITH_INDEX(Inputs::Audio, ChannelIndex), InParams.OperatorSettings));
 				}
 
 				FString GraphName;
@@ -213,11 +221,11 @@ namespace Metasound
 		{
 			using namespace AudioBusWriterNode;
 
-			InOutVertexData.BindReadVertex(METASOUND_GET_PARAM_NAME(InParamAudioBusOutput), AudioBusAsset);
+			InOutVertexData.BindReadVertex(METASOUND_GET_PARAM_NAME(Inputs::AudioBus), AudioBusAsset);
 
 			for (int32 ChannelIndex = 0; ChannelIndex < NumChannels; ++ChannelIndex)
 			{
-				InOutVertexData.BindReadVertex(METASOUND_GET_PARAM_NAME_WITH_INDEX(InParamAudio, ChannelIndex), AudioInputs[ChannelIndex]);
+				InOutVertexData.BindReadVertex(METASOUND_GET_PARAM_NAME_WITH_INDEX(Inputs::Audio, ChannelIndex), AudioInputs[ChannelIndex]);
 			}
 		}
 
diff --git a/Engine/Plugins/Runtime/Metasound/Source/MetasoundEngine/Public/MetasoundAudioBusWriterNode.h b/Engine/Plugins/Runtime/Metasound/Source/MetasoundEngine/Public/MetasoundAudioBusWriterNode.h
new file mode 100644
--- /dev/null
+++ b/Engine/Plugins/Runtime/Metasound/Source/MetasoundEngine/Public/MetasoundAudioBusWriterNode.h
@@ -0,0 +1,26 @@
+// Copyright Epic Games, Inc. All Rights Reserved.
+
+#pragma once
+
+#include "MetasoundEngineNodesNames.h"
+#include "MetasoundNodeInterface.h"
+#include "MetasoundParamHelper.h"
+
+namespace Metasound::AudioBusWriterNode
+{
+	template<uint32 NumChannels>
+	const FNodeClassName& GetClassName()
+	{
+		static const FName OperatorName = *FString::Printf(TEXT("Audio Bus Writer (%d)"), NumChannels);
+		static const FNodeClassName ClassName = { EngineNodes::Namespace, OperatorName, TEXT("") };
+		return ClassName;
+	}
+	
+	METASOUNDENGINE_API int32 GetCurrentMajorVersion();
+
+	namespace Inputs
+	{
+		DECLARE_METASOUND_PARAM(METASOUNDENGINE_API, AudioBus);
+		DECLARE_METASOUND_PARAM(METASOUNDENGINE_API, Audio);
+	}
+}
diff --git a/Engine/Plugins/Runtime/Metasound/Source/MetasoundGraphCore/Public/MetasoundParamHelper.h b/Engine/Plugins/Runtime/Metasound/Source/MetasoundGraphCore/Public/MetasoundParamHelper.h
--- a/Engine/Plugins/Runtime/Metasound/Source/MetasoundGraphCore/Public/MetasoundParamHelper.h
+++ b/Engine/Plugins/Runtime/Metasound/Source/MetasoundGraphCore/Public/MetasoundParamHelper.h
@@ -46,6 +46,29 @@
 
 */
 
+/**
+ * Declare a node parameter that can be accessed by other code (e.g. tests)
+ * @param API_EXPORT If this parameter should be available to other modules (common), provide your module API export here
+ * @param PARAM_NAME The parameter name
+ */
+#define DECLARE_METASOUND_PARAM(API_EXPORT, PARAM_NAME) \
+	API_EXPORT extern const FLazyName PARAM_NAME##Name;
+
+/**
+ * Use this macro to define your parameter if you used DECLARE_METASOUND_PARAM, otherwise just use METASOUND_PARAM in a source file
+ */
+#if WITH_EDITOR
+#define DEFINE_METASOUND_PARAM(NAME, NAME_TEXT, TOOLTIP_TEXT) \
+	const FLazyName NAME##Name = TEXT(NAME_TEXT); \
+	static const FText NAME##Tooltip = LOCTEXT(#NAME "Tooltip", TOOLTIP_TEXT); \
+	static const FText NAME##DisplayName = LOCTEXT(#NAME "DisplayName", NAME_TEXT);
+#else 
+#define DEFINE_METASOUND_PARAM(NAME, NAME_TEXT, TOOLTIP_TEXT) \
+	const FLazyName NAME##Name = TEXT(NAME_TEXT); \
+	static const FText NAME##Tooltip = FText::GetEmpty(); \
+	static const FText NAME##DisplayName = FText::GetEmpty();
+#endif // WITH_EDITOR
+
 #if WITH_EDITOR
 #define LOC_DEFINE_REGION
 #define METASOUND_PARAM(NAME, NAME_TEXT, TOOLTIP_TEXT) \
