Fix failing audio unit test on metasound graph output
Rename default graph page initialization for clarity.
#rnx

diff --git a/Engine/Plugins/Runtime/Metasound/Source/MetasoundEngineTest/Private/EngineTestMetasoundSource.cpp b/Engine/Plugins/Runtime/Metasound/Source/MetasoundEngineTest/Private/EngineTestMetasoundSource.cpp
--- a/Engine/Plugins/Runtime/Metasound/Source/MetasoundEngineTest/Private/EngineTestMetasoundSource.cpp
+++ b/Engine/Plugins/Runtime/Metasound/Source/MetasoundEngineTest/Private/EngineTestMetasoundSource.cpp
@@ -177,7 +177,7 @@ namespace EngineTestMetaSoundSourcePrivate
 
 		Document.RootGraph.Metadata.SetClassName(FMetasoundFrontendClassName { "Namespace", "Unit Test Node", *LexToString(FGuid::NewGuid()) });
 		Document.RootGraph.Metadata.SetType(EMetasoundFrontendClassType::Graph);
-		Document.RootGraph.InitDefaultGraph();
+		Document.RootGraph.InitDefaultGraphPage();
 
 		FDocumentHandle DocumentHandle = IDocumentController::CreateDocumentHandle(Document);
 		FGraphHandle RootGraph = DocumentHandle->GetRootGraph();
diff --git a/Engine/Plugins/Runtime/Metasound/Source/MetasoundEngineTest/Private/NodeTestGraphBuilder.cpp b/Engine/Plugins/Runtime/Metasound/Source/MetasoundEngineTest/Private/NodeTestGraphBuilder.cpp
--- a/Engine/Plugins/Runtime/Metasound/Source/MetasoundEngineTest/Private/NodeTestGraphBuilder.cpp
+++ b/Engine/Plugins/Runtime/Metasound/Source/MetasoundEngineTest/Private/NodeTestGraphBuilder.cpp
@@ -14,7 +14,7 @@ namespace Metasound::Test
 	{
 		Document.RootGraph.Metadata.SetClassName({ "Metasound", "TestNodes", *LexToString(FGuid::NewGuid()) });
 		Document.RootGraph.Metadata.SetType(EMetasoundFrontendClassType::Graph);
-		Document.RootGraph.InitDefaultGraph();
+		Document.RootGraph.InitDefaultGraphPage();
 
 		DocumentHandle = IDocumentController::CreateDocumentHandle(Document);
 		RootGraph = DocumentHandle->GetRootGraph();
diff --git a/Engine/Plugins/Runtime/Metasound/Source/MetasoundFrontend/Private/MetasoundFrontendDocument.cpp b/Engine/Plugins/Runtime/Metasound/Source/MetasoundFrontend/Private/MetasoundFrontendDocument.cpp
--- a/Engine/Plugins/Runtime/Metasound/Source/MetasoundFrontend/Private/MetasoundFrontendDocument.cpp
+++ b/Engine/Plugins/Runtime/Metasound/Source/MetasoundFrontend/Private/MetasoundFrontendDocument.cpp
@@ -1106,7 +1106,7 @@ const FMetasoundFrontendGraph& FMetasoundFrontendGraphClass::GetConstDefaultGrap
 	return FindConstGraphChecked(Metasound::Frontend::DefaultPageID);
 }
 
-FMetasoundFrontendGraph& FMetasoundFrontendGraphClass::InitDefaultGraph()
+FMetasoundFrontendGraph& FMetasoundFrontendGraphClass::InitDefaultGraphPage()
 {
 	checkf(PagedGraphs.IsEmpty(), TEXT("Attempting to initialize default page for graph class with existing graph implementation"));
 	FMetasoundFrontendGraph& NewGraph = PagedGraphs.AddDefaulted_GetRef();
diff --git a/Engine/Plugins/Runtime/Metasound/Source/MetasoundFrontend/Private/MetasoundFrontendDocumentBuilder.cpp b/Engine/Plugins/Runtime/Metasound/Source/MetasoundFrontend/Private/MetasoundFrontendDocumentBuilder.cpp
--- a/Engine/Plugins/Runtime/Metasound/Source/MetasoundFrontend/Private/MetasoundFrontendDocumentBuilder.cpp
+++ b/Engine/Plugins/Runtime/Metasound/Source/MetasoundFrontend/Private/MetasoundFrontendDocumentBuilder.cpp
@@ -1364,7 +1364,7 @@ void FMetaSoundFrontendDocumentBuilder::ClearDocument(TSharedRef<Metasound::Fron
 	checkf(!DocObject.IsAsset(), TEXT("Cannot call clear document on asset '%s': builder API does not support document mutation on serialized objects without editor data loaded"), *GetDebugName());
 #endif // !WITH_EDITORONLY_DATA
 
-	GraphClass.InitDefaultGraph();
+	GraphClass.InitDefaultGraphPage();
 
 	Doc.Interfaces.Empty();
 	Doc.Dependencies.Empty();
@@ -2384,7 +2384,7 @@ void FMetaSoundFrontendDocumentBuilder::InitDocument(const FMetasoundFrontendDoc
 	METASOUND_TRACE_CPUPROFILER_EVENT_SCOPE(FMetaSoundFrontendDocumentBuilder::InitDocument);
 
 	FMetasoundFrontendDocument& Document = GetDocumentChecked();
-	Document.RootGraph.InitDefaultGraph();
+	Document.RootGraph.InitDefaultGraphPage();
 
 	// 1. Set default class Metadata.
 	if (InDocumentTemplate)
diff --git a/Engine/Plugins/Runtime/Metasound/Source/MetasoundFrontend/Private/MetasoundFrontendDocumentController.cpp b/Engine/Plugins/Runtime/Metasound/Source/MetasoundFrontend/Private/MetasoundFrontendDocumentController.cpp
--- a/Engine/Plugins/Runtime/Metasound/Source/MetasoundFrontend/Private/MetasoundFrontendDocumentController.cpp
+++ b/Engine/Plugins/Runtime/Metasound/Source/MetasoundFrontend/Private/MetasoundFrontendDocumentController.cpp
@@ -403,7 +403,7 @@ namespace Metasound
 							FMetasoundFrontendGraphClass NewClass;
 							NewClass.ID = FGuid::NewGuid();
 							NewClass.Metadata = InMetadata;
-							NewClass.InitDefaultGraph();
+							NewClass.InitDefaultGraphPage();
 							Document->Subgraphs.Add(NewClass);
 						}
 						break;
diff --git a/Engine/Plugins/Runtime/Metasound/Source/MetasoundFrontend/Public/MetasoundFrontendDocument.h b/Engine/Plugins/Runtime/Metasound/Source/MetasoundFrontend/Public/MetasoundFrontendDocument.h
--- a/Engine/Plugins/Runtime/Metasound/Source/MetasoundFrontend/Public/MetasoundFrontendDocument.h
+++ b/Engine/Plugins/Runtime/Metasound/Source/MetasoundFrontend/Public/MetasoundFrontendDocument.h
@@ -1748,7 +1748,7 @@ public:
 
 	bool ContainsGraphPage(const FGuid& InPageID) const;
 
-	FMetasoundFrontendGraph& InitDefaultGraph();
+	FMetasoundFrontendGraph& InitDefaultGraphPage();
 
 	void IterateGraphPages(TFunctionRef<void(FMetasoundFrontendGraph&)> IterFunc);
 	void IterateGraphPages(TFunctionRef<void(const FMetasoundFrontendGraph&)> IterFunc) const;
