Enable editing MetaSound graphs that are no longer referencing valid page IDs defined in MetaSound Settings

#rnx

diff --git a/Engine/Plugins/Runtime/Metasound/Source/MetasoundEditor/Private/MetasoundDetailCustomization.cpp b/Engine/Plugins/Runtime/Metasound/Source/MetasoundEditor/Private/MetasoundDetailCustomization.cpp
--- a/Engine/Plugins/Runtime/Metasound/Source/MetasoundEditor/Private/MetasoundDetailCustomization.cpp
+++ b/Engine/Plugins/Runtime/Metasound/Source/MetasoundEditor/Private/MetasoundDetailCustomization.cpp
@@ -49,6 +49,11 @@
 
 namespace Metasound::Editor
 {
+	const FName GetMissingPageName(const FGuid& InPageID)
+	{
+		return FName(*FString::Printf(TEXT("Invalid (...%s)"), *InPageID.ToString().Mid(28, 8)));
+	}
+
 	FName BuildChildPath(const FString& InBasePath, FName InPropertyName)
 	{
 		return FName(InBasePath + TEXT(".") + InPropertyName.ToString());
@@ -470,16 +475,22 @@ namespace Metasound::Editor
 	{
 		EntryWidgets->ClearChildren();
 
-		auto CreateEntryWidget = [this](bool bIsDefault, FName InName) -> TSharedRef<SWidget>
+		TSet<FGuid> ImplementedGuids;
+		const FMetasoundFrontendDocument& Document = Builder->GetBuilder().GetConstDocumentChecked();
+		Document.RootGraph.IterateGraphPages([&ImplementedGuids](const FMetasoundFrontendGraph& Graph)
+		{
+			ImplementedGuids.Add(Graph.PageID);
+		});
+
+		auto CreateEntryWidget = [this](bool bIsDefault, const FGuid& InEntryPageID, FName InName) -> TSharedRef<SWidget>
 		{
 			using namespace Frontend;
 
-			TSharedRef<SWidget> SelectButtonWidget = PropertyCustomizationHelpers::MakeUseSelectedButton(FSimpleDelegate::CreateLambda([this, InName]()
+			TSharedRef<SWidget> SelectButtonWidget = PropertyCustomizationHelpers::MakeUseSelectedButton(FSimpleDelegate::CreateLambda([this, InEntryPageID, InName]()
 			{
-				EMetaSoundBuilderResult Result = EMetaSoundBuilderResult::Failed;
-				constexpr bool bFocusEditor = false; // Already focused by user action
-				UMetaSoundEditorSubsystem::GetConstChecked().SetFocusedPage(Builder.Get(), InName, bFocusEditor, Result);
-				if (Result == EMetaSoundBuilderResult::Succeeded)
+				constexpr bool bOpenEditor = false; // Already focused by user action
+				const bool bFocusedPage = UMetaSoundEditorSubsystem::GetConstChecked().SetFocusedPage(*Builder.Get(), InEntryPageID, bOpenEditor);
+				if (bFocusedPage)
 				{
 					BuildPageName = InName;
 				}
@@ -508,7 +519,7 @@ namespace Metasound::Editor
 
 			if (!bIsDefault)
 			{
-				TSharedRef<SWidget> RemoveButtonWidget = PropertyCustomizationHelpers::MakeDeleteButton(FSimpleDelegate::CreateLambda([this, InName]()
+				TSharedRef<SWidget> RemoveButtonWidget = PropertyCustomizationHelpers::MakeDeleteButton(FSimpleDelegate::CreateLambda([this, InName, InEntryPageID]()
 				{
 					using namespace Frontend;
 					const FScopedTransaction Transaction(FText::Format(LOCTEXT("RemovePageTransactionFormat", "Remove MetaSound Page '{0}'"), FText::FromName(InName)));
@@ -518,12 +529,12 @@ namespace Metasound::Editor
 					// Removal may modify the builder's build page ID if it is the currently set value
 					Builder->Modify();
 
-					EMetaSoundBuilderResult Result = EMetaSoundBuilderResult::Failed;
-					Builder->RemoveGraphPage(InName, Result);
-					if (Result == EMetaSoundBuilderResult::Succeeded)
+					const bool bGraphRemoved = Builder->GetBuilder().RemoveGraphPage(InEntryPageID);
+					if (bGraphRemoved)
 					{
 						UpdateItemNames();
 						ComboBox->RefreshOptions();
+						RebuildImplemented();
 					}
 				}), LOCTEXT("RemovePageTooltip2", "Removes the associated page from the MetaSound."));
 				EntryWidget->AddSlot()
@@ -547,39 +558,53 @@ namespace Metasound::Editor
 		check(Settings);
 		for (const FMetaSoundPageSettings& PageSettings : Settings->GetPageSettings())
 		{
-			if (ImplementedNames.Contains(PageSettings.Name))
+			if (ImplementedGuids.Remove(PageSettings.UniqueId) > 0)
 			{
+				const bool bIsDefault = PageSettings.UniqueId == Frontend::DefaultPageID;
 				EntryWidgets->AddSlot()
 				.HAlign(HAlign_Left)
 				.VAlign(VAlign_Center)
 				.AutoHeight()
 				[
-					CreateEntryWidget(!PageSettings.UniqueId.IsValid(), PageSettings.Name)
+					CreateEntryWidget(bIsDefault, PageSettings.UniqueId, PageSettings.Name)
 				];
 			}
 		}
+
+		for (const FGuid& MissingPageID : ImplementedGuids)
+		{
+			constexpr bool bIsDefault = false;
+			const FName MissingName = Editor::GetMissingPageName(MissingPageID);
+			EntryWidgets->AddSlot()
+				.HAlign(HAlign_Left)
+				.VAlign(VAlign_Center)
+				.AutoHeight()
+				[
+					CreateEntryWidget(bIsDefault, MissingPageID, MissingName)
+				];
+		}
 	}
 
 	void FMetasoundPagesDetailCustomization::RefreshView()
 	{
 		if (Builder.IsValid())
 		{
-			const UMetaSoundSettings* Settings = GetDefault<UMetaSoundSettings>();
-			check(Settings);
 			FMetaSoundFrontendDocumentBuilder& DocBuilder = Builder->GetBuilder();
 			const FGuid& PageID = DocBuilder.GetBuildPageID();
+
+			const UMetaSoundSettings* Settings = GetDefault<UMetaSoundSettings>();
+			check(Settings);
 			if (const FMetaSoundPageSettings* PageSettings = Settings->FindPageSettings(PageID))
 			{
 				BuildPageName = PageSettings->Name;
 			}
 			else
 			{
-				constexpr bool bFocusEditor = false; // Already focused by user action
-				EMetaSoundBuilderResult Result = EMetaSoundBuilderResult::Failed;
-				UMetaSoundEditorSubsystem::GetConstChecked().SetFocusedPage(Builder.Get(), Frontend::DefaultPageName, bFocusEditor, Result);
-				if (Result == EMetaSoundBuilderResult::Succeeded)
+				constexpr bool bOpenEditor = false; // Already open/focused by user action
+				const bool bFocusedPage = UMetaSoundEditorSubsystem::GetConstChecked().SetFocusedPage(*Builder.Get(), PageID, bOpenEditor);
+				if (bFocusedPage)
 				{
-					BuildPageName = Frontend::DefaultPageName;
+					BuildPageName = GetMissingPageName(PageID);
 				}
 			}
 		}
@@ -595,16 +620,13 @@ namespace Metasound::Editor
 
 	void FMetasoundPagesDetailCustomization::UpdateItemNames()
 	{
-		using namespace Frontend;
-
 		AddableItems.Reset();
 		ImplementedNames.Reset();
 
 		const UMetaSoundSettings* Settings = GetDefault<UMetaSoundSettings>();
 		check(Settings);
 
-		TSet<FGuid> ImplementedGuids { FGuid() }; // Default "no guid" is always implemented (base graph for back compat & cook safety to ensure there's always at least one)
-
+		TSet<FGuid> ImplementedGuids;
 		const FMetasoundFrontendDocument& Document = Builder->GetBuilder().GetConstDocumentChecked();
 		Document.RootGraph.IterateGraphPages([&ImplementedGuids](const FMetasoundFrontendGraph& Graph)
 		{
@@ -618,12 +640,12 @@ namespace Metasound::Editor
 		auto GetPageName = [&Settings](const FGuid& PageID)
 		{
 			const FMetaSoundPageSettings* Page = Settings->FindPageSettings(PageID);
-			if (ensure(Page))
+			if (Page)
 			{
 				return Page->Name;
 			}
 
-			return FName();
+			return GetMissingPageName(PageID);
 		};
 
 		Algo::Transform(ImplementedGuids, ImplementedNames, GetPageName);
diff --git a/Engine/Plugins/Runtime/Metasound/Source/MetasoundEditor/Private/MetasoundDetailCustomization.h b/Engine/Plugins/Runtime/Metasound/Source/MetasoundEditor/Private/MetasoundDetailCustomization.h
--- a/Engine/Plugins/Runtime/Metasound/Source/MetasoundEditor/Private/MetasoundDetailCustomization.h
+++ b/Engine/Plugins/Runtime/Metasound/Source/MetasoundEditor/Private/MetasoundDetailCustomization.h
@@ -20,6 +20,8 @@ namespace Metasound
 {
 	namespace Editor
 	{
+		const FName GetMissingPageName(const FGuid& InPageID);
+
 		class FMetaSoundDetailCustomizationBase : public IDetailCustomization
 		{
 		public:
diff --git a/Engine/Plugins/Runtime/Metasound/Source/MetasoundEditor/Private/MetasoundEditorSubsystem.cpp b/Engine/Plugins/Runtime/Metasound/Source/MetasoundEditor/Private/MetasoundEditorSubsystem.cpp
--- a/Engine/Plugins/Runtime/Metasound/Source/MetasoundEditor/Private/MetasoundEditorSubsystem.cpp
+++ b/Engine/Plugins/Runtime/Metasound/Source/MetasoundEditor/Private/MetasoundEditorSubsystem.cpp
@@ -262,7 +262,7 @@ void UMetaSoundEditorSubsystem::SetFocusedPage(UMetaSoundBuilderBase* Builder, F
 	check(Settings);
 	if (const FMetaSoundPageSettings* PageSettings = Settings->FindPageSettings(PageName))
 	{
-		const bool bFocusedPage = SetFocusedPageInternal(*PageSettings, *Builder, bOpenEditor);
+		const bool bFocusedPage = SetFocusedPageInternal(PageSettings->Name, PageSettings->UniqueId, *Builder, bOpenEditor);
 		if (bFocusedPage)
 		{
 			OutResult = EMetaSoundBuilderResult::Succeeded;
@@ -279,34 +279,38 @@ bool UMetaSoundEditorSubsystem::SetFocusedPage(UMetaSoundBuilderBase& Builder, c
 
 	const UMetaSoundSettings* Settings = GetDefault<UMetaSoundSettings>();
 	check(Settings);
+
+	FName PageName;
 	if (const FMetaSoundPageSettings* PageSettings = Settings->FindPageSettings(InPageID))
 	{
-		return SetFocusedPageInternal(*PageSettings, Builder, bOpenEditor);
+		PageName = PageSettings->Name;
 	}
 
-	return false;
+	return SetFocusedPageInternal(PageName, InPageID, Builder, bOpenEditor);
 }
 
-bool UMetaSoundEditorSubsystem::SetFocusedPageInternal(const FMetaSoundPageSettings& InPageSettings, UMetaSoundBuilderBase& Builder, bool bOpenEditor) const
+bool UMetaSoundEditorSubsystem::SetFocusedPageInternal(FName PageName, const FGuid& InPageID, UMetaSoundBuilderBase& Builder, bool bOpenEditor) const
 {
 	using namespace Metasound::Frontend;
 
-	const FScopedTransaction Transaction(FText::Format(LOCTEXT("SetFocusedPageTransactionFormat", "Set Focused Page '{0}'"), FText::FromName(InPageSettings.Name)));
+	const FScopedTransaction Transaction(FText::Format(LOCTEXT("SetFocusedPageTransactionFormat", "Set Focused Page '{0}'"), FText::FromName(PageName)));
 	Builder.Modify();
 
 	if (UMetasoundEditorSettings* EditorSettings = GetMutableDefault<UMetasoundEditorSettings>())
 	{
-		if (EditorSettings->AuditionPageMode == EAuditionPageMode::Focused)
+		if (!PageName.IsNone())
 		{
-			// Must set audition target page before setting build page ID as listeners
-			// to build page ID changes need to reliably be able to adjust to newly assigned
-			// audition target page.
-			EditorSettings->AuditionTargetPage = InPageSettings.Name;
+			if (EditorSettings->AuditionPageMode == EAuditionPageMode::Focused)
+			{
+				// Must set audition target page before setting build page ID as listeners
+				// to build page ID changes need to reliably be able to adjust to newly assigned
+				// audition target page.
+				EditorSettings->AuditionTargetPage = PageName;
+			}
 		}
 
-		if (Builder.GetBuilder().SetBuildPageID(InPageSettings.UniqueId))
+		if (Builder.GetBuilder().SetBuildPageID(InPageID))
 		{
-
 			// Reregister to ensure all future audible instances are using the new page implementation.
 			RegisterGraphWithFrontend(Builder.GetBuilder().CastDocumentObjectChecked<UObject>());
 		}
diff --git a/Engine/Plugins/Runtime/Metasound/Source/MetasoundEditor/Public/MetasoundEditorSubsystem.h b/Engine/Plugins/Runtime/Metasound/Source/MetasoundEditor/Public/MetasoundEditorSubsystem.h
--- a/Engine/Plugins/Runtime/Metasound/Source/MetasoundEditor/Public/MetasoundEditorSubsystem.h
+++ b/Engine/Plugins/Runtime/Metasound/Source/MetasoundEditor/Public/MetasoundEditorSubsystem.h
@@ -104,7 +104,7 @@ public:
 	static const UMetaSoundEditorSubsystem& GetConstChecked();
 
 private:
-	bool SetFocusedPageInternal(const FMetaSoundPageSettings& InPageSettings, UMetaSoundBuilderBase& Builder, bool bOpenEditor) const;
+	bool SetFocusedPageInternal(FName PageName, const FGuid& InPageID, UMetaSoundBuilderBase& Builder, bool bOpenEditor) const;
 
 	// Copy over sound wave settings such as attenuation, modulation, and sound class from the template sound wave to the MetaSound
 	void SetSoundWaveSettingsFromTemplate(USoundWave& NewMetasound, const USoundWave& TemplateSoundWave) const;
