- Only enable interact input widgets on graph pages that implement the given input's default
- Support breadcrumbing all class input default values
- Reimplement input breadcrumb caching using MetaSoundFrontendDocumentBuilder API

#rnx

diff --git a/Engine/Plugins/Runtime/Metasound/Source/MetasoundEditor/Private/MetasoundEditorDocumentClipboardUtils.cpp b/Engine/Plugins/Runtime/Metasound/Source/MetasoundEditor/Private/MetasoundEditorDocumentClipboardUtils.cpp
--- a/Engine/Plugins/Runtime/Metasound/Source/MetasoundEditor/Private/MetasoundEditorDocumentClipboardUtils.cpp
+++ b/Engine/Plugins/Runtime/Metasound/Source/MetasoundEditor/Private/MetasoundEditorDocumentClipboardUtils.cpp
@@ -67,7 +67,7 @@ namespace Metasound::Editor
 							FCreateNodeVertexParams VertexParams;
 							VertexParams.DataType = Breadcrumb.DataType;
 
-							FMetasoundFrontendClassInput ClassInput = FGraphBuilder::CreateUniqueClassInput(*OutAsset.GetOwningAsset(), VertexParams, &Breadcrumb.DefaultLiteral, &Breadcrumb.MemberName);
+							FMetasoundFrontendClassInput ClassInput = FGraphBuilder::CreateUniqueClassInput(*OutAsset.GetOwningAsset(), VertexParams, &Breadcrumb.DefaultLiterals.FindChecked(Frontend::DefaultPageID), &Breadcrumb.MemberName);
 							if (const FMetasoundFrontendNode* NewNode = Builder.AddGraphInput(ClassInput))
 							{
 								Input = Graph.FindOrAddInput(NewNode->GetID());
diff --git a/Engine/Plugins/Runtime/Metasound/Source/MetasoundEditor/Private/MetasoundEditorGraphInputNode.cpp b/Engine/Plugins/Runtime/Metasound/Source/MetasoundEditor/Private/MetasoundEditorGraphInputNode.cpp
--- a/Engine/Plugins/Runtime/Metasound/Source/MetasoundEditor/Private/MetasoundEditorGraphInputNode.cpp
+++ b/Engine/Plugins/Runtime/Metasound/Source/MetasoundEditor/Private/MetasoundEditorGraphInputNode.cpp
@@ -41,7 +41,7 @@ const FMetasoundEditorGraphVertexNodeBreadcrumb& UMetasoundEditorGraphInputNode:
 
 void UMetasoundEditorGraphInputNode::CacheBreadcrumb()
 {
-	using namespace Metasound::Frontend;
+	using namespace Metasound;
 
 	Breadcrumb = { };
 
@@ -50,18 +50,25 @@ void UMetasoundEditorGraphInputNode::CacheBreadcrumb()
 	// to associate or create new associated input.
 	if (Input)
 	{
-		FConstNodeHandle NodeHandle = Input->GetConstNodeHandle();
-
-		Breadcrumb.MemberName = NodeHandle->GetNodeName();
-		Breadcrumb.ClassName = NodeHandle->GetClassMetadata().GetClassName();
-
-		FConstOutputHandle OutputHandle = NodeHandle->GetConstOutputs().Last();
-		Breadcrumb.AccessType = OutputHandle->GetVertexAccessType();
-		Breadcrumb.DataType = OutputHandle->GetDataType();
+		Breadcrumb.MemberName = Input->GetMemberName();
 
-		if (const UMetasoundEditorGraphMemberDefaultLiteral* Literal = Input->GetLiteral())
+		const FMetaSoundFrontendDocumentBuilder& Builder = Input->GetFrontendBuilderChecked();
+		if (const FMetasoundFrontendClassInput* ClassInput = Builder.FindGraphInput(Breadcrumb.MemberName))
 		{
-			Breadcrumb.DefaultLiteral = Literal->GetDefault();
+			if (const FMetasoundFrontendNode* Node = Builder.FindGraphInputNode(Breadcrumb.MemberName))
+			{
+				if (const FMetasoundFrontendClass* Class = Builder.FindDependency(Node->ClassID))
+				{
+					Breadcrumb.ClassName = Class->Metadata.GetClassName();
+					Breadcrumb.AccessType = ClassInput->AccessType;
+					Breadcrumb.DataType = ClassInput->TypeName;
+
+					ClassInput->IterateDefaults([this](const FGuid& PageID, const FMetasoundFrontendLiteral& Literal)
+					{
+						Breadcrumb.DefaultLiterals.Add(PageID, Literal);
+					});
+				}
+			}
 		}
 	}
 }
@@ -208,13 +215,26 @@ bool UMetasoundEditorGraphInputNode::EnableInteractWidgets() const
 {
 	using namespace Metasound::Frontend;
 	
-	//If Constructor input
-	if (Input && Input->GetVertexAccessType() == EMetasoundFrontendVertexAccessType::Value)
+	if (Input)
 	{
-		UMetasoundEditorGraph* Graph = CastChecked<UMetasoundEditorGraph>(GetGraph());
-		return !Graph->IsPreviewing();
+		if (Input->GetVertexAccessType() == EMetasoundFrontendVertexAccessType::Value)
+		{
+			UMetasoundEditorGraph* Graph = CastChecked<UMetasoundEditorGraph>(GetGraph());
+			return !Graph->IsPreviewing();
+		}
+
+		bool bPageDefaultImplemented = false;
+		const FGuid& BuildPageID = Input->GetFrontendBuilderChecked().GetBuildPageID();
+		if (UMetasoundEditorGraphMemberDefaultLiteral* Literal = Input->GetLiteral())
+		{
+			Literal->IterateDefaults([&bPageDefaultImplemented, &BuildPageID](const FGuid& InDefaultPageID, FMetasoundFrontendLiteral)
+			{
+				bPageDefaultImplemented |= InDefaultPageID == BuildPageID;
+			});
+		}
+		return bPageDefaultImplemented;
 	}
 
-	return true;
+	return false;
 }
 #undef LOCTEXT_NAMESPACE
diff --git a/Engine/Plugins/Runtime/Metasound/Source/MetasoundEditor/Private/MetasoundEditorGraphNode.cpp b/Engine/Plugins/Runtime/Metasound/Source/MetasoundEditor/Private/MetasoundEditorGraphNode.cpp
--- a/Engine/Plugins/Runtime/Metasound/Source/MetasoundEditor/Private/MetasoundEditorGraphNode.cpp
+++ b/Engine/Plugins/Runtime/Metasound/Source/MetasoundEditor/Private/MetasoundEditorGraphNode.cpp
@@ -972,20 +972,28 @@ const FMetasoundEditorGraphNodeBreadcrumb& UMetasoundEditorGraphOutputNode::GetB
 
 void UMetasoundEditorGraphOutputNode::CacheBreadcrumb()
 {
-	using namespace Metasound::Frontend;
-	FConstNodeHandle NodeHandle = GetConstNodeHandle();
-
-	Breadcrumb.MemberName = NodeHandle->GetNodeName();
-	Breadcrumb.ClassName = NodeHandle->GetClassMetadata().GetClassName();
-
-	FConstInputHandle InputHandle = NodeHandle->GetConstInputs().Last();
-	Breadcrumb.AccessType = InputHandle->GetVertexAccessType();
-	Breadcrumb.DataType = InputHandle->GetDataType();
+	if (Output)
+	{
+		Breadcrumb.MemberName = Output->GetMemberName();
 
+		const FMetaSoundFrontendDocumentBuilder& Builder = Output->GetFrontendBuilderChecked();
+		if (const FMetasoundFrontendClassOutput* ClassOutput = Builder.FindGraphOutput(Breadcrumb.MemberName))
+		{
+			if (const FMetasoundFrontendNode* Node = Builder.FindGraphInputNode(Breadcrumb.MemberName))
+			{
+				if (const FMetasoundFrontendClass* Class = Builder.FindDependency(Node->ClassID))
+				{
+					Breadcrumb.ClassName = Class->Metadata.GetClassName();
+					Breadcrumb.AccessType = ClassOutput->AccessType;
+					Breadcrumb.DataType = ClassOutput->TypeName;
+				}
+			}
+		}
 
-	if (const UMetasoundEditorGraphMemberDefaultLiteral* Literal = Output->GetLiteral())
-	{
-		Breadcrumb.DefaultLiteral = Literal->GetDefault();
+		if (const UMetasoundEditorGraphMemberDefaultLiteral* Literal = Output->GetLiteral())
+		{
+			Breadcrumb.DefaultLiterals.Add(Metasound::Frontend::DefaultPageID, Literal->GetDefault());
+		}
 	}
 }
 
diff --git a/Engine/Plugins/Runtime/Metasound/Source/MetasoundEditor/Public/MetasoundEditorGraphNode.h b/Engine/Plugins/Runtime/Metasound/Source/MetasoundEditor/Public/MetasoundEditorGraphNode.h
--- a/Engine/Plugins/Runtime/Metasound/Source/MetasoundEditor/Public/MetasoundEditorGraphNode.h
+++ b/Engine/Plugins/Runtime/Metasound/Source/MetasoundEditor/Public/MetasoundEditorGraphNode.h
@@ -25,7 +25,7 @@ class UMetasoundEditorGraphMember;
 class UMetasoundEditorGraphVariable;
 class UMetasoundEditorGraphMemberDefaultFloat;
 
-namespace EMessageSeverity { enum Type : int; }
+namespace EMessageSeverity { enum Type : int32; }
 
 namespace Metasound::Editor
 {
@@ -67,7 +67,7 @@ struct FMetasoundEditorGraphMemberNodeBreadcrumb : public FMetasoundEditorGraphN
 	FName DataType;
 
 	UPROPERTY()
-	FMetasoundFrontendLiteral DefaultLiteral;
+	TMap<FGuid, FMetasoundFrontendLiteral> DefaultLiterals;
 };
 
 USTRUCT()
diff --git a/Engine/Plugins/Runtime/Metasound/Source/MetasoundFrontend/Private/MetasoundFrontendDocument.cpp b/Engine/Plugins/Runtime/Metasound/Source/MetasoundFrontend/Private/MetasoundFrontendDocument.cpp
--- a/Engine/Plugins/Runtime/Metasound/Source/MetasoundFrontend/Private/MetasoundFrontendDocument.cpp
+++ b/Engine/Plugins/Runtime/Metasound/Source/MetasoundFrontend/Private/MetasoundFrontendDocument.cpp
@@ -1141,6 +1141,13 @@ TArray<FMetasoundFrontendGraph>& FMetasoundFrontendGraphClass::IPropertyVersionT
 	return GraphClass.PagedGraphs;
 }
 
+void UMetaSoundFrontendMemberMetadata::IterateDefaults(TFunctionRef<void(const FGuid&, FMetasoundFrontendLiteral)> Iter) const
+{
+	// TODO: Replace with follow-up change using page-addressed defaults
+	Iter(Metasound::Frontend::DefaultPageID, GetDefault());
+}
+
+
 FMetasoundFrontendVersionNumber FMetasoundFrontendDocument::GetMaxVersion()
 {
 	return Metasound::Frontend::GetMaxDocumentVersion();
diff --git a/Engine/Plugins/Runtime/Metasound/Source/MetasoundFrontend/Public/MetasoundFrontendDocument.h b/Engine/Plugins/Runtime/Metasound/Source/MetasoundFrontend/Public/MetasoundFrontendDocument.h
--- a/Engine/Plugins/Runtime/Metasound/Source/MetasoundFrontend/Public/MetasoundFrontendDocument.h
+++ b/Engine/Plugins/Runtime/Metasound/Source/MetasoundFrontend/Public/MetasoundFrontendDocument.h
@@ -1785,6 +1785,7 @@ class METASOUNDFRONTEND_API UMetaSoundFrontendMemberMetadata : public UObject
 public:
 	virtual void ForceRefresh() { }
 	virtual FMetasoundFrontendLiteral GetDefault() const { return FMetasoundFrontendLiteral(); }
+	virtual void IterateDefaults(TFunctionRef<void(const FGuid&, FMetasoundFrontendLiteral)> Iter) const;
 	virtual EMetasoundFrontendLiteralType GetLiteralType() const { return EMetasoundFrontendLiteralType::None; }
 	virtual void SetFromLiteral(const FMetasoundFrontendLiteral& InLiteral) { }
 #endif // WITH_EDITOR
