UE-190433 - Need to allow users to reset delay line audio data.

diff --git a/Engine/Plugins/Runtime/Metasound/Source/MetasoundStandardNodes/Private/MetasoundDelayNode.cpp b/Engine/Plugins/Runtime/Metasound/Source/MetasoundStandardNodes/Private/MetasoundDelayNode.cpp
--- a/Engine/Plugins/Runtime/Metasound/Source/MetasoundStandardNodes/Private/MetasoundDelayNode.cpp
+++ b/Engine/Plugins/Runtime/Metasound/Source/MetasoundStandardNodes/Private/MetasoundDelayNode.cpp
@@ -19,6 +19,7 @@ namespace Metasound
 {
 	namespace Delay
 	{
+		METASOUND_PARAM(InputResetDelay, "Reset", "Resets the delay buffer.");
 		METASOUND_PARAM(InParamAudioInput, "In", "Audio input.")
 		METASOUND_PARAM(InParamDelayTime, "Delay Time", "The amount of time to delay the audio.")
 		METASOUND_PARAM(InParamDryLevel, "Dry Level", "The dry level of the delay.")
@@ -40,7 +41,8 @@ namespace Metasound
 		static const FVertexInterface& GetVertexInterface();
 		static TUniquePtr<IOperator> CreateOperator(const FBuildOperatorParams& InParams, FBuildResults& OutResults);
 
-		FDelayOperator(const FBuildOperatorParams& InParams, 
+		FDelayOperator(const FBuildOperatorParams& InParams,
+			const FTriggerReadRef& InTriggerReset,
 			const FAudioBufferReadRef& InAudioInput, 
 			const FTimeReadRef& InDelayTime, 
 			const FFloatReadRef& InDryLevel, 
@@ -58,6 +60,9 @@ namespace Metasound
 	private:
 		float GetInputDelayTimeMsec() const;
 
+		// The reset trigger
+		FTriggerReadRef TriggerReset;
+
 		// The input audio buffer
 		FAudioBufferReadRef AudioInput;
 
@@ -90,6 +95,7 @@ namespace Metasound
 	};
 
 	FDelayOperator::FDelayOperator(const FBuildOperatorParams& InParams,
+		const FTriggerReadRef& InTriggerReset,
 		const FAudioBufferReadRef& InAudioInput,
 		const FTimeReadRef& InDelayTime,
 		const FFloatReadRef& InDryLevel,
@@ -97,7 +103,8 @@ namespace Metasound
 		const FFloatReadRef& InFeedback,
 		float InMaxDelayTimeSeconds)
 
-		: AudioInput(InAudioInput)
+		: TriggerReset(InTriggerReset)
+		, AudioInput(InAudioInput)
 		, DelayTime(InDelayTime)
 		, DryLevel(InDryLevel)
 		, WetLevel(InWetLevel)
@@ -114,6 +121,7 @@ namespace Metasound
 	{
 		using namespace Delay;
 
+		InOutVertexData.BindReadVertex(METASOUND_GET_PARAM_NAME(InputResetDelay), TriggerReset);
 		InOutVertexData.BindReadVertex(METASOUND_GET_PARAM_NAME(InParamAudioInput), AudioInput);
 		InOutVertexData.BindReadVertex(METASOUND_GET_PARAM_NAME(InParamDelayTime), DelayTime);
 		InOutVertexData.BindReadVertex(METASOUND_GET_PARAM_NAME(InParamDryLevel), DryLevel);
@@ -163,6 +171,16 @@ namespace Metasound
 
 	void FDelayOperator::Execute()
 	{
+		TriggerReset->ExecuteBlock(
+			[&](int32 StartFrame, int32 EndFrame)
+			{
+			},
+			[this](int32 StartFrame, int32 EndFrame)
+			{
+				DelayBuffer.Reset();
+			}
+		);
+
 		// Get clamped delay time
 		float CurrentInputDelayTime = GetInputDelayTimeMsec();
 
@@ -212,6 +230,7 @@ namespace Metasound
 
 		static const FVertexInterface Interface(
 			FInputVertexInterface(
+				TInputDataVertex<FTrigger>(METASOUND_GET_PARAM_NAME_AND_METADATA(InputResetDelay)),
 				TInputDataVertex<FAudioBuffer>(METASOUND_GET_PARAM_NAME_AND_METADATA(InParamAudioInput)),
 				TInputDataVertex<FTime>(METASOUND_GET_PARAM_NAME_AND_METADATA(InParamDelayTime), 1.0f),
 				TInputDataVertex<float>(METASOUND_GET_PARAM_NAME_AND_METADATA(InParamDryLevel), 0.0f),
@@ -255,6 +274,7 @@ namespace Metasound
 
 		const FInputVertexInterfaceData& InputData = InParams.InputData;
 
+		FTriggerReadRef TriggerReset = InputData.GetOrConstructDataReadReference<FTrigger>(METASOUND_GET_PARAM_NAME(InputResetDelay), InParams.OperatorSettings);
 		FAudioBufferReadRef AudioIn = InputData.GetOrConstructDataReadReference<FAudioBuffer>(METASOUND_GET_PARAM_NAME(InParamAudioInput), InParams.OperatorSettings);
 		FTimeReadRef DelayTime = InputData.GetOrCreateDefaultDataReadReference<FTime>(METASOUND_GET_PARAM_NAME(InParamDelayTime), InParams.OperatorSettings);
 		FFloatReadRef DryLevel = InputData.GetOrCreateDefaultDataReadReference<float>(METASOUND_GET_PARAM_NAME(InParamDryLevel), InParams.OperatorSettings);
@@ -262,7 +282,7 @@ namespace Metasound
 		FFloatReadRef Feedback = InputData.GetOrCreateDefaultDataReadReference<float>(METASOUND_GET_PARAM_NAME(InParamFeedbackAmount), InParams.OperatorSettings);
 		FTime MaxDelayTime = InputData.GetOrCreateDefaultValue<FTime>(METASOUND_GET_PARAM_NAME(InParamMaxDelayTime), InParams.OperatorSettings);
 
-		return MakeUnique<FDelayOperator>(InParams, AudioIn, DelayTime, DryLevel, WetLevel, Feedback, MaxDelayTime.GetSeconds());
+		return MakeUnique<FDelayOperator>(InParams, TriggerReset, AudioIn, DelayTime, DryLevel, WetLevel, Feedback, MaxDelayTime.GetSeconds());
 	}
 
 	class FDelayNode : public FNodeFacade
diff --git a/Engine/Plugins/Runtime/Metasound/Source/MetasoundStandardNodes/Private/MetasoundStereoDelayNode.cpp b/Engine/Plugins/Runtime/Metasound/Source/MetasoundStandardNodes/Private/MetasoundStereoDelayNode.cpp
--- a/Engine/Plugins/Runtime/Metasound/Source/MetasoundStandardNodes/Private/MetasoundStereoDelayNode.cpp
+++ b/Engine/Plugins/Runtime/Metasound/Source/MetasoundStandardNodes/Private/MetasoundStereoDelayNode.cpp
@@ -21,6 +21,7 @@ namespace Metasound
 {
 	namespace StereoDelay
 	{
+		METASOUND_PARAM(InputResetDelay, "Reset", "Resets the delay buffer.");
 		METASOUND_PARAM(InAudioInputLeft, "In Left", "Left channel audio input.")
 		METASOUND_PARAM(InAudioInputRight, "In Right", "Right channel audio input.")
 		METASOUND_PARAM(InDelayMode, "Delay Mode", "Delay mode.")
@@ -62,7 +63,8 @@ namespace Metasound
 		static const FVertexInterface& GetVertexInterface();
 		static TUniquePtr<IOperator> CreateOperator(const FBuildOperatorParams& InParams, FBuildResults& OutResults);
 
-		FStereoDelayOperator(const FOperatorSettings& InSettings, 
+		FStereoDelayOperator(
+			const FOperatorSettings& InSettings,
 			const FAudioBufferReadRef& InLeftAudioInput, 
 			const FAudioBufferReadRef& InRightAudioInput,
 			const FStereoDelayModeReadRef& InStereoDelayMode,
@@ -71,7 +73,8 @@ namespace Metasound
 			const FFloatReadRef& InDryLevel, 
 			const FFloatReadRef& InWetLevel, 
 			const FFloatReadRef& InFeedback,
-			float InMaxDelayTimeSeconds);
+			float InMaxDelayTimeSeconds,
+			const FTriggerReadRef& InTriggerReset);
 
 
 		virtual void BindInputs(FInputVertexInterfaceData& InOutVertexData) override;
@@ -121,6 +124,9 @@ namespace Metasound
 
 		// Maximum delay time
 		float MaxDelayTimeSeconds = StereoDelay::DefaultMaxDelaySeconds;
+
+		// The reset trigger
+		FTriggerReadRef TriggerReset;
 	};
 
 	FStereoDelayOperator::FStereoDelayOperator(
@@ -133,7 +139,8 @@ namespace Metasound
 		const FFloatReadRef& InDryLevel,
 		const FFloatReadRef& InWetLevel, 
 		const FFloatReadRef& InFeedback,
-		float InMaxDelayTimeSeconds)
+		float InMaxDelayTimeSeconds,
+		const FTriggerReadRef& InTriggerReset)
 
 		: LeftAudioInput(InLeftAudioInput)
 		, RightAudioInput(InRightAudioInput)
@@ -148,6 +155,7 @@ namespace Metasound
 		, PrevDelayTimeMsec(GetInputDelayTimeMsecClamped())
 		, PrevDelayRatio(GetInputDelayRatioClamped())
 		, MaxDelayTimeSeconds(FMath::Clamp(InMaxDelayTimeSeconds, StereoDelay::MinDelaySeconds, StereoDelay::MaxDelaySeconds))
+		, TriggerReset(InTriggerReset)
 	{
 		LeftDelayBuffer.Init(InSettings.GetSampleRate(), MaxDelayTimeSeconds);
 		LeftDelayBuffer.SetDelayMsec(PrevDelayTimeMsec * (1.0f + PrevDelayRatio));
@@ -169,6 +177,7 @@ namespace Metasound
 		InOutVertexData.BindReadVertex(METASOUND_GET_PARAM_NAME(InWetLevel), WetLevel);
 		InOutVertexData.BindReadVertex(METASOUND_GET_PARAM_NAME(InFeedbackAmount), Feedback);
 		InOutVertexData.SetValue(METASOUND_GET_PARAM_NAME(InParamMaxDelayTime), FTime::FromSeconds(MaxDelayTimeSeconds));
+		InOutVertexData.BindReadVertex(METASOUND_GET_PARAM_NAME(InputResetDelay), TriggerReset);
 	}
 
 	void FStereoDelayOperator::BindOutputs(FOutputVertexInterfaceData& InOutVertexData)
@@ -208,6 +217,17 @@ namespace Metasound
 
 	void FStereoDelayOperator::Execute()
 	{
+		TriggerReset->ExecuteBlock(
+			[&](int32 StartFrame, int32 EndFrame)
+			{
+			},
+			[this](int32 StartFrame, int32 EndFrame)
+			{
+				LeftDelayBuffer.Reset();
+				RightDelayBuffer.Reset();
+			}
+		);
+
 		// Get clamped delay time
 		float CurrentInputDelayTime = GetInputDelayTimeMsecClamped();
 		float CurrentDelayRatio = GetInputDelayRatioClamped();
@@ -343,7 +363,8 @@ namespace Metasound
 				TInputDataVertex<float>(METASOUND_GET_PARAM_NAME_AND_METADATA(InDryLevel), 0.0f),
 				TInputDataVertex<float>(METASOUND_GET_PARAM_NAME_AND_METADATA(InWetLevel), 1.0f),
 				TInputDataVertex<float>(METASOUND_GET_PARAM_NAME_AND_METADATA(InFeedbackAmount), 0.0f),
-				TInputConstructorVertex<FTime>(METASOUND_GET_PARAM_NAME(InParamMaxDelayTime), MaxDelayTimeMetadata, DefaultMaxDelaySeconds)
+				TInputConstructorVertex<FTime>(METASOUND_GET_PARAM_NAME(InParamMaxDelayTime), MaxDelayTimeMetadata, DefaultMaxDelaySeconds),
+				TInputDataVertex<FTrigger>(METASOUND_GET_PARAM_NAME_AND_METADATA(InputResetDelay))
 			),
 			FOutputVertexInterface(
 				TOutputDataVertex<FAudioBuffer>(METASOUND_GET_PARAM_NAME_AND_METADATA(OutAudioLeft)),
@@ -390,8 +411,9 @@ namespace Metasound
 		FFloatReadRef WetLevel = InputData.GetOrCreateDefaultDataReadReference<float>(METASOUND_GET_PARAM_NAME(InWetLevel), InParams.OperatorSettings);
 		FFloatReadRef Feedback = InputData.GetOrCreateDefaultDataReadReference<float>(METASOUND_GET_PARAM_NAME(InFeedbackAmount), InParams.OperatorSettings);
 		FTime MaxDelayTime = InputData.GetOrCreateDefaultValue<FTime>(METASOUND_GET_PARAM_NAME(InParamMaxDelayTime), InParams.OperatorSettings);
+		FTriggerReadRef TriggerReset = InputData.GetOrConstructDataReadReference<FTrigger>(METASOUND_GET_PARAM_NAME(InputResetDelay), InParams.OperatorSettings);
 
-		return MakeUnique<FStereoDelayOperator>(InParams.OperatorSettings, LeftAudioIn, RightAudioIn, StereoDelayMode, DelayTime, DelayRatio, DryLevel, WetLevel, Feedback, MaxDelayTime.GetSeconds());
+		return MakeUnique<FStereoDelayOperator>(InParams.OperatorSettings, LeftAudioIn, RightAudioIn, StereoDelayMode, DelayTime, DelayRatio, DryLevel, WetLevel, Feedback, MaxDelayTime.GetSeconds(), TriggerReset);
 	}
 
 
