Making new metasound editor restricted mode

* Restricted mode allows preventing authoring new metasounds but allows creating new presets of existing metasounds

#rnx

diff --git a/Engine/Plugins/Runtime/Metasound/Source/MetasoundEditor/Private/MetasoundAssetDefinitions.cpp b/Engine/Plugins/Runtime/Metasound/Source/MetasoundEditor/Private/MetasoundAssetDefinitions.cpp
--- a/Engine/Plugins/Runtime/Metasound/Source/MetasoundEditor/Private/MetasoundAssetDefinitions.cpp
+++ b/Engine/Plugins/Runtime/Metasound/Source/MetasoundEditor/Private/MetasoundAssetDefinitions.cpp
@@ -115,10 +115,14 @@ TConstArrayView<FAssetCategoryPath> UAssetDefinition_MetaSoundPatch::GetAssetCat
 
 EAssetCommandResult UAssetDefinition_MetaSoundPatch::OpenAssets(const FAssetOpenArgs& OpenArgs) const
 {
-	for (UMetaSoundPatch* Metasound : OpenArgs.LoadObjects<UMetaSoundPatch>())
+	Metasound::Editor::IMetasoundEditorModule& MetaSoundEditorModule = FModuleManager::GetModuleChecked<Metasound::Editor::IMetasoundEditorModule>("MetaSoundEditor");
+	if (!MetaSoundEditorModule.IsRestrictedMode())
 	{
-		TSharedRef<Metasound::Editor::FEditor> NewEditor = MakeShared<Metasound::Editor::FEditor>();
-		NewEditor->InitMetasoundEditor(OpenArgs.GetToolkitMode(), OpenArgs.ToolkitHost, Metasound);
+		for (UMetaSoundPatch* Metasound : OpenArgs.LoadObjects<UMetaSoundPatch>())
+		{
+			TSharedRef<Metasound::Editor::FEditor> NewEditor = MakeShared<Metasound::Editor::FEditor>();
+			NewEditor->InitMetasoundEditor(OpenArgs.GetToolkitMode(), OpenArgs.ToolkitHost, Metasound);
+		}
 	}
 	return EAssetCommandResult::Handled;
 }
@@ -164,10 +168,15 @@ TConstArrayView<FAssetCategoryPath> UAssetDefinition_MetaSoundSource::GetAssetCa
 
 EAssetCommandResult UAssetDefinition_MetaSoundSource::OpenAssets(const FAssetOpenArgs& OpenArgs) const
 {
+	Metasound::Editor::IMetasoundEditorModule& MetaSoundEditorModule = FModuleManager::GetModuleChecked<Metasound::Editor::IMetasoundEditorModule>("MetaSoundEditor");
 	for (UMetaSoundSource* Metasound : OpenArgs.LoadObjects<UMetaSoundSource>())
 	{
-		TSharedRef<Metasound::Editor::FEditor> NewEditor = MakeShared<Metasound::Editor::FEditor>();
-		NewEditor->InitMetasoundEditor(OpenArgs.GetToolkitMode(), OpenArgs.ToolkitHost, Metasound);
+		// In restricted mode, we only want to open up the editor if we're preset otherwise, we don't want to open the editor
+		if (Metasound->bIsPreset || !MetaSoundEditorModule.IsRestrictedMode())
+		{
+			TSharedRef<Metasound::Editor::FEditor> NewEditor = MakeShared<Metasound::Editor::FEditor>();
+			NewEditor->InitMetasoundEditor(OpenArgs.GetToolkitMode(), OpenArgs.ToolkitHost, Metasound);
+		}
 	}
 
 	return EAssetCommandResult::Handled;
diff --git a/Engine/Plugins/Runtime/Metasound/Source/MetasoundEditor/Private/MetasoundEditor.cpp b/Engine/Plugins/Runtime/Metasound/Source/MetasoundEditor/Private/MetasoundEditor.cpp
--- a/Engine/Plugins/Runtime/Metasound/Source/MetasoundEditor/Private/MetasoundEditor.cpp
+++ b/Engine/Plugins/Runtime/Metasound/Source/MetasoundEditor/Private/MetasoundEditor.cpp
@@ -2557,10 +2557,14 @@ namespace Metasound
 						// ensure when generating asset key.
 						if (Class->Metadata.GetType() == EMetasoundFrontendClassType::External)
 						{
-							const FAssetKey AssetKey(Class->Metadata);
-							if (const FMetasoundAssetBase* Asset = IMetaSoundAssetManager::GetChecked().FindAsset(AssetKey))
+							IMetasoundEditorModule& MetaSoundEditorModule = FModuleManager::GetModuleChecked<IMetasoundEditorModule>("MetaSoundEditor");
+							if (!MetaSoundEditorModule.IsRestrictedMode())
 							{
-								AssetSubsystem->OpenEditorForAsset(Asset->GetOwningAsset());
+								const FAssetKey AssetKey(Class->Metadata);
+								if (const FMetasoundAssetBase* Asset = IMetaSoundAssetManager::GetChecked().FindAsset(AssetKey))
+								{
+									AssetSubsystem->OpenEditorForAsset(Asset->GetOwningAsset());
+								}
 							}
 						}
 					}
diff --git a/Engine/Plugins/Runtime/Metasound/Source/MetasoundEditor/Private/MetasoundEditorModule.cpp b/Engine/Plugins/Runtime/Metasound/Source/MetasoundEditor/Private/MetasoundEditorModule.cpp
--- a/Engine/Plugins/Runtime/Metasound/Source/MetasoundEditor/Private/MetasoundEditorModule.cpp
+++ b/Engine/Plugins/Runtime/Metasound/Source/MetasoundEditor/Private/MetasoundEditorModule.cpp
@@ -406,6 +406,16 @@ namespace Metasound
 				FGraphNodeVisualizationRegistry::Get().RegisterVisualization(InNodeClassName, OnCreateGraphNodeVisualizationWidget);
 			}
 
+			virtual bool IsRestrictedMode() const override
+			{
+				return bIsRestrictedMode;
+			}
+
+			virtual void SetRestrictedMode(bool bInRestrictedMode) override
+			{
+				bIsRestrictedMode = bInRestrictedMode;
+			}
+
 			void RegisterSettingsDelegates()
 			{
 				using namespace Engine;
@@ -778,6 +788,9 @@ namespace Metasound
 			TSharedPtr<FSlateStyleSet> StyleSet;
 
 			TSet<const UClass*> ExplicitProxyClasses;
+
+			// Whether or not the editor is in restricted mode: can only make new presets and not modify graphs
+			bool bIsRestrictedMode = false;
 		};
 	} // namespace Editor
 } // namespace Metasound
diff --git a/Engine/Plugins/Runtime/Metasound/Source/MetasoundEditor/Public/MetasoundEditorModule.h b/Engine/Plugins/Runtime/Metasound/Source/MetasoundEditor/Public/MetasoundEditorModule.h
--- a/Engine/Plugins/Runtime/Metasound/Source/MetasoundEditor/Public/MetasoundEditorModule.h
+++ b/Engine/Plugins/Runtime/Metasound/Source/MetasoundEditor/Public/MetasoundEditorModule.h
@@ -124,6 +124,12 @@ namespace Metasound
 
 			// For the given node class, register a delegate that can be used for creating in-graph node visualizations.
 			virtual void RegisterGraphNodeVisualization(FName InNodeClassName, FOnCreateGraphNodeVisualizationWidget OnCreateGraphNodeVisualizationWidget) = 0;
+
+			// Queries if the MetaSound Editor is in "restricted mode" (i.e. can only make new presets and not make new assets or edit asset graphs)
+			virtual bool IsRestrictedMode() const = 0;
+
+			// Sets if the MetaSound editor is in "restricted mode" (i.e. can only make new presets and not make new assets or edit asset graphs)
+			virtual void SetRestrictedMode(bool bInRestricted) = 0;
 		};
 	} // namespace Editor
 } // namespace Metasound
