MetaSound Source - Add AssetRegistrySearchable to the QualitySetting property

#rnx

diff --git a/Engine/Plugins/Runtime/Metasound/Source/MetasoundEngine/Public/MetasoundSource.h b/Engine/Plugins/Runtime/Metasound/Source/MetasoundEngine/Public/MetasoundSource.h
--- a/Engine/Plugins/Runtime/Metasound/Source/MetasoundEngine/Public/MetasoundSource.h
+++ b/Engine/Plugins/Runtime/Metasound/Source/MetasoundEngine/Public/MetasoundSource.h
@@ -130,7 +130,7 @@ public:
 
 #if WITH_EDITORONLY_DATA
 	// The QualitySetting MetaSound will use, as defined in 'MetaSound' Settings.
-	UPROPERTY(EditAnywhere, BlueprintReadWrite, meta = (GetOptions="MetasoundEngine.MetaSoundSettings.GetQualityNames"), Category = "Metasound")
+	UPROPERTY(EditAnywhere, BlueprintReadWrite, AssetRegistrySearchable, meta = (GetOptions="MetasoundEngine.MetaSoundSettings.GetQualityNames"), Category = "Metasound")
 	FName QualitySetting;
 
 	// This a editor only look up for the Quality Setting above. Preventing orphaning of the original name.
