UE-221785 - Promoting a graph input and output with right click does not highlight member to rename

diff --git a/Engine/Plugins/Runtime/Metasound/Source/MetasoundEditor/Private/MetasoundEditor.cpp b/Engine/Plugins/Runtime/Metasound/Source/MetasoundEditor/Private/MetasoundEditor.cpp
--- a/Engine/Plugins/Runtime/Metasound/Source/MetasoundEditor/Private/MetasoundEditor.cpp
+++ b/Engine/Plugins/Runtime/Metasound/Source/MetasoundEditor/Private/MetasoundEditor.cpp
@@ -4798,7 +4798,7 @@ namespace Metasound
 
 				UEdGraphNode* OwningNode = TargetPin->GetOwningNode();
 				FVector2D Location = FVector2D(OwningNode->NodePosX, OwningNode->NodePosY);
-				Metasound::SchemaUtils::PromoteToInput(&Graph, TargetPin, Location - DisplayStyle::NodeLayout::DefaultOffsetX, false);
+				Metasound::SchemaUtils::PromoteToInput(&Graph, TargetPin, Location - DisplayStyle::NodeLayout::DefaultOffsetX, /*bSelectNewNode=*/ true);
 			}
 		}
 
@@ -4832,7 +4832,7 @@ namespace Metasound
 
 				UEdGraphNode* OwningNode = TargetPin->GetOwningNode();
 				FVector2D Location = FVector2D(OwningNode->NodePosX, OwningNode->NodePosY);
-				Metasound::SchemaUtils::PromoteToOutput(&Graph, TargetPin, Location + DisplayStyle::NodeLayout::DefaultOffsetX * 2.f, false);
+				Metasound::SchemaUtils::PromoteToOutput(&Graph, TargetPin, Location + DisplayStyle::NodeLayout::DefaultOffsetX * 2.f, /*bSelectNewNode=*/ true);
 			}
 		}
 
@@ -4856,11 +4856,11 @@ namespace Metasound
 				FVector2D Location = FVector2D(OwningNode->NodePosX, OwningNode->NodePosY);
 				if (TargetPin->Direction == EGPD_Input)
 				{
-					Metasound::SchemaUtils::PromoteToVariable(&Graph, TargetPin, Location - DisplayStyle::NodeLayout::DefaultOffsetX, false);
+					Metasound::SchemaUtils::PromoteToVariable(&Graph, TargetPin, Location - DisplayStyle::NodeLayout::DefaultOffsetX, /*bSelectNewNode=*/ true);
 				}
 				else
 				{
-					Metasound::SchemaUtils::PromoteToMutatorVariable(&Graph, TargetPin, Location + DisplayStyle::NodeLayout::DefaultOffsetX * 2.f, false);
+					Metasound::SchemaUtils::PromoteToMutatorVariable(&Graph, TargetPin, Location + DisplayStyle::NodeLayout::DefaultOffsetX * 2.f, /*bSelectNewNode=*/ true);
 				}
 			}
 		}
@@ -4895,7 +4895,7 @@ namespace Metasound
 
 				UEdGraphNode* OwningNode = TargetPin->GetOwningNode();
 				FVector2D Location = FVector2D(OwningNode->NodePosX, OwningNode->NodePosY);
-				Metasound::SchemaUtils::PromoteToDeferredVariable(&Graph, TargetPin, Location - DisplayStyle::NodeLayout::DefaultOffsetX, false);
+				Metasound::SchemaUtils::PromoteToDeferredVariable(&Graph, TargetPin, Location - DisplayStyle::NodeLayout::DefaultOffsetX, /*bSelectNewNode=*/ true);
 			}
 		}
 
diff --git a/Engine/Plugins/Runtime/Metasound/Source/MetasoundEditor/Private/MetasoundEditorGraphSchema.cpp b/Engine/Plugins/Runtime/Metasound/Source/MetasoundEditor/Private/MetasoundEditorGraphSchema.cpp
--- a/Engine/Plugins/Runtime/Metasound/Source/MetasoundEditor/Private/MetasoundEditorGraphSchema.cpp
+++ b/Engine/Plugins/Runtime/Metasound/Source/MetasoundEditor/Private/MetasoundEditorGraphSchema.cpp
@@ -327,7 +327,14 @@ PRAGMA_ENABLE_DEPRECATION_WARNINGS
 							if (ensure(SchemaPrivate::TryConnectNewNodeToMatchingDataTypePin(*EdGraphNode, &FromPin)))
 							{
 								FGraphBuilder::RegisterGraphWithFrontend(ParentMetaSound);
-								SelectNodeInEditorForRename(*MetaSoundGraph, *NewGraphNode);
+								if (bSelectNode)
+								{
+									SchemaPrivate::SelectNodeInEditorForRename(*MetaSoundGraph, *NewGraphNode);
+								}
+								else
+								{
+									SchemaPrivate::SelectNodeInEditor(*MetaSoundGraph, *NewGraphNode);
+								}
 								return EdGraphNode;
 							}
 						}
