Fixing pin order

diff --git a/Engine/Plugins/Runtime/Metasound/Source/MetasoundStandardNodes/Private/MetasoundDelayNode.cpp b/Engine/Plugins/Runtime/Metasound/Source/MetasoundStandardNodes/Private/MetasoundDelayNode.cpp
--- a/Engine/Plugins/Runtime/Metasound/Source/MetasoundStandardNodes/Private/MetasoundDelayNode.cpp
+++ b/Engine/Plugins/Runtime/Metasound/Source/MetasoundStandardNodes/Private/MetasoundDelayNode.cpp
@@ -42,13 +42,13 @@ namespace Metasound
 		static TUniquePtr<IOperator> CreateOperator(const FBuildOperatorParams& InParams, FBuildResults& OutResults);
 
 		FDelayOperator(const FBuildOperatorParams& InParams,
-			const FTriggerReadRef& InTriggerReset,
 			const FAudioBufferReadRef& InAudioInput, 
 			const FTimeReadRef& InDelayTime, 
 			const FFloatReadRef& InDryLevel, 
 			const FFloatReadRef& InWetLevel, 
 			const FFloatReadRef& InFeedback,
-			float InMaxDelayTimeSeconds);
+			float InMaxDelayTimeSeconds,
+			const FTriggerReadRef& InTriggerReset);
 
 		virtual void BindInputs(FInputVertexInterfaceData& InOutVertexData) override;
 		virtual void BindOutputs(FOutputVertexInterfaceData& InOutVertexData) override;
@@ -60,9 +60,6 @@ namespace Metasound
 	private:
 		float GetInputDelayTimeMsec() const;
 
-		// The reset trigger
-		FTriggerReadRef TriggerReset;
-
 		// The input audio buffer
 		FAudioBufferReadRef AudioInput;
 
@@ -92,24 +89,27 @@ namespace Metasound
 
 		// Maximum delay time
 		float MaxDelayTimeSeconds = Delay::DefaultMaxDelaySeconds;
+
+		// The reset trigger
+		FTriggerReadRef TriggerReset;
 	};
 
 	FDelayOperator::FDelayOperator(const FBuildOperatorParams& InParams,
-		const FTriggerReadRef& InTriggerReset,
 		const FAudioBufferReadRef& InAudioInput,
 		const FTimeReadRef& InDelayTime,
 		const FFloatReadRef& InDryLevel,
 		const FFloatReadRef& InWetLevel,
 		const FFloatReadRef& InFeedback,
-		float InMaxDelayTimeSeconds)
+		float InMaxDelayTimeSeconds,
+		const FTriggerReadRef& InTriggerReset)
 
-		: TriggerReset(InTriggerReset)
-		, AudioInput(InAudioInput)
+		: AudioInput(InAudioInput)
 		, DelayTime(InDelayTime)
 		, DryLevel(InDryLevel)
 		, WetLevel(InWetLevel)
 		, Feedback(InFeedback)
 		, AudioOutput(FAudioBufferWriteRef::CreateNew(InParams.OperatorSettings))
+		, TriggerReset(InTriggerReset)
 	{
 		MaxDelayTimeSeconds = FMath::Clamp(InMaxDelayTimeSeconds, Delay::MinMaxDelaySeconds, Delay::MaxMaxDelaySeconds);
 
@@ -121,13 +121,13 @@ namespace Metasound
 	{
 		using namespace Delay;
 
-		InOutVertexData.BindReadVertex(METASOUND_GET_PARAM_NAME(InputResetDelay), TriggerReset);
 		InOutVertexData.BindReadVertex(METASOUND_GET_PARAM_NAME(InParamAudioInput), AudioInput);
 		InOutVertexData.BindReadVertex(METASOUND_GET_PARAM_NAME(InParamDelayTime), DelayTime);
 		InOutVertexData.BindReadVertex(METASOUND_GET_PARAM_NAME(InParamDryLevel), DryLevel);
 		InOutVertexData.BindReadVertex(METASOUND_GET_PARAM_NAME(InParamWetLevel), WetLevel);
 		InOutVertexData.BindReadVertex(METASOUND_GET_PARAM_NAME(InParamFeedbackAmount), Feedback);
 		InOutVertexData.SetValue(METASOUND_GET_PARAM_NAME(InParamMaxDelayTime), FTime::FromSeconds(MaxDelayTimeSeconds));
+		InOutVertexData.BindReadVertex(METASOUND_GET_PARAM_NAME(InputResetDelay), TriggerReset);
 	}
 
 	void FDelayOperator::BindOutputs(FOutputVertexInterfaceData& InOutVertexData)
@@ -230,13 +230,13 @@ namespace Metasound
 
 		static const FVertexInterface Interface(
 			FInputVertexInterface(
-				TInputDataVertex<FTrigger>(METASOUND_GET_PARAM_NAME_AND_METADATA(InputResetDelay)),
 				TInputDataVertex<FAudioBuffer>(METASOUND_GET_PARAM_NAME_AND_METADATA(InParamAudioInput)),
 				TInputDataVertex<FTime>(METASOUND_GET_PARAM_NAME_AND_METADATA(InParamDelayTime), 1.0f),
 				TInputDataVertex<float>(METASOUND_GET_PARAM_NAME_AND_METADATA(InParamDryLevel), 0.0f),
 				TInputDataVertex<float>(METASOUND_GET_PARAM_NAME_AND_METADATA(InParamWetLevel), 1.0f),
 				TInputDataVertex<float>(METASOUND_GET_PARAM_NAME_AND_METADATA(InParamFeedbackAmount), 0.0f),
-				TInputConstructorVertex<FTime>(METASOUND_GET_PARAM_NAME(InParamMaxDelayTime), MaxDelayTimeMetadata, DefaultMaxDelaySeconds)
+				TInputConstructorVertex<FTime>(METASOUND_GET_PARAM_NAME(InParamMaxDelayTime), MaxDelayTimeMetadata, DefaultMaxDelaySeconds),
+				TInputDataVertex<FTrigger>(METASOUND_GET_PARAM_NAME_AND_METADATA(InputResetDelay))
 			),
 			FOutputVertexInterface(
 				TOutputDataVertex<FAudioBuffer>(METASOUND_GET_PARAM_NAME_AND_METADATA(OutParamAudio))
@@ -274,15 +274,15 @@ namespace Metasound
 
 		const FInputVertexInterfaceData& InputData = InParams.InputData;
 
-		FTriggerReadRef TriggerReset = InputData.GetOrConstructDataReadReference<FTrigger>(METASOUND_GET_PARAM_NAME(InputResetDelay), InParams.OperatorSettings);
 		FAudioBufferReadRef AudioIn = InputData.GetOrConstructDataReadReference<FAudioBuffer>(METASOUND_GET_PARAM_NAME(InParamAudioInput), InParams.OperatorSettings);
 		FTimeReadRef DelayTime = InputData.GetOrCreateDefaultDataReadReference<FTime>(METASOUND_GET_PARAM_NAME(InParamDelayTime), InParams.OperatorSettings);
 		FFloatReadRef DryLevel = InputData.GetOrCreateDefaultDataReadReference<float>(METASOUND_GET_PARAM_NAME(InParamDryLevel), InParams.OperatorSettings);
 		FFloatReadRef WetLevel = InputData.GetOrCreateDefaultDataReadReference<float>(METASOUND_GET_PARAM_NAME(InParamWetLevel), InParams.OperatorSettings);
 		FFloatReadRef Feedback = InputData.GetOrCreateDefaultDataReadReference<float>(METASOUND_GET_PARAM_NAME(InParamFeedbackAmount), InParams.OperatorSettings);
 		FTime MaxDelayTime = InputData.GetOrCreateDefaultValue<FTime>(METASOUND_GET_PARAM_NAME(InParamMaxDelayTime), InParams.OperatorSettings);
+		FTriggerReadRef TriggerReset = InputData.GetOrConstructDataReadReference<FTrigger>(METASOUND_GET_PARAM_NAME(InputResetDelay), InParams.OperatorSettings);
 
-		return MakeUnique<FDelayOperator>(InParams, TriggerReset, AudioIn, DelayTime, DryLevel, WetLevel, Feedback, MaxDelayTime.GetSeconds());
+		return MakeUnique<FDelayOperator>(InParams, AudioIn, DelayTime, DryLevel, WetLevel, Feedback, MaxDelayTime.GetSeconds(), TriggerReset);
 	}
 
 	class FDelayNode : public FNodeFacade
