[Metasound] Include LOC_DEFINE_REGION so LOCTEXT works

#tests local build
#rnx

diff --git a/Engine/Plugins/Runtime/Metasound/Source/MetasoundGraphCore/Public/MetasoundParamHelper.h b/Engine/Plugins/Runtime/Metasound/Source/MetasoundGraphCore/Public/MetasoundParamHelper.h
--- a/Engine/Plugins/Runtime/Metasound/Source/MetasoundGraphCore/Public/MetasoundParamHelper.h
+++ b/Engine/Plugins/Runtime/Metasound/Source/MetasoundGraphCore/Public/MetasoundParamHelper.h
@@ -58,10 +58,12 @@
  * Use this macro to define your parameter if you used DECLARE_METASOUND_PARAM, otherwise just use METASOUND_PARAM in a source file
  */
 #if WITH_EDITOR
+#define LOC_DEFINE_REGION
 #define DEFINE_METASOUND_PARAM(NAME, NAME_TEXT, TOOLTIP_TEXT) \
 	const FLazyName NAME##Name = TEXT(NAME_TEXT); \
 	static const FText NAME##Tooltip = LOCTEXT(#NAME "Tooltip", TOOLTIP_TEXT); \
 	static const FText NAME##DisplayName = LOCTEXT(#NAME "DisplayName", NAME_TEXT);
+#undef LOC_DEFINE_REGION
 #else 
 #define DEFINE_METASOUND_PARAM(NAME, NAME_TEXT, TOOLTIP_TEXT) \
 	const FLazyName NAME##Name = TEXT(NAME_TEXT); \
