Fix MetaSound node output pin connection synchronization when deleting and undoing/redoing connections

diff --git a/Engine/Plugins/Runtime/Metasound/Source/MetasoundEditor/Private/MetasoundEditorGraphBuilder.cpp b/Engine/Plugins/Runtime/Metasound/Source/MetasoundEditor/Private/MetasoundEditorGraphBuilder.cpp
--- a/Engine/Plugins/Runtime/Metasound/Source/MetasoundEditor/Private/MetasoundEditorGraphBuilder.cpp
+++ b/Engine/Plugins/Runtime/Metasound/Source/MetasoundEditor/Private/MetasoundEditorGraphBuilder.cpp
@@ -2053,9 +2053,11 @@ namespace Metasound
 
 				TArray<UEdGraphPin*> Pins = EditorNode->GetAllPins();
 				TArray<FConstInputHandle> NodeInputs = Node->GetConstInputs();
+				TArray<FConstOutputHandle> NodeOutputs = Node->GetConstOutputs();
 
 				// Ignore connections which are not handled by the editor.
 				NodeInputs.RemoveAll([](const FConstInputHandle& FrontendInput) { return !FrontendInput->IsConnectionUserModifiable(); });
+				NodeOutputs.RemoveAll([](const FConstOutputHandle& FrontendOutput) { return !FrontendOutput->IsConnectionUserModifiable(); });
 
 				for (FConstInputHandle& NodeInput : NodeInputs)
 				{
@@ -2140,6 +2142,60 @@ namespace Metasound
 					SynchronizePinLiteral(*MatchingPin);
 				}
 
+				// Handle node outputs to break connections for the case 
+				// where the connected node input is on a node that has been deleted, 
+				// so it wasn't handled by the node inputs direction above
+				for (FConstOutputHandle& NodeOutput : NodeOutputs)
+				{
+					// Find pin
+					auto IsMatchingOutputPin = [&](const UEdGraphPin* Pin) -> bool
+					{
+						return IsMatchingOutputHandleAndPin(NodeOutput, *Pin);
+					};
+
+					UEdGraphPin* MatchingPin = nullptr;
+					if (UEdGraphPin** DoublePointer = Pins.FindByPredicate(IsMatchingOutputPin))
+					{
+						MatchingPin = *DoublePointer;
+					}
+
+					if (!ensure(MatchingPin))
+					{
+						continue;
+					}
+
+					// Remove connected pins from removal list
+					TArray<UEdGraphPin*> PinsToBreak = MatchingPin->LinkedTo;
+					TArray<FConstInputHandle> InputHandles = NodeOutput->GetConstConnectedInputs();
+					for (FConstInputHandle InputHandle : InputHandles)
+					{
+						UEdGraphPin** InputPin = PinsToBreak.FindByPredicate(
+							[&](const UEdGraphPin* Pin) -> bool
+							{
+								return IsMatchingInputHandleAndPin(InputHandle, *Pin);
+							});
+						if (InputPin)
+						{
+							UEdGraphPin* FoundPin = *InputPin;
+							if (FoundPin)
+							{
+								PinsToBreak.Remove(FoundPin);
+							}
+						}
+					}
+
+					// Break remaining invalid connections
+					for (UEdGraphPin* PinToBreak : PinsToBreak)
+					{
+						constexpr bool bMarkDirty = false;
+						MatchingPin->BreakLinkTo(PinToBreak, bMarkDirty);
+						const FText OwningNodeName = EditorNode->GetDisplayName();
+						const FText OutputName = FGraphBuilder::GetDisplayName(*NodeOutput);
+						UE_LOG(LogMetasoundEditor, Verbose, TEXT("Synchronizing Node '%s' Connection: Breaking a pin link to '%s'"), *OwningNodeName.ToString(), *OutputName.ToString());
+						bIsNodeDirty = true;
+					}
+				}
+
 				bIsGraphDirty |= bIsNodeDirty;
 			}
 
