Parse property sheets and apply properties they contain when creating a metasound preset
#rnx

diff --git a/Engine/Plugins/Runtime/Metasound/Source/MetasoundEditor/Private/MetasoundEditorSubsystem.cpp b/Engine/Plugins/Runtime/Metasound/Source/MetasoundEditor/Private/MetasoundEditorSubsystem.cpp
--- a/Engine/Plugins/Runtime/Metasound/Source/MetasoundEditor/Private/MetasoundEditorSubsystem.cpp
+++ b/Engine/Plugins/Runtime/Metasound/Source/MetasoundEditor/Private/MetasoundEditorSubsystem.cpp
@@ -1,6 +1,7 @@
 // Copyright Epic Games, Inc. All Rights Reserved.
 #include "MetasoundEditorSubsystem.h"
 
+#include "AudioPropertiesSheetAssetBase.h"
 #include "IAssetTools.h"
 #include "Kismet/KismetSystemLibrary.h"
 #include "MetasoundDocumentBuilderRegistry.h"
@@ -366,7 +367,7 @@ void UMetaSoundEditorSubsystem::SetSoundWaveSettingsFromTemplate(USoundWave& New
 	NewMetaSoundWave.bEnableBusSends = TemplateSoundWave.bEnableBusSends;
 	NewMetaSoundWave.SourceEffectChain = TemplateSoundWave.SourceEffectChain;
 	NewMetaSoundWave.BusSends = TemplateSoundWave.BusSends;
-	NewMetaSoundWave.PreEffectBusSends = TemplateSoundWave.PreEffectBusSends; 
+	NewMetaSoundWave.PreEffectBusSends = TemplateSoundWave.PreEffectBusSends;
 
 	NewMetaSoundWave.bEnableBaseSubmix = TemplateSoundWave.bEnableBaseSubmix;
 	NewMetaSoundWave.SoundSubmixObject = TemplateSoundWave.SoundSubmixObject;
@@ -375,7 +376,7 @@ void UMetaSoundEditorSubsystem::SetSoundWaveSettingsFromTemplate(USoundWave& New
 
 	// Modulation 
 	NewMetaSoundWave.ModulationSettings = TemplateSoundWave.ModulationSettings;
-	
+
 	// Voice Management 
 	NewMetaSoundWave.VirtualizationMode = TemplateSoundWave.VirtualizationMode;
 	NewMetaSoundWave.bOverrideConcurrency = TemplateSoundWave.bOverrideConcurrency;
@@ -384,6 +385,16 @@ void UMetaSoundEditorSubsystem::SetSoundWaveSettingsFromTemplate(USoundWave& New
 
 	NewMetaSoundWave.bBypassVolumeScaleForPriority = TemplateSoundWave.bBypassVolumeScaleForPriority;
 	NewMetaSoundWave.Priority = TemplateSoundWave.Priority;
+
+	//Property Sheets - keep this last so that properties in the sheet will be applied
+	NewMetaSoundWave.AudioPropertiesSheet = TemplateSoundWave.AudioPropertiesSheet;
+
+	if (NewMetaSoundWave.AudioPropertiesSheet)
+	{
+		NewMetaSoundWave.AudioPropertiesSheet->CopyToObjectProperties(&NewMetaSoundWave);
+	}
+
+	return;
 }
 
 #undef LOCTEXT_NAMESPACE // "MetaSoundEditor"
