- Change "RemoveAll<GraphPages/PageDefaultValues>Pages" to be "Reset", as the default value should always exist and be valid
- Add utility functions to EdGraphBuilder & FrontendDocBuilder to be used for editor updates to accomidate paged input defaults
- General minor controller to docbuilder API clean-up spots
- Misc comment clean-up

#rnx

diff --git a/Engine/Plugins/Runtime/Metasound/Source/MetasoundEditor/Private/MetasoundDetailCustomization.cpp b/Engine/Plugins/Runtime/Metasound/Source/MetasoundEditor/Private/MetasoundDetailCustomization.cpp
--- a/Engine/Plugins/Runtime/Metasound/Source/MetasoundEditor/Private/MetasoundDetailCustomization.cpp
+++ b/Engine/Plugins/Runtime/Metasound/Source/MetasoundEditor/Private/MetasoundDetailCustomization.cpp
@@ -432,17 +432,19 @@ namespace Metasound::Editor
 					using namespace Frontend;
 					UObject& MetaSound = GetMetaSound();
 
-					const FScopedTransaction Transaction(LOCTEXT("RemoveAllPagesTransaction", "Remove All MetaSound Pages"));
+					const FScopedTransaction Transaction(LOCTEXT("ResetPagesTransaction", "Reset MetaSound Pages"));
 					MetaSound.Modify();
 
+					constexpr bool bClearDefaultGraph = false;
+
 					// Underlying DocBuilder's pageID is a property that is tracked by transaction stack, so signal as modifying behavior
 					Builder->Modify();
-					Builder->RemoveAllGraphPages();
+					Builder->ResetGraphPages(bClearDefaultGraph);
 
 					UpdateItemNames();
 					ComboBox->RefreshOptions();
 					FGraphBuilder::RegisterGraphWithFrontend(MetaSound);
-				}), LOCTEXT("RemoveImplementablePagesTooltip1", "Removes all pages from the given MetaSound."))
+				}), LOCTEXT("ResetPagesTooltip1", "Removes all pages from the given MetaSound less the default."))
 			];
 
 		Utilities->SetEnabled(TAttribute<bool>::Create(TAttribute<bool>::FGetter::CreateSP(this, &FMetaSoundDetailCustomizationBase::IsGraphEditable)));
diff --git a/Engine/Plugins/Runtime/Metasound/Source/MetasoundEditor/Private/MetasoundEditorGraphBuilder.cpp b/Engine/Plugins/Runtime/Metasound/Source/MetasoundEditor/Private/MetasoundEditorGraphBuilder.cpp
--- a/Engine/Plugins/Runtime/Metasound/Source/MetasoundEditor/Private/MetasoundEditorGraphBuilder.cpp
+++ b/Engine/Plugins/Runtime/Metasound/Source/MetasoundEditor/Private/MetasoundEditorGraphBuilder.cpp
@@ -653,6 +653,13 @@ namespace Metasound
 			return Results;
 		}
 
+		UMetaSoundBuilderBase& FGraphBuilder::GetBuilderFromPinChecked(const UEdGraphPin& InPin)
+		{
+			const UMetasoundEditorGraphNode* Node = CastChecked<UMetasoundEditorGraphNode>(InPin.GetOwningNode());
+			check(Node);
+			return Node->GetBuilderChecked();
+		}
+
 		TArray<FString> FGraphBuilder::GetDataTypeNameCategories(const FName& InDataTypeName)
 		{
 			FString CategoryString = InDataTypeName.ToString();
@@ -1202,14 +1209,21 @@ namespace Metasound
 
 			IMetasoundEditorModule& EditorModule = FModuleManager::GetModuleChecked<IMetasoundEditorModule>("MetaSoundEditor");
 
-			FConstInputHandle InputHandle = GetConstInputHandleFromPin(&InInputPin);
-			if (!ensure(InputHandle->IsValid()))
+			const FMetaSoundFrontendDocumentBuilder& Builder = GetBuilderFromPinChecked(InInputPin).GetConstBuilder();
+			FMetasoundFrontendVertexHandle InputHandle = GetPinVertexHandle(Builder, &InInputPin);
+			if (!ensure(InputHandle.IsSet()))
+			{
+				return false;
+			}
+
+			const FMetasoundFrontendVertex* Vertex = GetPinVertex(Builder, &InInputPin);
+			if (!ensure(Vertex))
 			{
 				return false;
 			}
 
 			const FString& InStringValue = InInputPin.DefaultValue;
-			const FName TypeName = InputHandle->GetDataType();
+			const FName TypeName = Vertex->TypeName;
 
 			FDataTypeRegistryInfo DataTypeInfo;
 			IDataTypeRegistry::Get().GetDataTypeInfo(TypeName, DataTypeInfo);
@@ -1287,15 +1301,29 @@ namespace Metasound
 						// However, if the class default literal is set to an object, the literal should be set to a valid, null object.
 						// This is used for reset to default behavior, where an valid object literal with a null value is a separate case
 						// from an inherited cleared default literal. 
-						const FMetasoundFrontendLiteral* ClassDefaultLiteral = InputHandle->GetClassDefaultLiteral();
-						if (ClassDefaultLiteral && ClassDefaultLiteral->GetType() == EMetasoundFrontendLiteralType::None)
+						const TArray<FMetasoundFrontendClassInputDefault>* ClassInputDefaults = Builder.FindNodeClassInputDefaults(InputHandle.NodeID, Vertex->Name);
+						if (ClassInputDefaults)
 						{
-							OutDefaultLiteral.Clear();	
-						}
-						else
-						{
-							OutDefaultLiteral.Set(static_cast<UObject*>(nullptr));
+							FGuid PageID;
+							if (Engine::FDocumentBuilderRegistry::GetChecked().TryResolveTargetPageID(*ClassInputDefaults, PageID))
+							{
+								auto MatchesPageID = [&PageID](const FMetasoundFrontendClassInputDefault& InputDefault) { return PageID == InputDefault.PageID; };
+								if (const FMetasoundFrontendClassInputDefault* ClassDefault = ClassInputDefaults->FindByPredicate(MatchesPageID))
+								{
+									if (ClassDefault->Literal.GetType() == EMetasoundFrontendLiteralType::None)
+									{
+										OutDefaultLiteral.Clear();
+									}
+									else
+									{
+										OutDefaultLiteral = ClassDefault->Literal;
+									}
+								}
+							}
+							return true;
 						}
+
+						OutDefaultLiteral.Set(static_cast<UObject*>(nullptr));
 					}
 				}
 				break;
@@ -1545,65 +1573,21 @@ namespace Metasound
 			return true;
 		}
 
-		void FGraphBuilder::DisconnectPinVertex(UEdGraphPin& InPin, bool bAddLiteralInputs)
+		void FGraphBuilder::DisconnectPinVertex(UEdGraphPin& InPin)
 		{
-			using namespace Editor;
-			using namespace Frontend;
-
-			TArray<FInputHandle> InputHandles;
-			TArray<UEdGraphPin*> InputPins;
-
-			UMetasoundEditorGraphNode* Node = CastChecked<UMetasoundEditorGraphNode>(InPin.GetOwningNode());
-
-			if (InPin.Direction == EGPD_Input)
-			{
-				const FName PinName = InPin.GetFName();
-
-				FNodeHandle NodeHandle = Node->GetNodeHandle();
-				FInputHandle InputHandle = NodeHandle->GetInputWithVertexName(PinName);
-
-				// Input can be invalid if renaming a vertex member
-				if (InputHandle->IsValid())
-				{
-					InputHandles.Add(InputHandle);
-					InputPins.Add(&InPin);
-				}
-			}
-			else
+			FMetaSoundFrontendDocumentBuilder& Builder = GetBuilderFromPinChecked(InPin).GetBuilder();
+			const FMetasoundFrontendVertexHandle VertexHandle = GetPinVertexHandle(Builder, &InPin);
+			if (VertexHandle.IsSet())
 			{
-				check(InPin.Direction == EGPD_Output);
-				for (UEdGraphPin* Pin : InPin.LinkedTo)
+				if (InPin.Direction == EGPD_Input)
 				{
-					check(Pin);
-					FNodeHandle NodeHandle = CastChecked<UMetasoundEditorGraphNode>(Pin->GetOwningNode())->GetNodeHandle();
-					FInputHandle InputHandle = NodeHandle->GetInputWithVertexName(Pin->GetFName());
-
-					// Input can be invalid if renaming a vertex member
-					if (InputHandle->IsValid())
-					{
-						InputHandles.Add(InputHandle);
-						InputPins.Add(Pin);
-					}
+					Builder.RemoveEdgeToNodeInput(VertexHandle.NodeID, VertexHandle.VertexID);
 				}
-			}
-
-			for (int32 i = 0; i < InputHandles.Num(); ++i)
-			{
-				FInputHandle InputHandle = InputHandles[i];
-				FConstOutputHandle OutputHandle = InputHandle->GetConnectedOutput();
-
-				InputHandle->Disconnect();
-
-				if (bAddLiteralInputs)
+				else
 				{
-					FNodeHandle NodeHandle = InputHandle->GetOwningNode();
-					SynchronizePinLiteral(*InputPins[i]);
+					Builder.RemoveEdgesFromNodeOutput(VertexHandle.NodeID, VertexHandle.VertexID);
 				}
 			}
-
-			UObject& MetaSound = Node->GetMetasoundChecked();
-			FMetasoundAssetBase* MetaSoundAsset = IMetasoundUObjectRegistry::Get().GetObjectAsAssetBase(&MetaSound);
-			MetaSoundAsset->GetModifyContext().SetDocumentModified();
 		}
 
 		bool FGraphBuilder::DeleteMemberNodes(const UMetasoundEditorGraphMember& InGraphMember)
@@ -1886,7 +1870,8 @@ namespace Metasound
 			if (ensure(NewPin))
 			{
 				RefreshPinMetadata(*NewPin, InInputHandle->GetMetadata());
-				SynchronizePinLiteral(*NewPin);
+				const FMetaSoundFrontendDocumentBuilder& Builder = InEditorNode.GetBuilderChecked().GetConstBuilder();
+				SynchronizePinLiteral(Builder, *NewPin);
 			}
 
 			return NewPin;
@@ -2139,7 +2124,7 @@ namespace Metasound
 						}
 					}
 
-					SynchronizePinLiteral(*MatchingPin);
+					SynchronizePinLiteral(InBuilder, *MatchingPin);
 				}
 
 				// Handle node outputs to break connections for the case 
@@ -2637,8 +2622,9 @@ namespace Metasound
 			return bIsNodeDirty;
 		}
 
-		bool FGraphBuilder::SynchronizePinLiteral(UEdGraphPin& InPin)
+		bool FGraphBuilder::SynchronizePinLiteral(const FMetaSoundFrontendDocumentBuilder& InBuilder, UEdGraphPin& InPin)
 		{
+			using namespace Engine;
 			using namespace Frontend;
 
 			if (!ensure(InPin.Direction == EGPD_Input))
@@ -2648,21 +2634,38 @@ namespace Metasound
 
 			const FString OldValue = InPin.DefaultValue;
 
-			FConstInputHandle InputHandle = GetConstInputHandleFromPin(&InPin);
-			if (const FMetasoundFrontendLiteral* NodeDefaultLiteral = InputHandle->GetLiteral())
+			const FMetasoundFrontendVertex* InputVertex = GetPinVertex(InBuilder, &InPin);
+			if (!ensure(InputVertex))
 			{
-				InPin.DefaultValue = NodeDefaultLiteral->ToString();
-				return OldValue != InPin.DefaultValue;
+				return false;
 			}
 
-			if (const FMetasoundFrontendLiteral* ClassDefaultLiteral = InputHandle->GetClassDefaultLiteral())
+			FMetasoundFrontendVertexHandle InputHandle = GetPinVertexHandle(InBuilder, &InPin);
+			check(InputHandle.IsSet());
+
+			if (const FMetasoundFrontendVertexLiteral* VertexLiteral = InBuilder.FindNodeInputDefault(InputHandle.NodeID, InputVertex->Name))
 			{
-				InPin.DefaultValue = ClassDefaultLiteral->ToString();
+				InPin.DefaultValue = VertexLiteral->Value.ToString();
 				return OldValue != InPin.DefaultValue;
 			}
 
+			const TArray<FMetasoundFrontendClassInputDefault>* ClassDefaults = InBuilder.FindNodeClassInputDefaults(InputHandle.NodeID, InputVertex->Name);
+			if (ClassDefaults)
+			{
+				FGuid PageID;
+				if (Engine::FDocumentBuilderRegistry::GetChecked().TryResolveTargetPageID(*ClassDefaults, PageID))
+				{
+					auto MatchesPageID = [&PageID](const FMetasoundFrontendClassInputDefault& InputDefault) { return InputDefault.PageID == PageID; };
+					if (const FMetasoundFrontendClassInputDefault* ClassDefault = ClassDefaults->FindByPredicate(MatchesPageID))
+					{
+						InPin.DefaultValue = ClassDefault->Literal.ToString();
+						return OldValue != InPin.DefaultValue;
+					}
+				}
+			}
+
 			FMetasoundFrontendLiteral DefaultLiteral;
-			DefaultLiteral.SetFromLiteral(IDataTypeRegistry::Get().CreateDefaultLiteral(InputHandle->GetDataType()));
+			DefaultLiteral.SetFromLiteral(IDataTypeRegistry::Get().CreateDefaultLiteral(InputVertex->TypeName));
 
 			InPin.DefaultValue = DefaultLiteral.ToString();
 			return OldValue != InPin.DefaultValue;
diff --git a/Engine/Plugins/Runtime/Metasound/Source/MetasoundEditor/Private/MetasoundEditorGraphBuilder.h b/Engine/Plugins/Runtime/Metasound/Source/MetasoundEditor/Private/MetasoundEditorGraphBuilder.h
--- a/Engine/Plugins/Runtime/Metasound/Source/MetasoundEditor/Private/MetasoundEditorGraphBuilder.h
+++ b/Engine/Plugins/Runtime/Metasound/Source/MetasoundEditor/Private/MetasoundEditorGraphBuilder.h
@@ -192,11 +192,13 @@ namespace Metasound
 			// Disconnects pin's associated frontend vertex from any linked input
 			// or output nodes, and reflects change in the Frontend graph. Does *not*
 			// disconnect the EdGraph pins.
-			static void DisconnectPinVertex(UEdGraphPin& InPin, bool bAddLiteralInputs = true);
+			static void DisconnectPinVertex(UEdGraphPin& InPin);
 
 			// Generates a unique output name for the given MetaSound object
 			static FName GenerateUniqueNameByClassType(UObject& InMetaSound, EMetasoundFrontendClassType InClassType, const FString& InBaseName);
 
+			static UMetaSoundBuilderBase& GetBuilderFromPinChecked(const UEdGraphPin& InPin);
+
 			static TArray<FString> GetDataTypeNameCategories(const FName& InDataTypeName);
 
 			// Get the input handle from an input pin.  Ensures pin is an input pin.
@@ -287,7 +289,7 @@ namespace Metasound
 			static bool SynchronizeConnections(const FMetaSoundFrontendDocumentBuilder& InBuilder, UMetasoundEditorGraph& OutGraph);
 
 			// Synchronizes literal for a given input with the EdGraph's pin value.
-			static bool SynchronizePinLiteral(UEdGraphPin& InPin);
+			static bool SynchronizePinLiteral(const FMetaSoundFrontendDocumentBuilder& InBuilder, UEdGraphPin& InPin);
 
 			// Synchronizes pin type for a given pin with that registered with the MetaSound editor module provided.
 			static bool SynchronizePinType(const IMetasoundEditorModule& InEditorModule, UEdGraphPin& InPin, const FName InDataType);
diff --git a/Engine/Plugins/Runtime/Metasound/Source/MetasoundEditor/Private/MetasoundEditorGraphNode.cpp b/Engine/Plugins/Runtime/Metasound/Source/MetasoundEditor/Private/MetasoundEditorGraphNode.cpp
--- a/Engine/Plugins/Runtime/Metasound/Source/MetasoundEditor/Private/MetasoundEditorGraphNode.cpp
+++ b/Engine/Plugins/Runtime/Metasound/Source/MetasoundEditor/Private/MetasoundEditorGraphNode.cpp
@@ -676,20 +676,16 @@ void UMetasoundEditorGraphNode::PostEditUndo()
 
 	UEdGraphPin::ResolveAllPinReferences();
 
-	// This can trigger and the handle is no longer valid if transaction
-	// is being undone on a graph node that is orphaned.  If orphaned,
-	// bail early.
-	FConstNodeHandle NodeHandle = GetConstNodeHandle();
-	if (!NodeHandle->IsValid())
-	{
-		return;
-	}
-
-	for (UEdGraphPin* Pin : Pins)
+	// Test for frontend node existence safely as it can no longer be so after rolling transactions back or forward.
+	const FMetaSoundFrontendDocumentBuilder& Builder = GetBuilderChecked().GetConstBuilder();
+	if (const FMetasoundFrontendNode* FrontendNode = Builder.FindNode(GetNodeID()))
 	{
-		if (Pin && Pin->Direction == EGPD_Input)
+		for (UEdGraphPin* Pin : Pins)
 		{
-			FGraphBuilder::SynchronizePinLiteral(*Pin);
+			if (Pin && Pin->Direction == EGPD_Input)
+			{
+				FGraphBuilder::SynchronizePinLiteral(Builder, *Pin);
+			}
 		}
 	}
 }
diff --git a/Engine/Plugins/Runtime/Metasound/Source/MetasoundEngine/Public/MetasoundBuilderBase.cpp b/Engine/Plugins/Runtime/Metasound/Source/MetasoundEngine/Public/MetasoundBuilderBase.cpp
--- a/Engine/Plugins/Runtime/Metasound/Source/MetasoundEngine/Public/MetasoundBuilderBase.cpp
+++ b/Engine/Plugins/Runtime/Metasound/Source/MetasoundEngine/Public/MetasoundBuilderBase.cpp
@@ -943,9 +943,9 @@ void UMetaSoundBuilderBase::ReloadCache(bool bPrimeCache)
 }
 
 #if WITH_EDITORONLY_DATA
-void UMetaSoundBuilderBase::RemoveAllGraphPages()
+void UMetaSoundBuilderBase::ResetGraphPages(bool bClearDefaultGraph)
 {
-	Builder.RemoveAllGraphPages();
+	Builder.ResetGraphPages(bClearDefaultGraph);
 	Builder.RemoveUnusedDependencies();
 }
 #endif // WITH_EDITORONLY_DATA
diff --git a/Engine/Plugins/Runtime/Metasound/Source/MetasoundEngine/Public/MetasoundBuilderBase.h b/Engine/Plugins/Runtime/Metasound/Source/MetasoundEngine/Public/MetasoundBuilderBase.h
--- a/Engine/Plugins/Runtime/Metasound/Source/MetasoundEngine/Public/MetasoundBuilderBase.h
+++ b/Engine/Plugins/Runtime/Metasound/Source/MetasoundEngine/Public/MetasoundBuilderBase.h
@@ -338,10 +338,9 @@ public:
 	UFUNCTION(BlueprintCallable, Category = "Audio|MetaSound|Builder", meta = (ExpandEnumAsExecs = "OutResult"))
 	void ConvertToPreset(const TScriptInterface<IMetaSoundDocumentInterface>& ReferencedNodeClass, EMetaSoundBuilderResult& OutResult);
 
-
 #if WITH_EDITORONLY_DATA
-	// Removes all graph pages, leaving just the default.
-	void RemoveAllGraphPages();
+	// Removes all graph pages except the default.  If bClearDefaultPage is true, clears the default graph page implementation.
+	void ResetGraphPages(bool bClearDefaultPage);
 #endif // WITH_EDITORONLY_DATA
 
 	// Removes graph input if it exists; sets result to succeeded if it was removed and failed if it was not.
@@ -466,7 +465,7 @@ public:
 	const FMetaSoundFrontendDocumentBuilder& GetConstBuilder() const;
 	int32 GetLastTransactionRegistered() const;
 
-	// Resets FrontendBuilder instance, creating a transient MetaSound document that is internal to this UObject Builder.
+	// Resets FrontendBuilder instance, creating a transient MetaSound document that is managed by this UObject Builder.
 	void Initialize();
 
 protected:
diff --git a/Engine/Plugins/Runtime/Metasound/Source/MetasoundFrontend/Private/MetasoundFrontendDocument.cpp b/Engine/Plugins/Runtime/Metasound/Source/MetasoundFrontend/Private/MetasoundFrontendDocument.cpp
--- a/Engine/Plugins/Runtime/Metasound/Source/MetasoundFrontend/Private/MetasoundFrontendDocument.cpp
+++ b/Engine/Plugins/Runtime/Metasound/Source/MetasoundFrontend/Private/MetasoundFrontendDocument.cpp
@@ -857,13 +857,15 @@ FMetasoundFrontendClassInput::FMetasoundFrontendClassInput(const Audio::FParamet
 #endif // WITH_EDITOR
 }
 
-#if WITH_EDITORONLY_DATA
 FMetasoundFrontendLiteral& FMetasoundFrontendClassInput::AddDefault(const FGuid& InPageID)
 {
+#if WITH_EDITORONLY_DATA
+	checkf(InPageID == Metasound::Frontend::DefaultPageID, TEXT("Non-default page mutation is not supported without editor data loaded."));
+#endif // WITH_EDITORONLY_DATA
+
 	checkf(!ContainsDefault(InPageID), TEXT("Page default with given ID already exists"));
 	return Defaults.Add_GetRef(FMetasoundFrontendClassInputDefault { InPageID }).Literal;
 }
-#endif // WITH_EDITORONLY_DATA
 
 bool FMetasoundFrontendClassInput::ContainsDefault(const FGuid& InPageID) const
 {
@@ -940,25 +942,38 @@ void FMetasoundFrontendClassInput::IterateDefaults(TFunctionRef<void(const FGuid
 	}
 }
 
-#if WITH_EDITORONLY_DATA
-void FMetasoundFrontendClassInput::RemoveAllDefaults()
+bool FMetasoundFrontendClassInput::RemoveDefault(const FGuid& InPageID)
 {
-	Defaults.Reset();
+	if (InPageID != Metasound::Frontend::DefaultPageID)
+	{
+		auto IsPage = [&InPageID](const FMetasoundFrontendClassInputDefault& Default) { return Default.PageID == InPageID; };
+		return Defaults.RemoveAllSwap(IsPage) > 0;
+	}
+
+	return false;
 }
 
-bool FMetasoundFrontendClassInput::RemoveDefault(const FGuid& InPageID)
+void FMetasoundFrontendClassInput::ResetDefaults()
 {
-	auto IsPage = [&InPageID](const FMetasoundFrontendClassInputDefault& Default) { return Default.PageID == InPageID; };
-	return Defaults.RemoveAllSwap(IsPage) > 0;
+	using namespace Metasound::Frontend;
+
+	Defaults.Reset();
+	FMetasoundFrontendLiteral TypeDefault;
+	TypeDefault.SetFromLiteral(IDataTypeRegistry::Get().CreateDefaultLiteral(TypeName));
+	InitDefault(MoveTemp(TypeDefault));
+	Defaults.Shrink();
 }
 
 void FMetasoundFrontendClassInput::SetDefaults(TArray<FMetasoundFrontendClassInputDefault> InputDefaults)
 {
+#if DO_CHECK
+	auto IsDefaultPageID = [](const FMetasoundFrontendClassInputDefault& Default) { return Default.PageID == Metasound::Frontend::DefaultPageID; };
+	check(InputDefaults.ContainsByPredicate(IsDefaultPageID));
+#endif // DO_CHECK
+
 	Defaults = MoveTemp(InputDefaults);
 }
 
-#endif // WITH_EDITORONLY_DATA
-
 FMetasoundFrontendClassVariable::FMetasoundFrontendClassVariable(const FMetasoundFrontendClassVertex& InOther)
 	: FMetasoundFrontendClassVertex(InOther)
 {
@@ -1034,11 +1049,6 @@ bool FMetasoundFrontendGraphClass::ContainsGraphPage(const FGuid& InPageID) cons
 }
 
 #if WITH_EDITORONLY_DATA
-void FMetasoundFrontendGraphClass::RemoveAllGraphPages()
-{
-	PagedGraphs.Empty();
-}
-
 bool FMetasoundFrontendGraphClass::RemoveGraphPage(const FGuid& InPageID, FGuid* OutAdjacentPageID)
 {
 	if (InPageID != FGuid())
@@ -1066,6 +1076,25 @@ bool FMetasoundFrontendGraphClass::RemoveGraphPage(const FGuid& InPageID, FGuid*
 
 	return false;
 }
+
+void FMetasoundFrontendGraphClass::ResetGraphPages(bool bClearDefaultGraph)
+{
+	PagedGraphs.RemoveAllSwap([](const FMetasoundFrontendGraph& Graph)
+	{
+		return Graph.PageID != Metasound::Frontend::DefaultPageID;
+	}, EAllowShrinking::Yes);
+
+	if (bClearDefaultGraph)
+	{
+		IterateGraphPages([](FMetasoundFrontendGraph& Graph)
+		{
+			Graph.Nodes.Empty();
+			Graph.Edges.Empty();
+			Graph.Variables.Empty();
+			Graph.Style = { };
+		});
+	}
+}
 #endif // WITH_EDITORONLY_DATA
 
 FMetasoundFrontendGraph* FMetasoundFrontendGraphClass::FindGraph(const FGuid& InPageID)
diff --git a/Engine/Plugins/Runtime/Metasound/Source/MetasoundFrontend/Private/MetasoundFrontendDocumentBuilder.cpp b/Engine/Plugins/Runtime/Metasound/Source/MetasoundFrontend/Private/MetasoundFrontendDocumentBuilder.cpp
--- a/Engine/Plugins/Runtime/Metasound/Source/MetasoundFrontend/Private/MetasoundFrontendDocumentBuilder.cpp
+++ b/Engine/Plugins/Runtime/Metasound/Source/MetasoundFrontend/Private/MetasoundFrontendDocumentBuilder.cpp
@@ -1351,20 +1351,30 @@ void FMetaSoundFrontendDocumentBuilder::ClearDocument(TSharedRef<Metasound::Fron
 	GraphClass.Interface.SetOutputStyle({ });
 #endif // WITH_EDITOR
 
-	GraphClass.PresetOptions.InputsInheritingDefault.Reset();
+	GraphClass.PresetOptions.InputsInheritingDefault.Empty();
 	GraphClass.PresetOptions.bIsPreset = false;
 
 	// Removing graph pages is not necessary when editor only data is not available as graph mutation
 	// is only supported in builds with editor data loaded. Otherwise, anything calling ClearDocument
 	// should only be a transient, non serialized asset graph which does not support page mutation.
 #if WITH_EDITORONLY_DATA
-	RemoveAllGraphPages();
+	constexpr bool bClearDefaultGraph = true;
+	ResetGraphPages(bClearDefaultGraph);
 #else // !WITH_EDITORONLY_DATA
 	UObject& DocObject = CastDocumentObjectChecked<UObject>();
 	checkf(!DocObject.IsAsset(), TEXT("Cannot call clear document on asset '%s': builder API does not support document mutation on serialized objects without editor data loaded"), *GetDebugName());
+
+	GraphClass.IterateGraphPages([] (FMetasoundFrontendGraph& Graph)
+	{
+		Graph.Nodes.Empty();
+		Graph.Edges.Empty();
+		Graph.Variables.Empty();
+	});
 #endif // !WITH_EDITORONLY_DATA
 
-	GraphClass.InitDefaultGraphPage();
+	GraphClass.Interface.Inputs.Empty();
+	GraphClass.Interface.Outputs.Empty();
+	GraphClass.Interface.Environment.Empty();
 
 	Doc.Interfaces.Empty();
 	Doc.Dependencies.Empty();
@@ -1813,6 +1823,15 @@ const FMetasoundFrontendNode* FMetaSoundFrontendDocumentBuilder::FindGraphOutput
 	return nullptr;
 }
 
+const FMetasoundFrontendVariable* FMetaSoundFrontendDocumentBuilder::FindGraphVariable(FName VariableName, const FGuid* InPageID) const
+{
+	const FGuid& PageID = InPageID ? *InPageID : BuildPageID;
+	const FMetasoundFrontendDocument& Document = GetDocumentChecked();
+	const FMetasoundFrontendGraph& Graph = Document.RootGraph.FindConstGraphChecked(PageID);
+	auto MatchesName = [&VariableName](const FMetasoundFrontendVariable& Variable) { return Variable.Name == VariableName; };
+	return Graph.Variables.FindByPredicate(MatchesName);
+}
+
 #if WITH_EDITOR
 UMetaSoundFrontendMemberMetadata* FMetaSoundFrontendDocumentBuilder::FindMemberMetadata(const FGuid& InMemberID)
 {
@@ -2773,24 +2792,6 @@ void FMetaSoundFrontendDocumentBuilder::FinishBuilding()
 	DocumentCache.Reset();
 }
 
-#if WITH_EDITORONLY_DATA
-void FMetaSoundFrontendDocumentBuilder::RemoveAllGraphPages()
-{
-	using namespace Metasound;
-
-	FMetasoundFrontendGraphClass& RootGraph = GetDocumentChecked().RootGraph;
-	RootGraph.IterateGraphPages([this](FMetasoundFrontendGraph& Graph)
-	{
-		if (Graph.PageID != Frontend::DefaultPageID)
-		{
-			DocumentDelegates->PageDelegates.OnRemovingPage.Broadcast(Frontend::FDocumentMutatePageArgs { Graph.PageID });
-		}
-	});
-	RootGraph.RemoveAllGraphPages();
-	SetBuildPageID(Frontend::DefaultPageID);
-}
-#endif // WITH_EDITORONLY_DATA
-
 bool FMetaSoundFrontendDocumentBuilder::RemoveDependency(const FGuid& InClassID)
 {
 	using namespace Metasound::Frontend;
@@ -3271,7 +3272,6 @@ bool FMetaSoundFrontendDocumentBuilder::RemoveGraphPage(const FGuid& InPageID)
 		{
 			ensureAlwaysMsgf(SetBuildPageID(AdjacentPageID), TEXT("AdjacentPageID returned is always expected to be valid"));
 		}
-
 	}
 
 	return bPageRemoved;
@@ -3426,6 +3426,55 @@ void FMetaSoundFrontendDocumentBuilder::ReloadCache()
 	Reload(DocumentDelegates, true);
 }
 
+#if WITH_EDITORONLY_DATA
+bool FMetaSoundFrontendDocumentBuilder::ResetGraphInputDefault(FName InputName)
+{
+	using namespace Metasound;
+
+	auto NameMatchesInput = [&InputName](const FMetasoundFrontendClassInput& Input) { return Input.Name == InputName; };
+	FMetasoundFrontendDocument& Document = GetDocumentChecked();
+	TArray<FMetasoundFrontendClassInput>& Inputs = Document.RootGraph.Interface.Inputs;
+
+	const int32 Index = Inputs.IndexOfByPredicate(NameMatchesInput);
+	if (Index != INDEX_NONE)
+	{
+		FMetasoundFrontendClassInput& Input = Inputs[Index];
+		Input.ResetDefaults();
+
+		DocumentDelegates->InterfaceDelegates.OnInputDefaultChanged.Broadcast(Index);
+
+		// Set the input as inheriting default for presets
+		if (IsPreset())
+		{
+			constexpr bool bInputInheritsDefault = true;
+			return SetGraphInputInheritsDefault(InputName, bInputInheritsDefault);
+		}
+
+		Document.Metadata.ModifyContext.AddMemberIDModified(Input.NodeID);
+		return true;
+	}
+
+	return false;
+}
+
+void FMetaSoundFrontendDocumentBuilder::ResetGraphPages(bool bClearDefaultGraph)
+{
+	using namespace Metasound;
+
+	FMetasoundFrontendGraphClass& RootGraph = GetDocumentChecked().RootGraph;
+	RootGraph.IterateGraphPages([this](FMetasoundFrontendGraph& Graph)
+	{
+		if (Graph.PageID != Frontend::DefaultPageID)
+		{
+			DocumentDelegates->PageDelegates.OnRemovingPage.Broadcast(Frontend::FDocumentMutatePageArgs{ Graph.PageID });
+		}
+	});
+
+	RootGraph.ResetGraphPages(bClearDefaultGraph);
+	SetBuildPageID(Frontend::DefaultPageID);
+}
+#endif // WITH_EDITORONLY_DATA
+
 #if WITH_EDITOR
 void FMetaSoundFrontendDocumentBuilder::SetAuthor(const FString& InAuthor)
 {
@@ -3556,6 +3605,7 @@ bool FMetaSoundFrontendDocumentBuilder::SetGraphInputDataType(FName InputName, F
 			DefaultLiteral.SetFromLiteral(Frontend::IDataTypeRegistry::Get().CreateDefaultLiteral(DataType));
 			GraphInput.FindDefaultChecked(Frontend::DefaultPageID) = DefaultLiteral;
 			GraphInput.TypeName = DataType;
+			GraphInput.ResetDefaults();
 
 			RootGraph.IterateGraphPages([this, &DataType, &GraphInput](FMetasoundFrontendGraph& Graph)
 			{
@@ -3606,12 +3656,13 @@ bool FMetaSoundFrontendDocumentBuilder::SetGraphInputDataType(FName InputName, F
 	return true;
 }
 
-bool FMetaSoundFrontendDocumentBuilder::SetGraphInputDefault(FName InputName, const FMetasoundFrontendLiteral& InDefaultLiteral)
+bool FMetaSoundFrontendDocumentBuilder::SetGraphInputDefault(FName InputName, FMetasoundFrontendLiteral InDefaultLiteral, const FGuid* InPageID)
 {
 	using namespace Metasound;
 
 	auto NameMatchesInput = [&InputName](const FMetasoundFrontendClassInput& Input) { return Input.Name == InputName; };
-	TArray<FMetasoundFrontendClassInput>& Inputs = GetDocumentChecked().RootGraph.Interface.Inputs;
+	FMetasoundFrontendDocument& Document = GetDocumentChecked();
+	TArray<FMetasoundFrontendClassInput>& Inputs = Document.RootGraph.Interface.Inputs;
 
 	const int32 Index = Inputs.IndexOfByPredicate(NameMatchesInput);
 	if (Index != INDEX_NONE)
@@ -3619,7 +3670,20 @@ bool FMetaSoundFrontendDocumentBuilder::SetGraphInputDefault(FName InputName, co
 		FMetasoundFrontendClassInput& Input = Inputs[Index];
 		if (Frontend::IDataTypeRegistry::Get().IsLiteralTypeSupported(Input.TypeName, InDefaultLiteral.GetType()))
 		{
-			Input.FindDefaultChecked(Frontend::DefaultPageID) = InDefaultLiteral;
+			const FGuid PageID = InPageID ? *InPageID : Frontend::DefaultPageID;
+			bool bFound = false;
+			Input.IterateDefaults([&bFound, &PageID, &InDefaultLiteral](const FGuid& InputPageID, FMetasoundFrontendLiteral& InputLiteral)
+			{
+				if (!bFound && InputPageID == PageID)
+				{
+					bFound = true;
+					InputLiteral = MoveTemp(InDefaultLiteral);
+				}
+			});
+			if (!bFound)
+			{
+				Input.AddDefault(PageID) = MoveTemp(InDefaultLiteral);
+			}
 			DocumentDelegates->InterfaceDelegates.OnInputDefaultChanged.Broadcast(Index);
 
 			// Set the input as no longer inheriting default for presets
@@ -3637,6 +3701,43 @@ bool FMetaSoundFrontendDocumentBuilder::SetGraphInputDefault(FName InputName, co
 	return false;
 }
 
+bool FMetaSoundFrontendDocumentBuilder::SetGraphInputDefaults(FName InputName, TArray<FMetasoundFrontendClassInputDefault> Defaults)
+{
+	using namespace Metasound;
+
+	auto NameMatchesInput = [&InputName](const FMetasoundFrontendClassInput& Input) { return Input.Name == InputName; };
+	FMetasoundFrontendDocument& Document = GetDocumentChecked();
+	TArray<FMetasoundFrontendClassInput>& Inputs = Document.RootGraph.Interface.Inputs;
+
+	const int32 Index = Inputs.IndexOfByPredicate(NameMatchesInput);
+	if (Index != INDEX_NONE)
+	{
+		FMetasoundFrontendClassInput& Input = Inputs[Index];
+		TSet<FGuid> ValidPageIDs;
+		bool bAllSupported = Algo::AllOf(Defaults, [&Input](const FMetasoundFrontendClassInputDefault& Default)
+		{
+			return Frontend::IDataTypeRegistry::Get().IsLiteralTypeSupported(Input.TypeName, Default.Literal.GetType());
+		});
+		if (bAllSupported)
+		{
+			Input.SetDefaults(MoveTemp(Defaults));
+			DocumentDelegates->InterfaceDelegates.OnInputDefaultChanged.Broadcast(Index);
+
+			// Set the input as no longer inheriting default for presets
+			if (IsPreset())
+			{
+				constexpr bool bInputInheritsDefault = false;
+				return SetGraphInputInheritsDefault(InputName, bInputInheritsDefault);
+			}
+
+			return true;
+		}
+		UE_LOG(LogMetaSound, Error, TEXT("Attempting to set graph input of type '%s' with unsupported literal type(s)"), *Input.TypeName.ToString());
+	}
+
+	return false;
+}
+
 bool FMetaSoundFrontendDocumentBuilder::SetGraphInputInheritsDefault(FName InName, bool bInputInheritsDefault)
 {
 	FMetasoundFrontendGraphClassPresetOptions& PresetOptions = GetDocumentChecked().RootGraph.PresetOptions;
@@ -3772,6 +3873,27 @@ bool FMetaSoundFrontendDocumentBuilder::SetGraphOutputDataType(FName OutputName,
 	return true;
 }
 
+bool FMetaSoundFrontendDocumentBuilder::SetGraphVariableDefault(FName VariableName, FMetasoundFrontendLiteral InDefaultLiteral, const FGuid* InPageID)
+{
+	using namespace Metasound;
+
+	const FGuid PageID = InPageID ? *InPageID : Frontend::DefaultPageID;
+	FMetasoundFrontendDocument& Document = GetDocumentChecked();
+	FMetasoundFrontendGraph& Graph = Document.RootGraph.FindGraphChecked(PageID);
+
+	auto NameMatchesVariable = [&VariableName](const FMetasoundFrontendVariable& Variable) { return Variable.Name == VariableName; };
+	if (FMetasoundFrontendVariable* Variable = Graph.Variables.FindByPredicate(NameMatchesVariable))
+	{
+		if (Frontend::IDataTypeRegistry::Get().IsLiteralTypeSupported(Variable->TypeName, InDefaultLiteral.GetType()))
+		{
+			Variable->Literal = MoveTemp(InDefaultLiteral);
+			return true;
+		}
+	}
+
+	return false;
+}
+
 #if WITH_EDITOR
 void FMetaSoundFrontendDocumentBuilder::SetDisplayName(const FText& InDisplayName)
 {
diff --git a/Engine/Plugins/Runtime/Metasound/Source/MetasoundFrontend/Public/MetasoundFrontendDocument.h b/Engine/Plugins/Runtime/Metasound/Source/MetasoundFrontend/Public/MetasoundFrontendDocument.h
--- a/Engine/Plugins/Runtime/Metasound/Source/MetasoundFrontend/Public/MetasoundFrontendDocument.h
+++ b/Engine/Plugins/Runtime/Metasound/Source/MetasoundFrontend/Public/MetasoundFrontendDocument.h
@@ -987,10 +987,7 @@ private:
 	TArray<FMetasoundFrontendClassInputDefault> Defaults;
 
 public:
-#if WITH_EDITORONLY_DATA
 	FMetasoundFrontendLiteral& AddDefault(const FGuid& InPageID);
-#endif // WITH_EDITORONLY_DATA
-
 	bool ContainsDefault(const FGuid& InPageID) const;
 	const FMetasoundFrontendLiteral* FindConstDefault(const FGuid& InPageID) const;
 	const FMetasoundFrontendLiteral& FindConstDefaultChecked(const FGuid& InPageID) const;
@@ -1001,12 +998,9 @@ public:
 	void InitDefault(FMetasoundFrontendLiteral InitLiteral);
 	void IterateDefaults(TFunctionRef<void(const FGuid&, FMetasoundFrontendLiteral&)> IterFunc);
 	void IterateDefaults(TFunctionRef<void(const FGuid&, const FMetasoundFrontendLiteral&)> IterFunc) const;
-
-#if WITH_EDITORONLY_DATA
-	void RemoveAllDefaults();
 	bool RemoveDefault(const FGuid& InPageID);
+	void ResetDefaults();
 	void SetDefaults(TArray<FMetasoundFrontendClassInputDefault> InputDefaults);
-#endif // WITH_EDITORONLY_DATA
 };
 
 // Contains info for variable vertex of a Metasound class.
@@ -1738,12 +1732,14 @@ public:
 public:
 #if WITH_EDITORONLY_DATA
 	const FMetasoundFrontendGraph& AddGraphPage(const FGuid& InPageID, bool bDuplicateLastGraph = true, bool bSetAsBuildGraph = true);
-	void RemoveAllGraphPages();
 
 	// Removes the page associated with the given PageID.  Returns true if removed, false if not.
 	// If provided an "AdjacentPageID," sets the value at the given pointer to a page ID adjacent to
 	// the removed page. If last page was removed, returns the default graph ID.
 	bool RemoveGraphPage(const FGuid& InPageID, FGuid* OutAdjacentPageID = nullptr);
+
+	// Removes all graph pages except the default.  If bClearDefaultPage is true, clears the default graph page implementation.
+	void ResetGraphPages(bool bClearDefaultGraph);
 #endif // WITH_EDITORONLY_DATA
 
 	bool ContainsGraphPage(const FGuid& InPageID) const;
diff --git a/Engine/Plugins/Runtime/Metasound/Source/MetasoundFrontend/Public/MetasoundFrontendDocumentBuilder.h b/Engine/Plugins/Runtime/Metasound/Source/MetasoundFrontend/Public/MetasoundFrontendDocumentBuilder.h
--- a/Engine/Plugins/Runtime/Metasound/Source/MetasoundFrontend/Public/MetasoundFrontendDocumentBuilder.h
+++ b/Engine/Plugins/Runtime/Metasound/Source/MetasoundFrontend/Public/MetasoundFrontendDocumentBuilder.h
@@ -229,6 +229,8 @@ public:
 	const FMetasoundFrontendClassOutput* FindGraphOutput(FName OutputName) const;
 	const FMetasoundFrontendNode* FindGraphOutputNode(FName OutputName, const FGuid* InPageID = nullptr) const;
 
+	const FMetasoundFrontendVariable* FindGraphVariable(FName InVariableName, const FGuid* InPageID = nullptr) const;
+
 	bool FindInterfaceInputNodes(FName InterfaceName, TArray<const FMetasoundFrontendNode*>& OutInputs, const FGuid* InPageID = nullptr) const;
 	bool FindInterfaceOutputNodes(FName InterfaceName, TArray<const FMetasoundFrontendNode*>& OutOutputs, const FGuid* InPageID = nullptr) const;
 
@@ -295,9 +297,11 @@ public:
 	int32 GetTransactionCount() const;
 
 	TArray<const FMetasoundFrontendNode*> GetGraphInputTemplateNodes(FName InInputName, const FGuid* InPageID = nullptr);
+
 	EMetasoundFrontendVertexAccessType GetNodeInputAccessType(const FGuid& InNodeID, const FGuid& InVertexID, const FGuid* InPageID = nullptr) const;
 	const FMetasoundFrontendLiteral* GetNodeInputClassDefault(const FGuid& InNodeID, const FGuid& InVertexID, const FGuid* InPageID = nullptr) const;
 	const FMetasoundFrontendLiteral* GetNodeInputDefault(const FGuid& InNodeID, const FGuid& InVertexID, const FGuid* InPageID = nullptr) const;
+
 	EMetasoundFrontendVertexAccessType GetNodeOutputAccessType(const FGuid& InNodeID, const FGuid& InVertexID, const FGuid* InPageID = nullptr) const;
 
 	// Initializes the builder's document, using the (optional) provided document template, (optional) class name, and (optionally) whether or not to reset the existing class version.
@@ -332,10 +336,6 @@ public:
 
 	bool ModifyInterfaces(Metasound::Frontend::FModifyInterfaceOptions&& InOptions);
 
-#if WITH_EDITORONLY_DATA
-	void RemoveAllGraphPages();
-#endif // WITH_EDITORONLY_DATA
-
 	UE_DEPRECATED(5.5,
 		"Cache invalidation may require new copy of delegates. In addition, re-priming is discouraged. "
 		"To enforce this, new recommended pattern is to construct a new builder instead")
@@ -379,6 +379,13 @@ public:
 	UE_DEPRECATED(5.5, "Use GenerateNewClassName instead")
 	bool RenameRootGraphClass(const FMetasoundFrontendClassName& InName);
 
+#if WITH_EDITORONLY_DATA
+	bool ResetGraphInputDefault(FName InputName);
+
+	// Removes all graph pages except the default.  If bClearDefaultPage is true, clears the default graph page implementation.
+	void ResetGraphPages(bool bClearDefaultGraph);
+#endif // WITH_EDITORONLY_DATA
+
 #if WITH_EDITOR
 	void SetAuthor(const FString& InAuthor);
 
@@ -403,7 +410,8 @@ public:
 	// already the given DataType.
 	bool SetGraphInputDataType(FName InputName, FName DataType);
 
-	bool SetGraphInputDefault(FName InputName, const FMetasoundFrontendLiteral& InDefaultLiteral);
+	bool SetGraphInputDefault(FName InputName, FMetasoundFrontendLiteral InDefaultLiteral, const FGuid* InPageID = nullptr);
+	bool SetGraphInputDefaults(FName InputName, TArray<FMetasoundFrontendClassInputDefault> Defaults);
 
 	// Sets the given graph output's access type. If connected to other nodes and access type is not compatible,
 	// associated edges/connections are removed.  Returns true if either DataType was successfully set to new
@@ -415,6 +423,9 @@ public:
 	// already the given DataType.
 	bool SetGraphOutputDataType(FName OutputName, FName DataType);
 
+	// Sets the given graph variable's default.
+	bool SetGraphVariableDefault(FName VariableName, FMetasoundFrontendLiteral InDefaultLiteral, const FGuid* InPageID = nullptr);
+
 #if WITH_EDITOR
 	void SetDisplayName(const FText& InDisplayName);
 
