Fix shadow variable compilation errors
#rnx

diff --git a/Engine/Plugins/Runtime/Metasound/Source/MetasoundFrontend/Private/MetasoundFrontendDocument.cpp b/Engine/Plugins/Runtime/Metasound/Source/MetasoundFrontend/Private/MetasoundFrontendDocument.cpp
--- a/Engine/Plugins/Runtime/Metasound/Source/MetasoundFrontend/Private/MetasoundFrontendDocument.cpp
+++ b/Engine/Plugins/Runtime/Metasound/Source/MetasoundFrontend/Private/MetasoundFrontendDocument.cpp
@@ -1079,19 +1079,19 @@ bool FMetasoundFrontendGraphClass::RemoveGraphPage(const FGuid& InPageID, FGuid*
 
 void FMetasoundFrontendGraphClass::ResetGraphPages(bool bClearDefaultGraph)
 {
-	PagedGraphs.RemoveAllSwap([](const FMetasoundFrontendGraph& Graph)
+	PagedGraphs.RemoveAllSwap([](const FMetasoundFrontendGraph& PageGraph)
 	{
-		return Graph.PageID != Metasound::Frontend::DefaultPageID;
+		return PageGraph.PageID != Metasound::Frontend::DefaultPageID;
 	}, EAllowShrinking::Yes);
 
 	if (bClearDefaultGraph)
 	{
-		IterateGraphPages([](FMetasoundFrontendGraph& Graph)
+		IterateGraphPages([](FMetasoundFrontendGraph& PageGraph)
 		{
-			Graph.Nodes.Empty();
-			Graph.Edges.Empty();
-			Graph.Variables.Empty();
-			Graph.Style = { };
+			PageGraph.Nodes.Empty();
+			PageGraph.Edges.Empty();
+			PageGraph.Variables.Empty();
+			PageGraph.Style = { };
 		});
 	}
 }
