Remove redundant call from bad local merge
#rnx

diff --git a/Engine/Plugins/Runtime/Metasound/Source/MetasoundFrontend/Private/MetasoundFrontendDocumentBuilder.cpp b/Engine/Plugins/Runtime/Metasound/Source/MetasoundFrontend/Private/MetasoundFrontendDocumentBuilder.cpp
--- a/Engine/Plugins/Runtime/Metasound/Source/MetasoundFrontend/Private/MetasoundFrontendDocumentBuilder.cpp
+++ b/Engine/Plugins/Runtime/Metasound/Source/MetasoundFrontend/Private/MetasoundFrontendDocumentBuilder.cpp
@@ -3601,9 +3601,6 @@ bool FMetaSoundFrontendDocumentBuilder::SetGraphInputDataType(FName InputName, F
 		FMetasoundFrontendClassInput& GraphInput = RootGraph.Interface.Inputs[*Index];
 		if (GraphInput.TypeName != DataType)
 		{
-			FMetasoundFrontendLiteral DefaultLiteral;
-			DefaultLiteral.SetFromLiteral(Frontend::IDataTypeRegistry::Get().CreateDefaultLiteral(DataType));
-			GraphInput.FindDefaultChecked(Frontend::DefaultPageID) = DefaultLiteral;
 			GraphInput.TypeName = DataType;
 			GraphInput.ResetDefaults();
 
