[MetaSound] [Editor] Add 'IsAdvanced' Field to Graph Vertices
#rnx

diff --git a/Engine/Plugins/Runtime/Metasound/Source/MetasoundEditor/Private/MetasoundEditorGraph.cpp b/Engine/Plugins/Runtime/Metasound/Source/MetasoundEditor/Private/MetasoundEditorGraph.cpp
--- a/Engine/Plugins/Runtime/Metasound/Source/MetasoundEditor/Private/MetasoundEditorGraph.cpp
+++ b/Engine/Plugins/Runtime/Metasound/Source/MetasoundEditor/Private/MetasoundEditorGraph.cpp
@@ -389,6 +389,47 @@ void UMetasoundEditorGraphVertex::SetDisplayName(const FText& InNewName, bool bP
 	Graph->RegisterGraphWithFrontend();
 	FGraphBuilder::GetOutermostMetaSoundChecked(*Graph).GetModifyContext().AddMemberIDsModified({ GetMemberID() });
 }
+#if WITH_EDITORONLY_DATA
+bool UMetasoundEditorGraphVertex::SetIsAdvancedDisplay(const bool IsAdvancedDisplay)
+{
+	UMetasoundEditorGraph* Graph = GetOwningGraph();
+	if (!ensure(Graph))
+	{
+		return false;
+	}
+
+	const FText TransactionLabel = FText::Format(LOCTEXT("SetGraphVertexIsAdvancedDisplayState", "Set Metasound {0} IsAdvancedDislay"), GetGraphMemberLabel());
+	const FScopedTransaction Transaction(TransactionLabel, true);
+
+	Graph->Modify();
+	Graph->GetMetasoundChecked().Modify();
+	Modify();
+
+	bool bSucceeded = false;
+	
+	if (GetClassType() == EMetasoundFrontendClassType::Input)
+	{
+		bSucceeded = GetFrontendBuilderChecked().SetGraphInputAdvancedDisplay(GetMemberName(), IsAdvancedDisplay);
+		
+	}
+	else if(GetClassType() == EMetasoundFrontendClassType::Output)
+	{
+		bSucceeded = GetFrontendBuilderChecked().SetGraphOutputAdvancedDisplay(GetMemberName(), IsAdvancedDisplay);
+	}
+
+	const FMetaSoundFrontendDocumentBuilder& Builder = GetFrontendBuilderChecked();	
+	if (const FMetasoundFrontendNode* Node = Builder.FindNode(NodeID))
+	{
+		if (const FMetasoundFrontendClass* Class = Builder.FindDependency(Node->ClassID))
+		{
+			ClassName = Class->Metadata.GetClassName();
+		}
+	}
+
+	Graph->RegisterGraphWithFrontend();
+	return bSucceeded;
+}
+#endif // WITH_EDITORONLY_DATA
 
 Metasound::Frontend::FNodeHandle UMetasoundEditorGraphVertex::GetNodeHandle()
 {
diff --git a/Engine/Plugins/Runtime/Metasound/Source/MetasoundEditor/Private/MetasoundNodeDetailCustomization.cpp b/Engine/Plugins/Runtime/Metasound/Source/MetasoundEditor/Private/MetasoundNodeDetailCustomization.cpp
--- a/Engine/Plugins/Runtime/Metasound/Source/MetasoundEditor/Private/MetasoundNodeDetailCustomization.cpp
+++ b/Engine/Plugins/Runtime/Metasound/Source/MetasoundEditor/Private/MetasoundNodeDetailCustomization.cpp
@@ -66,6 +66,10 @@ namespace Metasound
 			static const FText ConstructorPinText = LOCTEXT("ConstructorPinText", "Is Constructor Pin");
 			static const FText ConstructorPinTooltip = LOCTEXT("ConstructorPinTooltip",
 				"Whether this input or output is a constructor pin. Constructor values are only read on construction (on play), and are not dynamically updated at runtime.");
+			
+			static const FText AdvancedPinText = LOCTEXT("AdvancedPinText", "Is Advanced Pin");
+			static const FText AdvancedPinTooltip = LOCTEXT("AdvancedPinTooltip",
+				"Advanced Pins are hidden by default on the node when this MetaSound is used in other graphs.");
 
 
 			// Retrieves the data type info if the literal property's member is found. Returns if the associated member is found, false if not.
@@ -1414,7 +1418,42 @@ PRAGMA_ENABLE_DEPRECATION_WARNINGS
 				];
 			}
 		}
-		
+
+#if WITH_EDITORONLY_DATA
+		void FMetasoundVertexDetailCustomization::AddAdvancedPinRow(IDetailLayoutBuilder& InDetailLayout)
+		{
+			//only add row if input or output
+			if (UMetasoundEditorGraphVertex* Vertex = Cast<UMetasoundEditorGraphVertex>(GraphMember.Get()))
+			{
+				InDetailLayout.EditCategory("General").AddCustomRow(MemberCustomizationPrivate::AdvancedPinText)
+					.IsEnabled(IsGraphEditable() && !IsInterfaceMember())
+					.NameContent()
+					[
+						SNew(STextBlock)
+							.Text(MemberCustomizationPrivate::AdvancedPinText)
+							.ToolTipText(MemberCustomizationPrivate::AdvancedPinTooltip)
+							.Font(IDetailLayoutBuilder::GetDetailFontBold())
+					]
+				.ValueContent()
+				[
+					SAssignNew(AdvancedPinCheckbox, SCheckBox)
+					.IsChecked_Lambda([this, Vertex]()
+					{
+						return OnGetAdvancedPinCheckboxState(Vertex);
+					})
+					.OnCheckStateChanged_Lambda([this, Vertex](ECheckBoxState InNewState)
+					{
+						OnAdvancedPinStateChanged(Vertex, InNewState);
+					})
+					[
+						SNew(STextBlock)
+						.Font(IDetailLayoutBuilder::GetDetailFont())
+					]
+				];
+			}
+		}
+#endif // WITH_EDITORONLY_DATA
+
 		void FMetasoundVertexDetailCustomization::CustomizeGeneralCategory(IDetailLayoutBuilder& InDetailLayout)
 		{
 			FMetasoundMemberDetailCustomization::CustomizeGeneralCategory(InDetailLayout);
@@ -1431,6 +1470,10 @@ PRAGMA_ENABLE_DEPRECATION_WARNINGS
 			{
 				AddConstructorPinRow(InDetailLayout);
 			}
+			
+#if WITH_EDITORONLY_DATA		
+				AddAdvancedPinRow(InDetailLayout);
+#endif // WITH_EDITORONLY_DATA
 
 			// Sort order
 			IDetailCategoryBuilder& CategoryBuilder = FMetasoundMemberDetailCustomization::GetGeneralCategoryBuilder(InDetailLayout);
@@ -1570,6 +1613,19 @@ PRAGMA_ENABLE_DEPRECATION_WARNINGS
 			return ECheckBoxState::Undetermined;
 		}
 
+#if WITH_EDITORONLY_DATA
+		ECheckBoxState FMetasoundVertexDetailCustomization::OnGetAdvancedPinCheckboxState(TWeakObjectPtr<UMetasoundEditorGraphVertex> InGraphVertex) const
+		{
+			if (InGraphVertex.IsValid())
+			{
+				FMetaSoundFrontendDocumentBuilder& DocumentBuilder = InGraphVertex->GetFrontendBuilderChecked();
+				return DocumentBuilder.GetIsAdvancedDisplay(InGraphVertex->GetMemberName(), InGraphVertex->GetClassType()) ? ECheckBoxState::Checked : ECheckBoxState::Unchecked;
+			}
+
+			return ECheckBoxState::Undetermined;
+		}
+#endif // WITH_EDITORONLY_DATA
+
 		void FMetasoundVertexDetailCustomization::OnConstructorPinStateChanged(TWeakObjectPtr<UMetasoundEditorGraphVertex> InGraphVertex, ECheckBoxState InNewState)
 		{
 			if (InGraphVertex.IsValid() && ConstructorPinCheckbox.IsValid())
@@ -1595,6 +1651,22 @@ PRAGMA_ENABLE_DEPRECATION_WARNINGS
 			}
 		}
 
+#if WITH_EDITORONLY_DATA
+		void FMetasoundVertexDetailCustomization::OnAdvancedPinStateChanged(TWeakObjectPtr<UMetasoundEditorGraphVertex> InGraphVertex, ECheckBoxState InNewState)
+		{
+			if (InGraphVertex.IsValid() && AdvancedPinCheckbox.IsValid())
+			{
+				bool Checked = InNewState == ECheckBoxState::Checked;		
+				InGraphVertex->SetIsAdvancedDisplay(Checked);
+
+				if (FMetasoundAssetBase* MetasoundAsset = Metasound::IMetasoundUObjectRegistry::Get().GetObjectAsAssetBase(GraphMember->GetOutermostObject()))
+				{
+					MetasoundAsset->GetModifyContext().AddMemberIDsModified({ GraphMember->GetMemberID() });
+				}				
+			}
+		}
+#endif // WITH_EDITORONLY_DATA
+
 		void FMetasoundInputDetailCustomization::CustomizeDetails(IDetailLayoutBuilder& InDetailLayout)
 		{
 			CacheMemberData(InDetailLayout);
diff --git a/Engine/Plugins/Runtime/Metasound/Source/MetasoundEditor/Private/MetasoundNodeDetailCustomization.h b/Engine/Plugins/Runtime/Metasound/Source/MetasoundEditor/Private/MetasoundNodeDetailCustomization.h
--- a/Engine/Plugins/Runtime/Metasound/Source/MetasoundEditor/Private/MetasoundNodeDetailCustomization.h
+++ b/Engine/Plugins/Runtime/Metasound/Source/MetasoundEditor/Private/MetasoundNodeDetailCustomization.h
@@ -383,6 +383,15 @@ namespace Metasound
 			void OnConstructorPinStateChanged(TWeakObjectPtr<UMetasoundEditorGraphVertex> InGraphMember, ECheckBoxState InNewState);
 
 			TSharedPtr<SCheckBox> ConstructorPinCheckbox;
+
+#if WITH_EDITORONLY_DATA
+			void AddAdvancedPinRow(IDetailLayoutBuilder& InDetailLayout);
+			ECheckBoxState OnGetAdvancedPinCheckboxState(TWeakObjectPtr<UMetasoundEditorGraphVertex> InGraphMember) const;
+			void OnAdvancedPinStateChanged(TWeakObjectPtr<UMetasoundEditorGraphVertex> InGraphMember, ECheckBoxState InNewState);
+
+			TSharedPtr<SCheckBox> AdvancedPinCheckbox;
+#endif // WITH_EDITORONLY_DATA
+
 		};
 
 		class FMetasoundInputDetailCustomization : public FMetasoundVertexDetailCustomization
diff --git a/Engine/Plugins/Runtime/Metasound/Source/MetasoundEditor/Private/SMetasoundGraphNode.cpp b/Engine/Plugins/Runtime/Metasound/Source/MetasoundEditor/Private/SMetasoundGraphNode.cpp
--- a/Engine/Plugins/Runtime/Metasound/Source/MetasoundEditor/Private/SMetasoundGraphNode.cpp
+++ b/Engine/Plugins/Runtime/Metasound/Source/MetasoundEditor/Private/SMetasoundGraphNode.cpp
@@ -374,6 +374,23 @@ namespace Metasound
 				FMetasoundFrontendNodeStyle Style = FrontendNode->Style;
 				InPin->SafeSetHidden(Style.bUnconnectedPinsHidden);
 			}
+			
+			if (InPin->Direction == EGPD_Input)
+			{
+				const Frontend::FConstInputHandle InputHandle = FGraphBuilder::GetConstInputHandleFromPin(InPin);
+				if (InputHandle->GetMetadata().bIsAdvancedDisplay != InPin->bAdvancedView)
+				{
+					FGraphBuilder::RefreshPinMetadata(*InPin, InputHandle->GetMetadata());
+				}
+			}
+			else if (InPin->Direction == EGPD_Output)
+			{
+				Frontend::FConstOutputHandle OutputHandle = FGraphBuilder::GetConstOutputHandleFromPin(InPin);
+				if (OutputHandle->GetMetadata().bIsAdvancedDisplay != InPin->bAdvancedView)
+				{				
+					FGraphBuilder::RefreshPinMetadata(*InPin, OutputHandle->GetMetadata());
+				}
+			}
 
 			const bool bShowPin = ShouldPinBeHidden(InPin);
 			if (bShowPin)
diff --git a/Engine/Plugins/Runtime/Metasound/Source/MetasoundEditor/Public/MetasoundEditorGraph.h b/Engine/Plugins/Runtime/Metasound/Source/MetasoundEditor/Public/MetasoundEditorGraph.h
--- a/Engine/Plugins/Runtime/Metasound/Source/MetasoundEditor/Public/MetasoundEditorGraph.h
+++ b/Engine/Plugins/Runtime/Metasound/Source/MetasoundEditor/Public/MetasoundEditorGraph.h
@@ -245,6 +245,11 @@ public:
 	/** Sets the VertexAccessType assigned to this member. */
 	virtual void SetVertexAccessType(EMetasoundFrontendVertexAccessType InNewAccessType, bool bPostTransaction = true) PURE_VIRTUAL(UMetasoundEditorGraphVertex::SetVertexAccessType, )
 
+#if WITH_EDITORONLY_DATA
+	/**Sets if the Vertex is shown as Advanced Display*/
+	bool SetIsAdvancedDisplay(const bool IsAdvancedDisplay);
+#endif // WITH_EDITORONLY_DATA
+
 	/** Returns the node handle associated with the vertex. */
 	Metasound::Frontend::FNodeHandle GetNodeHandle();
 
diff --git a/Engine/Plugins/Runtime/Metasound/Source/MetasoundFrontend/Private/MetasoundFrontendDocument.cpp b/Engine/Plugins/Runtime/Metasound/Source/MetasoundFrontend/Private/MetasoundFrontendDocument.cpp
--- a/Engine/Plugins/Runtime/Metasound/Source/MetasoundFrontend/Private/MetasoundFrontendDocument.cpp
+++ b/Engine/Plugins/Runtime/Metasound/Source/MetasoundFrontend/Private/MetasoundFrontendDocument.cpp
@@ -328,7 +328,12 @@ void FMetasoundFrontendClassVertex::SplitName(FName& OutNamespace, FName& OutPar
 
 bool FMetasoundFrontendClassVertex::IsFunctionalEquivalent(const FMetasoundFrontendClassVertex& InLHS, const FMetasoundFrontendClassVertex& InRHS)
 {
-	return FMetasoundFrontendVertex::IsFunctionalEquivalent(InLHS, InRHS) && (InLHS.AccessType == InRHS.AccessType);
+	bool bEquivalentAdvancedDisplay = true;
+#if WITH_EDITORONLY_DATA
+	bEquivalentAdvancedDisplay = InLHS.Metadata.bIsAdvancedDisplay == InRHS.Metadata.bIsAdvancedDisplay;
+#endif // WITH_EDITORONLY_DATA
+
+	return FMetasoundFrontendVertex::IsFunctionalEquivalent(InLHS, InRHS) && (InLHS.AccessType == InRHS.AccessType) && bEquivalentAdvancedDisplay;
 }
 
 bool FMetasoundFrontendClassVertex::CanConnectVertexAccessTypes(EMetasoundFrontendVertexAccessType InFromType, EMetasoundFrontendVertexAccessType InToType)
diff --git a/Engine/Plugins/Runtime/Metasound/Source/MetasoundFrontend/Private/MetasoundFrontendDocumentBuilder.cpp b/Engine/Plugins/Runtime/Metasound/Source/MetasoundFrontend/Private/MetasoundFrontendDocumentBuilder.cpp
--- a/Engine/Plugins/Runtime/Metasound/Source/MetasoundFrontend/Private/MetasoundFrontendDocumentBuilder.cpp
+++ b/Engine/Plugins/Runtime/Metasound/Source/MetasoundFrontend/Private/MetasoundFrontendDocumentBuilder.cpp
@@ -2395,6 +2395,33 @@ EMetasoundFrontendVertexAccessType FMetaSoundFrontendDocumentBuilder::GetNodeOut
 	return EMetasoundFrontendVertexAccessType::Unset;
 }
 
+#if WITH_EDITORONLY_DATA
+const bool FMetaSoundFrontendDocumentBuilder::GetIsAdvancedDisplay(const FName MemberName, const EMetasoundFrontendClassType Type) const
+{
+	const FMetasoundFrontendDocument& Document = GetConstDocumentChecked();
+
+	//Input
+	if (Type == EMetasoundFrontendClassType::Input)
+	{
+		if (const int32* Index = DocumentCache->GetInterfaceCache().FindInputIndex(MemberName))
+		{
+			const FMetasoundFrontendClassInput& GraphInput = Document.RootGraph.Interface.Inputs[*Index];
+			return GraphInput.Metadata.bIsAdvancedDisplay;
+		}
+	}
+	//Output
+	else if (Type == EMetasoundFrontendClassType::Output)
+	{
+		if (const int32* Index = DocumentCache->GetInterfaceCache().FindOutputIndex(MemberName))
+		{
+			const FMetasoundFrontendClassOutput& GraphOutput = Document.RootGraph.Interface.Outputs[*Index];
+			return GraphOutput.Metadata.bIsAdvancedDisplay;
+		}
+	}
+	return false;
+}
+#endif // WITH_EDITORONLY_DATA
+
 void FMetaSoundFrontendDocumentBuilder::InitDocument(const FMetasoundFrontendDocument* InDocumentTemplate, const FMetasoundFrontendClassName* InNewClassName, bool bResetVersion)
 {
 	using namespace Metasound;
@@ -3504,6 +3531,25 @@ bool FMetaSoundFrontendDocumentBuilder::SetBuildPageID(const FGuid& InBuildPageI
 
 	return false;
 }
+
+bool FMetaSoundFrontendDocumentBuilder::SetGraphInputAdvancedDisplay(const FName InputName, const bool InAdvancedDisplay)
+{
+	FMetasoundFrontendDocument& Document = GetDocumentChecked();
+	FMetasoundFrontendGraphClass& RootGraph = Document.RootGraph;
+
+	if (const int32* Index = DocumentCache->GetInterfaceCache().FindInputIndex(InputName))
+	{
+		FMetasoundFrontendClassInput& GraphInput = RootGraph.Interface.Inputs[*Index];			
+		if (GraphInput.Metadata.bIsAdvancedDisplay != InAdvancedDisplay)
+		{
+			GraphInput.Metadata.SetIsAdvancedDisplay(InAdvancedDisplay);
+			Document.Metadata.ModifyContext.AddMemberIDModified(GraphInput.VertexID);
+			return true;
+		}
+	}
+
+	return false;
+}
 #endif // WITH_EDITORONLY_DATA
 
 bool FMetaSoundFrontendDocumentBuilder::SetGraphInputAccessType(FName InputName, EMetasoundFrontendVertexAccessType AccessType)
@@ -3735,6 +3781,27 @@ bool FMetaSoundFrontendDocumentBuilder::SetGraphInputDefaults(FName InputName, T
 	return false;
 }
 
+#if WITH_EDITORONLY_DATA
+bool FMetaSoundFrontendDocumentBuilder::SetGraphOutputAdvancedDisplay(const FName OutputName, const bool InAdvancedDisplay)
+{
+	FMetasoundFrontendDocument& Document = GetDocumentChecked();
+	FMetasoundFrontendGraphClass& RootGraph = Document.RootGraph;
+
+	if (const int32* Index = DocumentCache->GetInterfaceCache().FindOutputIndex(OutputName))
+	{
+		FMetasoundFrontendClassOutput& GraphOutput = Document.RootGraph.Interface.Outputs[*Index];
+		if (GraphOutput.Metadata.bIsAdvancedDisplay != InAdvancedDisplay)
+		{
+			GraphOutput.Metadata.SetIsAdvancedDisplay(InAdvancedDisplay);
+			Document.Metadata.ModifyContext.AddMemberIDModified(GraphOutput.VertexID);	
+			return true;
+		}
+	}
+
+	return false;
+}
+#endif // WITH_EDITORONLY_DATA
+
 bool FMetaSoundFrontendDocumentBuilder::SetGraphInputInheritsDefault(FName InName, bool bInputInheritsDefault)
 {
 	FMetasoundFrontendGraphClassPresetOptions& PresetOptions = GetDocumentChecked().RootGraph.PresetOptions;
diff --git a/Engine/Plugins/Runtime/Metasound/Source/MetasoundFrontend/Public/MetasoundFrontendDocument.h b/Engine/Plugins/Runtime/Metasound/Source/MetasoundFrontend/Public/MetasoundFrontendDocument.h
--- a/Engine/Plugins/Runtime/Metasound/Source/MetasoundFrontend/Public/MetasoundFrontendDocument.h
+++ b/Engine/Plugins/Runtime/Metasound/Source/MetasoundFrontend/Public/MetasoundFrontendDocument.h
@@ -843,6 +843,11 @@ public:
 		GetDisplayName() = InText;
 	}
 
+	void SetIsAdvancedDisplay(const bool InIsAdvancedDisplay)
+	{		
+		bIsAdvancedDisplay = InIsAdvancedDisplay;
+	}
+
 	void SetSerializeText(bool bInSerializeText)
 	{
 		if (bSerializeText)
@@ -885,6 +890,8 @@ struct METASOUNDFRONTEND_API FMetasoundFrontendClassVertex : public FMetasoundFr
 	// Metadata associated with vertex.
 	UPROPERTY(EditAnywhere, Category = CustomView)
 	FMetasoundFrontendVertexMetadata Metadata;
+
+	const bool GetIsAdvancedDisplay() const { return Metadata.bIsAdvancedDisplay; };
 #endif // WITH_EDITORONLY_DATA
 
 	UPROPERTY()
diff --git a/Engine/Plugins/Runtime/Metasound/Source/MetasoundFrontend/Public/MetasoundFrontendDocumentBuilder.h b/Engine/Plugins/Runtime/Metasound/Source/MetasoundFrontend/Public/MetasoundFrontendDocumentBuilder.h
--- a/Engine/Plugins/Runtime/Metasound/Source/MetasoundFrontend/Public/MetasoundFrontendDocumentBuilder.h
+++ b/Engine/Plugins/Runtime/Metasound/Source/MetasoundFrontend/Public/MetasoundFrontendDocumentBuilder.h
@@ -304,6 +304,10 @@ public:
 
 	EMetasoundFrontendVertexAccessType GetNodeOutputAccessType(const FGuid& InNodeID, const FGuid& InVertexID, const FGuid* InPageID = nullptr) const;
 
+#if WITH_EDITORONLY_DATA
+	const bool GetIsAdvancedDisplay(const FName MemberName, const EMetasoundFrontendClassType Type) const;
+#endif // WITH_EDITORONLY_DATA
+
 	// Initializes the builder's document, using the (optional) provided document template, (optional) class name, and (optionally) whether or not to reset the existing class version.
 	void InitDocument(const FMetasoundFrontendDocument* InDocumentTemplate = nullptr, const FMetasoundFrontendClassName* InNewClassName = nullptr, bool bResetVersion = true);
 
@@ -398,6 +402,10 @@ public:
 	// Swapping the targeted build graph ID clears the local cache, so swapping frequently can
 	// induce cash thrashing.
 	bool SetBuildPageID(const FGuid& InBuildPageID);
+
+	// Sets the given input`s IsAdvancedDisplay state. AdvancedDisplay pins are hidden in the node by default.
+	// returns true if state was changed.
+	bool SetGraphInputAdvancedDisplay(const FName InputName, const bool InAdvancedDisplay);
 #endif // WITH_EDITORONLY_DATA
 
 	// Sets the given graph input's access type. If connected to other nodes and access type is not compatible,
@@ -413,6 +421,12 @@ public:
 	bool SetGraphInputDefault(FName InputName, FMetasoundFrontendLiteral InDefaultLiteral, const FGuid* InPageID = nullptr);
 	bool SetGraphInputDefaults(FName InputName, TArray<FMetasoundFrontendClassInputDefault> Defaults);
 
+#if WITH_EDITORONLY_DATA
+	// Sets the given output`s IsAdvancedDisplay state. AdvancedDisplay pins are hidden in the node by default.
+	// returns true if state was changed.
+	bool SetGraphOutputAdvancedDisplay(const FName OutputName, const bool InAdvancedDisplay);
+#endif // WITH_EDITORONLY_DATA
+
 	// Sets the given graph output's access type. If connected to other nodes and access type is not compatible,
 	// associated edges/connections are removed.  Returns true if either DataType was successfully set to new
 	// value or if AccessType is already the given AccessType.
