Restrict target page to existing page names and reset if page provided is removed from the array of valid pages (trivial)
#rnx

diff --git a/Engine/Plugins/Runtime/Metasound/Source/MetasoundEngine/Private/MetasoundSettings.cpp b/Engine/Plugins/Runtime/Metasound/Source/MetasoundEngine/Private/MetasoundSettings.cpp
--- a/Engine/Plugins/Runtime/Metasound/Source/MetasoundEngine/Private/MetasoundSettings.cpp
+++ b/Engine/Plugins/Runtime/Metasound/Source/MetasoundEngine/Private/MetasoundSettings.cpp
@@ -153,6 +153,7 @@ void UMetaSoundSettings::ConformPageSettingsDefault(bool bNotifyDefaultConformed
 
 	bool bContainsPageDefault = false;
 	bool bDefaultConformed = false;
+	bool bTargetFound = TargetPageName == Frontend::DefaultPageName;
 	for (int32 Index = PageSettings.Num() - 1; Index >= 0; --Index)
 	{
 		FMetaSoundPageSettings& Page = PageSettings[Index];
@@ -175,6 +176,7 @@ void UMetaSoundSettings::ConformPageSettingsDefault(bool bNotifyDefaultConformed
 				bDefaultConformed = true;
 			}
 		}
+		bTargetFound |= TargetPageName == Page.Name;
 	}
 
 	if (!bContainsPageDefault)
@@ -185,6 +187,11 @@ void UMetaSoundSettings::ConformPageSettingsDefault(bool bNotifyDefaultConformed
 		bDefaultConformed = true;
 	}
 
+	if (!bTargetFound)
+	{
+		TargetPageName = Frontend::DefaultPageName;
+	}
+
 	if (bNotifyDefaultConformed && bDefaultConformed)
 	{
 		OnDefaultConformed.Broadcast();
diff --git a/Engine/Plugins/Runtime/Metasound/Source/MetasoundEngine/Public/MetasoundSettings.h b/Engine/Plugins/Runtime/Metasound/Source/MetasoundEngine/Public/MetasoundSettings.h
--- a/Engine/Plugins/Runtime/Metasound/Source/MetasoundEngine/Public/MetasoundSettings.h
+++ b/Engine/Plugins/Runtime/Metasound/Source/MetasoundEngine/Public/MetasoundSettings.h
@@ -142,7 +142,7 @@ private:
 	/** Page Name to target when attempting to execute MetaSound. If target page is not implemented (or cooked in a runtime build)
 	  * for the active platform, uses order of cooked pages (see 'Page Settings' for order) falling back to lower index-ordered page
 	  * implemented in MetaSound asset. */
-	UPROPERTY(EditAnywhere, config, Category = Pages)
+	UPROPERTY(EditAnywhere, config, Category = Pages, meta = (GetOptions = "MetasoundEngine.MetaSoundSettings.GetPageNames"))
 	FName TargetPageName = Metasound::Frontend::DefaultPageName;
 
 	/** Array of possible page settings that can be added to a MetaSound object. Order
