- Fix up missing/newly added calls not respecting page IDs in document builder API (input/output vertex swapping, comment nodes, etc.)

#rnx

diff --git a/Engine/Plugins/Runtime/Metasound/Source/MetasoundEditor/Private/MetasoundEditorSubsystem.cpp b/Engine/Plugins/Runtime/Metasound/Source/MetasoundEditor/Private/MetasoundEditorSubsystem.cpp
--- a/Engine/Plugins/Runtime/Metasound/Source/MetasoundEditor/Private/MetasoundEditorSubsystem.cpp
+++ b/Engine/Plugins/Runtime/Metasound/Source/MetasoundEditor/Private/MetasoundEditorSubsystem.cpp
@@ -263,7 +263,8 @@ void UMetaSoundEditorSubsystem::SetFocusedPage(UMetaSoundBuilderBase* Builder, F
 	check(Settings);
 	if (const FMetaSoundPageSettings* PageSettings = Settings->FindPageSettings(PageName))
 	{
-		const bool bFocusedPage = SetFocusedPageInternal(PageSettings->Name, PageSettings->UniqueId, *Builder, bOpenEditor);
+		constexpr bool bPostTransaction = true;
+		const bool bFocusedPage = SetFocusedPageInternal(PageSettings->Name, PageSettings->UniqueId, *Builder, bOpenEditor, bPostTransaction);
 		if (bFocusedPage)
 		{
 			OutResult = EMetaSoundBuilderResult::Succeeded;
@@ -274,7 +275,7 @@ void UMetaSoundEditorSubsystem::SetFocusedPage(UMetaSoundBuilderBase* Builder, F
 	OutResult = EMetaSoundBuilderResult::Failed;
 }
 
-bool UMetaSoundEditorSubsystem::SetFocusedPage(UMetaSoundBuilderBase& Builder, const FGuid& InPageID, bool bOpenEditor) const
+bool UMetaSoundEditorSubsystem::SetFocusedPage(UMetaSoundBuilderBase& Builder, const FGuid& InPageID, bool bOpenEditor, bool bPostTransaction) const
 {
 	using namespace Metasound::Frontend;
 
@@ -287,19 +288,16 @@ bool UMetaSoundEditorSubsystem::SetFocusedPage(UMetaSoundBuilderBase& Builder, c
 		PageName = PageSettings->Name;
 	}
 
-	return SetFocusedPageInternal(PageName, InPageID, Builder, bOpenEditor);
+	return SetFocusedPageInternal(PageName, InPageID, Builder, bOpenEditor, bPostTransaction);
 }
 
-bool UMetaSoundEditorSubsystem::SetFocusedPageInternal(FName PageName, const FGuid& InPageID, UMetaSoundBuilderBase& Builder, bool bOpenEditor) const
+bool UMetaSoundEditorSubsystem::SetFocusedPageInternal(FName PageName, const FGuid& InPageID, UMetaSoundBuilderBase& Builder, bool bOpenEditor, bool bPostTransaction) const
 {
 	using namespace Metasound::Frontend;
 
-	const FScopedTransaction Transaction(FText::Format(LOCTEXT("SetFocusedPageTransactionFormat", "Set Focused Page '{0}'"), FText::FromName(PageName)));
-	Builder.Modify();
-
-	if (UMetasoundEditorSettings* EditorSettings = GetMutableDefault<UMetasoundEditorSettings>())
+	if (!PageName.IsNone())
 	{
-		if (!PageName.IsNone())
+		if (UMetasoundEditorSettings* EditorSettings = GetMutableDefault<UMetasoundEditorSettings>())
 		{
 			if (EditorSettings->AuditionPageMode == EAuditionPageMode::Focused)
 			{
@@ -310,17 +308,24 @@ bool UMetaSoundEditorSubsystem::SetFocusedPageInternal(FName PageName, const FGu
 			}
 		}
 
-		if (Builder.GetBuilder().SetBuildPageID(InPageID))
+		const FMetaSoundFrontendDocumentBuilder& DocBuilder = Builder.GetConstBuilder();
+		if (DocBuilder.GetBuildPageID() != InPageID)
 		{
-			// Reregister to ensure all future audible instances are using the new page implementation.
-			RegisterGraphWithFrontend(Builder.GetBuilder().CastDocumentObjectChecked<UObject>());
-		}
+			const FScopedTransaction Transaction(FText::Format(LOCTEXT("SetFocusedPageTransactionFormat", "Set Focused Page '{0}'"), FText::FromName(PageName)), bPostTransaction);
+			Builder.Modify();
+			UObject& MetaSound = DocBuilder.CastDocumentObjectChecked<UObject>();
+			if (Builder.GetBuilder().SetBuildPageID(InPageID))
+			{
+				// Reregister to ensure all future audible instances are using the new page implementation.
+				RegisterGraphWithFrontend(MetaSound);
+			}
 
-		if (GEditor && bOpenEditor)
-		{
-			if (UAssetEditorSubsystem* AssetEditorSubsystem = GEditor->GetEditorSubsystem<UAssetEditorSubsystem>())
+			if (GEditor && bOpenEditor)
 			{
-				AssetEditorSubsystem->OpenEditorForAsset(&Builder.GetConstBuilder().CastDocumentObjectChecked<UObject>());
+				if (UAssetEditorSubsystem* AssetEditorSubsystem = GEditor->GetEditorSubsystem<UAssetEditorSubsystem>())
+				{
+					AssetEditorSubsystem->OpenEditorForAsset(&MetaSound);
+				}
 			}
 		}
 
diff --git a/Engine/Plugins/Runtime/Metasound/Source/MetasoundEditor/Public/MetasoundEditorSubsystem.h b/Engine/Plugins/Runtime/Metasound/Source/MetasoundEditor/Public/MetasoundEditorSubsystem.h
--- a/Engine/Plugins/Runtime/Metasound/Source/MetasoundEditor/Public/MetasoundEditorSubsystem.h
+++ b/Engine/Plugins/Runtime/Metasound/Source/MetasoundEditor/Public/MetasoundEditorSubsystem.h
@@ -89,7 +89,7 @@ public:
 	// the provided builder to the given PageID if and sets the audition target page to
 	// the provided ID. If the given builder has an asset editor open, optionally opens or brings
 	// that editor's associated PageID into user focus.
-	bool SetFocusedPage(UMetaSoundBuilderBase& Builder, const FGuid& InPageID, bool bOpenEditor) const;
+	bool SetFocusedPage(UMetaSoundBuilderBase& Builder, const FGuid& InPageID, bool bOpenEditor, bool bPostTransaction = true) const;
 
 	// Unregisters toolbar extender that is displayed in the MetaSound Asset Editor.
 	bool UnregisterToolbarExtender(TSharedRef<FExtender> InExtender);
@@ -104,7 +104,7 @@ public:
 	static const UMetaSoundEditorSubsystem& GetConstChecked();
 
 private:
-	bool SetFocusedPageInternal(FName PageName, const FGuid& InPageID, UMetaSoundBuilderBase& Builder, bool bOpenEditor) const;
+	bool SetFocusedPageInternal(FName PageName, const FGuid& InPageID, UMetaSoundBuilderBase& Builder, bool bOpenEditor, bool bPostTransaction) const;
 
 	// Copy over sound wave settings such as attenuation, modulation, and sound class from the template sound wave to the MetaSound
 	void SetSoundWaveSettingsFromTemplate(USoundWave& NewMetasound, const USoundWave& TemplateSoundWave) const;
diff --git a/Engine/Plugins/Runtime/Metasound/Source/MetasoundFrontend/Private/MetasoundFrontendDocumentBuilder.cpp b/Engine/Plugins/Runtime/Metasound/Source/MetasoundFrontend/Private/MetasoundFrontendDocumentBuilder.cpp
--- a/Engine/Plugins/Runtime/Metasound/Source/MetasoundFrontend/Private/MetasoundFrontendDocumentBuilder.cpp
+++ b/Engine/Plugins/Runtime/Metasound/Source/MetasoundFrontend/Private/MetasoundFrontendDocumentBuilder.cpp
@@ -506,32 +506,35 @@ namespace Metasound::Frontend
 #if WITH_EDITORONLY_DATA
 				if (bAddedVertices && Options.bSetDefaultNodeLocations && !IsRunningCookCommandlet())
 				{
-					TArray<FMetasoundFrontendNode>& Nodes = OutBuilder.FindBuildGraphChecked().Nodes;
-					// Sort/Place Inputs
+					Document.RootGraph.IterateGraphPages([&](FMetasoundFrontendGraph& Graph)
 					{
-						TSet<FName> NamesToSort;
-						Algo::Transform(InputsToAdd, NamesToSort, [](const FInputInterfacePair& Pair) { return Pair.Key.Name; });
-						auto GetInputSortOrder = [&OutBuilder](const FVertexName& InVertexName)
+						TArray<FMetasoundFrontendNode>& Nodes = Graph.Nodes;
+						// Sort/Place Inputs
 						{
-							const FMetasoundFrontendClassInput* Input = OutBuilder.FindGraphInput(InVertexName);
-							checkf(Input, TEXT("Input must exist by this point of modifying the document's interfaces and respective members"));
-							return Input->Metadata.SortOrderIndex;
-						};
-						UpdateAddedVertexNodePositions(EMetasoundFrontendClassType::Input, OutBuilder, NamesToSort, GetInputSortOrder, FVector2D::Zero(), Nodes);
-					}
+							TSet<FName> NamesToSort;
+							Algo::Transform(InputsToAdd, NamesToSort, [](const FInputInterfacePair& Pair) { return Pair.Key.Name; });
+							auto GetInputSortOrder = [&OutBuilder](const FVertexName& InVertexName)
+							{
+								const FMetasoundFrontendClassInput* Input = OutBuilder.FindGraphInput(InVertexName);
+								checkf(Input, TEXT("Input must exist by this point of modifying the document's interfaces and respective members"));
+								return Input->Metadata.SortOrderIndex;
+							};
+							UpdateAddedVertexNodePositions(EMetasoundFrontendClassType::Input, OutBuilder, NamesToSort, GetInputSortOrder, FVector2D::Zero(), Nodes);
+						}
 
-					// Sort/Place Outputs
-					{
-						TSet<FName> NamesToSort;
-						Algo::Transform(OutputsToAdd, NamesToSort, [](const FOutputInterfacePair& OutputInterfacePair) { return OutputInterfacePair.Key.Name; });
-						auto GetOutputSortOrder = [&OutBuilder](const FVertexName& InVertexName)
+						// Sort/Place Outputs
 						{
-							const FMetasoundFrontendClassOutput* Output = OutBuilder.FindGraphOutput(InVertexName);
-							checkf(Output, TEXT("Output must exist by this point of modifying the document's interfaces and respective members"));
-							return Output->Metadata.SortOrderIndex;
-						};
-						UpdateAddedVertexNodePositions(EMetasoundFrontendClassType::Output, OutBuilder, NamesToSort, GetOutputSortOrder, 3 * DisplayStyle::NodeLayout::DefaultOffsetX, Nodes);
-					}
+							TSet<FName> NamesToSort;
+							Algo::Transform(OutputsToAdd, NamesToSort, [](const FOutputInterfacePair& OutputInterfacePair) { return OutputInterfacePair.Key.Name; });
+							auto GetOutputSortOrder = [&OutBuilder](const FVertexName& InVertexName)
+							{
+								const FMetasoundFrontendClassOutput* Output = OutBuilder.FindGraphOutput(InVertexName);
+								checkf(Output, TEXT("Output must exist by this point of modifying the document's interfaces and respective members"));
+								return Output->Metadata.SortOrderIndex;
+							};
+							UpdateAddedVertexNodePositions(EMetasoundFrontendClassType::Output, OutBuilder, NamesToSort, GetOutputSortOrder, 3 * DisplayStyle::NodeLayout::DefaultOffsetX, Nodes);
+						}
+					});
 				}
 #endif // WITH_EDITORONLY_DATA
 
@@ -1685,15 +1688,20 @@ TArray<const FMetasoundFrontendEdge*> FMetaSoundFrontendDocumentBuilder::FindEdg
 }
 
 #if WITH_EDITOR
-const FMetaSoundFrontendGraphComment* FMetaSoundFrontendDocumentBuilder::FindGraphComment(const FGuid& InCommentID) const
+const FMetaSoundFrontendGraphComment* FMetaSoundFrontendDocumentBuilder::FindGraphComment(const FGuid& InCommentID, const FGuid* InPageID) const
 {
-	const TMap<FGuid, FMetaSoundFrontendGraphComment>& Comments = FindConstBuildGraphChecked().Style.Comments;
+	check(InCommentID.IsValid());
+	const FGuid& PageID = InPageID ? *InPageID : BuildPageID;
+	const FMetasoundFrontendDocument& Document = GetConstDocumentChecked();
+	const TMap<FGuid, FMetaSoundFrontendGraphComment>& Comments = Document.RootGraph.FindConstGraphChecked(PageID).Style.Comments;
 	return Comments.Find(InCommentID);
 }
 
-FMetaSoundFrontendGraphComment* FMetaSoundFrontendDocumentBuilder::FindGraphComment(const FGuid& InCommentID)
+FMetaSoundFrontendGraphComment* FMetaSoundFrontendDocumentBuilder::FindGraphComment(const FGuid& InCommentID, const FGuid* InPageID)
 {
-	TMap<FGuid, FMetaSoundFrontendGraphComment>& Comments = FindBuildGraphChecked().Style.Comments;
+	const FGuid& PageID = InPageID ? *InPageID : BuildPageID;
+	FMetasoundFrontendDocument& Document = GetDocumentChecked();
+	TMap<FGuid, FMetaSoundFrontendGraphComment>& Comments = Document.RootGraph.FindGraphChecked(PageID).Style.Comments;
 	return Comments.Find(InCommentID);
 }
 #endif // WITH_EDITOR
@@ -1873,7 +1881,7 @@ const FMetasoundFrontendVertex* FMetaSoundFrontendDocumentBuilder::FindNodeInput
 {
 	using namespace Metasound::Frontend;
 	const FGuid& PageID = InPageID ? *InPageID : BuildPageID;
-	const IDocumentGraphNodeCache& NodeCache = DocumentCache->GetNodeCache();
+	const IDocumentGraphNodeCache& NodeCache = DocumentCache->GetNodeCache(PageID);
 	return NodeCache.FindInputVertex(InNodeID, InVertexID);
 }
 
@@ -2079,10 +2087,12 @@ const FMetasoundFrontendVertex* FMetaSoundFrontendDocumentBuilder::FindNodeOutpu
 }
 
 #if WITH_EDITOR
-FMetaSoundFrontendGraphComment& FMetaSoundFrontendDocumentBuilder::FindOrAddGraphComment(const FGuid& InCommentID)
+FMetaSoundFrontendGraphComment& FMetaSoundFrontendDocumentBuilder::FindOrAddGraphComment(const FGuid& InCommentID, const FGuid* InPageID)
 {
 	check(InCommentID.IsValid());
-	TMap<FGuid, FMetaSoundFrontendGraphComment>& Comments = FindBuildGraphChecked().Style.Comments;
+	const FGuid& PageID = InPageID ? *InPageID : BuildPageID;
+	FMetasoundFrontendDocument& Document = GetDocumentChecked();
+	TMap<FGuid, FMetaSoundFrontendGraphComment>& Comments = Document.RootGraph.FindGraphChecked(PageID).Style.Comments;
 	return Comments.FindOrAdd(InCommentID);
 }
 #endif // WITH_EDITOR
@@ -2513,48 +2523,52 @@ void FMetaSoundFrontendDocumentBuilder::InitNodeLocations()
 	FVector2D ExternalNodeLocation = InputNodeLocation + DisplayStyle::NodeLayout::DefaultOffsetX;
 	FVector2D OutputNodeLocation = ExternalNodeLocation + DisplayStyle::NodeLayout::DefaultOffsetX;
 
-	TArray<FMetasoundFrontendNode>& Nodes = FindBuildGraphChecked().Nodes;
-	for (FMetasoundFrontendNode& Node : Nodes)
+	FMetasoundFrontendDocument& Document = GetDocumentChecked();
+	Document.RootGraph.IterateGraphPages([&](FMetasoundFrontendGraph& Graph)
 	{
-		if (const int32* ClassIndex = DocumentCache->FindDependencyIndex(Node.ClassID))
+		TArray<FMetasoundFrontendNode>& Nodes = Graph.Nodes;
+		for (FMetasoundFrontendNode& Node : Nodes)
 		{
-			FMetasoundFrontendClass& Class = GetDocumentChecked().Dependencies[*ClassIndex];
-
-			const EMetasoundFrontendClassType NodeType = Class.Metadata.GetType();
-			FVector2D NewLocation;
-			if (NodeType == EMetasoundFrontendClassType::Input)
+			if (const int32* ClassIndex = DocumentCache->FindDependencyIndex(Node.ClassID))
 			{
-				NewLocation = InputNodeLocation;
-				InputNodeLocation += DisplayStyle::NodeLayout::DefaultOffsetY;
-			}
-			else if (NodeType == EMetasoundFrontendClassType::Output)
-			{
-				NewLocation = OutputNodeLocation;
-				OutputNodeLocation += DisplayStyle::NodeLayout::DefaultOffsetY;
-			}
-			else
-			{
-				NewLocation = ExternalNodeLocation;
-				ExternalNodeLocation += DisplayStyle::NodeLayout::DefaultOffsetY;
-			}
+				FMetasoundFrontendClass& Class = Document.Dependencies[*ClassIndex];
 
-			// TODO: Find consistent location for controlling node locations.
-			// Currently it is split between MetasoundEditor and MetasoundFrontend modules.
-			FMetasoundFrontendNodeStyle& Style = Node.Style;
-			if (Style.Display.Locations.IsEmpty())
-			{
-				Style.Display.Locations = { { FGuid::NewGuid(), NewLocation } };
-			}
-			// Initialize the position if the location hasn't been assigned yet.  This can happen
-			// if default interfaces were assigned to the given MetaSound but not placed with respect
-			// to one another.  In this case, node location initialization takes "priority" to avoid
-			// visual overlap.
-			else if (Style.Display.Locations.Num() == 1 && Style.Display.Locations.Contains(FGuid()))
-			{
-				Style.Display.Locations = { { FGuid::NewGuid(), NewLocation } };
+				const EMetasoundFrontendClassType NodeType = Class.Metadata.GetType();
+				FVector2D NewLocation;
+				if (NodeType == EMetasoundFrontendClassType::Input)
+				{
+					NewLocation = InputNodeLocation;
+					InputNodeLocation += DisplayStyle::NodeLayout::DefaultOffsetY;
+				}
+				else if (NodeType == EMetasoundFrontendClassType::Output)
+				{
+					NewLocation = OutputNodeLocation;
+					OutputNodeLocation += DisplayStyle::NodeLayout::DefaultOffsetY;
+				}
+				else
+				{
+					NewLocation = ExternalNodeLocation;
+					ExternalNodeLocation += DisplayStyle::NodeLayout::DefaultOffsetY;
+				}
+
+				// TODO: Find consistent location for controlling node locations.
+				// Currently it is split between MetasoundEditor and MetasoundFrontend modules.
+				FMetasoundFrontendNodeStyle& Style = Node.Style;
+				if (Style.Display.Locations.IsEmpty())
+				{
+					Style.Display.Locations = { { FGuid::NewGuid(), NewLocation } };
+				}
+				// Initialize the position if the location hasn't been assigned yet.  This can happen
+				// if default interfaces were assigned to the given MetaSound but not placed with respect
+				// to one another.  In this case, node location initialization takes "priority" to avoid
+				// visual overlap.
+				else if (Style.Display.Locations.Num() == 1 && Style.Display.Locations.Contains(FGuid()))
+				{
+					Style.Display.Locations = { { FGuid::NewGuid(), NewLocation } };
+				}
 			}
 		}
-	}
+	});
 #endif // WITH_EDITORONLY_DATA
 }
 
@@ -2823,55 +2837,57 @@ bool FMetaSoundFrontendDocumentBuilder::RemoveDependency(const FGuid& InClassID)
 {
 	using namespace Metasound::Frontend;
 
-	FMetasoundFrontendDocument& Document = GetDocumentChecked();
-	const IDocumentGraphNodeCache& NodeCache = DocumentCache->GetNodeCache();
-
+	bool bSuccess = false;
 	if (const int32* IndexPtr = DocumentCache->FindDependencyIndex(InClassID))
 	{
+		FMetasoundFrontendDocument& Document = GetDocumentChecked();
 		TArray<FMetasoundFrontendClass>& Dependencies = Document.Dependencies;
 		const int32 Index = *IndexPtr;
 
-		TArray<const FMetasoundFrontendNode*> Nodes = NodeCache.FindNodesOfClassID(InClassID);
-		for (const FMetasoundFrontendNode* Node : Nodes)
+		bSuccess = true;
+		Document.RootGraph.IterateGraphPages([this, &bSuccess, &InClassID](const FMetasoundFrontendGraph& Graph)
 		{
-			if (!RemoveNode(Node->GetID()))
+			const IDocumentGraphNodeCache& NodeCache = DocumentCache->GetNodeCache(Graph.PageID);
+			TArray<const FMetasoundFrontendNode*> Nodes = NodeCache.FindNodesOfClassID(InClassID);
+			for (const FMetasoundFrontendNode* Node : Nodes)
 			{
-				return false;
+				bSuccess &= RemoveNode(Node->GetID());
 			}
-		}
+		});
 
 		RemoveSwapDependencyInternal(Index);
 	}
 
-	return true;
+	return bSuccess;
 }
 
 bool FMetaSoundFrontendDocumentBuilder::RemoveDependency(EMetasoundFrontendClassType ClassType, const FMetasoundFrontendClassName& InClassName, const FMetasoundFrontendVersionNumber& InClassVersionNumber)
 {
 	using namespace Metasound::Frontend;
 
-	FMetasoundFrontendDocument& Document = GetDocumentChecked();
-	const IDocumentGraphNodeCache& NodeCache = DocumentCache->GetNodeCache();
-
+	bool bSuccess = false;
 	const FNodeRegistryKey ClassKey(ClassType, InClassName, InClassVersionNumber);
 	if (const int32* IndexPtr = DocumentCache->FindDependencyIndex(ClassKey))
 	{
+		FMetasoundFrontendDocument& Document = GetDocumentChecked();
 		TArray<FMetasoundFrontendClass>& Dependencies = Document.Dependencies;
 		const int32 Index = *IndexPtr;
 
-		TArray<const FMetasoundFrontendNode*> Nodes = NodeCache.FindNodesOfClassID(Dependencies[Index].ID);
-		for (const FMetasoundFrontendNode* Node : Nodes)
+		bSuccess = true;
+		Document.RootGraph.IterateGraphPages([this, &bSuccess, &Dependencies, &Index](const FMetasoundFrontendGraph& Graph)
 		{
-			if (!RemoveNode(Node->GetID()))
+			const IDocumentGraphNodeCache& NodeCache = DocumentCache->GetNodeCache(Graph.PageID);
+			TArray<const FMetasoundFrontendNode*> Nodes = NodeCache.FindNodesOfClassID(Dependencies[Index].ID);
+			for (const FMetasoundFrontendNode* Node : Nodes)
 			{
-				return false;
+				bSuccess &= RemoveNode(Node->GetID());
 			}
-		}
+		});
 
 		RemoveSwapDependencyInternal(Index);
 	}
 
-	return true;
+	return bSuccess;
 }
 
 void FMetaSoundFrontendDocumentBuilder::RemoveSwapDependencyInternal(int32 Index)
@@ -3652,7 +3668,7 @@ bool FMetaSoundFrontendDocumentBuilder::SetGraphInputDataType(FName InputName, F
 
 			RootGraph.IterateGraphPages([this, &DataType, &GraphInput](FMetasoundFrontendGraph& Graph)
 			{
-				const Frontend::IDocumentGraphNodeCache& NodeCache = DocumentCache->GetNodeCache();
+				const Frontend::IDocumentGraphNodeCache& NodeCache = DocumentCache->GetNodeCache(Graph.PageID);
 				if (const int32* NodeIndex = NodeCache.FindNodeIndex(GraphInput.NodeID))
 				{
 					FMetasoundFrontendNode& Node = Graph.Nodes[*NodeIndex];
@@ -3848,10 +3864,10 @@ bool FMetaSoundFrontendDocumentBuilder::SetGraphOutputAccessType(FName OutputNam
 		{
 			RootGraph.IterateGraphPages([this, &GraphOutput, &AccessType](FMetasoundFrontendGraph& Graph)
 			{
-				const IDocumentGraphNodeCache& NodeCache = DocumentCache->GetNodeCache();
+				const IDocumentGraphNodeCache& NodeCache = DocumentCache->GetNodeCache(Graph.PageID);
 				if (const int32* NodeIndex = NodeCache.FindNodeIndex(GraphOutput.NodeID))
 				{
-					FMetasoundFrontendNode& Node = FindBuildGraphChecked().Nodes[*NodeIndex];
+					FMetasoundFrontendNode& Node = Graph.Nodes[*NodeIndex];
 					const FMetasoundFrontendVertex& NodeInput = Node.Interface.Inputs.Last();
 					IterateNodesConnectedWithVertex({ GraphOutput.NodeID, NodeInput.VertexID }, [this, &Graph, &AccessType](const FMetasoundFrontendEdge& Edge, FMetasoundFrontendNode& ConnectedNode)
 					{
@@ -3907,10 +3923,10 @@ bool FMetaSoundFrontendDocumentBuilder::SetGraphOutputDataType(FName OutputName,
 
 		RootGraph.IterateGraphPages([this, &GraphOutput, &DataType](FMetasoundFrontendGraph& Graph)
 		{
-			const IDocumentGraphNodeCache& NodeCache = DocumentCache->GetNodeCache();
+			const IDocumentGraphNodeCache& NodeCache = DocumentCache->GetNodeCache(Graph.PageID);
 			if (const int32* NodeIndex = NodeCache.FindNodeIndex(GraphOutput.NodeID))
 			{
-				FMetasoundFrontendNode& Node = FindBuildGraphChecked().Nodes[*NodeIndex];
+				FMetasoundFrontendNode& Node = Graph.Nodes[*NodeIndex];
 
 				FMetasoundFrontendLiteral DefaultLiteral;
 				DefaultLiteral.SetFromLiteral(IDataTypeRegistry::Get().CreateDefaultLiteral(DataType));
@@ -4160,9 +4176,13 @@ bool FMetaSoundFrontendDocumentBuilder::SwapGraphInput(const FMetasoundFrontendC
 		}
 	}
 
-	const IDocumentGraphNodeCache& NodeCache = DocumentCache->GetNodeCache();
 	const IDocumentGraphInterfaceCache& InterfaceCache = DocumentCache->GetInterfaceCache();
 
+#if WITH_EDITOR
+	using FPageNodeLocations = TMap<FGuid, FVector2D>;
+	TMap<FGuid, FPageNodeLocations> PageNodeLocations;
+#endif // WITH_EDITOR
+
 	// 2. Gather data from existing member/node needed to swap
 	TArray<FMetasoundFrontendEdge> RemovedEdges;
 
@@ -4170,21 +4190,22 @@ bool FMetaSoundFrontendDocumentBuilder::SwapGraphInput(const FMetasoundFrontendC
 	checkf(ExistingInputClass, TEXT("'SwapGraphInput' failed to find original graph input"));
 	const FGuid NodeID = ExistingInputClass->NodeID;
 
-#if WITH_EDITOR
-	TMap<FGuid, FVector2D> Locations;
-#endif // WITH_EDITOR
+	FMetasoundFrontendDocument& Document = GetDocumentChecked();
+	FMetasoundFrontendGraphClass& RootGraph = Document.RootGraph;
+	Document.RootGraph.IterateGraphPages([&](FMetasoundFrontendGraph& Graph)
 	{
+		const IDocumentGraphNodeCache& NodeCache = DocumentCache->GetNodeCache(Graph.PageID);
 		const FMetasoundFrontendNode* ExistingInputNode = NodeCache.FindNode(NodeID);
 		check(ExistingInputNode);
 
 #if WITH_EDITOR
-		Locations = ExistingInputNode->Style.Display.Locations;
+		PageNodeLocations.Add(Graph.PageID, ExistingInputNode->Style.Display.Locations);
 #endif // WITH_EDITOR
 
 		const FGuid VertexID = ExistingInputNode->Interface.Outputs.Last().VertexID;
-		TArray<const FMetasoundFrontendEdge*> Edges = DocumentCache->GetEdgeCache().FindEdges(NodeID, VertexID);
+		TArray<const FMetasoundFrontendEdge*> Edges = DocumentCache->GetEdgeCache(Graph.PageID).FindEdges(NodeID, VertexID);
 		Algo::Transform(Edges, RemovedEdges, [](const FMetasoundFrontendEdge* Edge) { return *Edge; });
-	}
+	});
 
 	// 3. Remove existing graph vertex
 	{
@@ -4201,26 +4222,32 @@ bool FMetaSoundFrontendDocumentBuilder::SwapGraphInput(const FMetasoundFrontendC
 
 	const FMetasoundFrontendNode* NewInputNode = AddGraphInput(NewInput);
 	checkf(NewInputNode, TEXT("Failed to add new Input node when swapping graph inputs"));
+	checkf(NewInputNode->GetID() == NewInput.NodeID, TEXT("Expected new node added to build graph to have same ID as provided input"));
 
-#if WITH_EDITOR
-	// 5a. Add to new copy existing node locations
-	if (!Locations.IsEmpty())
+	Document.RootGraph.IterateGraphPages([&](FMetasoundFrontendGraph& Graph)
 	{
-		const int32* NodeIndex = NodeCache.FindNodeIndex(NewInputNode->GetID());
-		checkf(NodeIndex, TEXT("Cache was not updated to reflect newly added input node"));
-		FMetasoundFrontendNode& NewNode = FindBuildGraphChecked().Nodes[*NodeIndex];
-		NewNode.Style.Display.Locations = Locations;
-	}
+#if WITH_EDITOR
+		// 5a. Add to new copy existing node locations
+		if (const FPageNodeLocations* Locations = PageNodeLocations.Find(Graph.PageID))
+		{
+			const IDocumentGraphNodeCache& NodeCache = DocumentCache->GetNodeCache(Graph.PageID);
+			const int32* NodeIndex = NodeCache.FindNodeIndex(NodeID);
+			checkf(NodeIndex, TEXT("Cache was not updated to reflect newly added input node"));
+			FMetasoundFrontendNode& NewNode = Graph.Nodes[*NodeIndex];
+			NewNode.Style.Display.Locations = *Locations;
+		}
 #endif // WITH_EDITOR
 
-	// 5b. Add to new copy existing node edges
-	for (const FMetasoundFrontendEdge& RemovedEdge : RemovedEdges)
-	{
-		FMetasoundFrontendEdge NewEdge = RemovedEdge;
-		NewEdge.FromNodeID = NewInputNode->GetID();
-		NewEdge.FromVertexID = NewInputNode->Interface.Outputs.Last().VertexID;
-		AddEdge(MoveTemp(NewEdge));
-	}
+		// 5b. Add to new copy existing node edges
+		for (const FMetasoundFrontendEdge& RemovedEdge : RemovedEdges)
+		{
+			FMetasoundFrontendEdge NewEdge = RemovedEdge;
+			NewEdge.FromNodeID = NewInputNode->GetID();
+			NewEdge.FromVertexID = NewInputNode->Interface.Outputs.Last().VertexID;
+			AddEdge(MoveTemp(NewEdge));
+		}
+	});
+
 	return true;
 }
 
@@ -4237,9 +4264,13 @@ bool FMetaSoundFrontendDocumentBuilder::SwapGraphOutput(const FMetasoundFrontend
 		}
 	}
 
-	const IDocumentGraphNodeCache& NodeCache = DocumentCache->GetNodeCache();
 	const IDocumentGraphInterfaceCache& InterfaceCache = DocumentCache->GetInterfaceCache();
 
+#if WITH_EDITOR
+	using FPageNodeLocations = TMap<FGuid, FVector2D>;
+	TMap<FGuid, FPageNodeLocations> PageNodeLocations;
+#endif // WITH_EDITOR
+
 	// 2. Gather data from existing member/node needed to swap
 	TArray<FMetasoundFrontendEdge> RemovedEdges;
 
@@ -4247,21 +4278,22 @@ bool FMetaSoundFrontendDocumentBuilder::SwapGraphOutput(const FMetasoundFrontend
 	checkf(ExistingOutputClass, TEXT("'SwapGraphOutput' failed to find original graph output"));
 	const FGuid NodeID = ExistingOutputClass->NodeID;
 
-#if WITH_EDITOR
-	TMap<FGuid, FVector2D> Locations;
-#endif // WITH_EDITOR
+	FMetasoundFrontendDocument& Document = GetDocumentChecked();
+	FMetasoundFrontendGraphClass& RootGraph = Document.RootGraph;
+	Document.RootGraph.IterateGraphPages([&](FMetasoundFrontendGraph& Graph)
 	{
+		const IDocumentGraphNodeCache& NodeCache = DocumentCache->GetNodeCache(Graph.PageID);
 		const FMetasoundFrontendNode* ExistingOutputNode = NodeCache.FindNode(NodeID);
 		check(ExistingOutputNode);
 
 #if WITH_EDITOR
-		Locations = ExistingOutputNode->Style.Display.Locations;
+		PageNodeLocations.Add(Graph.PageID, ExistingOutputNode->Style.Display.Locations);
 #endif // WITH_EDITOR
 
 		const FGuid VertexID = ExistingOutputNode->Interface.Inputs.Last().VertexID;
-		TArray<const FMetasoundFrontendEdge*> Edges = DocumentCache->GetEdgeCache().FindEdges(ExistingOutputNode->GetID(), VertexID);
+		TArray<const FMetasoundFrontendEdge*> Edges = DocumentCache->GetEdgeCache(Graph.PageID).FindEdges(NodeID, VertexID);
 		Algo::Transform(Edges, RemovedEdges, [](const FMetasoundFrontendEdge* Edge) { return *Edge; });
-	}
+	});
 
 	// 3. Remove existing graph vertex
 	{
@@ -4278,26 +4310,31 @@ bool FMetaSoundFrontendDocumentBuilder::SwapGraphOutput(const FMetasoundFrontend
 
 	const FMetasoundFrontendNode* NewOutputNode = AddGraphOutput(NewOutput);
 	checkf(NewOutputNode, TEXT("Failed to add new output node when swapping graph outputs"));
+	checkf(NewOutputNode->GetID() == NewOutput.NodeID, TEXT("Expected new node added to build graph to have same ID as provided output"));
 
-#if WITH_EDITOR
-	// 5a. Add to new copy existing node locations
-	if (!Locations.IsEmpty())
+	Document.RootGraph.IterateGraphPages([&](FMetasoundFrontendGraph& Graph)
 	{
-		const int32* NodeIndex = NodeCache.FindNodeIndex(NewOutputNode->GetID());
-		checkf(NodeIndex, TEXT("Cache was not updated to reflect newly added input node"));
-		FMetasoundFrontendNode& NewNode = FindBuildGraphChecked().Nodes[*NodeIndex];
-		NewNode.Style.Display.Locations = Locations;
-	}
+#if WITH_EDITOR
+		// 5a. Add to new copy existing node locations
+		if (const FPageNodeLocations* Locations = PageNodeLocations.Find(Graph.PageID))
+		{
+			const IDocumentGraphNodeCache& NodeCache = DocumentCache->GetNodeCache(Graph.PageID);
+			const int32* NodeIndex = NodeCache.FindNodeIndex(NodeID);
+			checkf(NodeIndex, TEXT("Cache was not updated to reflect newly added output node"));
+			FMetasoundFrontendNode& NewNode = Graph.Nodes[*NodeIndex];
+			NewNode.Style.Display.Locations = *Locations;
+		}
 #endif // WITH_EDITOR
 
-	// 5b. Add to new copy existing node edges
-	for (const FMetasoundFrontendEdge& RemovedEdge : RemovedEdges)
-	{
-		FMetasoundFrontendEdge NewEdge = RemovedEdge;
-		NewEdge.ToNodeID = NewOutputNode->GetID();
-		NewEdge.ToVertexID = NewOutputNode->Interface.Inputs.Last().VertexID;
-		AddEdge(MoveTemp(NewEdge));
-	}
+		// 5b. Add to new copy existing node edges
+		for (const FMetasoundFrontendEdge& RemovedEdge : RemovedEdges)
+		{
+			FMetasoundFrontendEdge NewEdge = RemovedEdge;
+			NewEdge.ToNodeID = NewOutputNode->GetID();
+			NewEdge.ToVertexID = NewOutputNode->Interface.Inputs.Last().VertexID;
+			AddEdge(MoveTemp(NewEdge));
+		}
+	});
 
 	return true;
 }
diff --git a/Engine/Plugins/Runtime/Metasound/Source/MetasoundFrontend/Private/MetasoundFrontendDocumentCache.cpp b/Engine/Plugins/Runtime/Metasound/Source/MetasoundFrontend/Private/MetasoundFrontendDocumentCache.cpp
--- a/Engine/Plugins/Runtime/Metasound/Source/MetasoundFrontend/Private/MetasoundFrontendDocumentCache.cpp
+++ b/Engine/Plugins/Runtime/Metasound/Source/MetasoundFrontend/Private/MetasoundFrontendDocumentCache.cpp
@@ -211,11 +211,6 @@ namespace Metasound::Frontend
 		return *DependencyCache;
 	}
 
-	const IDocumentGraphEdgeCache& FDocumentCache::GetEdgeCache() const
-	{
-		return GetEdgeCache(BuildPageID);
-	}
-
 	const IDocumentGraphEdgeCache& FDocumentCache::GetEdgeCache(const FGuid& InPageID) const
 	{
 		TSharedPtr<FDocumentGraphEdgeCache> EdgeCache = EdgeCacheMap.FindRef(InPageID);
@@ -229,11 +224,6 @@ namespace Metasound::Frontend
 		return *EdgeCache.Get();
 	}
 
-	const IDocumentGraphNodeCache& FDocumentCache::GetNodeCache() const
-	{
-		return GetNodeCache(BuildPageID);
-	}
-
 	const IDocumentGraphNodeCache& FDocumentCache::GetNodeCache(const FGuid& InPageID) const
 	{
 		TSharedPtr<FDocumentGraphNodeCache> NodeCache = NodeCacheMap.FindRef(InPageID);
@@ -643,7 +633,7 @@ namespace Metasound::Frontend
 					*bOutIsRerouted = true;
 				}
 
-				const IDocumentGraphEdgeCache& EdgeCache = Parent->GetEdgeCache();
+				const IDocumentGraphEdgeCache& EdgeCache = Parent->GetEdgeCache(PageID);
 				const FMetasoundFrontendVertex& RerouteOutput = Node->Interface.Outputs.Last();
 				TArrayView<const int32> EdgeIndices = EdgeCache.FindEdgeIndicesFromNodeOutput(InNodeID, RerouteOutput.VertexID);
 				const FMetasoundFrontendDocument& Doc = Parent->GetDocument();
@@ -698,7 +688,7 @@ namespace Metasound::Frontend
 					*bOutIsRerouted = true;
 				}
 
-				const IDocumentGraphEdgeCache& EdgeCache = Parent->GetEdgeCache();
+				const IDocumentGraphEdgeCache& EdgeCache = Parent->GetEdgeCache(PageID);
 				const FMetasoundFrontendVertex& RerouteOutput = Node->Interface.Outputs.Last();
 				TArrayView<const int32> EdgeIndices = EdgeCache.FindEdgeIndicesFromNodeOutput(InNodeID, RerouteOutput.VertexID);
 				const FMetasoundFrontendDocument& Doc = Parent->GetDocument();
@@ -752,7 +742,7 @@ namespace Metasound::Frontend
 					*bOutIsRerouted = true;
 				}
 
-				const IDocumentGraphEdgeCache& EdgeCache = Parent->GetEdgeCache();
+				const IDocumentGraphEdgeCache& EdgeCache = Parent->GetEdgeCache(PageID);
 				const FMetasoundFrontendVertex& RerouteInput = Node->Interface.Inputs.Last();
 				const FMetasoundFrontendDocument& Doc = Parent->GetDocument();
 				if (const int32* ConnectedEdgeIndex = EdgeCache.FindEdgeIndexToNodeInput(InNodeID, RerouteInput.VertexID))
@@ -791,7 +781,7 @@ namespace Metasound::Frontend
 					*bOutIsRerouted = true;
 				}
 
-				const IDocumentGraphEdgeCache& EdgeCache = Parent->GetEdgeCache();
+				const IDocumentGraphEdgeCache& EdgeCache = Parent->GetEdgeCache(PageID);
 				const FMetasoundFrontendVertex& RerouteInput = Node->Interface.Inputs.Last();
 				const FMetasoundFrontendDocument& Doc = Parent->GetDocument();
 				if (const int32* ConnectedEdgeIndex = EdgeCache.FindEdgeIndexToNodeInput(InNodeID, RerouteInput.VertexID))
diff --git a/Engine/Plugins/Runtime/Metasound/Source/MetasoundFrontend/Private/MetasoundFrontendDocumentCache.h b/Engine/Plugins/Runtime/Metasound/Source/MetasoundFrontend/Private/MetasoundFrontendDocumentCache.h
--- a/Engine/Plugins/Runtime/Metasound/Source/MetasoundFrontend/Private/MetasoundFrontendDocumentCache.h
+++ b/Engine/Plugins/Runtime/Metasound/Source/MetasoundFrontend/Private/MetasoundFrontendDocumentCache.h
@@ -174,9 +174,7 @@ namespace Metasound::Frontend
 		virtual const int32* FindDependencyIndex(const FGuid& InClassID) const override;
 
 		virtual const FMetasoundFrontendDocument& GetDocument() const override;
-		virtual const IDocumentGraphEdgeCache& GetEdgeCache() const override;
 		virtual const IDocumentGraphEdgeCache& GetEdgeCache(const FGuid& InPageID) const override;
-		virtual const IDocumentGraphNodeCache& GetNodeCache() const override;
 		virtual const IDocumentGraphNodeCache& GetNodeCache(const FGuid& InPageID) const override;
 		virtual const IDocumentGraphInterfaceCache& GetInterfaceCache() const override;
 
diff --git a/Engine/Plugins/Runtime/Metasound/Source/MetasoundFrontend/Public/MetasoundFrontendDocumentBuilder.h b/Engine/Plugins/Runtime/Metasound/Source/MetasoundFrontend/Public/MetasoundFrontendDocumentBuilder.h
--- a/Engine/Plugins/Runtime/Metasound/Source/MetasoundFrontend/Public/MetasoundFrontendDocumentBuilder.h
+++ b/Engine/Plugins/Runtime/Metasound/Source/MetasoundFrontend/Public/MetasoundFrontendDocumentBuilder.h
@@ -211,9 +211,9 @@ public:
 	const FMetasoundFrontendNode* DuplicateGraphOutput(const FMetasoundFrontendClassOutput& InClassOutput, const FName InName, const FGuid* InPageID = nullptr);
 
 #if WITH_EDITORONLY_DATA
-	const FMetaSoundFrontendGraphComment* FindGraphComment(const FGuid& InCommentID) const;
-	FMetaSoundFrontendGraphComment* FindGraphComment(const FGuid& InCommentID);
-	FMetaSoundFrontendGraphComment& FindOrAddGraphComment(const FGuid& InCommentID);
+	const FMetaSoundFrontendGraphComment* FindGraphComment(const FGuid& InCommentID, const FGuid* InPageID = nullptr) const;
+	FMetaSoundFrontendGraphComment* FindGraphComment(const FGuid& InCommentID, const FGuid* InPageID = nullptr);
+	FMetaSoundFrontendGraphComment& FindOrAddGraphComment(const FGuid& InCommentID, const FGuid* InPageID = nullptr);
 	UMetaSoundFrontendMemberMetadata* FindMemberMetadata(const FGuid& InMemberID);
 #endif // WITH_EDITORONLY_DATA
 
diff --git a/Engine/Plugins/Runtime/Metasound/Source/MetasoundFrontend/Public/MetasoundFrontendDocumentCacheInterface.h b/Engine/Plugins/Runtime/Metasound/Source/MetasoundFrontend/Public/MetasoundFrontendDocumentCacheInterface.h
--- a/Engine/Plugins/Runtime/Metasound/Source/MetasoundFrontend/Public/MetasoundFrontendDocumentCacheInterface.h
+++ b/Engine/Plugins/Runtime/Metasound/Source/MetasoundFrontend/Public/MetasoundFrontendDocumentCacheInterface.h
@@ -108,14 +108,14 @@ namespace Metasound::Frontend
 
 		virtual const FMetasoundFrontendDocument& GetDocument() const = 0;
 
-		// Returns NodeCache associated with the currently set build graph PageID
-		virtual const IDocumentGraphNodeCache& GetNodeCache() const = 0;
+		UE_DEPRECATED(5.5, "Use overload providing pageID")
+		virtual const IDocumentGraphNodeCache& GetNodeCache() const { return GetNodeCache(::Metasound::Frontend::DefaultPageID); }
 
 		// Returns NodeCache associated with the given PageID
 		virtual const IDocumentGraphNodeCache& GetNodeCache(const FGuid& InPageID) const = 0;
 
-		// Returns EdgeCache associated with the currently set build graph PageID
-		virtual const IDocumentGraphEdgeCache& GetEdgeCache() const = 0;
+		UE_DEPRECATED(5.5, "Use overload providing pageID")
+		virtual const IDocumentGraphEdgeCache& GetEdgeCache() const { return GetEdgeCache(::Metasound::Frontend::DefaultPageID); }
 
 		// Returns EdgeCache associated with the given PageID
 		virtual const IDocumentGraphEdgeCache& GetEdgeCache(const FGuid& InPageID) const = 0;
