- Always show audition page in analyzer stack (vs. resolved graph page name)
- Fix-up crash from swapping between editor page cvar and tab visibility attribute not being respected

#rnx

diff --git a/Engine/Plugins/Runtime/Metasound/Source/MetasoundEditor/Private/MetasoundEditor.cpp b/Engine/Plugins/Runtime/Metasound/Source/MetasoundEditor/Private/MetasoundEditor.cpp
--- a/Engine/Plugins/Runtime/Metasound/Source/MetasoundEditor/Private/MetasoundEditor.cpp
+++ b/Engine/Plugins/Runtime/Metasound/Source/MetasoundEditor/Private/MetasoundEditor.cpp
@@ -848,6 +848,7 @@ namespace Metasound
 
 			InTabManager->RegisterTabSpawner(TabNamesPrivate::GraphCanvas, FOnSpawnTab::CreateLambda(
 				[
+					bShowPageTab = ShowPageDetails(),
 					InFocusPageWidget = FocusPageWidget,
 					InPlayTimeWidget = PlayTimeWidget,
 					InMetasoundGraphEditor = MetasoundGraphEditor,
@@ -865,7 +866,7 @@ namespace Metasound
 						InRenderStatsWidget.ToSharedRef()
 					];
 
-				if (AssetEditorPrivate::EnablePageEditor)
+				if (bShowPageTab)
 				{
 					StatsWidget->AddSlot()
 						.HAlign(HAlign_Left)
@@ -933,22 +934,27 @@ namespace Metasound
 			.SetGroup(WorkspaceMenuCategoryRef)
 			.SetIcon(FSlateIcon(FAppStyle::GetAppStyleSetName(), "Kismet.Tabs.Palette"));
 
-			if (AssetEditorPrivate::EnablePageEditor)
+			if (ShowPageDetails())
 			{
-				InTabManager->RegisterTabSpawner(TabNamesPrivate::Pages, FOnSpawnTab::CreateLambda([this, InPagesDetails = PagesDetails](const FSpawnTabArgs&)
+				if (Builder.IsValid() && !Builder->IsPreset())
 				{
-					return SNew(SDockTab)
-						.Visibility(TAttribute<EVisibility>::Create([this]()
-						{
-							return Builder.IsValid() && Builder->IsPreset()
-								? EVisibility::Hidden
-								: EVisibility::Visible;
-						}))
-						.Label(LOCTEXT("MetasoundPagesDetailsTitle", "Pages"))[ InPagesDetails.ToSharedRef() ];
-				}))
-				.SetDisplayName(LOCTEXT("PagesTab", "Pages"))
-					.SetGroup(WorkspaceMenuCategoryRef)
-					.SetIcon(FSlateIcon(FAppStyle::GetAppStyleSetName(), "Kismet.Tabs.Palette"));
+					const FCanSpawnTab CanSpawnTab = FCanSpawnTab::CreateLambda([this](const FSpawnTabArgs&)
+					{
+						return Builder.IsValid() && !Builder->IsPreset();
+					});
+
+					InTabManager->RegisterTabSpawner(TabNamesPrivate::Pages, FOnSpawnTab::CreateLambda([this, InPagesDetails = PagesDetails](const FSpawnTabArgs&)
+					{
+						return SNew(SDockTab)
+							.Label(LOCTEXT("MetasoundPagesDetailsTitle", "Pages"))
+							[
+								InPagesDetails.ToSharedRef()
+							];
+					}), CanSpawnTab)
+					.SetDisplayName(LOCTEXT("PagesTab", "Pages"))
+						.SetGroup(WorkspaceMenuCategoryRef)
+						.SetIcon(FSlateIcon(FAppStyle::GetAppStyleSetName(), "Kismet.Tabs.Palette"));
+				}
 			}
 
 			InTabManager->RegisterTabSpawner(TabNamesPrivate::Interfaces, FOnSpawnTab::CreateLambda([InInterfacesDetails = InterfacesDetails](const FSpawnTabArgs&)
@@ -978,7 +984,12 @@ namespace Metasound
 			InTabManager->UnregisterTabSpawner(TabNamesPrivate::GraphCanvas);
 			InTabManager->UnregisterTabSpawner(TabNamesPrivate::Details);
 			InTabManager->UnregisterTabSpawner(TabNamesPrivate::Members);
-			InTabManager->UnregisterTabSpawner(TabNamesPrivate::Pages);
+
+			if (ShowPageDetails())
+			{
+				InTabManager->UnregisterTabSpawner(TabNamesPrivate::Pages);
+			}
+
 			InTabManager->UnregisterTabSpawner(TabNamesPrivate::Interfaces);
 			InTabManager->UnregisterTabSpawner(TabNamesPrivate::Find);
 		}
@@ -1156,7 +1167,21 @@ namespace Metasound
 				CreateAnalyzers(*MetaSoundSource);
 			}
 
-			const TSharedRef<FTabManager::FLayout> StandaloneDefaultLayout = FTabManager::NewLayout("Standalone_MetasoundEditor_Layout_v11")
+			TSharedRef<FTabManager::FStack> DetailsStack = FTabManager::NewStack()
+				->SetSizeCoefficient(0.50f)
+				->SetHideTabWell(false)
+				->AddTab(TabNamesPrivate::Details, ETabState::OpenedTab);
+
+			if (ShowPageDetails())
+			{
+				DetailsStack->AddTab(TabNamesPrivate::Pages, ETabState::OpenedTab);
+			}
+			else
+			{
+				DetailsStack->AddTab(TabNamesPrivate::Pages, ETabState::InvalidTab);
+			}
+
+			const TSharedRef<FTabManager::FLayout> StandaloneDefaultLayout = FTabManager::NewLayout("Standalone_MetasoundEditor_Layout_v14")
 			->AddArea
 			(
 				FTabManager::NewPrimaryArea()
@@ -1184,10 +1209,7 @@ namespace Metasound
 						)
 						->Split
 						(
-							FTabManager::NewStack()
-							->SetSizeCoefficient(0.50f)
-							->SetHideTabWell(false)
-							->AddTab(TabNamesPrivate::Details, ETabState::OpenedTab)
+							DetailsStack
 						)
 					)
 					->Split
@@ -1290,15 +1312,38 @@ namespace Metasound
 
 			if (MetasoundDetails.IsValid())
 			{
-				MetasoundDetails->SetObjects(SelectedObjects);
-				MetasoundDetails->HideFilterArea(false);
-				if (TabManager.IsValid())
+				if (SelectedObjects.IsEmpty())
+				{
+					if (TabManager.IsValid())
+					{
+						if (ShowPageDetails())
+						{
+							TabManager->TryInvokeTab(TabNamesPrivate::Pages);
+						}
+					}
+				}
+				else
 				{
-					TabManager->TryInvokeTab(TabFactory::Names::Details);
+					MetasoundDetails->SetObjects(SelectedObjects);
+					MetasoundDetails->HideFilterArea(false);
+					if (TabManager.IsValid())
+					{
+						TabManager->TryInvokeTab(TabNamesPrivate::Details);
+					}
 				}
 			}
 		}
 
+		bool FEditor::ShowPageDetails() const
+		{
+			if (AssetEditorPrivate::EnablePageEditor)
+			{
+				return Builder.IsValid() && !Builder->IsPreset();
+			}
+
+			return false;
+		}
+
 		bool FEditor::GetBoundsForSelectedNodes(FSlateRect& Rect, float Padding)
 		{
 			return MetasoundGraphEditor->GetBoundsForSelectedNodes(Rect, Padding);
@@ -1578,7 +1623,7 @@ namespace Metasound
 				InterfacesDetails->HideFilterArea(true);
 			}
 
-			if (AssetEditorPrivate::EnablePageEditor)
+			if (ShowPageDetails())
 			{
 				PagesDetails = PropertyModule.CreateDetailView(Args);
 				if (PagesDetails.IsValid())
@@ -1589,6 +1634,12 @@ namespace Metasound
 
 					PagesDetails->SetObjects(PagesViewObj);
 					PagesDetails->HideFilterArea(true);
+
+					TAttribute<bool> EnabledAttr = TAttribute<bool>::Create([this]()
+					{
+						return ShowPageDetails();
+					});
+					PagesDetails->SetEnabled(EnabledAttr);
 				}
 			}
 
@@ -2516,8 +2567,9 @@ namespace Metasound
 				if (Builder.IsValid())
 				{
 					constexpr bool bOpenEditor = false; // Already Focused
+					constexpr bool bPostTransaction = false;
 					const FMetaSoundFrontendDocumentBuilder& DocBuilder = Builder->GetConstBuilder();
-					UMetaSoundEditorSubsystem::GetChecked().SetFocusedPage(*Builder.Get(), DocBuilder.GetBuildPageID(), bOpenEditor);
+					UMetaSoundEditorSubsystem::GetChecked().SetFocusedPage(*Builder.Get(), DocBuilder.GetBuildPageID(), bOpenEditor, bPostTransaction);
 				}
 			}
 		}
@@ -4623,7 +4675,7 @@ namespace Metasound
 
 		void FEditor::ShowFindInMetaSound()
 		{
-			TabManager->TryInvokeTab(TabFactory::Names::Find);
+			TabManager->TryInvokeTab(TabNamesPrivate::Find);
 			if (FindWidget.IsValid())
 			{
 				FindWidget->FocusForUse();
@@ -4632,7 +4684,7 @@ namespace Metasound
 
 		void FEditor::FindSelectedNodeInGraph()
 		{
-			TabManager->TryInvokeTab(TabFactory::Names::Find);
+			TabManager->TryInvokeTab(TabNamesPrivate::Find);
 			if (FindWidget.IsValid())
 			{		
 				const FGraphPanelSelectionSet& SelectedNodes = MetasoundGraphEditor->GetSelectedNodes();
@@ -4710,12 +4762,12 @@ namespace Metasound
 			if (RenderPageWidget.IsValid())
 			{
 				const FMetaSoundPageSettings* PageSettings = nullptr;
-				if (bIsPlaying && Builder.IsValid())
+				if (Builder.IsValid())
 				{
-					const IMetaSoundDocumentInterface& DocInterface = Builder->GetConstBuilder().GetConstDocumentInterfaceChecked();
-					FGuid PageID = Frontend::DefaultPageID;
-					FDocumentBuilderRegistry::GetChecked().TryResolveTargetPageID(DocInterface.GetConstDocument().RootGraph, PageID);
-					PageSettings = Settings->FindPageSettings(PageID);
+					if (const UMetasoundEditorSettings* EditorSettings = GetDefault<UMetasoundEditorSettings>())
+					{
+						PageSettings = Settings->FindPageSettings(EditorSettings->AuditionTargetPage);
+					}
 				}
 				RenderPageWidget->Update(PageSettings);
 			}
diff --git a/Engine/Plugins/Runtime/Metasound/Source/MetasoundEditor/Private/MetasoundEditor.h b/Engine/Plugins/Runtime/Metasound/Source/MetasoundEditor/Private/MetasoundEditor.h
--- a/Engine/Plugins/Runtime/Metasound/Source/MetasoundEditor/Private/MetasoundEditor.h
+++ b/Engine/Plugins/Runtime/Metasound/Source/MetasoundEditor/Private/MetasoundEditor.h
@@ -460,6 +460,9 @@ namespace Metasound::Editor
 		/** Converts the MetaSound from a preset to a fully modifiable MetaSound. */
 		void ConvertFromPreset();
 
+		/* Whether or not page details should be visible. */
+		bool ShowPageDetails() const;
+
 		/** Creates audition menu options */
 		TSharedRef<SWidget> CreateAuditionMenuOptions();
 
