Trivial - Deprecate and swap class default builder API that returns multiple class defaults (and conform node input API for visibility)
#rnx

diff --git a/Engine/Plugins/Runtime/Metasound/Source/MetasoundEngine/Public/MetasoundBuilderBase.cpp b/Engine/Plugins/Runtime/Metasound/Source/MetasoundEngine/Public/MetasoundBuilderBase.cpp
--- a/Engine/Plugins/Runtime/Metasound/Source/MetasoundEngine/Public/MetasoundBuilderBase.cpp
+++ b/Engine/Plugins/Runtime/Metasound/Source/MetasoundEngine/Public/MetasoundBuilderBase.cpp
@@ -1161,10 +1161,10 @@ void UMetaSoundBuilderBase::GetNodeInputData(const FMetaSoundBuilderNodeInputHan
 
 FMetasoundFrontendLiteral UMetaSoundBuilderBase::GetNodeInputDefault(const FMetaSoundBuilderNodeInputHandle& InputHandle, EMetaSoundBuilderResult& OutResult)
 {
-	if (const FMetasoundFrontendLiteral* Default = Builder.GetNodeInputDefault(InputHandle.NodeID, InputHandle.VertexID))
+	if (const FMetasoundFrontendVertexLiteral* VertexLiteral = Builder.FindNodeInputDefault(InputHandle.NodeID, InputHandle.VertexID))
 	{
 		OutResult = EMetaSoundBuilderResult::Succeeded;
-		return *Default;
+		return VertexLiteral->Value;
 	}
 
 	OutResult = EMetaSoundBuilderResult::Failed;
@@ -1173,10 +1173,27 @@ FMetasoundFrontendLiteral UMetaSoundBuilderBase::GetNodeInputDefault(const FMeta
 
 FMetasoundFrontendLiteral UMetaSoundBuilderBase::GetNodeInputClassDefault(const FMetaSoundBuilderNodeInputHandle& InputHandle, EMetaSoundBuilderResult& OutResult)
 {
-	if (const FMetasoundFrontendLiteral* Default = Builder.GetNodeInputClassDefault(InputHandle.NodeID, InputHandle.VertexID))
+	using namespace Metasound::Engine;
+
+	if (const FMetasoundFrontendVertex* Vertex = Builder.FindNodeInput(InputHandle.NodeID, InputHandle.VertexID))
 	{
-		OutResult = EMetaSoundBuilderResult::Succeeded;
-		return *Default;
+		if (const TArray<FMetasoundFrontendClassInputDefault>* ClassDefaults = Builder.FindNodeClassInputDefaults(InputHandle.NodeID, Vertex->Name))
+		{
+			FGuid ResolvedPageID;
+			if (FDocumentBuilderRegistry::GetChecked().TryResolveTargetPageID(*ClassDefaults, ResolvedPageID))
+			{
+				OutResult = EMetaSoundBuilderResult::Succeeded;
+				auto MatchesPageID = [&ResolvedPageID](const FMetasoundFrontendClassInputDefault& Default)
+				{
+					return Default.PageID == ResolvedPageID;
+				};
+				if (const FMetasoundFrontendClassInputDefault* Default = ClassDefaults->FindByPredicate(MatchesPageID))
+				{
+					OutResult = EMetaSoundBuilderResult::Succeeded;
+					return Default->Literal;
+				}
+			}
+		}
 	}
 
 	OutResult = EMetaSoundBuilderResult::Failed;
diff --git a/Engine/Plugins/Runtime/Metasound/Source/MetasoundFrontend/Public/MetasoundFrontendDocumentBuilder.h b/Engine/Plugins/Runtime/Metasound/Source/MetasoundFrontend/Public/MetasoundFrontendDocumentBuilder.h
--- a/Engine/Plugins/Runtime/Metasound/Source/MetasoundFrontend/Public/MetasoundFrontendDocumentBuilder.h
+++ b/Engine/Plugins/Runtime/Metasound/Source/MetasoundFrontend/Public/MetasoundFrontendDocumentBuilder.h
@@ -299,7 +299,11 @@ public:
 	TArray<const FMetasoundFrontendNode*> GetGraphInputTemplateNodes(FName InInputName, const FGuid* InPageID = nullptr);
 
 	EMetasoundFrontendVertexAccessType GetNodeInputAccessType(const FGuid& InNodeID, const FGuid& InVertexID, const FGuid* InPageID = nullptr) const;
+
+	UE_DEPRECATED(5.5, "Use FindNodeInputClass overloads instead and use GetDefaults() on result (now supports page values)")
 	const FMetasoundFrontendLiteral* GetNodeInputClassDefault(const FGuid& InNodeID, const FGuid& InVertexID, const FGuid* InPageID = nullptr) const;
+
+	UE_DEPRECATED(5.5, "Use FindNodeInputDefault and returned struct Value member instead")
 	const FMetasoundFrontendLiteral* GetNodeInputDefault(const FGuid& InNodeID, const FGuid& InVertexID, const FGuid* InPageID = nullptr) const;
 
 	EMetasoundFrontendVertexAccessType GetNodeOutputAccessType(const FGuid& InNodeID, const FGuid& InVertexID, const FGuid* InPageID = nullptr) const;
