Update MetaSound cook to strip unsupported MetaSound platform page data

#rnx

diff --git a/Engine/Plugins/Runtime/Metasound/Source/MetasoundEditor/Private/MetasoundDetailCustomization.cpp b/Engine/Plugins/Runtime/Metasound/Source/MetasoundEditor/Private/MetasoundDetailCustomization.cpp
--- a/Engine/Plugins/Runtime/Metasound/Source/MetasoundEditor/Private/MetasoundDetailCustomization.cpp
+++ b/Engine/Plugins/Runtime/Metasound/Source/MetasoundEditor/Private/MetasoundDetailCustomization.cpp
@@ -558,7 +558,7 @@ namespace Metasound::Editor
 
 		const UMetaSoundSettings* Settings = GetDefault<UMetaSoundSettings>();
 		check(Settings);
-		for (const FMetaSoundPageSettings& PageSettings : Settings->GetPageSettings())
+		Settings->IteratePageSettings([&](const FMetaSoundPageSettings& PageSettings)
 		{
 			if (ImplementedGuids.Remove(PageSettings.UniqueId) > 0)
 			{
@@ -571,7 +571,7 @@ namespace Metasound::Editor
 					CreateEntryWidget(bIsDefault, PageSettings.UniqueId, PageSettings.Name)
 				];
 			}
-		}
+		});
 
 		for (const FGuid& MissingPageID : ImplementedGuids)
 		{
@@ -635,9 +635,13 @@ namespace Metasound::Editor
 			ImplementedGuids.Add(Graph.PageID);
 		});
 
-		auto GetStringName = [](const FMetaSoundPageSettings& Page) -> TSharedPtr<FString> { return MakeShared<FString>(Page.Name.ToString()); };
-		auto RemoveImplementedItem = [&ImplementedGuids](const FMetaSoundPageSettings& Page) { return !ImplementedGuids.Contains(Page.UniqueId); };
-		Algo::TransformIf(Settings->GetPageSettings(), AddableItems, RemoveImplementedItem, GetStringName);
+		Settings->IteratePageSettings([this, &ImplementedGuids](const FMetaSoundPageSettings& Page)
+		{
+			if (!ImplementedGuids.Contains(Page.UniqueId))
+			{
+				AddableItems.Add(MakeShared<FString>(Page.Name.ToString()));
+			}
+		});
 
 		auto GetPageName = [&Settings](const FGuid& PageID)
 		{
diff --git a/Engine/Plugins/Runtime/Metasound/Source/MetasoundEditor/Private/MetasoundEditor.cpp b/Engine/Plugins/Runtime/Metasound/Source/MetasoundEditor/Private/MetasoundEditor.cpp
--- a/Engine/Plugins/Runtime/Metasound/Source/MetasoundEditor/Private/MetasoundEditor.cpp
+++ b/Engine/Plugins/Runtime/Metasound/Source/MetasoundEditor/Private/MetasoundEditor.cpp
@@ -1859,22 +1859,14 @@ namespace Metasound
 
 		void FEditor::CreateAuditionPageSubMenuOptions(FMenuBuilder& MenuBuilder)
 		{
-			TWeakObjectPtr<const UMetaSoundSettings> SettingsPtr = GetDefault<UMetaSoundSettings>();
-			if (!SettingsPtr.IsValid())
+			const UMetaSoundSettings* Settings = GetDefault<UMetaSoundSettings>();
+			if (!Settings)
 			{
 				return;
 			}
 
-			const TArray<FMetaSoundPageSettings>& PageSettingsArray = SettingsPtr->GetPageSettings();
 			MenuBuilder.BeginSection("SetAuditionPlatformSectionHeader", LOCTEXT("SetAuditionPlatformDescription", "Platform"));
 			{
-				TSet<FName> ImplementedPlatforms;
-				for (const FMetaSoundPageSettings& PageSettings : PageSettingsArray)
-				{
-					auto PairToKey = [](const TPair<FName, int32>& PerPlatformPair) { return PerPlatformPair.Key; };
-					Algo::Transform(PageSettings.IsCooked.PerPlatform, ImplementedPlatforms, PairToKey);
-				}
-
 				auto CreatePlatformEntry = [this, &MenuBuilder](FName PlatformName, const FText& PlatformText)
 				{
 					FUIAction SetPlatformAction;
@@ -1903,8 +1895,8 @@ namespace Metasound
 						SetPlatformAction);
 				};
 
-				CreatePlatformEntry(FName(), LOCTEXT("DefaultPlatformDisplayName", "Default"));
-				for (const FName& PlatformName : ImplementedPlatforms)
+				const TArray<FName> AuditionPlatforms = UMetasoundEditorSettings::GetAuditionPlatformNames();
+				for (const FName& PlatformName : AuditionPlatforms)
 				{
 					const FText PlatformText = FText::FromName(PlatformName);
 					CreatePlatformEntry(PlatformName, PlatformText);
@@ -1953,19 +1945,18 @@ namespace Metasound
 						return ECheckBoxState::Unchecked;
 					})
 					.ToolTipText(FocusPageTooltip),
-					LOCTEXT("EnableFocusTargetPageSwap", "Set To Focused"),
+					LOCTEXT("EnableFocusTargetPageGraphSwap", "Set Target To Focused Graph"),
 					true,
 					true,
 					FocusPageTooltip
 				);
 
-				TSet<FName> PageNames;
-				Algo::Transform(PageSettingsArray, PageNames, [](const FMetaSoundPageSettings& PageSettings) { return PageSettings.Name; });
-
-				auto CreateTargetPageEntry = [this, &MenuBuilder](FName AuditionTargetPage, const FText& PageText)
+				auto TryAddPageEntry = [this, &MenuBuilder](const FMetaSoundPageSettings& PageSettings)
 				{
-					FUIAction SetPlatformAction;
-					SetPlatformAction.ExecuteAction = FExecuteAction::CreateLambda([this, AuditionTargetPage]()
+					const FName AuditionTargetPage = PageSettings.Name;
+					const FText PageText = FText::FromName(PageSettings.Name);
+					FUIAction SetTargetPageAction;
+					SetTargetPageAction.ExecuteAction = FExecuteAction::CreateLambda([this, AuditionTargetPage]()
 					{
 						if (UMetasoundEditorSettings* EditorSettings = GetMutableDefault<UMetasoundEditorSettings>())
 						{
@@ -1974,7 +1965,7 @@ namespace Metasound
 						}
 					});
 
-					SetPlatformAction.CanExecuteAction = FCanExecuteAction::CreateLambda([this, AuditionTargetPage]()
+					SetTargetPageAction.CanExecuteAction = FCanExecuteAction::CreateLambda([this, AuditionTargetPage]()
 					{
 						if (const UMetasoundEditorSettings* EditorSettings = GetDefault<UMetasoundEditorSettings>())
 						{
@@ -1984,21 +1975,31 @@ namespace Metasound
 								return EditorSettings->AuditionTargetPage != AuditionTargetPage;
 							}
 						}
+						return false;
+					});
 
+					SetTargetPageAction.IsActionVisibleDelegate.BindLambda([PageName = PageSettings.Name]()
+					{
+						if (const UMetasoundEditorSettings* EditorSettings = GetDefault<UMetasoundEditorSettings>())
+						{
+							if (const UMetaSoundSettings* Settings = GetDefault<UMetaSoundSettings>())
+							{
+								if (const FMetaSoundPageSettings* PageSetting = Settings->FindPageSettings(PageName))
+								{
+									return PageSetting->PlatformCanTargetPage(EditorSettings->AuditionPlatform);
+								}
+							}
+						}
 						return false;
 					});
 
 					MenuBuilder.AddMenuEntry(PageText,
 						FText::Format(LOCTEXT("SetAuditionPageToolTip", "Sets the audition target page to '{0}'."), PageText),
 						FSlateIcon(),
-						SetPlatformAction);
+						SetTargetPageAction);
 				};
 
-				for (const FName& PageName : PageNames)
-				{
-					const FText PageText = FText::FromName(PageName);
-					CreateTargetPageEntry(PageName, PageText);
-				}
+				Settings->IteratePageSettings(TryAddPageEntry);
 			}
 			MenuBuilder.EndSection();
 		}
diff --git a/Engine/Plugins/Runtime/Metasound/Source/MetasoundEditor/Private/MetasoundEditorGraphBuilder.cpp b/Engine/Plugins/Runtime/Metasound/Source/MetasoundEditor/Private/MetasoundEditorGraphBuilder.cpp
--- a/Engine/Plugins/Runtime/Metasound/Source/MetasoundEditor/Private/MetasoundEditorGraphBuilder.cpp
+++ b/Engine/Plugins/Runtime/Metasound/Source/MetasoundEditor/Private/MetasoundEditorGraphBuilder.cpp
@@ -1304,20 +1304,17 @@ namespace Metasound
 						const TArray<FMetasoundFrontendClassInputDefault>* ClassInputDefaults = Builder.FindNodeClassInputDefaults(InputHandle.NodeID, Vertex->Name);
 						if (ClassInputDefaults)
 						{
-							FGuid PageID;
-							if (Engine::FDocumentBuilderRegistry::GetChecked().TryResolveTargetPageID(*ClassInputDefaults, PageID))
+							const FGuid PageID = Engine::FDocumentBuilderRegistry::GetChecked().ResolveTargetPageID(*ClassInputDefaults);
+							auto MatchesPageID = [&PageID](const FMetasoundFrontendClassInputDefault& InputDefault) { return PageID == InputDefault.PageID; };
+							if (const FMetasoundFrontendClassInputDefault* ClassDefault = ClassInputDefaults->FindByPredicate(MatchesPageID))
 							{
-								auto MatchesPageID = [&PageID](const FMetasoundFrontendClassInputDefault& InputDefault) { return PageID == InputDefault.PageID; };
-								if (const FMetasoundFrontendClassInputDefault* ClassDefault = ClassInputDefaults->FindByPredicate(MatchesPageID))
+								if (ClassDefault->Literal.GetType() == EMetasoundFrontendLiteralType::None)
 								{
-									if (ClassDefault->Literal.GetType() == EMetasoundFrontendLiteralType::None)
-									{
-										OutDefaultLiteral.Clear();
-									}
-									else
-									{
-										OutDefaultLiteral = ClassDefault->Literal;
-									}
+									OutDefaultLiteral.Clear();
+								}
+								else
+								{
+									OutDefaultLiteral = ClassDefault->Literal;
 								}
 							}
 							return true;
@@ -2652,15 +2649,12 @@ namespace Metasound
 			const TArray<FMetasoundFrontendClassInputDefault>* ClassDefaults = InBuilder.FindNodeClassInputDefaults(InputHandle.NodeID, InputVertex->Name);
 			if (ClassDefaults)
 			{
-				FGuid PageID;
-				if (Engine::FDocumentBuilderRegistry::GetChecked().TryResolveTargetPageID(*ClassDefaults, PageID))
+				const FGuid PageID = Engine::FDocumentBuilderRegistry::GetChecked().ResolveTargetPageID(*ClassDefaults);
+				auto MatchesPageID = [&PageID](const FMetasoundFrontendClassInputDefault& InputDefault) { return InputDefault.PageID == PageID; };
+				if (const FMetasoundFrontendClassInputDefault* ClassDefault = ClassDefaults->FindByPredicate(MatchesPageID))
 				{
-					auto MatchesPageID = [&PageID](const FMetasoundFrontendClassInputDefault& InputDefault) { return InputDefault.PageID == PageID; };
-					if (const FMetasoundFrontendClassInputDefault* ClassDefault = ClassDefaults->FindByPredicate(MatchesPageID))
-					{
-						InPin.DefaultValue = ClassDefault->Literal.ToString();
-						return OldValue != InPin.DefaultValue;
-					}
+					InPin.DefaultValue = ClassDefault->Literal.ToString();
+					return OldValue != InPin.DefaultValue;
 				}
 			}
 
diff --git a/Engine/Plugins/Runtime/Metasound/Source/MetasoundEditor/Private/MetasoundEditorModule.cpp b/Engine/Plugins/Runtime/Metasound/Source/MetasoundEditor/Private/MetasoundEditorModule.cpp
--- a/Engine/Plugins/Runtime/Metasound/Source/MetasoundEditor/Private/MetasoundEditorModule.cpp
+++ b/Engine/Plugins/Runtime/Metasound/Source/MetasoundEditor/Private/MetasoundEditorModule.cpp
@@ -422,9 +422,9 @@ namespace Metasound
 
 				if (UMetaSoundSettings* Settings = GetMutableDefault<UMetaSoundSettings>())
 				{
-					Settings->GetOnDefaultConformedDelegate().AddLambda([]()
+					Settings->GetOnDefaultRenamedDelegate().AddLambda([]()
 					{
-						FNotificationInfo Info(LOCTEXT("MetaSoundSettings_CannotDeleteDefaultPage", "Cannot change name of nor delete 'Default' MetaSound Page"));
+						FNotificationInfo Info(LOCTEXT("MetaSoundSettings_CannotNameDefaultPage", "Cannot name 'Default': reserved MetaSound page name"));
 						Info.bFireAndForget = true;
 						Info.ExpireDuration = 2.0f;
 						Info.bUseThrobber = true;
diff --git a/Engine/Plugins/Runtime/Metasound/Source/MetasoundEditor/Private/MetasoundEditorSettings.cpp b/Engine/Plugins/Runtime/Metasound/Source/MetasoundEditor/Private/MetasoundEditorSettings.cpp
--- a/Engine/Plugins/Runtime/Metasound/Source/MetasoundEditor/Private/MetasoundEditorSettings.cpp
+++ b/Engine/Plugins/Runtime/Metasound/Source/MetasoundEditor/Private/MetasoundEditorSettings.cpp
@@ -112,39 +112,51 @@ const FAudioMaterialMeterStyle* UMetasoundEditorSettings::GetMeterStyle() const
 
 Metasound::Engine::FAuditionPageInfo UMetasoundEditorSettings::ResolveAuditionPageInfo(const TSet<FGuid>& InPageIDs) const
 {
+	using namespace Metasound;
 	using namespace Metasound::Engine;
 
 	FAuditionPageInfo PreviewInfo { .PlatformName = AuditionPlatform };
-	auto PageIsCooked = [&InPageIDs, &PreviewInfo](const FMetaSoundPageSettings& PageSettings)
-	{
-		if (InPageIDs.Contains(PageSettings.UniqueId))
-		{
-			return PageSettings.IsCooked.GetValueForPlatform(PreviewInfo.PlatformName);
-		}
 
-		return false;
-	};
+	if (AuditionTargetPage == Frontend::DefaultPageName)
+	{
+		check(InPageIDs.Contains(Frontend::DefaultPageID));
+		PreviewInfo.PageID = Frontend::DefaultPageID;
+		return PreviewInfo;
+	}
 
 	if (const UMetaSoundSettings* Settings = GetDefault<UMetaSoundSettings>())
 	{
 		if (const FMetaSoundPageSettings* TargetPageSettings = Settings->FindPageSettings(AuditionTargetPage))
 		{
+			auto PageIsCooked = [&InPageIDs, &PreviewInfo](const FMetaSoundPageSettings& PageSettings)
+			{
+				if (InPageIDs.Contains(PageSettings.UniqueId))
+				{
+					return !PageSettings.ExcludePageFromCook(PreviewInfo.PlatformName);
+				}
+
+				return false;
+			};
+
 			const FGuid& TargetPageID = TargetPageSettings->UniqueId;
-			const TArray<FMetaSoundPageSettings>& PageSettingsArray = Settings->GetPageSettings();
+			constexpr bool bReverse = true;
 			bool bFoundMatch = false;
-			for (int32 Index = PageSettingsArray.Num() - 1; Index >= 0; --Index)
+			bool bPageSelected = false;
+			Settings->IteratePageSettings([&](const FMetaSoundPageSettings& PageSettings)
 			{
-				const FMetaSoundPageSettings& PageSettings = PageSettingsArray[Index];
-				bFoundMatch |= PageSettings.UniqueId == TargetPageID;
-				if (bFoundMatch)
+				if (!bPageSelected)
 				{
-					if (PageIsCooked(PageSettings))
+					bFoundMatch |= PageSettings.UniqueId == TargetPageID;
+					if (bFoundMatch)
 					{
-						PreviewInfo.PageID = PageSettings.UniqueId;
-						return PreviewInfo;
+						if (PageIsCooked(PageSettings))
+						{
+							bPageSelected = true;
+							PreviewInfo.PageID = PageSettings.UniqueId;
+						}
 					}
 				}
-			}
+			}, bReverse);
 		}
 	}
 
@@ -156,16 +168,8 @@ TArray<FName> UMetasoundEditorSettings::GetAuditionPlatformNames()
 {
 	if (const UMetaSoundSettings* Settings = GetDefault<UMetaSoundSettings>())
 	{
-		TSet<FName> PlatformNames { "Default" };
-		for (const FMetaSoundPageSettings& PageSettings : Settings->GetPageSettings())
-		{
-			TSet<FName> PagePlatforms;
-			PageSettings.IsCooked.PerPlatform.GetKeys(PagePlatforms);
-			PlatformNames.Append(PagePlatforms);
-		}
-		return PlatformNames.Array();
+		return Settings->GetImplementedPagePlatforms();
 	}
-
-	return { };
+	return { FPlatformProperties::IniPlatformName() };
 }
 #undef LOCTEXT_NAMESPACE // "MetaSoundEditor"
diff --git a/Engine/Plugins/Runtime/Metasound/Source/MetasoundEngine/Private/MetasoundAssetSubsystem.cpp b/Engine/Plugins/Runtime/Metasound/Source/MetasoundEngine/Private/MetasoundAssetSubsystem.cpp
--- a/Engine/Plugins/Runtime/Metasound/Source/MetasoundEngine/Private/MetasoundAssetSubsystem.cpp
+++ b/Engine/Plugins/Runtime/Metasound/Source/MetasoundEngine/Private/MetasoundAssetSubsystem.cpp
@@ -175,6 +175,8 @@ namespace Metasound::Engine
 #if WITH_EDITOR
 		virtual TSet<FAssetInfo> GetReferencedAssetClasses(const FMetasoundAssetBase& InAssetBase) const override;
 #endif // WITH_EDITOR
+		virtual void IterateAssets(TFunctionRef<void(const FAssetKey, const TArray<FTopLevelAssetPath>&)> Iter) const override;
+		virtual void ReloadMetaSoundAssets() const override;
 		virtual void RemoveAsset(const UObject& InObject) override;
 		virtual void RemoveAsset(const FAssetData& InAssetData) override;
 		virtual void RenameAsset(const FAssetData& InAssetData, const FString& InOldObjectPath) override;
@@ -588,6 +590,14 @@ namespace Metasound::Engine
 	}
 #endif // WITH_EDITOR
 
+	void FMetaSoundAssetManager::IterateAssets(TFunctionRef<void(const FAssetKey, const TArray<FTopLevelAssetPath>&)> Iter) const
+	{
+		for (const TPair<FAssetKey, TArray<FTopLevelAssetPath>>& Pair : PathMap)
+		{
+			Iter(Pair.Key, Pair.Value);
+		}
+	}
+
 	void FMetaSoundAssetManager::RebuildDenyListCache(const UAssetManager& InAssetManager)
 	{
 		using namespace Metasound::Frontend;
@@ -758,6 +768,34 @@ namespace Metasound::Engine
 
 	}
 
+	void FMetaSoundAssetManager::ReloadMetaSoundAssets() const
+	{
+		TSet<FMetasoundAssetBase*> ToReregister;
+		IterateAssets([&ToReregister](const FAssetKey& AssetKey, const TArray<FTopLevelAssetPath>& Paths)
+		{
+			if (FMetasoundAssetBase* Asset = IMetaSoundAssetManager::GetChecked().FindAsset(AssetKey))
+			{
+				if (Asset->IsRegistered())
+				{
+					ToReregister.Add(Asset);
+					Asset->UnregisterGraphWithFrontend();
+				}
+			}
+		});
+
+		// Handled in second loop to avoid re-registering referenced graphs more than once
+		IterateAssets([&ToReregister](const FAssetKey& AssetKey, const TArray<FTopLevelAssetPath>& Paths)
+		{
+			if (FMetasoundAssetBase* Asset = IMetaSoundAssetManager::GetChecked().FindAsset(AssetKey))
+			{
+				if (ToReregister.Contains(Asset))
+				{
+					Asset->UpdateAndRegisterForExecution();
+				}
+			}
+		});
+	}
+
 	void FMetaSoundAssetManager::SearchAndIterateDirectoryAssets(const TArray<FDirectoryPath>& InDirectories, TFunctionRef<void(const FAssetData&)> InFunction)
 	{
 		if (InDirectories.IsEmpty())
diff --git a/Engine/Plugins/Runtime/Metasound/Source/MetasoundEngine/Private/MetasoundBuilderSubsystem.cpp b/Engine/Plugins/Runtime/Metasound/Source/MetasoundEngine/Private/MetasoundBuilderSubsystem.cpp
--- a/Engine/Plugins/Runtime/Metasound/Source/MetasoundEngine/Private/MetasoundBuilderSubsystem.cpp
+++ b/Engine/Plugins/Runtime/Metasound/Source/MetasoundEngine/Private/MetasoundBuilderSubsystem.cpp
@@ -37,6 +37,21 @@
 
 namespace Metasound::Engine
 {
+	FAutoConsoleCommand CVarMetaSoundSetTargetPage(
+		TEXT("au.MetaSound.Pages.SetTarget"),
+		TEXT("Sets the target page to that with the given name. If name not specified or not found, command is ignored.\n"),
+		FConsoleCommandWithArgsDelegate::CreateLambda([](const TArray<FString>& Args)
+		{
+			if (!Args.IsEmpty())
+			{
+				if (UMetaSoundBuilderSubsystem* Subsystem = GEngine->GetEngineSubsystem<UMetaSoundBuilderSubsystem>())
+				{
+					Subsystem->SetTargetPage(FName { *Args.Last() });
+				}
+			}
+		})
+	);
+
 	namespace BuilderSubsystemPrivate
 	{
 		template <typename TLiteralType>
@@ -68,14 +83,11 @@ namespace Metasound::Engine
 			}
 			else if (const TArray<FMetasoundFrontendClassInputDefault>* ClassDefaults = Builder.FindNodeClassInputDefaults(InNodeID, VertexName))
 			{
-				FGuid PageID;
-				if (FDocumentBuilderRegistry::GetChecked().TryResolveTargetPageID(*ClassDefaults, PageID))
+				const FGuid PageID = FDocumentBuilderRegistry::GetChecked().ResolveTargetPageID(*ClassDefaults);
+				auto MatchesPageID = [&PageID](const FMetasoundFrontendClassInputDefault& InputDefault) { return InputDefault.PageID == PageID; };
+				if (const FMetasoundFrontendClassInputDefault* ClassDefault = ClassDefaults->FindByPredicate(MatchesPageID))
 				{
-					auto MatchesPageID = [&PageID](const FMetasoundFrontendClassInputDefault& InputDefault) { return InputDefault.PageID == PageID; };
-					if (const FMetasoundFrontendClassInputDefault* ClassDefault = ClassDefaults->FindByPredicate(MatchesPageID))
-					{
-						DefaultLiteral = ClassDefault->Literal;
-					}
+					DefaultLiteral = ClassDefault->Literal;
 				}
 			}
 
@@ -267,6 +279,12 @@ bool UMetaSoundSourceBuilder::GetLiveUpdatesEnabled() const
 	return GetMetaSoundSource().GetDynamicGeneratorTransactor().IsValid();
 }
 
+const FMetasoundFrontendGraph& UMetaSoundSourceBuilder::GetConstTargetPageGraphChecked() const
+{
+	const FMetasoundFrontendGraphClass& RootGraph = Builder.GetConstDocumentChecked().RootGraph;
+	return RootGraph.FindConstGraphChecked(TargetPageID);
+}
+
 const UMetaSoundSource& UMetaSoundSourceBuilder::GetMetaSoundSource() const
 {
 	return GetConstBuilder().CastDocumentObjectChecked<UMetaSoundSource>();
@@ -329,19 +347,17 @@ void UMetaSoundSourceBuilder::InitTargetPageDelegates(Metasound::Frontend::FDocu
 	});
 
 	const IMetaSoundDocumentInterface& DocInterface = GetConstBuilder().GetConstDocumentInterfaceChecked();
-	FGuid PageID;
-	if (FDocumentBuilderRegistry::GetChecked().TryResolveTargetPageID(DocInterface.GetConstDocument().RootGraph, PageID))
-	{
-		FEdgeModifyDelegates& EdgeDelegates = OutDocumentDelegates.FindEdgeDelegatesChecked(PageID);
-		EdgeDelegates.OnEdgeAdded.AddUObject(this, &UMetaSoundSourceBuilder::OnEdgeAdded);
-		EdgeDelegates.OnRemoveSwappingEdge.AddUObject(this, &UMetaSoundSourceBuilder::OnRemoveSwappingEdge);
-
-		FNodeModifyDelegates& NodeDelegates = OutDocumentDelegates.FindNodeDelegatesChecked(PageID);
-		NodeDelegates.OnNodeAdded.AddUObject(this, &UMetaSoundSourceBuilder::OnNodeAdded);
-		NodeDelegates.OnNodeInputLiteralSet.AddUObject(this, &UMetaSoundSourceBuilder::OnNodeInputLiteralSet);
-		NodeDelegates.OnRemoveSwappingNode.AddUObject(this, &UMetaSoundSourceBuilder::OnRemoveSwappingNode);
-		NodeDelegates.OnRemovingNodeInputLiteral.AddUObject(this, &UMetaSoundSourceBuilder::OnRemovingNodeInputLiteral);
-	}
+	TargetPageID = FDocumentBuilderRegistry::GetChecked().ResolveTargetPageID(DocInterface.GetConstDocument().RootGraph);
+
+	FEdgeModifyDelegates& EdgeDelegates = OutDocumentDelegates.FindEdgeDelegatesChecked(TargetPageID);
+	EdgeDelegates.OnEdgeAdded.AddUObject(this, &UMetaSoundSourceBuilder::OnEdgeAdded);
+	EdgeDelegates.OnRemoveSwappingEdge.AddUObject(this, &UMetaSoundSourceBuilder::OnRemoveSwappingEdge);
+
+	FNodeModifyDelegates& NodeDelegates = OutDocumentDelegates.FindNodeDelegatesChecked(TargetPageID);
+	NodeDelegates.OnNodeAdded.AddUObject(this, &UMetaSoundSourceBuilder::OnNodeAdded);
+	NodeDelegates.OnNodeInputLiteralSet.AddUObject(this, &UMetaSoundSourceBuilder::OnNodeInputLiteralSet);
+	NodeDelegates.OnRemoveSwappingNode.AddUObject(this, &UMetaSoundSourceBuilder::OnRemoveSwappingNode);
+	NodeDelegates.OnRemovingNodeInputLiteral.AddUObject(this, &UMetaSoundSourceBuilder::OnRemovingNodeInputLiteral);
 }
 
 void UMetaSoundSourceBuilder::OnAssetReferenceAdded(TScriptInterface<IMetaSoundDocumentInterface> DocInterface)
@@ -360,7 +376,7 @@ void UMetaSoundSourceBuilder::OnEdgeAdded(int32 EdgeIndex) const
 {
 	using namespace Metasound::DynamicGraph;
 
-	const FMetasoundFrontendEdge& NewEdge = Builder.FindConstBuildGraphChecked().Edges[EdgeIndex];
+	const FMetasoundFrontendEdge& NewEdge = GetConstTargetPageGraphChecked().Edges[EdgeIndex];
 	ExecuteAuditionableTransaction([this, &NewEdge](Metasound::DynamicGraph::FDynamicOperatorTransactor& Transactor)
 	{
 		const FMetaSoundFrontendDocumentBuilder& DocBuilder = GetConstBuilder();
@@ -434,7 +450,7 @@ void UMetaSoundSourceBuilder::OnNodeAdded(int32 NodeIndex) const
 		using namespace Metasound;
 		using namespace Metasound::Frontend;
 
-		const FMetasoundFrontendNode& AddedNode = Builder.FindConstBuildGraphChecked().Nodes[NodeIndex];
+		const FMetasoundFrontendNode& AddedNode = GetConstTargetPageGraphChecked().Nodes[NodeIndex];
 		const FMetasoundFrontendClass* NodeClass = Builder.FindDependency(AddedNode.ClassID);
 		checkf(NodeClass, TEXT("Node successfully added to graph but document is missing associated dependency"));
 
@@ -541,7 +557,7 @@ void UMetaSoundSourceBuilder::OnNodeInputLiteralSet(int32 NodeIndex, int32 Verte
 {
 	using namespace Metasound::DynamicGraph;
 
-	const FMetasoundFrontendNode& Node = Builder.FindConstBuildGraphChecked().Nodes[NodeIndex];
+	const FMetasoundFrontendNode& Node = GetConstTargetPageGraphChecked().Nodes[NodeIndex];
 	const FMetasoundFrontendVertex& Input = Node.Interface.Inputs[VertexIndex];
 
 	// Only send the literal down if not connected, as the graph core layer
@@ -592,14 +608,13 @@ void UMetaSoundSourceBuilder::OnRemovingPage(const Metasound::Frontend::FDocumen
 
 	FDocumentModifyDelegates& DocDelegates = Builder.GetDocumentDelegates();
 	InitTargetPageDelegates(DocDelegates);
-
 }
 
 void UMetaSoundSourceBuilder::OnRemoveSwappingEdge(int32 SwapIndex, int32 LastIndex) const
 {
 	using namespace Metasound::DynamicGraph;
 
-	const FMetasoundFrontendEdge& EdgeBeingRemoved = Builder.FindConstBuildGraphChecked().Edges[SwapIndex];
+	const FMetasoundFrontendEdge& EdgeBeingRemoved = GetConstTargetPageGraphChecked().Edges[SwapIndex];
 	ExecuteAuditionableTransaction([this, EdgeBeingRemoved](FDynamicOperatorTransactor& Transactor)
 	{
 		using namespace Metasound;
@@ -683,7 +698,7 @@ void UMetaSoundSourceBuilder::OnRemoveSwappingNode(int32 SwapIndex, int32 LastIn
 	{
 		using namespace Metasound::Frontend;
 
-		const FMetasoundFrontendNode& NodeBeingRemoved = Builder.FindConstBuildGraphChecked().Nodes[SwapIndex];
+		const FMetasoundFrontendNode& NodeBeingRemoved = GetConstTargetPageGraphChecked().Nodes[SwapIndex];
 		const FGuid& NodeID = NodeBeingRemoved.GetID();
 		Transactor.RemoveNode(NodeID);
 		return true;
@@ -694,22 +709,19 @@ void UMetaSoundSourceBuilder::OnRemovingNodeInputLiteral(int32 NodeIndex, int32
 {
 	using namespace Metasound::DynamicGraph;
 
-	const FMetasoundFrontendNode& Node = Builder.FindConstBuildGraphChecked().Nodes[NodeIndex];
+	const TArray<FMetasoundFrontendNode>& Nodes = GetConstTargetPageGraphChecked().Nodes;
+	const FMetasoundFrontendNode& Node = Nodes[NodeIndex];
 	const FMetasoundFrontendVertex& Input = Node.Interface.Inputs[VertexIndex];
 
 	// Only send the literal down if not connected, as the graph core layer will disconnect.
 	if (!Builder.IsNodeInputConnected(Node.GetID(), Input.VertexID))
 	{
-		ExecuteAuditionableTransaction([this, &NodeIndex, &VertexIndex, &LiteralIndex](FDynamicOperatorTransactor& Transactor)
+		ExecuteAuditionableTransaction([this, &Node, &Input, &NodeIndex, &VertexIndex, &LiteralIndex](FDynamicOperatorTransactor& Transactor)
 		{
 			using namespace Metasound;
 			using namespace Metasound::Engine;
 			using namespace Metasound::Frontend;
 
-			const TArray<FMetasoundFrontendNode>& Nodes = Builder.FindConstBuildGraphChecked().Nodes;
-			const FMetasoundFrontendNode& Node = Nodes[NodeIndex];
-			const FMetasoundFrontendVertex& Input = Node.Interface.Inputs[VertexIndex];
-
 			TOptional<FMetasoundFrontendLiteral> InputDefault = BuilderSubsystemPrivate::TryResolveNodeInputDefault(Builder, Node.GetID(), Input.Name);
 			if (ensureAlwaysMsgf(InputDefault.IsSet(), TEXT("Could not dynamically assign default literal from class definition upon removing input '%s' literal: document's dependency entry invalid and has no default assigned"), *Input.Name.ToString()))
 			{
@@ -1160,9 +1172,16 @@ void UMetaSoundBuilderSubsystem::RegisterSourceBuilder(FName BuilderName, UMetaS
 
 bool UMetaSoundBuilderSubsystem::SetTargetPage(FName PageName)
 {
+	using namespace Metasound::Frontend;
+
 	if (UMetaSoundSettings* Settings = GetMutableDefault<UMetaSoundSettings>())
 	{
-		return Settings->SetTargetPage(PageName);
+		const bool bTargetChanged = Settings->SetTargetPage(PageName);
+		if (bTargetChanged)
+		{
+			IMetaSoundAssetManager::GetChecked().ReloadMetaSoundAssets();
+		}
+		return bTargetChanged;
 	}
 
 	return false;
diff --git a/Engine/Plugins/Runtime/Metasound/Source/MetasoundEngine/Private/MetasoundDocumentBuilderRegistry.cpp b/Engine/Plugins/Runtime/Metasound/Source/MetasoundEngine/Private/MetasoundDocumentBuilderRegistry.cpp
--- a/Engine/Plugins/Runtime/Metasound/Source/MetasoundEngine/Private/MetasoundDocumentBuilderRegistry.cpp
+++ b/Engine/Plugins/Runtime/Metasound/Source/MetasoundEngine/Private/MetasoundDocumentBuilderRegistry.cpp
@@ -97,6 +97,64 @@ namespace Metasound::Engine
 	}
 
 #if WITH_EDITORONLY_DATA
+	bool FDocumentBuilderRegistry::CookPages(FName PlatformName, FMetaSoundFrontendDocumentBuilder& Builder) const
+	{
+		bool bModified = false;
+
+		const UMetaSoundSettings* Settings = GetDefault<UMetaSoundSettings>();
+		check(Settings);
+		const TArray<FGuid>& PlatformTargetPageIDs = Settings->GetCookedTargetPageIDs(PlatformName);
+		auto StripPageEntries = [&](TFunctionRef<bool(const FGuid&)> RemovePageItem, TSet<FGuid>& ImplementedPageIDs)
+		{
+			TSet<FGuid> ResolvedTargetIDs;
+			for (const FGuid& TargetPage : PlatformTargetPageIDs)
+			{
+				const FGuid PageID = ResolveTargetPageIDInternal(*Settings, ImplementedPageIDs, TargetPage, PlatformName);
+				ResolvedTargetIDs.Add(PageID);
+			}
+
+			TSet<FGuid> ToRemove = ImplementedPageIDs.Difference(ResolvedTargetIDs);
+			for (const FGuid& PageID : ToRemove)
+			{
+				bModified |= RemovePageItem(PageID);
+			}
+		};
+
+		const FMetasoundFrontendDocument& Document = Builder.GetConstDocumentChecked();
+
+		{ // Strip graphs
+			TSet<FGuid> ImplementedGraphPageIDs;
+			auto AddPageID = [&ImplementedGraphPageIDs](const FMetasoundFrontendGraph& Graph) { ImplementedGraphPageIDs.Add(Graph.PageID); };
+			Document.RootGraph.IterateGraphPages(AddPageID);
+
+			auto RemoveGraphPage = [&Builder](const FGuid& InPageID)
+			{
+				return Builder.RemoveGraphPage(InPageID);
+			};
+
+			StripPageEntries(RemoveGraphPage, ImplementedGraphPageIDs);
+		}
+
+		{ // Strip default input values
+			TSet<FGuid> ImplementedDefaultPageIDs;
+			for (const FMetasoundFrontendClassInput& GraphInput : Document.RootGraph.Interface.Inputs)
+			{
+				ImplementedDefaultPageIDs.Reset();
+				auto AddPageID = [&ImplementedDefaultPageIDs](const FGuid& PageID, const FMetasoundFrontendLiteral&) { ImplementedDefaultPageIDs.Add(PageID); };
+				GraphInput.IterateDefaults(AddPageID);
+
+				auto RemoveDefault = [&Builder, &GraphInput](const FGuid& InPageID)
+				{
+					return Builder.RemoveGraphInputDefault(GraphInput.Name, InPageID);
+				};
+
+				StripPageEntries(RemoveDefault, ImplementedDefaultPageIDs);
+			}
+		}
+
+		return bModified;
+	}
+
 	FMetaSoundFrontendDocumentBuilder& FDocumentBuilderRegistry::FindOrBeginBuilding(TScriptInterface<IMetaSoundDocumentInterface> MetaSound)
 	{
 		UObject* Object = MetaSound.GetObject();
@@ -324,7 +382,7 @@ namespace Metasound::Engine
 		return bReloaded;
 	}
 
-	bool FDocumentBuilderRegistry::TryResolveTargetPageID(const FMetasoundFrontendGraphClass& InGraphClass, FGuid& OutResolvedPageID) const
+	FGuid FDocumentBuilderRegistry::ResolveTargetPageID(const FMetasoundFrontendGraphClass& InGraphClass) const
 	{
 		TSet<FGuid> DocPageIds;
 		InGraphClass.IterateGraphPages([&DocPageIds](const FMetasoundFrontendGraph& PageGraph)
@@ -332,32 +390,28 @@ namespace Metasound::Engine
 			DocPageIds.Add(PageGraph.PageID);
 		});
 
-		return TryResolveTargetPageID(DocPageIds, OutResolvedPageID);
+		return ResolveTargetPageIDInternal(DocPageIds);
 	}
 
-	bool FDocumentBuilderRegistry::TryResolveTargetPageID(const FMetasoundFrontendClassInput& InClassInput, FGuid& OutResolvedPageID) const
+	FGuid FDocumentBuilderRegistry::ResolveTargetPageID(const FMetasoundFrontendClassInput& InClassInput) const
 	{
-		TSet<FGuid> DocPageIds;
-		InClassInput.IterateDefaults([&DocPageIds](const FGuid& PageID, const FMetasoundFrontendLiteral&)
+		TSet<FGuid> DefaultPageIds;
+		InClassInput.IterateDefaults([&DefaultPageIds](const FGuid& PageID, const FMetasoundFrontendLiteral&)
 		{
-			DocPageIds.Add(PageID);
+			DefaultPageIds.Add(PageID);
 		});
 
-		return TryResolveTargetPageID(DocPageIds, OutResolvedPageID);
+		return ResolveTargetPageIDInternal(DefaultPageIds);
 	}
 
-	bool FDocumentBuilderRegistry::TryResolveTargetPageID(const TArray<FMetasoundFrontendClassInputDefault>& InClassDefaults, FGuid& OutResolvedPageID) const
+	FGuid FDocumentBuilderRegistry::ResolveTargetPageID(const TArray<FMetasoundFrontendClassInputDefault>& InClassDefaults) const
 	{
-		TSet<FGuid> DocPageIds;
-		for (const FMetasoundFrontendClassInputDefault& ClassDefault : InClassDefaults)
-		{
-			DocPageIds.Add(ClassDefault.PageID);
-		}
-
-		return TryResolveTargetPageID(DocPageIds, OutResolvedPageID);
+		TSet<FGuid> DefaultPageIds;
+		Algo::Transform(InClassDefaults, DefaultPageIds, [](const FMetasoundFrontendClassInputDefault& ClassDefault) { return ClassDefault.PageID; });
+		return ResolveTargetPageIDInternal(DefaultPageIds);
 	}
 
-	bool FDocumentBuilderRegistry::TryResolveTargetPageID(const TSet<FGuid>& InPageIDs, FGuid& OutResolvedPageID) const
+	FGuid FDocumentBuilderRegistry::ResolveTargetPageIDInternal(const TSet<FGuid>& InPageIDs) const
 	{
 		FName PlatformName = FPlatformProperties::IniPlatformName();
 
@@ -367,58 +421,52 @@ namespace Metasound::Engine
 			FAuditionPageInfo PreviewInfo = OnResolveAuditionPageInfo.Execute(InPageIDs);
 			if (PreviewInfo.PageID.IsSet())
 			{
-				OutResolvedPageID = PreviewInfo.PageID.GetValue();
-				return InPageIDs.Contains(OutResolvedPageID);
+				return PreviewInfo.PageID.GetValue();
 			}
 
 			PlatformName = PreviewInfo.PlatformName;
 		}
 #endif // WITH_EDITOR
 
-		const UMetaSoundSettings* Settings = GetDefault<UMetaSoundSettings>();
-		if (Settings)
+		if (const UMetaSoundSettings* Settings = GetDefault<UMetaSoundSettings>())
 		{
-			const FGuid& TargetPageID = Settings->GetTargetPageID();
-			const TArray<FMetaSoundPageSettings>& PageSettingsArray = Settings->GetPageSettings();
+			const FGuid& TargetPageID = Settings->GetTargetPageSettings().UniqueId;
+			return ResolveTargetPageIDInternal(*Settings, InPageIDs, TargetPageID, PlatformName);
+		}
 
-			bool bFoundMatch = false;
-			for (int32 Index = PageSettingsArray.Num() - 1; Index >= 0; --Index)
+		return Frontend::DefaultPageID;
+	}
+
+	FGuid FDocumentBuilderRegistry::ResolveTargetPageIDInternal(const UMetaSoundSettings& Settings, const TSet<FGuid>& InPageIDs, const FGuid& TargetPageID, FName PlatformName) const
+	{
+		bool bResolved = false;
+		FGuid ResolvedPageID = Frontend::DefaultPageID;
+		constexpr bool bReverse = true;
+		bool bFoundMatch = false;
+		Settings.IteratePageSettings([&](const FMetaSoundPageSettings& PageSettings)
+		{
+			bFoundMatch |= PageSettings.UniqueId == TargetPageID;
+			if (bFoundMatch && !bResolved)
 			{
-				const FMetaSoundPageSettings& PageSettings = PageSettingsArray[Index];
-				bFoundMatch |= PageSettings.UniqueId == TargetPageID;
 				const bool bAssetImplementsPage = InPageIDs.Contains(PageSettings.UniqueId);
-				if (bFoundMatch && bAssetImplementsPage)
+				if (bAssetImplementsPage)
 				{
 #if WITH_EDITOR
-					const bool bIsCooked = PageSettings.IsCooked.GetValueForPlatform(PlatformName);
+					const bool bIsCooked = !PageSettings.ExcludePageFromCook(PlatformName);
 					if (bIsCooked)
 					{
-						OutResolvedPageID = PageSettings.UniqueId;
-						return true;
+						bResolved = true;
+						ResolvedPageID = PageSettings.UniqueId;
 					}
 #else // !WITH_EDITOR
-					OutResolvedPageID = PageSettings.UniqueId;
-					return true;
+					bResolved = true;
+					ResolvedPageID = PageSettings.UniqueId;
 #endif // !WITH_EDITOR
 				}
 			}
-		}
-
-		// In the case that no id was able to be resolved, returned value should use
-		// default graph. If no pages were implemented, error and attempt to return
-		// a valid pageID, but there is no guarantee the page ID represents an implemented
-		// page.
-		auto IsValidPageID = [](const FGuid& InPageID) { return InPageID != Frontend::DefaultPageID; };
-		if (const FGuid* ImplPageID = Algo::FindByPredicate(InPageIDs, IsValidPageID))
-		{
-			OutResolvedPageID = *ImplPageID;
-		}
-		else
-		{
-			OutResolvedPageID = Frontend::DefaultPageID;
-		}
+		}, bReverse);
 
-		return InPageIDs.Contains(OutResolvedPageID);
+		return ResolvedPageID;
 	}
 
 	void FDocumentBuilderRegistry::SetEventLogVerbosity(ELogEvent Event, ELogVerbosity::Type Verbosity)
diff --git a/Engine/Plugins/Runtime/Metasound/Source/MetasoundEngine/Private/MetasoundEngineAsset.h b/Engine/Plugins/Runtime/Metasound/Source/MetasoundEngine/Private/MetasoundEngineAsset.h
--- a/Engine/Plugins/Runtime/Metasound/Source/MetasoundEngine/Private/MetasoundEngineAsset.h
+++ b/Engine/Plugins/Runtime/Metasound/Source/MetasoundEngine/Private/MetasoundEngineAsset.h
@@ -16,6 +16,7 @@
 
 #if WITH_EDITORONLY_DATA
 #include "Algo/Transform.h"
+#include "Interfaces/ITargetPlatform.h"
 #include "MetasoundFrontendRegistryContainer.h"
 #include "Misc/DataValidation.h"
 #include "UObject/GarbageCollection.h"
@@ -210,9 +211,14 @@ namespace Metasound::Engine
 			const bool bCanEverExecute = Metasound::CanEverExecuteGraph(bIsCooking);
 			if (!bCanEverExecute)
 			{
+				FName PlatformName;
+				if (const ITargetPlatform* TargetPlatform = InSaveContext.GetTargetPlatform())
+				{
+					PlatformName = *TargetPlatform->PlatformName();
+				}
 				const bool bIsDeterministic = SerializationRequiresDeterminism(bIsCooking);
 				FDocumentIDGenerator::FScopeDeterminism DeterminismScope(bIsDeterministic);
-				InMetaSound.UpdateAndRegisterForSerialization();
+				InMetaSound.UpdateAndRegisterForSerialization(PlatformName);
 			}
  			else if (FApp::CanEverRenderAudio())
 			{
diff --git a/Engine/Plugins/Runtime/Metasound/Source/MetasoundEngine/Private/MetasoundSettings.cpp b/Engine/Plugins/Runtime/Metasound/Source/MetasoundEngine/Private/MetasoundSettings.cpp
--- a/Engine/Plugins/Runtime/Metasound/Source/MetasoundEngine/Private/MetasoundSettings.cpp
+++ b/Engine/Plugins/Runtime/Metasound/Source/MetasoundEngine/Private/MetasoundSettings.cpp
@@ -11,21 +11,6 @@
 
 namespace Metasound::SettingsPrivate
 {
-	FAutoConsoleCommand CVarMetaSoundSetTargetPage(
-		TEXT("au.MetaSound.Pages.SetTarget"),
-		TEXT("Sets the target page to that with the given name. If name not specified or not found, command is ignored.\n"),
-		FConsoleCommandWithArgsDelegate::CreateLambda([](const TArray<FString>& Args)
-		{
-			if (!Args.IsEmpty())
-			{
-				if (UMetaSoundSettings* Settings = GetMutableDefault<UMetaSoundSettings>())
-				{
-					Settings->SetTargetPage(FName { *Args.Last() });
-				}
-			}
-		})
-	);
-
 #if WITH_EDITOR
 	template<typename SettingsStructType>
 	TSet<FName> GetStructNames(const TArray<SettingsStructType>& InSettings, int32 IgnoreIndex = INDEX_NONE)
@@ -145,67 +130,94 @@ namespace Metasound::SettingsPrivate
 #endif // WITH_EDITOR
 } // namespace Metasound::SettingsPrivate
 
-
 #if WITH_EDITOR
-void UMetaSoundSettings::ConformPageSettingsDefault(bool bNotifyDefaultConformed)
+bool FMetaSoundPageSettings::ExcludePageFromCook(FName PlatformName) const
+{
+	if (PlatformCanTargetPage(PlatformName))
+	{
+		return false;
+	}
+
+	return Exclude.GetValueForPlatform(PlatformName);
+}
+
+TArray<FName> FMetaSoundPageSettings::GetImplementedPlatforms() const
+{
+	TArray<FName> PlatformNames;
+	Target.PerPlatform.GetKeys(PlatformNames);
+	return PlatformNames;
+}
+
+bool FMetaSoundPageSettings::PlatformCanTargetPage(FName PlatformName) const
+{
+	const bool bIsTargeted = Target.GetValueForPlatform(PlatformName);
+	return bIsTargeted;
+}
+
+void UMetaSoundSettings::ConformPageSettings(bool bNotifyDefaultRenamed)
 {
 	using namespace Metasound;
 
-	bool bContainsPageDefault = false;
-	bool bDefaultConformed = false;
-	bool bTargetFound = TargetPageName == Frontend::DefaultPageName;
+	bool bTargetFound = false;
+	bool bDefaultRenamed = false;
 	for (int32 Index = PageSettings.Num() - 1; Index >= 0; --Index)
 	{
 		FMetaSoundPageSettings& Page = PageSettings[Index];
 		const bool bIsDefaultName = Page.Name == Frontend::DefaultPageName;
 		if (bIsDefaultName)
 		{
-			if (Page.UniqueId != Frontend::DefaultPageID)
-			{
-				Page.UniqueId = { };
-				bDefaultConformed = true;
-			}
-
-			bContainsPageDefault = true;
-		}
-		else
-		{
-			if (Page.UniqueId == Frontend::DefaultPageID)
-			{
-				Page.UniqueId = FGuid::NewGuid();
-				bDefaultConformed = true;
-			}
+			const TSet<FName> PageNames(GetPageNames());
+			Page.Name = SettingsPrivate::GenerateUniqueName(PageNames, *Page.Name.ToString());
+			bDefaultRenamed = true;
 		}
 		bTargetFound |= TargetPageName == Page.Name;
 	}
 
-	if (!bContainsPageDefault)
+	DefaultPageSettings.Target = PageSettings.IsEmpty();
+	if (DefaultPageSettings.Target.GetValue())
 	{
-		FMetaSoundPageSettings DefaultSettings;
-		DefaultSettings.Name = Frontend::DefaultPageName;
-		PageSettings.Insert(MoveTemp(DefaultSettings), 0);
-		bDefaultConformed = true;
+		TargetPageName = DefaultPageSettings.Name;
+	}
+	else
+	{
+		if (!bTargetFound)
+		{
+			TargetPageName = PageSettings.Last().Name;
+		}
 	}
 
-	if (!bTargetFound)
+#if WITH_EDITORONLY_DATA
 	{
-		TargetPageName = Frontend::DefaultPageName;
+		FScopeLock Lock(&CookPlatformTargetCritSec);
+		CookPlatformTargetPageIDs.Reset();
+		CookPlatformTargetPage = { };
 	}
+#endif // WITH_EDITORONLY_DATA
 
-	if (bNotifyDefaultConformed && bDefaultConformed)
+	TargetPageNameOverride.Reset();
+
+	if (bNotifyDefaultRenamed && bDefaultRenamed)
 	{
-		OnDefaultConformed.Broadcast();
+		OnDefaultRenamed.Broadcast();
 	}
 }
 #endif // WITH_EDITOR
 
 const FMetaSoundPageSettings* UMetaSoundSettings::FindPageSettings(FName Name) const
 {
+	if (Name == Metasound::Frontend::DefaultPageName)
+	{
+		return &GetDefaultPageSettings();
+	}
 	return Metasound::SettingsPrivate::FindSettingsStruct(PageSettings, Name);
 }
 
 const FMetaSoundPageSettings* UMetaSoundSettings::FindPageSettings(const FGuid& InPageID) const
 {
+	if (InPageID == Metasound::Frontend::DefaultPageID)
+	{
+		return &GetDefaultPageSettings();
+	}
 	return Metasound::SettingsPrivate::FindSettingsStruct(PageSettings, InPageID);
 }
 
@@ -219,19 +231,72 @@ const FMetaSoundQualitySettings* UMetaSoundSettings::FindQualitySettings(const F
 	return Metasound::SettingsPrivate::FindSettingsStruct(QualitySettings, InQualityID);
 }
 
-const FGuid& UMetaSoundSettings::GetTargetPageID() const
+const FMetaSoundPageSettings& UMetaSoundSettings::GetDefaultPageSettings() const
+{
+	return DefaultPageSettings;
+}
+
+#if WITH_EDITOR
+TArray<FName> UMetaSoundSettings::GetImplementedPagePlatforms() const
+{
+	TSet<FName> PlatformNames { FPlatformProperties::IniPlatformName() };
+	IteratePageSettings([&](const FMetaSoundPageSettings& PageSettings)
+	{
+		TArray<FName> PagePlatforms = PageSettings.GetImplementedPlatforms();
+		PlatformNames.Append(MoveTemp(PagePlatforms));
+	});
+	return PlatformNames.Array();
+}
+#endif // WITH_EDITOR
+
+#if WITH_EDITORONLY_DATA
+const TArray<FGuid>& UMetaSoundSettings::GetCookedTargetPageIDs(FName PlatformName) const
+{
+	FScopeLock Lock(&CookPlatformTargetCritSec);
+	if (PlatformName != CookPlatformTargetPage)
+	{
+		CookPlatformTargetPage = PlatformName;
+		CookPlatformTargetPageIDs.Reset();
+		Algo::TransformIf(PageSettings, CookPlatformTargetPageIDs,
+		[&PlatformName](const FMetaSoundPageSettings& PageSetting)
+		{
+#if WITH_EDITOR
+			return PageSetting.PlatformCanTargetPage(PlatformName);
+#else // !WITH_EDITOR
+			return true;
+#endif // !WITH_EDITOR
+		},
+		[](const FMetaSoundPageSettings& PageSetting)
+		{
+			return PageSetting.UniqueId;
+		});
+	}
+
+	return CookPlatformTargetPageIDs;
+}
+#endif // WITH_EDITORONLY_DATA
+
+const FMetaSoundPageSettings& UMetaSoundSettings::GetTargetPageSettings() const
 {
-	if (const FMetaSoundPageSettings* TargetSettings = FindPageSettings(TargetPageName))
+	const FName TargetPage = TargetPageNameOverride.IsSet() ? *TargetPageNameOverride : TargetPageName;
+
+	if (const FMetaSoundPageSettings* TargetSettings = FindPageSettings(TargetPage))
 	{
-		return TargetSettings->UniqueId;
+		if (TargetSettings->Target.GetValue())
+		{
+			return *TargetSettings;
+		}
 	}
 
-	if (PageSettings.IsEmpty())
+	for (const FMetaSoundPageSettings& Setting : PageSettings)
 	{
-		return Metasound::Frontend::DefaultPageID;
+		if (Setting.Target.GetValue())
+		{
+			return Setting;
+		}
 	}
 
-	return PageSettings.Last().UniqueId;
+	return DefaultPageSettings;
 }
 
 #if WITH_EDITOR
@@ -242,8 +307,8 @@ void UMetaSoundSettings::PostEditChangeChainProperty(FPropertyChangedChainEvent&
 	PostEditChainChangedStructMember(PostEditChangeChainProperty, PageSettings, GetPageSettingPropertyName(), TEXT("New Page"));
 	PostEditChainChangedStructMember(PostEditChangeChainProperty, QualitySettings, GetQualitySettingPropertyName(), TEXT("New Quality"));
 
-	constexpr bool bNotifyDefaultConformed = true;
-	ConformPageSettingsDefault(bNotifyDefaultConformed);
+	constexpr bool bNotifyDefaultRenamed = true;
+	ConformPageSettings(bNotifyDefaultRenamed);
 
 	Super::PostEditChangeChainProperty(PostEditChangeChainProperty);
 }
@@ -252,6 +317,9 @@ void UMetaSoundSettings::PostEditChangeProperty(FPropertyChangedEvent& PropertyC
 {	
 	Super::PostEditChangeProperty(PropertyChangedEvent);
 	
+	constexpr bool bNotifyDefaultRenamed = true;
+	ConformPageSettings(bNotifyDefaultRenamed);
+
 	if (PropertyChangedEvent.MemberProperty->GetName() == GetPageSettingPropertyName())
 	{
 		OnPageSettingsUpdated.Broadcast();
@@ -262,9 +330,6 @@ void UMetaSoundSettings::PostEditChangeProperty(FPropertyChangedEvent& PropertyC
 
 void UMetaSoundSettings::PostInitProperties()
 {
-	constexpr bool bNotifyDefaultConformed = false;
-	ConformPageSettingsDefault(bNotifyDefaultConformed);
-
 	Super::PostInitProperties();
 }
 #endif // WITH_EDITOR
@@ -275,8 +340,11 @@ bool UMetaSoundSettings::SetTargetPage(FName PageName)
 	{
 		if (TargetPageName != PageSetting->Name)
 		{
-			TargetPageName = PageSetting->Name;
-			return true;
+			if (PageSetting->Target.GetValue())
+			{
+				TargetPageNameOverride = PageSetting->Name;
+				return true;
+			}
 		}
 	}
 
@@ -284,9 +352,9 @@ bool UMetaSoundSettings::SetTargetPage(FName PageName)
 }
 
 #if WITH_EDITORONLY_DATA
-Metasound::Engine::FOnSettingsDefaultConformed& UMetaSoundSettings::GetOnDefaultConformedDelegate()
+Metasound::Engine::FOnSettingsDefaultConformed& UMetaSoundSettings::GetOnDefaultRenamedDelegate()
 {
-	return OnDefaultConformed;
+	return OnDefaultRenamed;
 }
 
 Metasound::Engine::FOnPageSettingsUpdated& UMetaSoundSettings::GetOnPageSettingsUpdatedDelegate()
@@ -311,8 +379,10 @@ TArray<FName> UMetaSoundSettings::GetPageNames()
 	if (const UMetaSoundSettings* Settings = GetDefault<UMetaSoundSettings>())
 	{
 		TArray<FName> Names;
-		auto GetName = [](const FMetaSoundPageSettings& Page) { return Page.Name; };
-		Algo::Transform(Settings->GetPageSettings(), Names, GetName);
+		Settings->IteratePageSettings([&Names](const FMetaSoundPageSettings& PageSetting)
+		{
+			Names.Add(PageSetting.Name);
+		});
 		return Names;
 	}
 
@@ -333,4 +403,23 @@ TArray<FName> UMetaSoundSettings::GetQualityNames()
 }
 #endif // WITH_EDITOR
 
+void UMetaSoundSettings::IteratePageSettings(TFunctionRef<void(const FMetaSoundPageSettings&)> Iter, bool bReverse) const
+{
+	if (bReverse)
+	{
+		for (int32 Index = PageSettings.Num() - 1; Index >= 0; --Index)
+		{
+			Iter(PageSettings[Index]);
+		}
+		Iter(GetDefaultPageSettings());
+	}
+	else
+	{
+		Iter(GetDefaultPageSettings());
+		for (const FMetaSoundPageSettings& Setting : PageSettings)
+		{
+			Iter(Setting);
+		}
+	}
+}
 #undef LOCTEXT_NAMESPACE // MetaSound
diff --git a/Engine/Plugins/Runtime/Metasound/Source/MetasoundEngine/Public/MetasoundBuilderBase.cpp b/Engine/Plugins/Runtime/Metasound/Source/MetasoundEngine/Public/MetasoundBuilderBase.cpp
--- a/Engine/Plugins/Runtime/Metasound/Source/MetasoundEngine/Public/MetasoundBuilderBase.cpp
+++ b/Engine/Plugins/Runtime/Metasound/Source/MetasoundEngine/Public/MetasoundBuilderBase.cpp
@@ -1179,19 +1179,16 @@ FMetasoundFrontendLiteral UMetaSoundBuilderBase::GetNodeInputClassDefault(const
 	{
 		if (const TArray<FMetasoundFrontendClassInputDefault>* ClassDefaults = Builder.FindNodeClassInputDefaults(InputHandle.NodeID, Vertex->Name))
 		{
-			FGuid ResolvedPageID;
-			if (FDocumentBuilderRegistry::GetChecked().TryResolveTargetPageID(*ClassDefaults, ResolvedPageID))
+			const FGuid ResolvedPageID = FDocumentBuilderRegistry::GetChecked().ResolveTargetPageID(*ClassDefaults);
+			OutResult = EMetaSoundBuilderResult::Succeeded;
+			auto MatchesPageID = [&ResolvedPageID](const FMetasoundFrontendClassInputDefault& Default)
+			{
+				return Default.PageID == ResolvedPageID;
+			};
+			if (const FMetasoundFrontendClassInputDefault* Default = ClassDefaults->FindByPredicate(MatchesPageID))
 			{
 				OutResult = EMetaSoundBuilderResult::Succeeded;
-				auto MatchesPageID = [&ResolvedPageID](const FMetasoundFrontendClassInputDefault& Default)
-				{
-					return Default.PageID == ResolvedPageID;
-				};
-				if (const FMetasoundFrontendClassInputDefault* Default = ClassDefaults->FindByPredicate(MatchesPageID))
-				{
-					OutResult = EMetaSoundBuilderResult::Succeeded;
-					return Default->Literal;
-				}
+				return Default->Literal;
 			}
 		}
 	}
diff --git a/Engine/Plugins/Runtime/Metasound/Source/MetasoundEngine/Public/MetasoundBuilderSubsystem.h b/Engine/Plugins/Runtime/Metasound/Source/MetasoundEngine/Public/MetasoundBuilderSubsystem.h
--- a/Engine/Plugins/Runtime/Metasound/Source/MetasoundEngine/Public/MetasoundBuilderSubsystem.h
+++ b/Engine/Plugins/Runtime/Metasound/Source/MetasoundEngine/Public/MetasoundBuilderSubsystem.h
@@ -118,6 +118,8 @@ protected:
 private:
 	static TOptional<Metasound::FAnyDataReference> CreateDataReference(const Metasound::FOperatorSettings& InOperatorSettings, FName DataType, const Metasound::FLiteral& InLiteral, Metasound::EDataReferenceAccessType AccessType);
 
+	const FMetasoundFrontendGraph& GetConstTargetPageGraphChecked() const;
+
 	const UMetaSoundSource& GetMetaSoundSource() const;
 	UMetaSoundSource& GetMetaSoundSource();
 
@@ -142,6 +144,7 @@ private:
 
 	TArray<uint64> LiveComponentIDs;
 	FDelegateHandle LiveComponentHandle;
+	FGuid TargetPageID = Metasound::Frontend::DefaultPageID;
 
 	friend class UMetaSoundBuilderSubsystem;
 };
@@ -267,6 +270,10 @@ public:
 	UFUNCTION(BlueprintCallable, Category = "Audio|MetaSound|Builder")
 	void RegisterSourceBuilder(FName BuilderName, UMetaSoundSourceBuilder* Builder);
 
+	// Sets the targeted page for all MetaSound graph & input default to resolve against.
+	// If target page is not implemented (or cooked in a runtime build) for the active platform,
+	// uses order of cooked pages(see 'Page Settings' for order) falling back to lower index - ordered page
+	// implemented in MetaSound asset. If no fallback is found, uses default graph/input default.
 	UFUNCTION(BlueprintCallable, Category = "MetaSounds|Pages")
 	UPARAM(DisplayName = "TargetPageChanged") bool SetTargetPage(FName PageName);
 
diff --git a/Engine/Plugins/Runtime/Metasound/Source/MetasoundEngine/Public/MetasoundDocumentBuilderRegistry.h b/Engine/Plugins/Runtime/Metasound/Source/MetasoundEngine/Public/MetasoundDocumentBuilderRegistry.h
--- a/Engine/Plugins/Runtime/Metasound/Source/MetasoundEngine/Public/MetasoundDocumentBuilderRegistry.h
+++ b/Engine/Plugins/Runtime/Metasound/Source/MetasoundEngine/Public/MetasoundDocumentBuilderRegistry.h
@@ -12,6 +12,9 @@
 #include "Misc/ScopeLock.h"
 
 
+// Forward Declarations
+class UMetaSoundSettings;
+
 namespace Metasound::Engine
 {
 #if WITH_EDITOR
@@ -107,6 +110,12 @@ namespace Metasound::Engine
 
 		// Frontend::IDocumentBuilderRegistry Implementation
 #if WITH_EDITORONLY_DATA
+		// Given the provided builder, removes paged data within the associated document for a cooked build.
+		// This function removes graphs and input defaults which are not to ever be used by a given cook
+		// platform, allowing users to optimize away data and scale the amount of memory required for
+		// initial load of input UObjects and graph topology, which can also positively effect runtime
+		// performance as well, etc. Returns true if builder modified the document, false if not.
+		virtual bool CookPages(FName PlatformName, FMetaSoundFrontendDocumentBuilder& Builder) const override;
 		virtual FMetaSoundFrontendDocumentBuilder& FindOrBeginBuilding(TScriptInterface<IMetaSoundDocumentInterface> MetaSound) override;
 #endif // WITH_EDITORONLY_DATA
 
@@ -137,9 +146,9 @@ namespace Metasound::Engine
 		bool ReloadBuilder(const FMetasoundFrontendClassName& InClassName) const override;
 
 		// Given the provided document and its respective pages, returns the PageID to be used for runtime IGraph and proxy generation.
-		virtual bool TryResolveTargetPageID(const FMetasoundFrontendGraphClass& InGraphClass, FGuid& OutResolvedPageID) const override;
-		virtual bool TryResolveTargetPageID(const FMetasoundFrontendClassInput& InClassInput, FGuid& OutResolvedPageID) const override;
-		virtual bool TryResolveTargetPageID(const TArray<FMetasoundFrontendClassInputDefault>& Defaults, FGuid& OutResolvedPageID) const override;
+		virtual FGuid ResolveTargetPageID(const FMetasoundFrontendGraphClass& InGraphClass) const override;
+		virtual FGuid ResolveTargetPageID(const FMetasoundFrontendClassInput& InClassInput) const override;
+		virtual FGuid ResolveTargetPageID(const TArray<FMetasoundFrontendClassInputDefault>& Defaults) const override;
 
 		void SetEventLogVerbosity(ELogEvent Event, ELogVerbosity::Type Verbosity);
 
@@ -147,7 +156,8 @@ namespace Metasound::Engine
 		void AddBuilderInternal(const FMetasoundFrontendClassName& InClassName, UMetaSoundBuilderBase* NewBuilder) const;
 		bool CanPostEventLog(ELogEvent Event, ELogVerbosity::Type Verbosity) const;
 		void FinishBuildingInternal(UMetaSoundBuilderBase& Builder, bool bForceUnregisterNodeClass) const;
-		bool TryResolveTargetPageID(const TSet<FGuid>& InPageIDs, FGuid& OutResolvedPageID) const;
+		FGuid ResolveTargetPageIDInternal(const TSet<FGuid>& InPageIDs) const;
+		FGuid ResolveTargetPageIDInternal(const UMetaSoundSettings& Settings, const TSet<FGuid>& InPageIDs, const FGuid& TargetPageID, FName PlatformName) const;
 
 #if WITH_EDITOR
 		FOnResolveAuditionPageInfo OnResolveAuditionPageInfo;
diff --git a/Engine/Plugins/Runtime/Metasound/Source/MetasoundEngine/Public/MetasoundSettings.h b/Engine/Plugins/Runtime/Metasound/Source/MetasoundEngine/Public/MetasoundSettings.h
--- a/Engine/Plugins/Runtime/Metasound/Source/MetasoundEngine/Public/MetasoundSettings.h
+++ b/Engine/Plugins/Runtime/Metasound/Source/MetasoundEngine/Public/MetasoundSettings.h
@@ -4,6 +4,8 @@
 #include "Delegates/DelegateCombinations.h"
 #include "Engine/DeveloperSettings.h"
 #include "MetasoundFrontendDocument.h"
+#include "Misc/Optional.h"
+#include "Misc/ScopeLock.h"
 #include "UObject/NoExportTypes.h"
 #include "UObject/Object.h"
 #include "UObject/ObjectMacros.h"
@@ -59,20 +61,33 @@ USTRUCT()
 struct METASOUNDENGINE_API FMetaSoundPageSettings
 {
 	GENERATED_BODY()
-	
+
 	UPROPERTY()
-	FGuid UniqueId;
+	FGuid UniqueId = Metasound::Frontend::DefaultPageID;
 
-	/** Name of this page's setting. This will appear in the MetaSound Asset Editor's 'Page Selector'.
-		The names should be unique and adequately describe the Entry. "High", "Low" etc. **/
+	/** Name of this page's setting to be displayed in editors and used for identification from Blueprint/native API. */
 	UPROPERTY(EditAnywhere, Category = "Pages")
-	FName Name;
+	FName Name = Metasound::Frontend::DefaultPageName;
+
+private:
+	// When true, page can be targeted for the assigned platform(s)/platform group(s).
+	UPROPERTY(EditAnywhere, config, Category = "Pages")
+	FPerPlatformBool Target = true;
 
 #if WITH_EDITORONLY_DATA
-	// When true, page data defined on serialized MetaSounds are included in cooked build (for the assigned platform(s)).
+	// When true, exclude page data when cooking from the assigned platform(s)/platform group(s) (ignored if target is true).
 	UPROPERTY(EditAnywhere, config, Category = "Pages")
-	FPerPlatformBool IsCooked = true;
+	FPerPlatformBool Exclude = false;
 #endif //WITH_EDITORONLY_DATA
+
+public:
+#if WITH_EDITOR
+	bool ExcludePageFromCook(FName PlatformName) const;
+	TArray<FName> GetImplementedPlatforms() const;
+	bool PlatformCanTargetPage(FName PlatformName) const;
+#endif //WITH_EDITOR
+
+	friend class UMetaSoundSettings;
 };
 
 USTRUCT()
@@ -135,13 +150,14 @@ public:
 
 private:
 #if WITH_EDITORONLY_DATA
-	Metasound::Engine::FOnSettingsDefaultConformed OnDefaultConformed;
+	Metasound::Engine::FOnSettingsDefaultConformed OnDefaultRenamed;
 	Metasound::Engine::FOnPageSettingsUpdated OnPageSettingsUpdated;
 #endif //WITH_EDITORONLY_DATA
 
 	/** Page Name to target when attempting to execute MetaSound. If target page is not implemented (or cooked in a runtime build)
 	  * for the active platform, uses order of cooked pages (see 'Page Settings' for order) falling back to lower index-ordered page
-	  * implemented in MetaSound asset. */
+	  * implemented in MetaSound asset. If no fallback is found, uses default implementation. If any pages are specified, target
+	  * cannot be set directly to default (allows for removal of default implementation when cooking for performance reasons). */
 	UPROPERTY(EditAnywhere, config, Category = Pages, meta = (GetOptions = "MetasoundEngine.MetaSoundSettings.GetPageNames"))
 	FName TargetPageName = Metasound::Frontend::DefaultPageName;
 
@@ -152,10 +168,23 @@ private:
 	UPROPERTY(EditAnywhere, config, Category = Pages)
 	TArray<FMetaSoundPageSettings> PageSettings;
 
+	/** Default page settings to be used in editor and if no other page settings are targeted or defined. */
+	UPROPERTY()
+	FMetaSoundPageSettings DefaultPageSettings;
+
 	/** Array of possible quality settings for Metasounds to chose from */
 	UPROPERTY(EditAnywhere, config, Category = Quality)
 	TArray<FMetaSoundQualitySettings> QualitySettings;
 
+	/** Target page override to use set from code/script (supercedes serialized 'TargetPageName' field). */
+	TOptional<FName> TargetPageNameOverride;
+
+#if WITH_EDITORONLY_DATA
+	mutable FCriticalSection CookPlatformTargetCritSec;
+	mutable TArray<FGuid> CookPlatformTargetPageIDs;
+	mutable FName CookPlatformTargetPage;
+#endif // WITH_EDITORONLY_DATA
+
 public:
 	// Returns the page settings with the provided name. If there are multiple settings
 	// with the same name, selection within the duplicates is undefined.
@@ -171,23 +200,36 @@ public:
 	// Returns the quality settings with the unique ID given.
 	const FMetaSoundQualitySettings* FindQualitySettings(const FGuid& InQualityID) const;
 
-	// Returns the target page name.
-	const FName& GetTargetPageName() const { return TargetPageName; }
+	const FMetaSoundPageSettings& GetDefaultPageSettings() const;
+
+#if WITH_EDITORONLY_DATA
+	const TArray<FGuid>& GetCookedTargetPageIDs(FName PlatformName) const;
+#endif // WITH_EDITORONLY_DATA
+
+#if WITH_EDITOR
+	// Returns superset of explicitly implemented page platforms (and groups) as
+	// defined in PageSettings.
+	TArray<FName> GetImplementedPagePlatforms() const;
+#endif // WITH_EDITOR
 
-	// Returns the target page ID.
-	const FGuid& GetTargetPageID() const;
+	// Returns the currently targeted page settings.
+	const FMetaSoundPageSettings& GetTargetPageSettings() const;
 
-	const TArray<FMetaSoundPageSettings>& GetPageSettings() const { return PageSettings; }
 	const TArray<FMetaSoundQualitySettings>& GetQualitySettings() const { return QualitySettings; }
 
 #if WITH_EDITORONLY_DATA
-	Metasound::Engine::FOnSettingsDefaultConformed& GetOnDefaultConformedDelegate();
+	Metasound::Engine::FOnSettingsDefaultConformed& GetOnDefaultRenamedDelegate();
 	Metasound::Engine::FOnPageSettingsUpdated& GetOnPageSettingsUpdatedDelegate();
 
 	static FName GetPageSettingPropertyName();
 	static FName GetQualitySettingPropertyName();
 #endif // WITH_EDITORONLY_DATA
 
+	// Iterates possible page settings in order. Does not include default page settings. If optionally set
+	// to reverse, iterates in reverse. Does *not* include default setting (as fallback therein is not related
+	// to user-defined, ordered page settings).
+	void IteratePageSettings(TFunctionRef<void(const FMetaSoundPageSettings&)> Iter, bool bReverse = false) const;
+
 	// Sets the target page to the given name. Returns true if associated page settings were found
 	// and target set, false if not found and not set.
 	bool SetTargetPage(FName PageName);
@@ -207,8 +249,9 @@ public:
 	*/
 	UFUNCTION()
 	static TArray<FName> GetQualityNames();
+
 private:
-	void ConformPageSettingsDefault(bool bNotifyDefaultConformed);
+	void ConformPageSettings(bool bNotifyDefaultRenamed);
 
 	virtual void PostEditChangeChainProperty(FPropertyChangedChainEvent& PropertyChangedEvent) override;
 	virtual void PostEditChangeProperty(FPropertyChangedEvent& PropertyChangedEvent) override;
diff --git a/Engine/Plugins/Runtime/Metasound/Source/MetasoundFrontend/Private/MetasoundAssetBase.cpp b/Engine/Plugins/Runtime/Metasound/Source/MetasoundFrontend/Private/MetasoundAssetBase.cpp
--- a/Engine/Plugins/Runtime/Metasound/Source/MetasoundFrontend/Private/MetasoundAssetBase.cpp
+++ b/Engine/Plugins/Runtime/Metasound/Source/MetasoundFrontend/Private/MetasoundAssetBase.cpp
@@ -313,18 +313,24 @@ void FMetasoundAssetBase::UpdateAndRegisterForExecution(Metasound::Frontend::FMe
 		RegisterAssetDependencies(InRegistrationOptions);
 	}
 
+	UObject* Owner = GetOwningAsset();
+	check(Owner);
+
 	// This should not be necessary as it should be added on asset load,
 	// but currently registration is required to be called prior to adding
 	// an object-defined graph class to the registry so it was placed here.
-	if (UObject* OwningObject = GetOwningAsset())
-	{
-		IMetaSoundAssetManager::GetChecked().AddOrUpdateAsset(*OwningObject);
-	}
+	IMetaSoundAssetManager::GetChecked().AddOrUpdateAsset(*Owner);
 
 	// Auto update must be done after all referenced asset classes are registered
 	if (InRegistrationOptions.bAutoUpdate)
 	{
-		const bool bDidUpdate = AutoUpdate(InRegistrationOptions.bAutoUpdateLogWarningOnDroppedConnection);
+#if WITH_EDITORONLY_DATA
+		FMetaSoundFrontendDocumentBuilder& DocBuilder = IDocumentBuilderRegistry::GetChecked().FindOrBeginBuilding(Owner);
+		const bool bDidUpdate = VersionDependencies(DocBuilder, InRegistrationOptions.bAutoUpdateLogWarningOnDroppedConnection);
+#else // !WITH_EDITORONLY_DATA
+		constexpr bool bDidUpdate = false;
+#endif // WITH_EDITORONLY_DATA
+
 #if WITH_EDITOR
 		if (bDidUpdate || InRegistrationOptions.bForceViewSynchronization)
 		{
@@ -348,10 +354,6 @@ void FMetasoundAssetBase::UpdateAndRegisterForExecution(Metasound::Frontend::FMe
 	CacheRegistryMetadata();
 #endif // WITH_EDITOR
 
-	UObject* Owner = GetOwningAsset();
-	check(Owner);
-	const FString AssetName = Owner->GetName();
-
 	GraphRegistryKey = FRegistryContainerImpl::Get().RegisterGraph(Owner);
 	if (GraphRegistryKey.IsValid())
 	{
@@ -364,6 +366,7 @@ void FMetasoundAssetBase::UpdateAndRegisterForExecution(Metasound::Frontend::FMe
 		UClass* Class = Owner->GetClass();
 		check(Class);
 		const FString ClassName = Class->GetName();
+		const FString AssetName = Owner->GetName();
 		UE_LOG(LogMetaSound, Error, TEXT("Registration failed for MetaSound node class '%s' of UObject class '%s'"), *AssetName, *ClassName);
 	}
 }
@@ -376,7 +379,7 @@ void FMetasoundAssetBase::CookMetaSound()
 }
 
 #if WITH_EDITORONLY_DATA
-void FMetasoundAssetBase::UpdateAndRegisterForSerialization()
+void FMetasoundAssetBase::UpdateAndRegisterForSerialization(FName CookPlatformName)
 {
 	using namespace Metasound;
 	using namespace Metasound::Frontend;
@@ -389,11 +392,22 @@ void FMetasoundAssetBase::UpdateAndRegisterForSerialization()
 		return;
 	}
 
-	UpdateAndRegisterReferencesForSerialization();
+	UpdateAndRegisterReferencesForSerialization(CookPlatformName);
 	IMetaSoundAssetManager::GetChecked().AddOrUpdateAsset(*GetOwningAsset());
 
+	UObject* Owner = GetOwningAsset();
+	check(Owner);
+
+	bool bDidUpdate = false;
+
+	FMetaSoundFrontendDocumentBuilder& DocBuilder = IDocumentBuilderRegistry::GetChecked().FindOrBeginBuilding(Owner);
+	if (CookPlatformName.IsValid())
+	{
+		bDidUpdate |= IDocumentBuilderRegistry::GetChecked().CookPages(CookPlatformName, DocBuilder);
+	}
+
 	// Auto update must be done after all referenced asset classes are registered
-	const bool bDidUpdate = AutoUpdate(/*bAutoUpdateLogWarningOnDroppedConnection=*/true);
+	bDidUpdate |= VersionDependencies(DocBuilder, /*bAutoUpdateLogWarningOnDroppedConnection=*/true);
 #if WITH_EDITOR
 	if (bDidUpdate)
 	{
@@ -407,12 +421,8 @@ void FMetasoundAssetBase::UpdateAndRegisterForSerialization()
 	CacheRegistryMetadata();
 #endif // WITH_EDITOR
 
-	UObject* Owner = GetOwningAsset();
-	check(Owner);
-
 	{
 		// Performs document transforms on local copy, which reduces document footprint & renders transforming unnecessary at runtime
-		FMetaSoundFrontendDocumentBuilder& DocBuilder = IDocumentBuilderRegistry::GetChecked().FindOrBeginBuilding(Owner);
 		const bool bContainsTemplateDependency = DocBuilder.ContainsDependencyOfType(EMetasoundFrontendClassType::Template);
 		if (bContainsTemplateDependency)
 		{
@@ -935,12 +945,12 @@ void FMetasoundAssetBase::RegisterAssetDependencies(const Metasound::Frontend::F
 void FMetasoundAssetBase::CookReferencedMetaSounds()
 {
 #if WITH_EDITORONLY_DATA
-	UpdateAndRegisterReferencesForSerialization();
+	UpdateAndRegisterReferencesForSerialization({ });
 #endif // WITH_EDITORONLY_DATA
 }
 
 #if WITH_EDITORONLY_DATA
-void FMetasoundAssetBase::UpdateAndRegisterReferencesForSerialization()
+void FMetasoundAssetBase::UpdateAndRegisterReferencesForSerialization(FName CookPlatformName)
 {
 	using namespace Metasound::Frontend;
 
@@ -952,7 +962,7 @@ void FMetasoundAssetBase::UpdateAndRegisterReferencesForSerialization()
 		{
 			// TODO: Check for infinite recursion and error if so
 			AssetManager.AddOrUpdateAsset(*(Reference->GetOwningAsset()));
-			Reference->UpdateAndRegisterForSerialization();
+			Reference->UpdateAndRegisterForSerialization(CookPlatformName);
 		}
 	}
 }
@@ -962,9 +972,15 @@ bool FMetasoundAssetBase::AutoUpdate(bool bInLogWarningsOnDroppedConnection)
 {
 	using namespace Metasound::Frontend;
 
-	FString OwningAssetName = GetOwningAssetName();
-	const bool bAutoUpdated = FAutoUpdateRootGraph(MoveTemp(OwningAssetName), bInLogWarningsOnDroppedConnection).Transform(GetDocumentHandle());
-	return bAutoUpdated;
+	UObject* Owner = GetOwningAsset();
+	check(Owner);
+
+#if WITH_EDITORONLY_DATA
+	FMetaSoundFrontendDocumentBuilder& DocBuilder = IDocumentBuilderRegistry::GetChecked().FindOrBeginBuilding(Owner);
+	return VersionDependencies(DocBuilder, bInLogWarningsOnDroppedConnection);
+#else // !WITH_EDITORONLY_DATA
+	return false;
+#endif // !WITH_EDITORONLY_DATA
 }
 
 #if WITH_EDITORONLY_DATA
@@ -1016,6 +1032,30 @@ bool FMetasoundAssetBase::TryUpdateInterfaceFromVersion(const FMetasoundFrontend
 	return false;
 }
 
+bool FMetasoundAssetBase::VersionDependencies(FMetaSoundFrontendDocumentBuilder& Builder, bool bInLogWarningsOnDroppedConnection)
+{
+	using namespace Metasound::Frontend;
+
+	bool bDocumentModified = false;
+
+#if WITH_EDITORONLY_DATA
+	const FGuid InitBuildPageID = Builder.GetBuildPageID();
+	Metasound::Frontend::FDocumentHandle DocHandle = GetDocumentHandle();
+	const FMetasoundFrontendGraphClass& RootGraph = Builder.GetConstDocumentChecked().RootGraph;
+	RootGraph.IterateGraphPages([&](const FMetasoundFrontendGraph& Graph)
+	{
+		// Set the build page ID to this graph as a hack to apply dependency versioning logic using
+		// the controller/handle API until auto-update is renamed & moved to use document builder API.
+		Builder.SetBuildPageID(Graph.PageID);
+		FString OwningAssetName = GetOwningAssetName();
+		bDocumentModified |= FAutoUpdateRootGraph(MoveTemp(OwningAssetName), bInLogWarningsOnDroppedConnection).Transform(DocHandle);
+	});
+	Builder.SetBuildPageID(InitBuildPageID);
+#endif // WITH_EDITORONLY_DATA
+
+	return bDocumentModified;
+}
+
 FMetasoundFrontendInterface FMetasoundAssetBase::GetInterfaceToVersion(const FMetasoundFrontendVersion& InterfaceVersion) const
 {
 	using namespace Metasound::Frontend;
diff --git a/Engine/Plugins/Runtime/Metasound/Source/MetasoundFrontend/Private/MetasoundFrontendDocument.cpp b/Engine/Plugins/Runtime/Metasound/Source/MetasoundFrontend/Private/MetasoundFrontendDocument.cpp
--- a/Engine/Plugins/Runtime/Metasound/Source/MetasoundFrontend/Private/MetasoundFrontendDocument.cpp
+++ b/Engine/Plugins/Runtime/Metasound/Source/MetasoundFrontend/Private/MetasoundFrontendDocument.cpp
@@ -31,24 +31,22 @@ namespace Metasound
 		namespace DocumentPrivate
 		{
 			template <typename ResolveType>
-			bool TryResolveTargetPageID(const ResolveType& InToResolve, FGuid& OutPageID)
+			FGuid ResolveTargetPageID(const ResolveType& InToResolve)
 			{
-				OutPageID = Frontend::DefaultPageID;
-
 				// Registry is not available in tests, so for now resolution is considered successful at this level
 				// if registry is not initialized and providing a resolved page ID. TODO: Add a test implementation
 				// that returns the default page (or whatever page behavior is desired for testing).
 				if (IDocumentBuilderRegistry* BuilderRegistry = IDocumentBuilderRegistry::Get())
 				{
-					return IDocumentBuilderRegistry::GetChecked().TryResolveTargetPageID(InToResolve, OutPageID);
+					return IDocumentBuilderRegistry::GetChecked().ResolveTargetPageID(InToResolve);
 				}
 
-				return true;
+				return Frontend::DefaultPageID;
 			}
 		} // namespace DocumentPrivate
 
 #if WITH_EDITORONLY_DATA
-		const FText DefaultGraphPageDisplayName = NSLOCTEXT("MetasoundFrontend", "DefaultGraphPageDisplayName", "Default");
+		const FText DefaultPageDisplayName = NSLOCTEXT("MetasoundFrontend", "DefaultPageDisplayName", "Default");
 #endif // WITH_EDITORONLY_DATA
 
 		namespace DisplayStyle
@@ -1283,20 +1281,16 @@ namespace Metasound::Frontend
 	{
 		ForEachLiteral(static_cast<const FMetasoundFrontendClass&>(InGraphClass), OnLiteral);
 
-		FGuid PageID = Frontend::DefaultPageID;
-		const bool bPageIDResolved = DocumentPrivate::TryResolveTargetPageID(InGraphClass, PageID);
-		if (bPageIDResolved)
+		const FGuid PageID = DocumentPrivate::ResolveTargetPageID(InGraphClass);
+		const FMetasoundFrontendGraph& Graph = InGraphClass.FindConstGraphChecked(PageID);
+		for (const FMetasoundFrontendNode& Node : Graph.Nodes)
 		{
-			const FMetasoundFrontendGraph& Graph = InGraphClass.FindConstGraphChecked(PageID);
-			for (const FMetasoundFrontendNode& Node : Graph.Nodes)
-			{
-				ForEachLiteral(Node, OnLiteral);
-			}
+			ForEachLiteral(Node, OnLiteral);
+		}
 
-			for (const FMetasoundFrontendVariable& Variable : Graph.Variables)
-			{
-				OnLiteral(Variable.TypeName, Variable.Literal);
-			}
+		for (const FMetasoundFrontendVariable& Variable : Graph.Variables)
+		{
+			OnLiteral(Variable.TypeName, Variable.Literal);
 		}
 	}
 
@@ -1304,13 +1298,9 @@ namespace Metasound::Frontend
 	{
 		for (const FMetasoundFrontendClassInput& ClassInput : InClass.Interface.Inputs)
 		{
-			FGuid PageID = Frontend::DefaultPageID;
-			const bool bPageIDResolved = DocumentPrivate::TryResolveTargetPageID(ClassInput, PageID);
-			if (bPageIDResolved)
-			{
-				const FMetasoundFrontendLiteral& DefaultLiteral = ClassInput.FindConstDefaultChecked(PageID);
-				OnLiteral(ClassInput.TypeName, DefaultLiteral);
-			}
+			const FGuid PageID = DocumentPrivate::ResolveTargetPageID(ClassInput);
+			const FMetasoundFrontendLiteral& DefaultLiteral = ClassInput.FindConstDefaultChecked(PageID);
+			OnLiteral(ClassInput.TypeName, DefaultLiteral);
 		}
 	}
 
diff --git a/Engine/Plugins/Runtime/Metasound/Source/MetasoundFrontend/Private/MetasoundFrontendDocumentBuilder.cpp b/Engine/Plugins/Runtime/Metasound/Source/MetasoundFrontend/Private/MetasoundFrontendDocumentBuilder.cpp
--- a/Engine/Plugins/Runtime/Metasound/Source/MetasoundFrontend/Private/MetasoundFrontendDocumentBuilder.cpp
+++ b/Engine/Plugins/Runtime/Metasound/Source/MetasoundFrontend/Private/MetasoundFrontendDocumentBuilder.cpp
@@ -2997,6 +2997,39 @@ void FMetaSoundFrontendDocumentBuilder::Reload(TSharedPtr<Metasound::Frontend::F
 	}
 }
 
+#if WITH_EDITORONLY_DATA
+bool FMetaSoundFrontendDocumentBuilder::RemoveGraphInputDefault(FName InputName, const FGuid& InPageID)
+{
+	using namespace Metasound;
+
+	auto NameMatchesInput = [&InputName](const FMetasoundFrontendClassInput& Input) { return Input.Name == InputName; };
+	FMetasoundFrontendDocument& Document = GetDocumentChecked();
+	TArray<FMetasoundFrontendClassInput>& Inputs = Document.RootGraph.Interface.Inputs;
+
+	const int32 Index = Inputs.IndexOfByPredicate(NameMatchesInput);
+	if (Index != INDEX_NONE)
+	{
+		FMetasoundFrontendClassInput& Input = Inputs[Index];
+		const bool bRemovedDefault = Input.RemoveDefault(InPageID);
+		if (bRemovedDefault)
+		{
+			DocumentDelegates->InterfaceDelegates.OnInputDefaultChanged.Broadcast(Index);
+
+			// Set the input as no longer inheriting default for presets
+			if (IsPreset())
+			{
+				constexpr bool bInputInheritsDefault = false;
+				return SetGraphInputInheritsDefault(InputName, bInputInheritsDefault);
+			}
+
+			return true;
+		}
+	}
+
+	return false;
+}
+#endif // WITH_EDITORONLY_DATA
+
 bool FMetaSoundFrontendDocumentBuilder::RemoveNodeInputDefault(const FGuid& InNodeID, const FGuid& InVertexID, const FGuid* InPageID)
 {
 	using namespace Metasound::Frontend;
diff --git a/Engine/Plugins/Runtime/Metasound/Source/MetasoundFrontend/Private/MetasoundFrontendGraph.cpp b/Engine/Plugins/Runtime/Metasound/Source/MetasoundFrontend/Private/MetasoundFrontendGraph.cpp
--- a/Engine/Plugins/Runtime/Metasound/Source/MetasoundFrontend/Private/MetasoundFrontendGraph.cpp
+++ b/Engine/Plugins/Runtime/Metasound/Source/MetasoundFrontend/Private/MetasoundFrontendGraph.cpp
@@ -37,19 +37,17 @@ namespace Metasound
 			}
 
 			template <typename ResolveType>
-			bool TryResolveTargetPageID(const ResolveType& InToResolve, FGuid& OutPageID)
+			FGuid ResolveTargetPageID(const ResolveType& InToResolve)
 			{
-				OutPageID = Frontend::DefaultPageID;
-
 				// Registry is not available in tests, so for now resolution is considered successful at this level
 				// if registry is not initialized and providing a resolved page ID. TODO: Add a test implementation
 				// that returns the default page (or whatever page behavior is desired for testing).
 				if (IDocumentBuilderRegistry* BuilderRegistry = IDocumentBuilderRegistry::Get())
 				{
-					return IDocumentBuilderRegistry::GetChecked().TryResolveTargetPageID(InToResolve, OutPageID);
+					return IDocumentBuilderRegistry::GetChecked().ResolveTargetPageID(InToResolve);
 				}
 
-				return true;
+				return Frontend::DefaultPageID;
 			}
 		} // namespace GraphPrivate
 	} // namespace Frontend
@@ -353,15 +351,11 @@ namespace Metasound
 		if (nullptr == Literal)
 		{
 			// Find Class Default that is not invalid
-			FGuid PageID = Frontend::DefaultPageID;
-			const bool bPageIDResolved = GraphPrivate::TryResolveTargetPageID(InOwningGraphClassInput, PageID);
-			if (bPageIDResolved)
+			const FGuid PageID = GraphPrivate::ResolveTargetPageID(InOwningGraphClassInput);
+			const FMetasoundFrontendLiteral& DefaultLiteral = InOwningGraphClassInput.FindConstDefaultChecked(PageID);
+			if (DefaultLiteral.IsValid())
 			{
-				const FMetasoundFrontendLiteral& DefaultLiteral = InOwningGraphClassInput.FindConstDefaultChecked(PageID);
-				if (DefaultLiteral.IsValid())
-				{
-					Literal = &DefaultLiteral;
-				}
+				Literal = &DefaultLiteral;
 			}
 		}
 
@@ -369,16 +363,11 @@ namespace Metasound
 		if (nullptr == Literal && ensure(InInputNodeClass.Interface.Inputs.Num() == 1))
 		{
 			const FMetasoundFrontendClassInput& InputNodeClassInput = InInputNodeClass.Interface.Inputs.Last();
-
-			FGuid PageID = Frontend::DefaultPageID;
-			const bool bPageIDResolved = GraphPrivate::TryResolveTargetPageID(InputNodeClassInput, PageID);
-			if (bPageIDResolved)
+			const FGuid PageID = GraphPrivate::ResolveTargetPageID(InputNodeClassInput);
+			const FMetasoundFrontendLiteral& DefaultLiteral = InputNodeClassInput.FindConstDefaultChecked(PageID);
+			if (DefaultLiteral.IsValid())
 			{
-				const FMetasoundFrontendLiteral& DefaultLiteral = InputNodeClassInput.FindConstDefaultChecked(PageID);
-				if (DefaultLiteral.IsValid())
-				{
-					Literal = &DefaultLiteral;
-				}
+				Literal = &DefaultLiteral;
 			}
 		}
 
@@ -781,13 +770,7 @@ PRAGMA_ENABLE_DEPRECATION_WARNINGS
 
 		const FString GraphName = InContext.DebugAssetName;
 
-		FGuid PageID = Frontend::DefaultPageID;
-		const bool bPageIDResolved = GraphPrivate::TryResolveTargetPageID(InGraphClass, PageID);
-		if (!bPageIDResolved)
-		{
-			return { };
-		}
-
+		const FGuid PageID = GraphPrivate::ResolveTargetPageID(InGraphClass);
 		const FMetasoundFrontendGraph& PageGraph = InGraphClass.FindConstGraphChecked(PageID);
 		FBuildGraphContext BuildGraphContext
 		{
diff --git a/Engine/Plugins/Runtime/Metasound/Source/MetasoundFrontend/Public/MetasoundAssetBase.h b/Engine/Plugins/Runtime/Metasound/Source/MetasoundFrontend/Public/MetasoundAssetBase.h
--- a/Engine/Plugins/Runtime/Metasound/Source/MetasoundFrontend/Public/MetasoundAssetBase.h
+++ b/Engine/Plugins/Runtime/Metasound/Source/MetasoundFrontend/Public/MetasoundAssetBase.h
@@ -98,8 +98,8 @@ public:
 #if WITH_EDITORONLY_DATA
 	// Updates and registers this and referenced MetaSound document objects with the NodeClass Registry. AutoUpdates and
 	// optimizes aforementioned documents for serialization. Unlike 'UpdateAndRegisterForRuntime', does not generate required
-	// runtime data for graph execution.
-	void UpdateAndRegisterForSerialization();
+	// runtime data for graph execution. If CookPlatformName is set, used to strip data not required for the provided platform.
+	void UpdateAndRegisterForSerialization(FName CookPlatformName = { });
 #endif // WITH_EDITORONLY_DATA
 
 #if WITH_EDITOR
@@ -267,6 +267,7 @@ protected:
 	const FRuntimeData& GetRuntimeData() const;
 	PRAGMA_ENABLE_DEPRECATION_WARNINGS
 
+	UE_DEPRECATED(5.5, "AutoUpdate implementation now private and implemented within 'Version Dependencies'")
 	bool AutoUpdate(bool bInLogWarningsOnDroppedConnection);
 
 	UE_DEPRECATED(5.5, "Moved to private, non-cook specific implementation")
@@ -281,12 +282,16 @@ protected:
 private:
 #if WITH_EDITORONLY_DATA
 	void UpdateAssetRegistry();
-	void UpdateAndRegisterReferencesForSerialization();
+	void UpdateAndRegisterReferencesForSerialization(FName CookPlatformName);
 #endif // WITH_EDITORONLY_DATA
 
 	// Checks if version is up-to-date. If so, returns true. If false, updates the interfaces within the given asset's document to the most recent version.
 	bool TryUpdateInterfaceFromVersion(const FMetasoundFrontendVersion& Version);
 
+	// Versions dependencies to most recent version where applicable. If asset is a preset, MetaSound is rebuilt to accommodate any referenced node class interface changes.
+	// Otherwise, automatically updates any nodes and respective dependent classes to accommodate changes to interfaces therein preserving edges/connections where possible.
+	bool VersionDependencies(FMetaSoundFrontendDocumentBuilder& Builder, bool bInLogWarningsOnDroppedConnection);
+
 	// Returns new interface to be versioned to from the given version. If no interface versioning is
 	// required, returns invalid interface (interface with no name and invalid version number).
 	FMetasoundFrontendInterface GetInterfaceToVersion(const FMetasoundFrontendVersion& InterfaceVersion) const;
diff --git a/Engine/Plugins/Runtime/Metasound/Source/MetasoundFrontend/Public/MetasoundAssetManager.h b/Engine/Plugins/Runtime/Metasound/Source/MetasoundFrontend/Public/MetasoundAssetManager.h
--- a/Engine/Plugins/Runtime/Metasound/Source/MetasoundFrontend/Public/MetasoundAssetManager.h
+++ b/Engine/Plugins/Runtime/Metasound/Source/MetasoundFrontend/Public/MetasoundAssetManager.h
@@ -186,6 +186,9 @@ namespace Metasound::Frontend
 		virtual TSet<FAssetInfo> GetReferencedAssetClasses(const FMetasoundAssetBase& InAssetBase) const = 0;
 #endif // WITH_EDITOR
 
+		// Iterate all known MetaSound asset paths
+		virtual void IterateAssets(TFunctionRef<void(const FAssetKey, const TArray<FTopLevelAssetPath>&)> Iter) const = 0;
+
 		UE_DEPRECATED(5.5, "Rescan no longer supported nor required by Frontend")
 		virtual void RescanAutoUpdateDenyList() { }
 
@@ -208,6 +211,9 @@ namespace Metasound::Frontend
 		// Requests an async load of all async referenced assets of the input asset.
 		virtual void RequestAsyncLoadReferencedAssets(FMetasoundAssetBase& InAssetBase) = 0;
 
+		// Synchronously requests unregister and reregister of all loaded MetaSound assets node class entries.
+		virtual void ReloadMetaSoundAssets() const = 0;
+
 		// Removes object from MetaSound asset manager
 		virtual void RemoveAsset(const UObject& InObject) = 0;
 
diff --git a/Engine/Plugins/Runtime/Metasound/Source/MetasoundFrontend/Public/MetasoundDocumentInterface.h b/Engine/Plugins/Runtime/Metasound/Source/MetasoundFrontend/Public/MetasoundDocumentInterface.h
--- a/Engine/Plugins/Runtime/Metasound/Source/MetasoundFrontend/Public/MetasoundDocumentInterface.h
+++ b/Engine/Plugins/Runtime/Metasound/Source/MetasoundFrontend/Public/MetasoundDocumentInterface.h
@@ -81,6 +81,15 @@ namespace Metasound::Frontend
 	public:
 		virtual ~IDocumentBuilderRegistry() = default;
 
+#if WITH_EDITORONLY_DATA
+		// Given the provided builder, removes paged data within the associated document for a cooked build.
+		// This function removes graphs and input defaults which are not to ever be used by a given cook
+		// platform, allowing users to optimize away data and scale the amount of memory required for
+		// initial load of input UObjects and graph topology, which can also positively effect runtime
+		// performance as well, etc. Returns true if builder modified the document, false if not.
+		virtual bool CookPages(FName CookPlatformName, FMetaSoundFrontendDocumentBuilder& Builder) const = 0;
+#endif // WITH_EDITORONLY_DATA
+
 		virtual FMetaSoundFrontendDocumentBuilder* FindBuilder(TScriptInterface<IMetaSoundDocumentInterface> MetaSound) const = 0;
 		virtual FMetaSoundFrontendDocumentBuilder* FindBuilder(const FMetasoundFrontendClassName& ClassName, const FTopLevelAssetPath& AssetPath) const = 0;
 		virtual FMetaSoundFrontendDocumentBuilder* FindOutermostBuilder(const UObject& InSubObject) const = 0;
@@ -108,15 +117,15 @@ namespace Metasound::Frontend
 
 		// Given the provided document and its respective pages, provides the PageID to be used for runtime IGraph and proxy generation.
 		// Returns true if valid page was found, false if not.
-		virtual bool TryResolveTargetPageID(const FMetasoundFrontendGraphClass& InGraphClass, FGuid& OutResolvedPageID) const = 0;
+		virtual FGuid ResolveTargetPageID(const FMetasoundFrontendGraphClass& InGraphClass) const = 0;
 
 		// Given the provided input and its respective paged default values, provides the default PageID to be used for runtime IGraph and proxy generation.
 		// Returns true if valid page was found, false if not.
-		virtual bool TryResolveTargetPageID(const FMetasoundFrontendClassInput& InClassInput, FGuid& OutResolvedPageID) const = 0;
+		virtual FGuid ResolveTargetPageID(const FMetasoundFrontendClassInput& InClassInput) const = 0;
 
 		// Given the provided array of default values, provides the default PageID to be used for runtime IGraph and proxy generation.
 		// Returns true if valid page was found, false if not.
-		virtual bool TryResolveTargetPageID(const TArray<FMetasoundFrontendClassInputDefault>& Defaults, FGuid& OutResolvedPageID) const = 0;
+		virtual FGuid ResolveTargetPageID(const TArray<FMetasoundFrontendClassInputDefault>& Defaults) const = 0;
 
 		static IDocumentBuilderRegistry* Get();
 		static IDocumentBuilderRegistry& GetChecked();
diff --git a/Engine/Plugins/Runtime/Metasound/Source/MetasoundFrontend/Public/MetasoundFrontendDocument.h b/Engine/Plugins/Runtime/Metasound/Source/MetasoundFrontend/Public/MetasoundFrontendDocument.h
--- a/Engine/Plugins/Runtime/Metasound/Source/MetasoundFrontend/Public/MetasoundFrontendDocument.h
+++ b/Engine/Plugins/Runtime/Metasound/Source/MetasoundFrontend/Public/MetasoundFrontendDocument.h
@@ -40,7 +40,7 @@ namespace Metasound
 		constexpr TCHAR DefaultPageName[] = TEXT("Default");
 
 #if WITH_EDITORONLY_DATA
-		extern const FText METASOUNDFRONTEND_API DefaultGraphPageDisplayName;
+		extern const FText METASOUNDFRONTEND_API DefaultPageDisplayName;
 #endif // WITH_EDITORONLY_DATA
 
 		namespace DisplayStyle
diff --git a/Engine/Plugins/Runtime/Metasound/Source/MetasoundFrontend/Public/MetasoundFrontendDocumentBuilder.h b/Engine/Plugins/Runtime/Metasound/Source/MetasoundFrontend/Public/MetasoundFrontendDocumentBuilder.h
--- a/Engine/Plugins/Runtime/Metasound/Source/MetasoundFrontend/Public/MetasoundFrontendDocumentBuilder.h
+++ b/Engine/Plugins/Runtime/Metasound/Source/MetasoundFrontend/Public/MetasoundFrontendDocumentBuilder.h
@@ -381,6 +381,10 @@ public:
 
 	void Reload(TSharedPtr<Metasound::Frontend::FDocumentModifyDelegates> Delegates = {}, bool bPrimeCache = false);
 
+#if WITH_EDITORONLY_DATA
+	bool RemoveGraphInputDefault(FName InputName, const FGuid& InPageID);
+#endif // WITH_EDITORONLY_DATA
+
 	bool RemoveNodeInputDefault(const FGuid& InNodeID, const FGuid& InVertexID, const FGuid* InPageID = nullptr);
 	bool RemoveUnusedDependencies();
 
