Fix for shadow var

#rnx

diff --git a/Engine/Plugins/Runtime/Metasound/Source/MetasoundEngine/Private/MetasoundSettings.cpp b/Engine/Plugins/Runtime/Metasound/Source/MetasoundEngine/Private/MetasoundSettings.cpp
--- a/Engine/Plugins/Runtime/Metasound/Source/MetasoundEngine/Private/MetasoundSettings.cpp
+++ b/Engine/Plugins/Runtime/Metasound/Source/MetasoundEngine/Private/MetasoundSettings.cpp
@@ -240,9 +240,9 @@ const FMetaSoundPageSettings& UMetaSoundSettings::GetDefaultPageSettings() const
 TArray<FName> UMetaSoundSettings::GetImplementedPagePlatforms() const
 {
 	TSet<FName> PlatformNames { FPlatformProperties::IniPlatformName() };
-	IteratePageSettings([&](const FMetaSoundPageSettings& PageSettings)
+	IteratePageSettings([&](const FMetaSoundPageSettings& PageSetting)
 	{
-		TArray<FName> PagePlatforms = PageSettings.GetImplementedPlatforms();
+		TArray<FName> PagePlatforms = PageSetting.GetImplementedPlatforms();
 		PlatformNames.Append(MoveTemp(PagePlatforms));
 	});
 	return PlatformNames.Array();
