Page resolution perf/selection improvements
- Add option to disable MetaSound audition (page & platform) settings in PIE (default to enabled)
- Add ability for project to override resolved pageID
- Avoid constant set reallocation for now frequently-used page ID resolution (as inputs now call during generation independent of graph)
- Add trace event scope for page resolution perf monitoring

#rnx

diff --git a/Engine/Plugins/Runtime/Metasound/Source/MetasoundEditor/Private/MetasoundEditorModule.cpp b/Engine/Plugins/Runtime/Metasound/Source/MetasoundEditor/Private/MetasoundEditorModule.cpp
--- a/Engine/Plugins/Runtime/Metasound/Source/MetasoundEditor/Private/MetasoundEditorModule.cpp
+++ b/Engine/Plugins/Runtime/Metasound/Source/MetasoundEditor/Private/MetasoundEditorModule.cpp
@@ -63,6 +63,11 @@ namespace Metasound
 {
 	namespace Editor
 	{
+		namespace ModulePrivate
+		{
+			static const FText PageValuePropertyText = LOCTEXT("Node_DefaultPageValuePropertyName", "Page Value");
+		} // namespace ModulePrivate
+
 		using FMetasoundGraphPanelPinFactory = FGraphPanelPinFactory;
 
 		static const FName AssetToolName { "AssetTools" };
@@ -230,6 +235,7 @@ namespace Metasound
 			const FSlateBrush* PinDisconnectedIcon = nullptr;
 		};
 
+
 		class FModule : public IMetasoundEditorModule
 		{
 			void RegisterInputDefaultClasses()
@@ -433,27 +439,34 @@ namespace Metasound
 				}
 
 				const UMetasoundEditorSettings* EditorSettings = GetDefault<UMetasoundEditorSettings>();
-				FDocumentBuilderRegistry::GetChecked().GetOnResolveAuditionPageInfoDelegate().BindUObject(EditorSettings, &UMetasoundEditorSettings::ResolveAuditionPageInfo);
+				FDocumentBuilderRegistry::GetChecked().GetOnResolveAuditionPageDelegate().BindUObject(EditorSettings, &UMetasoundEditorSettings::ResolveAuditionPage);
 
-				FEditorDelegates::PreBeginPIE.AddWeakLambda(EditorSettings, [](const bool /* bSimulating */)
+				FEditorDelegates::PreBeginPIE.AddWeakLambda(EditorSettings, [this](const bool /* bSimulating */)
 				{
+					using namespace Metasound::Frontend;
+
 					if (const UMetasoundEditorSettings* EdSettings = GetDefault<UMetasoundEditorSettings>())
 					{
-						FOnResolveAuditionPageInfo& OnResolveAuditionPageInfo = FDocumentBuilderRegistry::GetChecked().GetOnResolveAuditionPageInfoDelegate();
-						if (OnResolveAuditionPageInfo.IsBoundToObject(EdSettings))
+						FOnResolveEditorPage& OnResolveAuditionPage = FDocumentBuilderRegistry::GetChecked().GetOnResolveAuditionPageDelegate();
+						if (OnResolveAuditionPage.IsBoundToObject(EdSettings))
 						{
-							OnResolveAuditionPageInfo.Unbind();
+							if (!EdSettings->bApplyAuditionSettingsInPIE)
+							{
+								OnResolveAuditionPage.Unbind();
+							}
 						}
+
+						IMetaSoundAssetManager::GetChecked().ReloadMetaSoundAssets();
 					}
 				});
 				FEditorDelegates::EndPIE.AddWeakLambda(EditorSettings, [](const bool /* bSimulating */)
 				{
 					if (const UMetasoundEditorSettings* EdSettings = GetDefault<UMetasoundEditorSettings>())
 					{
-						FOnResolveAuditionPageInfo& OnResolveAuditionPageInfo = FDocumentBuilderRegistry::GetChecked().GetOnResolveAuditionPageInfoDelegate();
-						if (!OnResolveAuditionPageInfo.IsBoundToObject(EdSettings))
+						FOnResolveEditorPage& OnResolveAuditionPage = FDocumentBuilderRegistry::GetChecked().GetOnResolveAuditionPageDelegate();
+						if (!OnResolveAuditionPage.IsBoundToObject(EdSettings))
 						{
-							OnResolveAuditionPageInfo.BindUObject(EdSettings, &UMetasoundEditorSettings::ResolveAuditionPageInfo);
+							OnResolveAuditionPage.BindUObject(EdSettings, &UMetasoundEditorSettings::ResolveAuditionPage);
 						}
 					}
 				});
@@ -726,6 +739,8 @@ namespace Metasound
 				AssetActions.Reset();
 				PinTypes.Reset();
 
+				LiteralCustomizationFactories.Reset();
+
 				FGraphNodeVisualizationRegistry::TearDown();
 			}
 
diff --git a/Engine/Plugins/Runtime/Metasound/Source/MetasoundEditor/Private/MetasoundEditorSettings.cpp b/Engine/Plugins/Runtime/Metasound/Source/MetasoundEditor/Private/MetasoundEditorSettings.cpp
--- a/Engine/Plugins/Runtime/Metasound/Source/MetasoundEditor/Private/MetasoundEditorSettings.cpp
+++ b/Engine/Plugins/Runtime/Metasound/Source/MetasoundEditor/Private/MetasoundEditorSettings.cpp
@@ -110,12 +110,12 @@ const FAudioMaterialMeterStyle* UMetasoundEditorSettings::GetMeterStyle() const
 	return &FAudioWidgetsStyle::Get().GetWidgetStyle<FAudioMaterialMeterStyle>("AudioMaterialMeter.Style");
 }
 
-Metasound::Engine::FAuditionPageInfo UMetasoundEditorSettings::ResolveAuditionPageInfo(const TSet<FGuid>& InPageIDs) const
+Metasound::Engine::FPageResolutionEditorResults UMetasoundEditorSettings::ResolveAuditionPage(const TArray<FGuid>& InPageIDs) const
 {
 	using namespace Metasound;
 	using namespace Metasound::Engine;
 
-	FAuditionPageInfo PreviewInfo { .PlatformName = AuditionPlatform };
+	FPageResolutionEditorResults PreviewInfo { .PlatformName = AuditionPlatform };
 
 	if (AuditionTargetPage == Frontend::DefaultPageName)
 	{
diff --git a/Engine/Plugins/Runtime/Metasound/Source/MetasoundEditor/Private/MetasoundEditorSettings.h b/Engine/Plugins/Runtime/Metasound/Source/MetasoundEditor/Private/MetasoundEditorSettings.h
--- a/Engine/Plugins/Runtime/Metasound/Source/MetasoundEditor/Private/MetasoundEditorSettings.h
+++ b/Engine/Plugins/Runtime/Metasound/Source/MetasoundEditor/Private/MetasoundEditorSettings.h
@@ -2,6 +2,7 @@
 #pragma once
 
 #include "AudioSpectrumAnalyzer.h"
+#include "MetasoundFrontendDocument.h"
 #include "Misc/CoreDefines.h"
 #include "UObject/NoExportTypes.h"
 #include "UObject/Object.h"
@@ -20,7 +21,7 @@ struct FMetasoundFrontendDocument;
 
 namespace Metasound::Engine
 {
-	struct FAuditionPageInfo;
+	struct FPageResolutionEditorResults;
 } // namespace Metasound::Engine
 
 
@@ -148,6 +149,12 @@ public:
 	UPROPERTY(EditAnywhere, config, DisplayName = "Pin MetaSound Source in Asset Menu", Category = AssetMenu)
 	bool bPinMetaSoundSourceInAssetMenu = true;
 
+	/** If true, uses editor page/platform audition settings in PIE. If false, uses project's defined values (see 'MetaSound'
+	  * Developer Settings for TargetPage, which can be manipulated via code/Blueprint.)
+	  */
+	UPROPERTY(EditAnywhere, config, Category = Audition)
+	bool bApplyAuditionSettingsInPIE = true;
+
 	/** Default author title to use when authoring a new
 	  * MetaSound.  If empty, uses machine name by default.
 	  */
@@ -175,7 +182,7 @@ public:
 		EditConditionHides = true,
 		GetOptions = "MetasoundEngine.MetaSoundSettings.GetPageNames")
 	)
-	FName AuditionTargetPage;
+	FName AuditionTargetPage = Metasound::Frontend::DefaultPageName;
 
 	/** Maps Pin Category To Pin Color */
 	TMap<FName, FLinearColor> CustomPinTypeColors;
@@ -294,7 +301,7 @@ public:
 	/** Get the AudioMaterialMeter Style. If MeterStyleOverride is not set, returns default style.*/
 	const FAudioMaterialMeterStyle* GetMeterStyle() const;
 
-	Metasound::Engine::FAuditionPageInfo ResolveAuditionPageInfo(const TSet<FGuid>& InPageIDs) const;
+	Metasound::Engine::FPageResolutionEditorResults ResolveAuditionPage(const TArray<FGuid>& InPageIDs) const;
 
 	UFUNCTION()
 	static TArray<FName> GetAuditionPlatformNames();
diff --git a/Engine/Plugins/Runtime/Metasound/Source/MetasoundEngine/Private/MetasoundDocumentBuilderRegistry.cpp b/Engine/Plugins/Runtime/Metasound/Source/MetasoundEngine/Private/MetasoundDocumentBuilderRegistry.cpp
--- a/Engine/Plugins/Runtime/Metasound/Source/MetasoundEngine/Private/MetasoundDocumentBuilderRegistry.cpp
+++ b/Engine/Plugins/Runtime/Metasound/Source/MetasoundEngine/Private/MetasoundDocumentBuilderRegistry.cpp
@@ -104,17 +104,17 @@ namespace Metasound::Engine
 		const UMetaSoundSettings* Settings = GetDefault<UMetaSoundSettings>();
 		check(Settings);
 		const TArray<FGuid>& PlatformTargetPageIDs = Settings->GetCookedTargetPageIDs(PlatformName);
-		auto StripPageEntries = [&](TFunctionRef<bool(const FGuid&)> RemovePageItem, TSet<FGuid>& ImplementedPageIDs)
+		auto StripPageEntries = [&](TFunctionRef<bool(const FGuid&)> RemovePageItem)
 		{
 			TSet<FGuid> ResolvedTargetIDs;
 			for (const FGuid& TargetPage : PlatformTargetPageIDs)
 			{
-				const FGuid PageID = ResolveTargetPageIDInternal(*Settings, ImplementedPageIDs, TargetPage, PlatformName);
+				const FGuid PageID = ResolveTargetPageIDInternal(*Settings, TargetPage, PlatformName);
 				ResolvedTargetIDs.Add(PageID);
 			}
 
-			TSet<FGuid> ToRemove = ImplementedPageIDs.Difference(ResolvedTargetIDs);
-			for (const FGuid& PageID : ToRemove)
+			TargetPageResolveScratch.RemoveAllSwap([&ResolvedTargetIDs](const FGuid& PageID) { return ResolvedTargetIDs.Contains(PageID); });
+			for (const FGuid& PageID : TargetPageResolveScratch)
 			{
 				bModified |= RemovePageItem(PageID);
 			}
@@ -123,8 +123,8 @@ namespace Metasound::Engine
 		const FMetasoundFrontendDocument& Document = Builder.GetConstDocumentChecked();
 
 		{ // Strip graphs
-			TSet<FGuid> ImplementedGraphPageIDs;
-			auto AddPageID = [&ImplementedGraphPageIDs](const FMetasoundFrontendGraph& Graph) { ImplementedGraphPageIDs.Add(Graph.PageID); };
+			TargetPageResolveScratch.Reset();
+			auto AddPageID = [this](const FMetasoundFrontendGraph& Graph) { TargetPageResolveScratch.Add(Graph.PageID); };
 			Document.RootGraph.IterateGraphPages(AddPageID);
 
 			auto RemoveGraphPage = [&Builder](const FGuid& InPageID)
@@ -132,15 +132,15 @@ namespace Metasound::Engine
 				return Builder.RemoveGraphPage(InPageID);
 			};
 
-			StripPageEntries(RemoveGraphPage, ImplementedGraphPageIDs);
+			StripPageEntries(RemoveGraphPage);
 		}
 
 		{ // Strip default input values
 			TSet<FGuid> ImplementedDefaultPageIDs;
 			for (const FMetasoundFrontendClassInput& GraphInput : Document.RootGraph.Interface.Inputs)
 			{
-				ImplementedDefaultPageIDs.Reset();
-				auto AddPageID = [&ImplementedDefaultPageIDs](const FGuid& PageID, const FMetasoundFrontendLiteral&) { ImplementedDefaultPageIDs.Add(PageID); };
+				TargetPageResolveScratch.Reset();
+				auto AddPageID = [this](const FGuid& PageID, const FMetasoundFrontendLiteral&) { TargetPageResolveScratch.Add(PageID); };
 				GraphInput.IterateDefaults(AddPageID);
 
 				auto RemoveDefault = [&Builder, &GraphInput](const FGuid& InPageID)
@@ -148,7 +148,7 @@ namespace Metasound::Engine
 					return Builder.RemoveGraphInputDefault(GraphInput.Name, InPageID);
 				};
 
-				StripPageEntries(RemoveDefault, ImplementedDefaultPageIDs);
+				StripPageEntries(RemoveDefault);
 			}
 		}
 
@@ -363,12 +363,17 @@ namespace Metasound::Engine
 	}
 
 #if WITH_EDITOR
-	FOnResolveAuditionPageInfo& FDocumentBuilderRegistry::GetOnResolveAuditionPageInfoDelegate()
+	FOnResolveEditorPage& FDocumentBuilderRegistry::GetOnResolveAuditionPageDelegate()
 	{
-		return OnResolveAuditionPageInfo;
+		return OnResolveAuditionPage;
 	}
 #endif // WITH_EDITOR
 
+	FOnResolvePage& FDocumentBuilderRegistry::GetOnResolveProjectPageOverrideDelegate()
+	{
+		return OnResolveProjectPage;
+	}
+
 	bool FDocumentBuilderRegistry::ReloadBuilder(const FMetasoundFrontendClassName& InClassName) const
 	{
 		bool bReloaded = false;
@@ -384,41 +389,44 @@ namespace Metasound::Engine
 
 	FGuid FDocumentBuilderRegistry::ResolveTargetPageID(const FMetasoundFrontendGraphClass& InGraphClass) const
 	{
-		TSet<FGuid> DocPageIds;
-		InGraphClass.IterateGraphPages([&DocPageIds](const FMetasoundFrontendGraph& PageGraph)
+		TargetPageResolveScratch.Reset();
+		InGraphClass.IterateGraphPages([this](const FMetasoundFrontendGraph& PageGraph)
 		{
-			DocPageIds.Add(PageGraph.PageID);
+			TargetPageResolveScratch.Add(PageGraph.PageID);
 		});
 
-		return ResolveTargetPageIDInternal(DocPageIds);
+		return ResolveTargetPageIDInternal();
 	}
 
 	FGuid FDocumentBuilderRegistry::ResolveTargetPageID(const FMetasoundFrontendClassInput& InClassInput) const
 	{
-		TSet<FGuid> DefaultPageIds;
-		InClassInput.IterateDefaults([&DefaultPageIds](const FGuid& PageID, const FMetasoundFrontendLiteral&)
+		TargetPageResolveScratch.Reset();
+		InClassInput.IterateDefaults([this](const FGuid& PageID, const FMetasoundFrontendLiteral&)
 		{
-			DefaultPageIds.Add(PageID);
+			TargetPageResolveScratch.Add(PageID);
 		});
 
-		return ResolveTargetPageIDInternal(DefaultPageIds);
+		return ResolveTargetPageIDInternal();
 	}
 
 	FGuid FDocumentBuilderRegistry::ResolveTargetPageID(const TArray<FMetasoundFrontendClassInputDefault>& InClassDefaults) const
 	{
-		TSet<FGuid> DefaultPageIds;
-		Algo::Transform(InClassDefaults, DefaultPageIds, [](const FMetasoundFrontendClassInputDefault& ClassDefault) { return ClassDefault.PageID; });
-		return ResolveTargetPageIDInternal(DefaultPageIds);
+		TargetPageResolveScratch.Reset();
+		Algo::Transform(InClassDefaults, TargetPageResolveScratch, [](const FMetasoundFrontendClassInputDefault& ClassDefault) { return ClassDefault.PageID; });
+		return ResolveTargetPageIDInternal();
 	}
 
-	FGuid FDocumentBuilderRegistry::ResolveTargetPageIDInternal(const TSet<FGuid>& InPageIDs) const
+	FGuid FDocumentBuilderRegistry::ResolveTargetPageIDInternal() const
 	{
+		METASOUND_LLM_SCOPE;
+		METASOUND_TRACE_CPUPROFILER_EVENT_SCOPE(FDocumentBuilderRegistry::ResolveTargetPageID);
+
 		FName PlatformName = FPlatformProperties::IniPlatformName();
 
 #if WITH_EDITOR
-		if (OnResolveAuditionPageInfo.IsBound())
+		if (OnResolveAuditionPage.IsBound())
 		{
-			FAuditionPageInfo PreviewInfo = OnResolveAuditionPageInfo.Execute(InPageIDs);
+			FPageResolutionEditorResults PreviewInfo = OnResolveAuditionPage.Execute(TargetPageResolveScratch);
 			if (PreviewInfo.PageID.IsSet())
 			{
 				return PreviewInfo.PageID.GetValue();
@@ -428,16 +436,23 @@ namespace Metasound::Engine
 		}
 #endif // WITH_EDITOR
 
+		if (OnResolveProjectPage.IsBound())
+		{
+			const FGuid ResolvedPageID = OnResolveProjectPage.Execute(TargetPageResolveScratch);
+			check (TargetPageResolveScratch.Contains(ResolvedPageID));
+			return ResolvedPageID;
+		}
+
 		if (const UMetaSoundSettings* Settings = GetDefault<UMetaSoundSettings>())
 		{
 			const FGuid& TargetPageID = Settings->GetTargetPageSettings().UniqueId;
-			return ResolveTargetPageIDInternal(*Settings, InPageIDs, TargetPageID, PlatformName);
+			return ResolveTargetPageIDInternal(*Settings, TargetPageID, PlatformName);
 		}
 
 		return Frontend::DefaultPageID;
 	}
 
-	FGuid FDocumentBuilderRegistry::ResolveTargetPageIDInternal(const UMetaSoundSettings& Settings, const TSet<FGuid>& InPageIDs, const FGuid& TargetPageID, FName PlatformName) const
+	FGuid FDocumentBuilderRegistry::ResolveTargetPageIDInternal(const UMetaSoundSettings& Settings, const FGuid& TargetPageID, FName PlatformName) const
 	{
 		bool bResolved = false;
 		FGuid ResolvedPageID = Frontend::DefaultPageID;
@@ -448,7 +463,7 @@ namespace Metasound::Engine
 			bFoundMatch |= PageSettings.UniqueId == TargetPageID;
 			if (bFoundMatch && !bResolved)
 			{
-				const bool bAssetImplementsPage = InPageIDs.Contains(PageSettings.UniqueId);
+				const bool bAssetImplementsPage = TargetPageResolveScratch.Contains(PageSettings.UniqueId);
 				if (bAssetImplementsPage)
 				{
 #if WITH_EDITOR
diff --git a/Engine/Plugins/Runtime/Metasound/Source/MetasoundEngine/Public/MetasoundDocumentBuilderRegistry.h b/Engine/Plugins/Runtime/Metasound/Source/MetasoundEngine/Public/MetasoundDocumentBuilderRegistry.h
--- a/Engine/Plugins/Runtime/Metasound/Source/MetasoundEngine/Public/MetasoundDocumentBuilderRegistry.h
+++ b/Engine/Plugins/Runtime/Metasound/Source/MetasoundEngine/Public/MetasoundDocumentBuilderRegistry.h
@@ -18,15 +18,17 @@ class UMetaSoundSettings;
 namespace Metasound::Engine
 {
 #if WITH_EDITOR
-	struct FAuditionPageInfo
+	struct FPageResolutionEditorResults
 	{
 		FName PlatformName;
 		TOptional<FGuid> PageID;
 	};
 
-	DECLARE_DELEGATE_RetVal_OneParam(FAuditionPageInfo, FOnResolveAuditionPageInfo, const TSet<FGuid>& /* InPageIDs */);
+	DECLARE_DELEGATE_RetVal_OneParam(FPageResolutionEditorResults, FOnResolveEditorPage, const TArray<FGuid>& /* InPageIDs */);
 #endif // WITH_EDITOR
 
+	DECLARE_DELEGATE_RetVal_OneParam(FGuid, FOnResolvePage, const TArray<FGuid>& /* InPageIDs */);
+
 	class METASOUNDENGINE_API FDocumentBuilderRegistry : public Frontend::IDocumentBuilderRegistry
 	{
 		mutable TMultiMap<FMetasoundFrontendClassName, TWeakObjectPtr<UMetaSoundBuilderBase>> Builders;
@@ -140,9 +142,11 @@ namespace Metasound::Engine
 		TArray<UMetaSoundBuilderBase*> FindBuilderObjects(const FMetasoundFrontendClassName& InClassName) const;
 
 #if WITH_EDITOR
-		FOnResolveAuditionPageInfo& GetOnResolveAuditionPageInfoDelegate();
+		FOnResolveEditorPage& GetOnResolveAuditionPageDelegate();
 #endif // WITH_EDITOR
 
+		FOnResolvePage& GetOnResolveProjectPageOverrideDelegate();
+
 		bool ReloadBuilder(const FMetasoundFrontendClassName& InClassName) const override;
 
 		// Given the provided document and its respective pages, returns the PageID to be used for runtime IGraph and proxy generation.
@@ -156,13 +160,16 @@ namespace Metasound::Engine
 		void AddBuilderInternal(const FMetasoundFrontendClassName& InClassName, UMetaSoundBuilderBase* NewBuilder) const;
 		bool CanPostEventLog(ELogEvent Event, ELogVerbosity::Type Verbosity) const;
 		void FinishBuildingInternal(UMetaSoundBuilderBase& Builder, bool bForceUnregisterNodeClass) const;
-		FGuid ResolveTargetPageIDInternal(const TSet<FGuid>& InPageIDs) const;
-		FGuid ResolveTargetPageIDInternal(const UMetaSoundSettings& Settings, const TSet<FGuid>& InPageIDs, const FGuid& TargetPageID, FName PlatformName) const;
+		FGuid ResolveTargetPageIDInternal() const;
+		FGuid ResolveTargetPageIDInternal(const UMetaSoundSettings& Settings, const FGuid& TargetPageID, FName PlatformName) const;
 
 #if WITH_EDITOR
-		FOnResolveAuditionPageInfo OnResolveAuditionPageInfo;
+		FOnResolveEditorPage OnResolveAuditionPage;
 #endif // WITH_EDITOR
 
+		FOnResolvePage OnResolveProjectPage;
+
+		mutable TArray<FGuid> TargetPageResolveScratch;
 		TSortedMap<ELogEvent, ELogVerbosity::Type> EventLogVerbosity;
 	};
 } // namespace Metasound::Engine
\ No newline at end of file
diff --git a/Engine/Plugins/Runtime/Metasound/Source/MetasoundFrontend/Private/MetasoundFrontendDocument.cpp b/Engine/Plugins/Runtime/Metasound/Source/MetasoundFrontend/Private/MetasoundFrontendDocument.cpp
--- a/Engine/Plugins/Runtime/Metasound/Source/MetasoundFrontend/Private/MetasoundFrontendDocument.cpp
+++ b/Engine/Plugins/Runtime/Metasound/Source/MetasoundFrontend/Private/MetasoundFrontendDocument.cpp
@@ -947,13 +947,8 @@ void FMetasoundFrontendClassInput::IterateDefaults(TFunctionRef<void(const FGuid
 
 bool FMetasoundFrontendClassInput::RemoveDefault(const FGuid& InPageID)
 {
-	if (InPageID != Metasound::Frontend::DefaultPageID)
-	{
-		auto IsPage = [&InPageID](const FMetasoundFrontendClassInputDefault& Default) { return Default.PageID == InPageID; };
-		return Defaults.RemoveAllSwap(IsPage) > 0;
-	}
-
-	return false;
+	auto IsPage = [&InPageID](const FMetasoundFrontendClassInputDefault& Default) { return Default.PageID == InPageID; };
+	return Defaults.RemoveAllSwap(IsPage) > 0;
 }
 
 void FMetasoundFrontendClassInput::ResetDefaults()
