[Backout] - CL35783705
Original CL Desc
-----------------------------------------------------------------
Fix for shadow var

#rnx

diff --git a/Engine/Plugins/Runtime/Metasound/Source/MetasoundEngine/Private/MetasoundSettings.cpp b/Engine/Plugins/Runtime/Metasound/Source/MetasoundEngine/Private/MetasoundSettings.cpp
--- a/Engine/Plugins/Runtime/Metasound/Source/MetasoundEngine/Private/MetasoundSettings.cpp
+++ b/Engine/Plugins/Runtime/Metasound/Source/MetasoundEngine/Private/MetasoundSettings.cpp
@@ -240,9 +240,9 @@ const FMetaSoundPageSettings& UMetaSoundSettings::GetDefaultPageSettings() const
 TArray<FName> UMetaSoundSettings::GetImplementedPagePlatforms() const
 {
 	TSet<FName> PlatformNames { FPlatformProperties::IniPlatformName() };
-	IteratePageSettings([&](const FMetaSoundPageSettings& PageSetting)
+	IteratePageSettings([&](const FMetaSoundPageSettings& PageSettings)
 	{
-		TArray<FName> PagePlatforms = PageSetting.GetImplementedPlatforms();
+		TArray<FName> PagePlatforms = PageSettings.GetImplementedPlatforms();
 		PlatformNames.Append(MoveTemp(PagePlatforms));
 	});
 	return PlatformNames.Array();
