Re-backout macro comment fix
#rnx

diff --git a/Engine/Plugins/Runtime/Metasound/Source/MetasoundEngine/Private/MetasoundSettings.cpp b/Engine/Plugins/Runtime/Metasound/Source/MetasoundEngine/Private/MetasoundSettings.cpp
--- a/Engine/Plugins/Runtime/Metasound/Source/MetasoundEngine/Private/MetasoundSettings.cpp
+++ b/Engine/Plugins/Runtime/Metasound/Source/MetasoundEngine/Private/MetasoundSettings.cpp
@@ -275,7 +275,7 @@ const TArray<FGuid>& UMetaSoundSettings::GetCookedTargetPageIDs(FName PlatformNa
 		{
 #if WITH_EDITOR
 			const bool bCanTargetDefault = DefaultPageSettings.Target.GetValueForPlatform(PlatformName);
-#else !WITH_EDITOR
+#else // !WITH_EDITOR
 			const bool bCanTargetDefault = DefaultPageSettings.Target.GetValue();
 #endif // !WITH_EDITOR
 
