- Rename feedback from prior MetaSound page cook submission
- Fix for runtime-generated MetaSounds asserting when registering from minor oversight in earlier CL
- Exposition of page editor enablement for subsequent page default CL in review

#tests automated PIE and unit

diff --git a/Engine/Plugins/Runtime/Metasound/Source/MetasoundEditor/Private/MetasoundDetailCustomization.cpp b/Engine/Plugins/Runtime/Metasound/Source/MetasoundEditor/Private/MetasoundDetailCustomization.cpp
--- a/Engine/Plugins/Runtime/Metasound/Source/MetasoundEditor/Private/MetasoundDetailCustomization.cpp
+++ b/Engine/Plugins/Runtime/Metasound/Source/MetasoundEditor/Private/MetasoundDetailCustomization.cpp
@@ -346,7 +346,6 @@ namespace Metasound::Editor
 	}
 
 	FMetasoundPagesDetailCustomization::FMetasoundPagesDetailCustomization()
-		: ItemName(TEXT("Page"))
 	{
 	}
 
@@ -392,7 +391,7 @@ namespace Metasound::Editor
 				{
 					UObject& MetaSound = GetMetaSound();
 
-					const FScopedTransaction Transaction(FText::Format(LOCTEXT("AddPageTransactionFormat", "Add MetaSound Page '{0}'"), FText::FromString(*NameToAdd)));
+					const FScopedTransaction Transaction(FText::Format(LOCTEXT("AddPageTransactionFormat", "Add MetaSound Page Graph '{0}'"), FText::FromString(*NameToAdd)));
 					MetaSound.Modify();
 
 					// Underlying DocBuilder's pageID is a property that is tracked by transaction stack, so signal as modifying behavior
@@ -408,7 +407,7 @@ namespace Metasound::Editor
 			.Content()
 			[
 				SNew(STextBlock)
-				.Text(LOCTEXT("UpdatePageAction", "Add Page..."))
+				.Text(LOCTEXT("AddPageGraphAction", "Add Page Graph..."))
 				.IsEnabled(TAttribute<bool>::Create(TAttribute<bool>::FGetter::CreateSP(this, &FMetaSoundDetailCustomizationBase::IsGraphEditable)))
 			];
 
@@ -432,7 +431,7 @@ namespace Metasound::Editor
 					using namespace Frontend;
 					UObject& MetaSound = GetMetaSound();
 
-					const FScopedTransaction Transaction(LOCTEXT("ResetPagesTransaction", "Reset MetaSound Pages"));
+					const FScopedTransaction Transaction(LOCTEXT("ResetGraphPagesTransaction", "Reset MetaSound Graph Pages"));
 					MetaSound.Modify();
 
 					constexpr bool bClearDefaultGraph = false;
@@ -444,14 +443,14 @@ namespace Metasound::Editor
 					UpdateItemNames();
 					ComboBox->RefreshOptions();
 					FGraphBuilder::RegisterGraphWithFrontend(MetaSound);
-				}), LOCTEXT("ResetPagesTooltip1", "Removes all pages from the given MetaSound less the default."))
+				}), LOCTEXT("ResetGraphPagesTooltip", "Removes all page graphs from the given MetaSound less the default."))
 			];
 
 		Utilities->SetEnabled(TAttribute<bool>::Create(TAttribute<bool>::FGetter::CreateSP(this, &FMetaSoundDetailCustomizationBase::IsGraphEditable)));
 
 		{
-			const FText HeaderName = LOCTEXT("PageGraphsDisplayName", "Pages");
-			IDetailCategoryBuilder& Category = DetailLayout.EditCategory("Pages", HeaderName);
+			const FText HeaderName = LOCTEXT("PageGraphsDisplayName", "Graphs");
+			IDetailCategoryBuilder& Category = DetailLayout.EditCategory("Graphs", HeaderName);
 			Category.AddCustomRow(HeaderName)[Utilities];
 			Category.AddCustomRow(LOCTEXT("ImplementedPagesLabel", "Graphs")) [ SAssignNew(EntryWidgets, SVerticalBox) ];
 		}
diff --git a/Engine/Plugins/Runtime/Metasound/Source/MetasoundEditor/Private/MetasoundDetailCustomization.h b/Engine/Plugins/Runtime/Metasound/Source/MetasoundEditor/Private/MetasoundDetailCustomization.h
--- a/Engine/Plugins/Runtime/Metasound/Source/MetasoundEditor/Private/MetasoundDetailCustomization.h
+++ b/Engine/Plugins/Runtime/Metasound/Source/MetasoundEditor/Private/MetasoundDetailCustomization.h
@@ -74,7 +74,6 @@ namespace Metasound
 			TSet<FName> ImplementedNames;
 			TSharedPtr<SSearchableComboBox> ComboBox;
 			TSharedPtr<SVerticalBox> EntryWidgets;
-			FString ItemName;
 			FName BuildPageName;
 
 			class FPageListener : public Frontend::IDocumentBuilderTransactionListener
diff --git a/Engine/Plugins/Runtime/Metasound/Source/MetasoundEditor/Private/MetasoundEditor.cpp b/Engine/Plugins/Runtime/Metasound/Source/MetasoundEditor/Private/MetasoundEditor.cpp
--- a/Engine/Plugins/Runtime/Metasound/Source/MetasoundEditor/Private/MetasoundEditor.cpp
+++ b/Engine/Plugins/Runtime/Metasound/Source/MetasoundEditor/Private/MetasoundEditor.cpp
@@ -117,6 +117,11 @@ namespace Metasound
 				ECVF_Default);
 		} // namespace AssetEditorPrivate
 
+		bool PageEditorEnabled()
+		{
+			return AssetEditorPrivate::EnablePageEditor != 0;
+		}
+
 		namespace TabNamesPrivate
 		{
 			const FName Analyzers = "MetasoundEditor_Analyzers";
diff --git a/Engine/Plugins/Runtime/Metasound/Source/MetasoundEditor/Private/MetasoundEditor.h b/Engine/Plugins/Runtime/Metasound/Source/MetasoundEditor/Private/MetasoundEditor.h
--- a/Engine/Plugins/Runtime/Metasound/Source/MetasoundEditor/Private/MetasoundEditor.h
+++ b/Engine/Plugins/Runtime/Metasound/Source/MetasoundEditor/Private/MetasoundEditor.h
@@ -102,6 +102,8 @@ class UMetasoundPagesView : public UMetasoundEditorViewBase
 
 namespace Metasound::Editor
 {
+	bool PageEditorEnabled();
+
 	// Forward Declarations
 	class FMetasoundGraphMemberSchemaAction;
 
diff --git a/Engine/Plugins/Runtime/Metasound/Source/MetasoundEditor/Private/MetasoundEditorSettings.cpp b/Engine/Plugins/Runtime/Metasound/Source/MetasoundEditor/Private/MetasoundEditorSettings.cpp
--- a/Engine/Plugins/Runtime/Metasound/Source/MetasoundEditor/Private/MetasoundEditorSettings.cpp
+++ b/Engine/Plugins/Runtime/Metasound/Source/MetasoundEditor/Private/MetasoundEditorSettings.cpp
@@ -132,7 +132,7 @@ Metasound::Engine::FPageResolutionEditorResults UMetasoundEditorSettings::Resolv
 			{
 				if (InPageIDs.Contains(PageSettings.UniqueId))
 				{
-					return !PageSettings.ExcludePageFromCook(PreviewInfo.PlatformName);
+					return !PageSettings.GetExcludeFromCook(PreviewInfo.PlatformName);
 				}
 
 				return false;
@@ -168,7 +168,7 @@ TArray<FName> UMetasoundEditorSettings::GetAuditionPlatformNames()
 {
 	if (const UMetaSoundSettings* Settings = GetDefault<UMetaSoundSettings>())
 	{
-		return Settings->GetImplementedPagePlatforms();
+		return Settings->GetAllPlatformNamesImplementingTargets();
 	}
 	return { FPlatformProperties::IniPlatformName() };
 }
diff --git a/Engine/Plugins/Runtime/Metasound/Source/MetasoundEngine/Private/MetasoundDocumentBuilderRegistry.cpp b/Engine/Plugins/Runtime/Metasound/Source/MetasoundEngine/Private/MetasoundDocumentBuilderRegistry.cpp
--- a/Engine/Plugins/Runtime/Metasound/Source/MetasoundEngine/Private/MetasoundDocumentBuilderRegistry.cpp
+++ b/Engine/Plugins/Runtime/Metasound/Source/MetasoundEngine/Private/MetasoundDocumentBuilderRegistry.cpp
@@ -472,17 +472,17 @@ namespace Metasound::Engine
 		bool bResolved = false;
 		FGuid ResolvedPageID = Frontend::DefaultPageID;
 		constexpr bool bReverse = true;
-		bool bFoundMatch = false;
+		bool bFoundTarget = false;
 		Settings.IteratePageSettings([&](const FMetaSoundPageSettings& PageSettings)
 		{
-			bFoundMatch |= PageSettings.UniqueId == TargetPageID;
-			if (bFoundMatch && !bResolved)
+			bFoundTarget |= PageSettings.UniqueId == TargetPageID;
+			if (bFoundTarget && !bResolved)
 			{
 				const bool bAssetImplementsPage = TargetPageResolveScratch.Contains(PageSettings.UniqueId);
 				if (bAssetImplementsPage)
 				{
 #if WITH_EDITOR
-					const bool bIsCooked = !PageSettings.ExcludePageFromCook(PlatformName);
+					const bool bIsCooked = !PageSettings.GetExcludeFromCook(PlatformName);
 					if (bIsCooked)
 					{
 						bResolved = true;
diff --git a/Engine/Plugins/Runtime/Metasound/Source/MetasoundEngine/Private/MetasoundSettings.cpp b/Engine/Plugins/Runtime/Metasound/Source/MetasoundEngine/Private/MetasoundSettings.cpp
--- a/Engine/Plugins/Runtime/Metasound/Source/MetasoundEngine/Private/MetasoundSettings.cpp
+++ b/Engine/Plugins/Runtime/Metasound/Source/MetasoundEngine/Private/MetasoundSettings.cpp
@@ -131,20 +131,22 @@ namespace Metasound::SettingsPrivate
 } // namespace Metasound::SettingsPrivate
 
 #if WITH_EDITOR
-bool FMetaSoundPageSettings::ExcludePageFromCook(FName PlatformName) const
+bool FMetaSoundPageSettings::GetExcludeFromCook(FName PlatformName) const
 {
 	if (PlatformCanTargetPage(PlatformName))
 	{
 		return false;
 	}
 
-	return Exclude.GetValueForPlatform(PlatformName);
+	return ExcludeFromCook.GetValueForPlatform(PlatformName);
 }
 
-TArray<FName> FMetaSoundPageSettings::GetImplementedPlatforms() const
+TArray<FName> FMetaSoundPageSettings::GetTargetPlatforms() const
 {
 	TArray<FName> PlatformNames;
-	Target.PerPlatform.GetKeys(PlatformNames);
+	Algo::TransformIf(Target.PerPlatform, PlatformNames,
+		[](const TPair<FName, bool>& Pair) { return Pair.Value; },
+		[](const TPair<FName, bool>& Pair) { return Pair.Key; });
 	return PlatformNames;
 }
 
@@ -237,12 +239,12 @@ const FMetaSoundPageSettings& UMetaSoundSettings::GetDefaultPageSettings() const
 }
 
 #if WITH_EDITOR
-TArray<FName> UMetaSoundSettings::GetImplementedPagePlatforms() const
+TArray<FName> UMetaSoundSettings::GetAllPlatformNamesImplementingTargets() const
 {
 	TSet<FName> PlatformNames { FPlatformProperties::IniPlatformName() };
-	IteratePageSettings([&](const FMetaSoundPageSettings& PageSetting)
+	IteratePageSettings([&](const FMetaSoundPageSettings& PageSettings)
 	{
-		TArray<FName> PagePlatforms = PageSetting.GetImplementedPlatforms();
+		TArray<FName> PagePlatforms = PageSettings.GetTargetPlatforms();
 		PlatformNames.Append(MoveTemp(PagePlatforms));
 	});
 	return PlatformNames.Array();
@@ -304,6 +306,8 @@ const FMetaSoundPageSettings& UMetaSoundSettings::GetTargetPageSettings() const
 		}
 	}
 
+	// Shouldn't hit this, but if for some reason the target page is in a bad state,
+	// try and return any page setting set as a valid target.
 	for (const FMetaSoundPageSettings& Setting : PageSettings)
 	{
 		if (Setting.Target.GetValue())
diff --git a/Engine/Plugins/Runtime/Metasound/Source/MetasoundEngine/Public/MetasoundSettings.h b/Engine/Plugins/Runtime/Metasound/Source/MetasoundEngine/Public/MetasoundSettings.h
--- a/Engine/Plugins/Runtime/Metasound/Source/MetasoundEngine/Public/MetasoundSettings.h
+++ b/Engine/Plugins/Runtime/Metasound/Source/MetasoundEngine/Public/MetasoundSettings.h
@@ -75,15 +75,20 @@ private:
 	FPerPlatformBool Target = true;
 
 #if WITH_EDITORONLY_DATA
-	// When true, exclude page data when cooking from the assigned platform(s)/platform group(s) (ignored if target is true).
+	// When true, exclude page data when cooking from the assigned platform(s)/platform group(s) (ignored if target is true for corresponding platform/group).
 	UPROPERTY(EditAnywhere, config, Category = "Pages")
-	FPerPlatformBool Exclude = false;
+	FPerPlatformBool ExcludeFromCook = false;
 #endif //WITH_EDITORONLY_DATA
 
 public:
 #if WITH_EDITOR
-	bool ExcludePageFromCook(FName PlatformName) const;
-	TArray<FName> GetImplementedPlatforms() const;
+	// Returns whether or not page should be excluded from being cooked on the given platform/platform group.
+	bool GetExcludeFromCook(FName PlatformName) const;
+
+	// Returns array of platforms/platform groups page is valid target on.
+	TArray<FName> GetTargetPlatforms() const;
+
+	// Returns whether platform/platform group provided can target the provided page.
 	bool PlatformCanTargetPage(FName PlatformName) const;
 #endif //WITH_EDITOR
 
@@ -207,11 +212,13 @@ public:
 #endif // WITH_EDITORONLY_DATA
 
 #if WITH_EDITOR
-	// Returns superset of explicitly implemented page platforms (and groups) as
-	// defined in PageSettings.
-	TArray<FName> GetImplementedPagePlatforms() const;
+	// Returns superset of page platforms (and groups) that define page target(s).
+	TArray<FName> GetAllPlatformNamesImplementingTargets() const;
 #endif // WITH_EDITOR
 
+	// Returns the project-specific page settings (does not include the required default settings).
+	const TArray<FMetaSoundPageSettings>& GetProjectPageSettings() const { return PageSettings; }
+
 	// Returns the currently targeted page settings.
 	const FMetaSoundPageSettings& GetTargetPageSettings() const;
 
@@ -225,9 +232,8 @@ public:
 	static FName GetQualitySettingPropertyName();
 #endif // WITH_EDITORONLY_DATA
 
-	// Iterates possible page settings in order. Does not include default page settings. If optionally set
-	// to reverse, iterates in reverse. Does *not* include default setting (as fallback therein is not related
-	// to user-defined, ordered page settings).
+	// Iterates possible page settings in order (including the default page settings, which is always last). If 
+	// optionally set to reverse, iterates in reverse.
 	void IteratePageSettings(TFunctionRef<void(const FMetaSoundPageSettings&)> Iter, bool bReverse = false) const;
 
 	// Sets the target page to the given name. Returns true if associated page settings were found
diff --git a/Engine/Plugins/Runtime/Metasound/Source/MetasoundFrontend/Private/MetasoundAssetBase.cpp b/Engine/Plugins/Runtime/Metasound/Source/MetasoundFrontend/Private/MetasoundAssetBase.cpp
--- a/Engine/Plugins/Runtime/Metasound/Source/MetasoundFrontend/Private/MetasoundAssetBase.cpp
+++ b/Engine/Plugins/Runtime/Metasound/Source/MetasoundFrontend/Private/MetasoundAssetBase.cpp
@@ -325,8 +325,14 @@ void FMetasoundAssetBase::UpdateAndRegisterForExecution(Metasound::Frontend::FMe
 	if (InRegistrationOptions.bAutoUpdate)
 	{
 #if WITH_EDITORONLY_DATA
-		FMetaSoundFrontendDocumentBuilder& DocBuilder = IDocumentBuilderRegistry::GetChecked().FindOrBeginBuilding(Owner);
-		const bool bDidUpdate = VersionDependencies(DocBuilder, InRegistrationOptions.bAutoUpdateLogWarningOnDroppedConnection);
+		bool bDidUpdate = false;
+		
+		// Only attempt asset versioning if owner is asset (dependency versioning on runtime MetaSound instances isn't supported nor necessary).
+		if (Owner->IsAsset())
+		{
+			FMetaSoundFrontendDocumentBuilder& DocBuilder = IDocumentBuilderRegistry::GetChecked().FindOrBeginBuilding(Owner);
+			VersionDependencies(DocBuilder, InRegistrationOptions.bAutoUpdateLogWarningOnDroppedConnection);
+		}
 #else // !WITH_EDITORONLY_DATA
 		constexpr bool bDidUpdate = false;
 #endif // WITH_EDITORONLY_DATA
