Giving a local variable a new name, allowing it to live its best life and removing it from the shadowy darkness
#rnx

diff --git a/Engine/Plugins/Runtime/Metasound/Source/MetasoundEngine/Private/MetasoundSettings.cpp b/Engine/Plugins/Runtime/Metasound/Source/MetasoundEngine/Private/MetasoundSettings.cpp
--- a/Engine/Plugins/Runtime/Metasound/Source/MetasoundEngine/Private/MetasoundSettings.cpp
+++ b/Engine/Plugins/Runtime/Metasound/Source/MetasoundEngine/Private/MetasoundSettings.cpp
@@ -242,9 +242,9 @@ const FMetaSoundPageSettings& UMetaSoundSettings::GetDefaultPageSettings() const
 TArray<FName> UMetaSoundSettings::GetAllPlatformNamesImplementingTargets() const
 {
 	TSet<FName> PlatformNames { FPlatformProperties::IniPlatformName() };
-	IteratePageSettings([&](const FMetaSoundPageSettings& PageSettings)
+	IteratePageSettings([&](const FMetaSoundPageSettings& PageSetting)
 	{
-		TArray<FName> PagePlatforms = PageSettings.GetTargetPlatforms();
+		TArray<FName> PagePlatforms = PageSetting.GetTargetPlatforms();
 		PlatformNames.Append(MoveTemp(PagePlatforms));
 	});
 	return PlatformNames.Array();
