- Don't bind MetaSound editor-only module delegates (including audition selector) when running cook commandlet
- Add lock around page resolve scratch array when resolving page IDs to protect against async registration resource contention
- Remove cruft set construction

#rnx

diff --git a/Engine/Plugins/Runtime/Metasound/Source/MetasoundEditor/Private/MetasoundEditorModule.cpp b/Engine/Plugins/Runtime/Metasound/Source/MetasoundEditor/Private/MetasoundEditorModule.cpp
--- a/Engine/Plugins/Runtime/Metasound/Source/MetasoundEditor/Private/MetasoundEditorModule.cpp
+++ b/Engine/Plugins/Runtime/Metasound/Source/MetasoundEditor/Private/MetasoundEditorModule.cpp
@@ -63,11 +63,6 @@ namespace Metasound
 {
 	namespace Editor
 	{
-		namespace ModulePrivate
-		{
-			static const FText PageValuePropertyText = LOCTEXT("Node_DefaultPageValuePropertyName", "Page Value");
-		} // namespace ModulePrivate
-
 		using FMetasoundGraphPanelPinFactory = FGraphPanelPinFactory;
 
 		static const FName AssetToolName { "AssetTools" };
@@ -426,6 +421,13 @@ namespace Metasound
 			{
 				using namespace Engine;
 
+				// All the following delegates are used for UX notification, audition
+				// and PIE which are not desired/necessary when cooking.
+				if (IsRunningCookCommandlet())
+				{
+					return;
+				}
+
 				if (UMetaSoundSettings* Settings = GetMutableDefault<UMetaSoundSettings>())
 				{
 					Settings->GetOnDefaultRenamedDelegate().AddLambda([]()
diff --git a/Engine/Plugins/Runtime/Metasound/Source/MetasoundEngine/Private/MetasoundDocumentBuilderRegistry.cpp b/Engine/Plugins/Runtime/Metasound/Source/MetasoundEngine/Private/MetasoundDocumentBuilderRegistry.cpp
--- a/Engine/Plugins/Runtime/Metasound/Source/MetasoundEngine/Private/MetasoundDocumentBuilderRegistry.cpp
+++ b/Engine/Plugins/Runtime/Metasound/Source/MetasoundEngine/Private/MetasoundDocumentBuilderRegistry.cpp
@@ -106,20 +106,20 @@ namespace Metasound::Engine
 		const TArray<FGuid>& PlatformTargetPageIDs = Settings->GetCookedTargetPageIDs(PlatformName);
 		checkf(!PlatformTargetPageIDs.IsEmpty(), TEXT("Must have at least one targeted page ID to cook MetaSound."));
 
-		auto StripPageEntries = [&](TFunctionRef<bool(const FGuid&)> RemovePageItem)
+		auto StripPageEntries = [&](TArray<FGuid>& PageIdsToResolve, TFunctionRef<bool(const FGuid&)> RemovePageItem)
 		{
 			TSet<FGuid> ResolvedTargetIDs;
 			for (const FGuid& TargetPage : PlatformTargetPageIDs)
 			{
-				const FGuid PageID = ResolveTargetPageIDInternal(*Settings, TargetPage, PlatformName);
+				const FGuid PageID = ResolveTargetPageIDInternal(*Settings, PageIdsToResolve, TargetPage, PlatformName);
 				ResolvedTargetIDs.Add(PageID);
 			}
 
 			auto IsResolvedTarget = [&ResolvedTargetIDs](const FGuid& PageID) { return ResolvedTargetIDs.Contains(PageID); };
 			checkf(!ResolvedTargetIDs.IsEmpty(), TEXT("Failed to resolve any valid target IDs, which will leave serialized page array in invalid state."));
-			TargetPageResolveScratch.RemoveAllSwap(IsResolvedTarget);
+			PageIdsToResolve.RemoveAllSwap(IsResolvedTarget, EAllowShrinking::No);
 
-			for (const FGuid& PageID : TargetPageResolveScratch)
+			for (const FGuid& PageID : PageIdsToResolve)
 			{
 				bModified |= RemovePageItem(PageID);
 			}
@@ -127,9 +127,9 @@ namespace Metasound::Engine
 
 		const FMetasoundFrontendDocument& Document = Builder.GetConstDocumentChecked();
 
+		TArray<FGuid> ResolvePageIDs;
 		{ // Strip graphs
-			TargetPageResolveScratch.Reset();
-			auto AddPageID = [this](const FMetasoundFrontendGraph& Graph) { TargetPageResolveScratch.Add(Graph.PageID); };
+			auto AddPageID = [&ResolvePageIDs](const FMetasoundFrontendGraph& Graph) { ResolvePageIDs.Add(Graph.PageID); };
 			Document.RootGraph.IterateGraphPages(AddPageID);
 
 			auto RemoveGraphPage = [&Builder](const FGuid& InPageID)
@@ -137,7 +137,7 @@ namespace Metasound::Engine
 				return Builder.RemoveGraphPage(InPageID);
 			};
 
-			StripPageEntries(RemoveGraphPage);
+			StripPageEntries(ResolvePageIDs, RemoveGraphPage);
 
 			checkf(!Document.RootGraph.GetConstGraphPages().IsEmpty(),
 				TEXT("Document in MetaSound asset '%s' had all default values "
@@ -147,11 +147,10 @@ namespace Metasound::Engine
 		}
 
 		{ // Strip default input values
-			TSet<FGuid> ImplementedDefaultPageIDs;
 			for (const FMetasoundFrontendClassInput& GraphInput : Document.RootGraph.Interface.Inputs)
 			{
-				TargetPageResolveScratch.Reset();
-				auto AddPageID = [this](const FGuid& PageID, const FMetasoundFrontendLiteral&) { TargetPageResolveScratch.Add(PageID); };
+				ResolvePageIDs.Reset();
+				auto AddPageID = [&ResolvePageIDs](const FGuid& PageID, const FMetasoundFrontendLiteral&) { ResolvePageIDs.Add(PageID); };
 				GraphInput.IterateDefaults(AddPageID);
 
 				auto RemoveDefault = [&Builder, &GraphInput](const FGuid& InPageID)
@@ -159,7 +158,7 @@ namespace Metasound::Engine
 					return Builder.RemoveGraphInputDefault(GraphInput.Name, InPageID);
 				};
 
-				StripPageEntries(RemoveDefault);
+				StripPageEntries(ResolvePageIDs, RemoveDefault);
 				checkf(!GraphInput.GetDefaults().IsEmpty(),
 					TEXT("Input '%s' had all default values stripped leaving it in an invalid state. "
 					"Input must always have at least one default value"),
@@ -404,34 +403,37 @@ namespace Metasound::Engine
 
 	FGuid FDocumentBuilderRegistry::ResolveTargetPageID(const FMetasoundFrontendGraphClass& InGraphClass) const
 	{
+		FScopeLock Lock(&TargetPageResolveScratchCritSec);
 		TargetPageResolveScratch.Reset();
 		InGraphClass.IterateGraphPages([this](const FMetasoundFrontendGraph& PageGraph)
 		{
 			TargetPageResolveScratch.Add(PageGraph.PageID);
 		});
 
-		return ResolveTargetPageIDInternal();
+		return ResolveTargetPageIDInternal(TargetPageResolveScratch);
 	}
 
 	FGuid FDocumentBuilderRegistry::ResolveTargetPageID(const FMetasoundFrontendClassInput& InClassInput) const
 	{
+		FScopeLock Lock(&TargetPageResolveScratchCritSec);
 		TargetPageResolveScratch.Reset();
 		InClassInput.IterateDefaults([this](const FGuid& PageID, const FMetasoundFrontendLiteral&)
 		{
 			TargetPageResolveScratch.Add(PageID);
 		});
 
-		return ResolveTargetPageIDInternal();
+		return ResolveTargetPageIDInternal(TargetPageResolveScratch);
 	}
 
 	FGuid FDocumentBuilderRegistry::ResolveTargetPageID(const TArray<FMetasoundFrontendClassInputDefault>& InClassDefaults) const
 	{
+		FScopeLock Lock(&TargetPageResolveScratchCritSec);
 		TargetPageResolveScratch.Reset();
 		Algo::Transform(InClassDefaults, TargetPageResolveScratch, [](const FMetasoundFrontendClassInputDefault& ClassDefault) { return ClassDefault.PageID; });
-		return ResolveTargetPageIDInternal();
+		return ResolveTargetPageIDInternal(TargetPageResolveScratch);
 	}
 
-	FGuid FDocumentBuilderRegistry::ResolveTargetPageIDInternal() const
+	FGuid FDocumentBuilderRegistry::ResolveTargetPageIDInternal(const TArray<FGuid>& InPageIDsToResolve) const
 	{
 		METASOUND_LLM_SCOPE;
 		METASOUND_TRACE_CPUPROFILER_EVENT_SCOPE(FDocumentBuilderRegistry::ResolveTargetPageID);
@@ -441,7 +443,7 @@ namespace Metasound::Engine
 #if WITH_EDITOR
 		if (OnResolveAuditionPage.IsBound())
 		{
-			FPageResolutionEditorResults PreviewInfo = OnResolveAuditionPage.Execute(TargetPageResolveScratch);
+			FPageResolutionEditorResults PreviewInfo = OnResolveAuditionPage.Execute(InPageIDsToResolve);
 			if (PreviewInfo.PageID.IsSet())
 			{
 				return PreviewInfo.PageID.GetValue();
@@ -453,21 +455,21 @@ namespace Metasound::Engine
 
 		if (OnResolveProjectPage.IsBound())
 		{
-			const FGuid ResolvedPageID = OnResolveProjectPage.Execute(TargetPageResolveScratch);
-			check (TargetPageResolveScratch.Contains(ResolvedPageID));
+			const FGuid ResolvedPageID = OnResolveProjectPage.Execute(InPageIDsToResolve);
+			check (InPageIDsToResolve.Contains(ResolvedPageID));
 			return ResolvedPageID;
 		}
 
 		if (const UMetaSoundSettings* Settings = GetDefault<UMetaSoundSettings>())
 		{
 			const FGuid& TargetPageID = Settings->GetTargetPageSettings().UniqueId;
-			return ResolveTargetPageIDInternal(*Settings, TargetPageID, PlatformName);
+			return ResolveTargetPageIDInternal(*Settings, InPageIDsToResolve, TargetPageID, PlatformName);
 		}
 
 		return Frontend::DefaultPageID;
 	}
 
-	FGuid FDocumentBuilderRegistry::ResolveTargetPageIDInternal(const UMetaSoundSettings& Settings, const FGuid& TargetPageID, FName PlatformName) const
+	FGuid FDocumentBuilderRegistry::ResolveTargetPageIDInternal(const UMetaSoundSettings& Settings, const TArray<FGuid>& InPageIDsToResolve, const FGuid& TargetPageID, FName PlatformName) const
 	{
 		bool bResolved = false;
 		FGuid ResolvedPageID = Frontend::DefaultPageID;
@@ -478,7 +480,7 @@ namespace Metasound::Engine
 			bFoundTarget |= PageSettings.UniqueId == TargetPageID;
 			if (bFoundTarget && !bResolved)
 			{
-				const bool bAssetImplementsPage = TargetPageResolveScratch.Contains(PageSettings.UniqueId);
+				const bool bAssetImplementsPage = InPageIDsToResolve.Contains(PageSettings.UniqueId);
 				if (bAssetImplementsPage)
 				{
 #if WITH_EDITOR
diff --git a/Engine/Plugins/Runtime/Metasound/Source/MetasoundEngine/Public/MetasoundDocumentBuilderRegistry.h b/Engine/Plugins/Runtime/Metasound/Source/MetasoundEngine/Public/MetasoundDocumentBuilderRegistry.h
--- a/Engine/Plugins/Runtime/Metasound/Source/MetasoundEngine/Public/MetasoundDocumentBuilderRegistry.h
+++ b/Engine/Plugins/Runtime/Metasound/Source/MetasoundEngine/Public/MetasoundDocumentBuilderRegistry.h
@@ -160,8 +160,8 @@ namespace Metasound::Engine
 		void AddBuilderInternal(const FMetasoundFrontendClassName& InClassName, UMetaSoundBuilderBase* NewBuilder) const;
 		bool CanPostEventLog(ELogEvent Event, ELogVerbosity::Type Verbosity) const;
 		void FinishBuildingInternal(UMetaSoundBuilderBase& Builder, bool bForceUnregisterNodeClass) const;
-		FGuid ResolveTargetPageIDInternal() const;
-		FGuid ResolveTargetPageIDInternal(const UMetaSoundSettings& Settings, const FGuid& TargetPageID, FName PlatformName) const;
+		FGuid ResolveTargetPageIDInternal(const TArray<FGuid>& PageIdsToResolve) const;
+		FGuid ResolveTargetPageIDInternal(const UMetaSoundSettings& Settings, const TArray<FGuid>& PageIdsToResolve, const FGuid& TargetPageID, FName PlatformName) const;
 
 #if WITH_EDITOR
 		FOnResolveEditorPage OnResolveAuditionPage;
@@ -170,8 +170,9 @@ namespace Metasound::Engine
 		FOnResolvePage OnResolveProjectPage;
 
 		// Reuseable scratch array of pages to resolve, which is used to
-		// optimize/reduce numboer of allocations required when resolving document.
+		// optimize/reduce number of allocations required when resolving document.
 		mutable TArray<FGuid> TargetPageResolveScratch;
+		mutable FCriticalSection TargetPageResolveScratchCritSec;
 
 		TSortedMap<ELogEvent, ELogVerbosity::Type> EventLogVerbosity;
 	};
