- Force refresh details panel when swapping Metasound object details view filter
Carry through rename suggestion to internal editor function & hide page input values if user has not provided page input defaults

#rnx

diff --git a/Engine/Plugins/Runtime/Metasound/Source/MetasoundEditor/Private/MetasoundEditor.cpp b/Engine/Plugins/Runtime/Metasound/Source/MetasoundEditor/Private/MetasoundEditor.cpp
--- a/Engine/Plugins/Runtime/Metasound/Source/MetasoundEditor/Private/MetasoundEditor.cpp
+++ b/Engine/Plugins/Runtime/Metasound/Source/MetasoundEditor/Private/MetasoundEditor.cpp
@@ -853,7 +853,7 @@ namespace Metasound
 
 			InTabManager->RegisterTabSpawner(TabNamesPrivate::GraphCanvas, FOnSpawnTab::CreateLambda(
 				[
-					bShowPageTab = ShowPageDetails(),
+					bShowPageTab = ShowPageGraphDetails(),
 					InFocusPageWidget = FocusPageWidget,
 					InPlayTimeWidget = PlayTimeWidget,
 					InMetasoundGraphEditor = MetasoundGraphEditor,
@@ -939,7 +939,7 @@ namespace Metasound
 			.SetGroup(WorkspaceMenuCategoryRef)
 			.SetIcon(FSlateIcon(FAppStyle::GetAppStyleSetName(), "Kismet.Tabs.Palette"));
 
-			if (ShowPageDetails())
+			if (ShowPageGraphDetails())
 			{
 				if (Builder.IsValid() && !Builder->IsPreset())
 				{
@@ -990,7 +990,7 @@ namespace Metasound
 			InTabManager->UnregisterTabSpawner(TabNamesPrivate::Details);
 			InTabManager->UnregisterTabSpawner(TabNamesPrivate::Members);
 
-			if (ShowPageDetails())
+			if (ShowPageGraphDetails())
 			{
 				InTabManager->UnregisterTabSpawner(TabNamesPrivate::Pages);
 			}
@@ -1176,7 +1176,7 @@ namespace Metasound
 				->SetHideTabWell(false)
 				->AddTab(TabNamesPrivate::Details, ETabState::OpenedTab);
 
-			if (ShowPageDetails())
+			if (ShowPageGraphDetails())
 			{
 				DetailsStack->AddTab(TabNamesPrivate::Pages, ETabState::OpenedTab);
 			}
@@ -1320,7 +1320,7 @@ namespace Metasound
 				{
 					if (TabManager.IsValid())
 					{
-						if (ShowPageDetails())
+						if (ShowPageGraphDetails())
 						{
 							TabManager->TryInvokeTab(TabNamesPrivate::Pages);
 						}
@@ -1338,14 +1338,25 @@ namespace Metasound
 			}
 		}
 
-		bool FEditor::ShowPageDetails() const
+		bool FEditor::ShowPageGraphDetails() const
 		{
-			if (AssetEditorPrivate::EnablePageEditor)
+			const UMetaSoundSettings* Settings = GetDefault<UMetaSoundSettings>();
+			if (!Settings)
 			{
-				return Builder.IsValid() && !Builder->IsPreset();
+				return false;
 			}
 
-			return false;
+			if (Settings->GetProjectPageSettings().IsEmpty())
+			{
+				return false;
+			}
+
+			if (!AssetEditorPrivate::EnablePageEditor)
+			{
+				return false;
+			}
+
+			return Builder.IsValid() && !Builder->IsPreset();
 		}
 
 		bool FEditor::GetBoundsForSelectedNodes(FSlateRect& Rect, float Padding)
@@ -1627,7 +1638,7 @@ namespace Metasound
 				InterfacesDetails->HideFilterArea(true);
 			}
 
-			if (ShowPageDetails())
+			if (ShowPageGraphDetails())
 			{
 				PagesDetails = PropertyModule.CreateDetailView(Args);
 				if (PagesDetails.IsValid())
@@ -1641,7 +1652,7 @@ namespace Metasound
 
 					TAttribute<bool> EnabledAttr = TAttribute<bool>::Create([this]()
 					{
-						return ShowPageDetails();
+						return ShowPageGraphDetails();
 					});
 					PagesDetails->SetEnabled(EnabledAttr);
 				}
@@ -2733,6 +2744,7 @@ namespace Metasound
 			}
 
 			EditObjectSettings();
+			RefreshDetails();
 		}
 
 		void FEditor::EditMetasoundSettings()
@@ -2743,6 +2755,7 @@ namespace Metasound
 			}
 
 			EditObjectSettings();
+			RefreshDetails();
 		}
 
 		void FEditor::SyncInBrowser()
diff --git a/Engine/Plugins/Runtime/Metasound/Source/MetasoundEditor/Private/MetasoundEditor.h b/Engine/Plugins/Runtime/Metasound/Source/MetasoundEditor/Private/MetasoundEditor.h
--- a/Engine/Plugins/Runtime/Metasound/Source/MetasoundEditor/Private/MetasoundEditor.h
+++ b/Engine/Plugins/Runtime/Metasound/Source/MetasoundEditor/Private/MetasoundEditor.h
@@ -463,7 +463,7 @@ namespace Metasound::Editor
 		void ConvertFromPreset();
 
 		/* Whether or not page details should be visible. */
-		bool ShowPageDetails() const;
+		bool ShowPageGraphDetails() const;
 
 		/** Creates audition menu options */
 		TSharedRef<SWidget> CreateAuditionMenuOptions();
