Editor support for MetaSound Paged Input Defaults

#rnx

diff --git a/Engine/Plugins/Runtime/Metasound/Source/MetasoundEditor/Private/FindInMetasound.cpp b/Engine/Plugins/Runtime/Metasound/Source/MetasoundEditor/Private/FindInMetasound.cpp
--- a/Engine/Plugins/Runtime/Metasound/Source/MetasoundEditor/Private/FindInMetasound.cpp
+++ b/Engine/Plugins/Runtime/Metasound/Source/MetasoundEditor/Private/FindInMetasound.cpp
@@ -7,6 +7,7 @@
 #include "MetasoundEditorGraphBuilder.h"
 #include "MetasoundEditorGraphInputNode.h"
 #include "MetasoundEditorModule.h"
+#include "MetasoundFrontendDocumentBuilder.h"
 #include "Styling/AppStyle.h"
 #include "Widgets/Images/SImage.h"
 #include "Widgets/Input/SSearchBox.h"
@@ -181,7 +182,9 @@ namespace Metasound::Editor
 			if (Member->GetDataType() != GetMetasoundDataTypeName<FTrigger>())
 			{
 				// Only list number of items for arrays to avoid string issues 
-				const FMetasoundFrontendLiteral& DefaultLiteral = Member->GetLiteral()->GetDefault();
+				FMetaSoundFrontendDocumentBuilder& Builder = Member->GetFrontendBuilderChecked();
+				FMetasoundFrontendLiteral DefaultLiteral;
+				Member->GetLiteral()->TryFindDefault(DefaultLiteral, &Builder.GetBuildPageID());
 				ValueText = FText::FromString(DefaultLiteral.ToString());
 				if (DefaultLiteral.IsArray())
 				{
@@ -321,7 +324,9 @@ namespace Metasound::Editor
 		if (const UMetasoundEditorGraphMember* Member = GetMetaSoundGraphMember(MemberNode))
 		{
 			// Get full object names and array contents (GetValueText shortens these)
-			const FMetasoundFrontendLiteral& DefaultLiteral = Member->GetLiteral()->GetDefault();
+			FMetaSoundFrontendDocumentBuilder& Builder = Member->GetFrontendBuilderChecked();
+			FMetasoundFrontendLiteral DefaultLiteral;
+			Member->GetLiteral()->TryFindDefault(DefaultLiteral, &Builder.GetBuildPageID());
 			return FText::FromString(DefaultLiteral.ToString());
 		}
 		// Get value from pin's external node 
diff --git a/Engine/Plugins/Runtime/Metasound/Source/MetasoundEditor/Private/MetasoundDefaultLiteralCustomization.cpp b/Engine/Plugins/Runtime/Metasound/Source/MetasoundEditor/Private/MetasoundDefaultLiteralCustomization.cpp
new file mode 100644
--- /dev/null
+++ b/Engine/Plugins/Runtime/Metasound/Source/MetasoundEditor/Private/MetasoundDefaultLiteralCustomization.cpp
@@ -0,0 +1,460 @@
+// Copyright Epic Games, Inc. All Rights Reserved.
+
+#include "MetasoundDefaultLiteralCustomization.h"
+
+#include "DetailLayoutBuilder.h"
+#include "MetasoundDetailCustomization.h"
+#include "MetasoundEditor.h"
+#include "MetasoundEditorGraph.h"
+#include "MetasoundEditorGraphBuilder.h"
+#include "MetasoundEditorGraphMemberDefaults.h"
+#include "MetasoundNodeDetailCustomization.h"
+#include "MetasoundSettings.h"
+#include "PropertyCustomizationHelpers.h"
+#include "SSearchableComboBox.h"
+
+#define LOCTEXT_NAMESPACE "MetaSoundEditor"
+
+
+namespace Metasound::Editor
+{
+	namespace LiteralCustomizationPrivate
+	{
+		FGuid GetGuidPropertyValue(TSharedPtr<IPropertyHandle> GuidProperty)
+		{
+			// Randomized so if property be invalid for whatever reason,
+			// the page default entry isn't mistakenly removed.
+			FGuid Guid = FGuid::NewGuid();
+
+			if (!ensure(GuidProperty.IsValid()))
+			{
+				return Guid;
+			}
+
+			auto GetGuidUInt = [&GuidProperty](FName Section)
+			{
+				int32 Value = 0;
+				TSharedPtr<IPropertyHandle> SectionHandle = GuidProperty->GetChildHandle(Section);
+				if (ensure(SectionHandle.IsValid()))
+				{
+					SectionHandle->GetValue(Value);
+				}
+				return Value;
+			};
+			return FGuid(
+				GetGuidUInt(GET_MEMBER_NAME_CHECKED(FGuid, A)),
+				GetGuidUInt(GET_MEMBER_NAME_CHECKED(FGuid, B)),
+				GetGuidUInt(GET_MEMBER_NAME_CHECKED(FGuid, C)),
+				GetGuidUInt(GET_MEMBER_NAME_CHECKED(FGuid, D))
+			);
+		}
+	} // namespace LiteralCustomizationPrivate
+
+	FMetasoundDefaultLiteralCustomizationBase::FMetasoundDefaultLiteralCustomizationBase(IDetailCategoryBuilder& InDefaultCategoryBuilder)
+		: DefaultCategoryBuilder(&InDefaultCategoryBuilder)
+	{
+	}
+
+	FMetasoundDefaultLiteralCustomizationBase::~FMetasoundDefaultLiteralCustomizationBase()
+	{
+		if (UMetaSoundSettings* Settings = GetMutableDefault<UMetaSoundSettings>())
+		{
+			if (OnPageSettingsUpdatedHandle.IsValid())
+			{
+				Settings->GetOnPageSettingsUpdatedDelegate().Remove(OnPageSettingsUpdatedHandle);
+			}
+		}
+	}
+
+	void FMetasoundDefaultLiteralCustomizationBase::BuildPageDefaultComboBox(UMetasoundEditorGraphMemberDefaultLiteral& Literal, FText RowName)
+	{
+		using namespace Engine;
+
+		check(DefaultCategoryBuilder);
+
+		TWeakObjectPtr<UMetasoundEditorGraphMemberDefaultLiteral> LiteralPtr(&Literal);
+		if (UMetaSoundSettings* Settings = GetMutableDefault<UMetaSoundSettings>())
+		{
+			FOnPageSettingsUpdated& UpdatedDelegate = Settings->GetOnPageSettingsUpdatedDelegate();
+			if (OnPageSettingsUpdatedHandle.IsValid())
+			{
+				UpdatedDelegate.Remove(OnPageSettingsUpdatedHandle);
+			}
+			OnPageSettingsUpdatedHandle = UpdatedDelegate.AddLambda([this, LiteralPtr]()
+			{
+				UpdatePagePickerNames(LiteralPtr);
+				if (PageDefaultComboBox.IsValid())
+				{
+					PageDefaultComboBox->RefreshOptions();
+				}
+			});
+		}
+		else
+		{
+			return;
+		}
+
+		SAssignNew(PageDefaultComboBox, SSearchableComboBox)
+		.OptionsSource(&AddablePageStringNames)
+		.OnGenerateWidget_Lambda([](TSharedPtr<FString> InItem)
+		{
+			return SNew(STextBlock).Text(FText::FromString(*InItem));
+		})
+		.OnSelectionChanged_Lambda([this, LiteralPtr](TSharedPtr<FString> NameToAdd, ESelectInfo::Type InSelectInfo)
+		{
+			using namespace Engine;
+			using namespace Frontend;
+
+			if (InSelectInfo != ESelectInfo::OnNavigation)
+			{
+				if (!LiteralPtr.IsValid())
+				{
+					return;
+				}
+
+				UMetasoundEditorGraphMember* Member = LiteralPtr->FindMember();
+				if (!Member)
+				{
+					return;
+				}
+
+				if (!NameToAdd.IsValid())
+				{
+					return;
+				}
+
+				const UMetaSoundSettings* Settings = GetDefault<UMetaSoundSettings>();
+				check(Settings);
+				if (const FMetaSoundPageSettings* PageSettings = Settings->FindPageSettings(FName(*NameToAdd)))
+				{
+					const FScopedTransaction Transaction(FText::Format(LOCTEXT("AddPageDefaultValueTransactionFormat", "Add '{0}' Page '{1}' Default Value"),
+						FText::FromName(Member->GetMemberName()),
+						FText::FromString(*NameToAdd)));
+
+					FMetaSoundFrontendDocumentBuilder& Builder = Member->GetFrontendBuilderChecked();
+					UObject& MetaSound = Builder.CastDocumentObjectChecked<UObject>();
+					MetaSound.Modify();
+					LiteralPtr->Modify();
+
+					LiteralPtr->InitDefault(PageSettings->UniqueId);
+
+					constexpr bool bPostTransaction = false;
+					Member->UpdateFrontendDefaultLiteral(bPostTransaction, &PageSettings->UniqueId);
+					FMetasoundAssetBase& MetasoundAsset = Editor::FGraphBuilder::GetOutermostMetaSoundChecked(*LiteralPtr);
+					MetasoundAsset.GetModifyContext().AddMemberIDsModified({ Member->GetMemberID() });
+
+					UpdatePagePickerNames(LiteralPtr);
+					PageDefaultComboBox->RefreshOptions();
+					FGraphBuilder::RegisterGraphWithFrontend(MetaSound);
+				}
+			}
+		})
+		.Content()
+		[
+			SNew(STextBlock)
+			.Text(LOCTEXT("AddPageDefaultValuePrompt", "Add Page Default Value..."))
+		];
+
+		DefaultCategoryBuilder->AddCustomRow(RowName)
+		.IsEnabled(GetEnabled())
+		.ValueContent()
+		[
+			SNew(SHorizontalBox)
+			.Visibility(Visibility)
+			+ SHorizontalBox::Slot()
+			.Padding(2.0f)
+			.HAlign(HAlign_Left)
+			.VAlign(VAlign_Center)
+			.AutoWidth()
+			[
+				PageDefaultComboBox->AsShared()
+			]
+			+ SHorizontalBox::Slot()
+			.Padding(2.0f)
+			.HAlign(HAlign_Left)
+			.VAlign(VAlign_Center)
+			.AutoWidth()
+			[
+				PropertyCustomizationHelpers::MakeDeleteButton(FSimpleDelegate::CreateLambda([this, LiteralPtr]()
+				{
+					using namespace Frontend;
+
+					if (!LiteralPtr.IsValid())
+					{
+						return;
+					}
+
+					UMetasoundEditorGraphMember* Member = LiteralPtr->FindMember();
+					if (!Member)
+					{
+						return;
+					}
+
+					const FScopedTransaction Transaction(LOCTEXT("ResetPageDefaultsTransaction", "Reset Paged Defaults"));
+
+					const UMetaSoundSettings* Settings = GetDefault<UMetaSoundSettings>();
+					check(Settings);
+
+					FMetaSoundFrontendDocumentBuilder& Builder = Member->GetFrontendBuilderChecked();
+					UObject& MetaSound = Builder.CastDocumentObjectChecked<UObject>();
+					MetaSound.Modify();
+					LiteralPtr->Modify();
+
+					LiteralPtr->ResetDefaults();
+
+					constexpr bool bPostTransaction = false;
+					Member->UpdateFrontendDefaultLiteral(bPostTransaction);
+
+					FMetasoundAssetBase& MetasoundAsset = Editor::FGraphBuilder::GetOutermostMetaSoundChecked(*LiteralPtr);
+					MetasoundAsset.GetModifyContext().AddMemberIDsModified({ Member->GetMemberID() });
+					UpdatePagePickerNames(LiteralPtr);
+					PageDefaultComboBox->RefreshOptions();
+					FGraphBuilder::RegisterGraphWithFrontend(MetaSound);
+				}), LOCTEXT("ResetPageDefaultsTooltip", "Resets page defaults for the given member, leaving just the initial value for the required 'Default' page."))
+			]
+		];
+	}
+
+	TSharedRef<SWidget> FMetasoundDefaultLiteralCustomizationBase::BuildPageDefaultNameWidget(UMetasoundEditorGraphMemberDefaultLiteral& Literal, TSharedRef<IPropertyHandle> ElementProperty)
+	{
+		TSharedPtr<IPropertyHandle> PageNameProperty = ElementProperty->GetChildHandle(GET_MEMBER_NAME_CHECKED(FMetasoundEditorMemberPageDefault, PageName));
+		if (!ensure(PageNameProperty.IsValid()))
+		{
+			return SNullWidget::NullWidget;
+		}
+
+		const FGuid PageID = LiteralCustomizationPrivate::GetGuidPropertyValue(ElementProperty->GetChildHandle(GET_MEMBER_NAME_CHECKED(FMetasoundEditorMemberPageDefault, PageID)));
+		FName PageName;
+		PageNameProperty->GetValue(PageName);
+		// Can't delete the default page, so just return name selector
+		if (PageID == Frontend::DefaultPageID)
+		{
+			return PageNameProperty->CreatePropertyValueWidget();
+		}
+
+		TWeakObjectPtr<UMetasoundEditorGraphMemberDefaultLiteral> LiteralPtr(&Literal);
+
+		UMetasoundEditorGraphMember* Member = LiteralPtr->FindMember();
+		if (!ensure(Member))
+		{
+			return SNullWidget::NullWidget;
+		}
+
+		const FText RemoveDescription = FText::Format(LOCTEXT("RemovePageDefaultTransactionFormat", "Remove '{0}' Page '{1}' Default Value"), FText::FromName(Member->GetMemberName()), FText::FromName(PageName));
+		TSharedRef<SWidget> RemovePageDefaultButton = PropertyCustomizationHelpers::MakeDeleteButton(FSimpleDelegate::CreateLambda([this, RemoveDescription, PageID, LiteralPtr]()
+		{
+			using namespace Frontend;
+
+			if (!LiteralPtr.IsValid())
+			{
+				return;
+			}
+
+			UMetasoundEditorGraphMember* Member = LiteralPtr->FindMember();
+			if (!Member)
+			{
+				return;
+			}
+
+			const FScopedTransaction Transaction(RemoveDescription);
+
+			const UMetaSoundSettings* Settings = GetDefault<UMetaSoundSettings>();
+			check(Settings);
+
+			FMetaSoundFrontendDocumentBuilder& Builder = Member->GetFrontendBuilderChecked();
+			UObject& MetaSound = Builder.CastDocumentObjectChecked<UObject>();
+			MetaSound.Modify();
+			LiteralPtr->Modify();
+
+			LiteralPtr->RemoveDefault(PageID);
+
+			constexpr bool bPostTransaction = false;
+			Member->UpdateFrontendDefaultLiteral(bPostTransaction);
+
+			FMetasoundAssetBase& MetasoundAsset = Editor::FGraphBuilder::GetOutermostMetaSoundChecked(*LiteralPtr);
+			MetasoundAsset.GetModifyContext().AddMemberIDsModified({ Member->GetMemberID() });
+			UpdatePagePickerNames(LiteralPtr);
+			PageDefaultComboBox->RefreshOptions();
+			FGraphBuilder::RegisterGraphWithFrontend(MetaSound);
+		}),
+		RemoveDescription);
+
+		return SNew(SHorizontalBox)
+		+ SHorizontalBox::Slot()
+		.AutoWidth()
+		.VAlign(VAlign_Center)
+		.Padding(2, 0)
+		[
+			PageNameProperty->CreatePropertyValueWidget()
+		]
+		+ SHorizontalBox::Slot()
+		.FillWidth(1)
+		.HAlign(HAlign_Left)
+		.VAlign(VAlign_Center)
+		.Padding(2, 0)
+		[
+			RemovePageDefaultButton
+		];
+	}
+
+	void FMetasoundDefaultLiteralCustomizationBase::CustomizeDefaults(UMetasoundEditorGraphMemberDefaultLiteral& InLiteral, IDetailLayoutBuilder& InDetailLayout)
+	{
+		CustomizePageDefaultRows(InLiteral, InDetailLayout, nullptr);
+	}
+
+	void FMetasoundDefaultLiteralCustomizationBase::CustomizePageDefaultRows(UMetasoundEditorGraphMemberDefaultLiteral& InLiteral, IDetailLayoutBuilder& InDetailLayout, FOnDefaultPageRowAdded* OnDefaultPageRowAdded)
+	{
+		DefaultProperties.Reset();
+
+		const UMetasoundEditorGraphMember* Member = InLiteral.FindMember();
+		if (!Member)
+		{
+			return;
+		}
+
+		const bool bIsPagedDefault = Member->IsDefaultPaged();
+
+		if (bIsPagedDefault)
+		{
+			TWeakObjectPtr<UMetasoundEditorGraphMemberDefaultLiteral> LiteralPtr(&InLiteral);
+			UpdatePagePickerNames(LiteralPtr);
+		}
+
+		TSharedPtr<IPropertyHandle> DefaultPageArrayHandle = InDetailLayout.AddObjectPropertyData(TArray<UObject*>({ &InLiteral }), UMetasoundEditorGraphMemberDefaultLiteral::GetDefaultsPropertyName());
+		if (DefaultPageArrayHandle.IsValid())
+		{
+			uint32 NumElements = 0;
+			TSharedPtr<IPropertyHandleArray> DefaultValueArray = DefaultPageArrayHandle->AsArray();
+			if (DefaultValueArray.IsValid())
+			{
+				DefaultValueArray->GetNumElements(NumElements);
+
+				bool bShowPageModifiers = Editor::PageEditorEnabled();
+				if (const UMetaSoundSettings* Settings = GetDefault<UMetaSoundSettings>())
+				{
+					// Don't show if project-specific defaults exist and/or project-level pages are implemented.
+					// If one of either, assumption is its the required default page.
+					bShowPageModifiers &= !Settings->GetProjectPageSettings().IsEmpty() || NumElements > 1;
+				}
+
+				if (bIsPagedDefault && bShowPageModifiers)
+				{
+					BuildPageDefaultComboBox(InLiteral, DefaultPageArrayHandle->GetPropertyDisplayName());
+				}
+
+				for (uint32 Index = 0; Index < NumElements; ++Index)
+				{
+					TSharedPtr<IPropertyHandle> ElementProperty = DefaultValueArray->GetElement(Index);
+					if (!ElementProperty.IsValid())
+					{
+						continue;
+					}
+
+					TSharedPtr<IPropertyHandle> ValueProperty = ElementProperty->GetChildHandle("Value");
+					if (!ValueProperty.IsValid())
+					{
+						continue;
+					}
+
+					DefaultProperties.Add(ValueProperty);
+
+					IDetailPropertyRow& ValueRow = InDetailLayout.AddPropertyToCategory(ValueProperty);
+					if (bIsPagedDefault && bShowPageModifiers)
+					{
+						ValueRow.CustomWidget(true /*bShowChilden */);
+
+						(*ValueRow.CustomNameWidget())
+						[
+							BuildPageDefaultNameWidget(InLiteral, ElementProperty.ToSharedRef())
+						];
+
+						ValueRow.ShowPropertyButtons(false);
+						(*ValueRow.CustomValueWidget())
+						[
+							ValueProperty->CreatePropertyValueWidget()
+						];
+					}
+					ValueRow.IsEnabled(GetEnabled());
+
+					// TODO: Fix half-implemented array multi-add broken by page implementation
+// 					if (bIsPagedDefault && bShowPageModifiers && OnDefaultPageRowAdded)
+// 					{
+// 						(*OnDefaultPageRowAdded)(ValueRow, ElementProperty->AsShared());
+// 					}
+				}
+			}
+		}
+	}
+
+	TArray<IDetailPropertyRow*> FMetasoundDefaultLiteralCustomizationBase::CustomizeLiteral(UMetasoundEditorGraphMemberDefaultLiteral& InLiteral, IDetailLayoutBuilder& InDetailLayout)
+	{
+		CustomizeDefaults(InLiteral, InDetailLayout);
+		return { };
+	}
+
+	TAttribute<EVisibility> FMetasoundDefaultLiteralCustomizationBase::GetDefaultVisibility() const
+	{
+		return Visibility;
+	}
+
+	TAttribute<bool> FMetasoundDefaultLiteralCustomizationBase::GetEnabled() const
+	{
+		return Enabled;
+	}
+
+	void FMetasoundDefaultLiteralCustomizationBase::SetDefaultVisibility(TAttribute<EVisibility> InVisibility)
+	{
+		Visibility = InVisibility;
+	}
+
+	void FMetasoundDefaultLiteralCustomizationBase::SetEnabled(TAttribute<bool> InEnabled)
+	{
+		Enabled = InEnabled;
+	}
+
+	void FMetasoundDefaultLiteralCustomizationBase::SetResetOverride(const TOptional<FResetToDefaultOverride>& InResetOverride)
+	{
+		ResetOverride = InResetOverride;
+	}
+
+	void FMetasoundDefaultLiteralCustomizationBase::UpdatePagePickerNames(TWeakObjectPtr<UMetasoundEditorGraphMemberDefaultLiteral> LiteralPtr)
+	{
+		using namespace Frontend;
+
+		AddablePageStringNames.Reset();
+		ImplementedPageNames.Reset();
+
+		if (!LiteralPtr.IsValid())
+		{
+			return;
+		}
+
+		const UMetaSoundSettings* Settings = GetDefault<UMetaSoundSettings>();
+		check(Settings);
+
+		TSet<FGuid> ImplementedGuids;
+		LiteralPtr->IterateDefaults([&ImplementedGuids](const FGuid& PageID, FMetasoundFrontendLiteral)
+		{
+			ImplementedGuids.Add(PageID);
+		});
+
+		auto GetStringName = [](const FMetaSoundPageSettings& Page) -> TSharedPtr<FString> { return MakeShared<FString>(Page.Name.ToString()); };
+		auto RemoveImplementedItem = [&ImplementedGuids](const FMetaSoundPageSettings& Page) { return !ImplementedGuids.Contains(Page.UniqueId); };
+		AddablePageStringNames.Add(GetStringName(Settings->GetDefaultPageSettings()));
+		Algo::TransformIf(Settings->GetProjectPageSettings(), AddablePageStringNames, RemoveImplementedItem, GetStringName);
+
+		auto GetPageName = [&Settings](const FGuid& PageID)
+		{
+			const FMetaSoundPageSettings* Page = Settings->FindPageSettings(PageID);
+			if (Page)
+			{
+				return Page->Name;
+			}
+
+			return Metasound::Editor::GetMissingPageName(PageID);
+		};
+
+		Algo::Transform(ImplementedGuids, ImplementedPageNames, GetPageName);
+	}
+} // namespace Metasound::Editor
+#undef LOCTEXT_NAMESPACE
diff --git a/Engine/Plugins/Runtime/Metasound/Source/MetasoundEditor/Private/MetasoundDetailCustomization.cpp b/Engine/Plugins/Runtime/Metasound/Source/MetasoundEditor/Private/MetasoundDetailCustomization.cpp
--- a/Engine/Plugins/Runtime/Metasound/Source/MetasoundEditor/Private/MetasoundDetailCustomization.cpp
+++ b/Engine/Plugins/Runtime/Metasound/Source/MetasoundEditor/Private/MetasoundDetailCustomization.cpp
@@ -443,7 +443,7 @@ namespace Metasound::Editor
 					UpdateItemNames();
 					ComboBox->RefreshOptions();
 					FGraphBuilder::RegisterGraphWithFrontend(MetaSound);
-				}), LOCTEXT("ResetGraphPagesTooltip", "Removes all page graphs from the given MetaSound less the default."))
+				}), LOCTEXT("ResetGraphPagesTooltip", "Removes all page graphs from the given MetaSound defined in the MetaSound project settings (does not remove the required 'Default' graph)."))
 			];
 
 		Utilities->SetEnabled(TAttribute<bool>::Create(TAttribute<bool>::FGetter::CreateSP(this, &FMetaSoundDetailCustomizationBase::IsGraphEditable)));
@@ -738,6 +738,7 @@ namespace Metasound::Editor
 		{
 			return;
 		}
+
 		if (UMetasoundInterfacesView* InterfacesView = CastChecked<UMetasoundInterfacesView>(Objects.Last()))
 		{
 			if (UObject* MetaSound = InterfacesView->GetMetasound())
@@ -746,6 +747,11 @@ namespace Metasound::Editor
 			}
 		}
 
+		TAttribute<bool> IsGraphEditableAttribute = TAttribute<bool>::Create(TAttribute<bool>::FGetter::CreateSP
+		(
+			this, &FMetaSoundDetailCustomizationBase::IsGraphEditable
+		));
+
 		UpdateInterfaceNames();
 
 		SAssignNew(InterfaceComboBox, SSearchableComboBox)
diff --git a/Engine/Plugins/Runtime/Metasound/Source/MetasoundEditor/Private/MetasoundDetailCustomization.h b/Engine/Plugins/Runtime/Metasound/Source/MetasoundEditor/Private/MetasoundDetailCustomization.h
--- a/Engine/Plugins/Runtime/Metasound/Source/MetasoundEditor/Private/MetasoundDetailCustomization.h
+++ b/Engine/Plugins/Runtime/Metasound/Source/MetasoundEditor/Private/MetasoundDetailCustomization.h
@@ -112,7 +112,6 @@ namespace Metasound
 
 			TSet<FName> ImplementedInterfaceNames;
 			TSharedPtr<SSearchableComboBox> InterfaceComboBox;
-			TAttribute<bool> IsGraphEditableAttribute;
 		};
 	} // namespace Editor
 } // namespace Metasound
\ No newline at end of file
diff --git a/Engine/Plugins/Runtime/Metasound/Source/MetasoundEditor/Private/MetasoundEditorDocumentClipboardUtils.cpp b/Engine/Plugins/Runtime/Metasound/Source/MetasoundEditor/Private/MetasoundEditorDocumentClipboardUtils.cpp
--- a/Engine/Plugins/Runtime/Metasound/Source/MetasoundEditor/Private/MetasoundEditorDocumentClipboardUtils.cpp
+++ b/Engine/Plugins/Runtime/Metasound/Source/MetasoundEditor/Private/MetasoundEditorDocumentClipboardUtils.cpp
@@ -39,18 +39,18 @@ namespace Metasound::Editor
 				bool bNameMatchFound = false;
 				const FMetasoundEditorGraphVertexNodeBreadcrumb& Breadcrumb = InputNode->GetBreadcrumb();
 				Graph.IterateInputs([&Input, &InputNode, &Graph, &Breadcrumb, &bNameMatchFound](UMetasoundEditorGraphInput& TestInput)
+				{
+					FConstNodeHandle InputHandle = TestInput.GetConstNodeHandle();
+					FConstOutputHandle TestOutput = InputHandle->GetConstOutputs().Last();
+					const bool bTypeMatches = TestOutput->GetDataType() == Breadcrumb.DataType;
+					const bool bAccessMatches = TestOutput->GetVertexAccessType() == Breadcrumb.AccessType;
+					const bool bNameMatches = InputHandle->GetNodeName() == Breadcrumb.MemberName;
+					bNameMatchFound |= bNameMatches;
+					if (bTypeMatches && bAccessMatches && bNameMatches)
 					{
-						FConstNodeHandle InputHandle = TestInput.GetConstNodeHandle();
-						FConstOutputHandle TestOutput = InputHandle->GetConstOutputs().Last();
-						const bool bTypeMatches = TestOutput->GetDataType() == Breadcrumb.DataType;
-						const bool bAccessMatches = TestOutput->GetVertexAccessType() == Breadcrumb.AccessType;
-						const bool bNameMatches = InputHandle->GetNodeName() == Breadcrumb.MemberName;
-						bNameMatchFound |= bNameMatches;
-						if (bTypeMatches && bAccessMatches && bNameMatches)
-						{
-							Input = &TestInput;
-						}
-					});
+						Input = &TestInput;
+					}
+				});
 
 				if (!Input)
 				{
@@ -67,7 +67,13 @@ namespace Metasound::Editor
 							FCreateNodeVertexParams VertexParams;
 							VertexParams.DataType = Breadcrumb.DataType;
 
-							FMetasoundFrontendClassInput ClassInput = FGraphBuilder::CreateUniqueClassInput(*OutAsset.GetOwningAsset(), VertexParams, &Breadcrumb.DefaultLiterals.FindChecked(Frontend::DefaultPageID), &Breadcrumb.MemberName);
+							TArray<FMetasoundFrontendClassInputDefault> InputDefaults;
+							Algo::Transform(Breadcrumb.DefaultLiterals, InputDefaults, [](const TPair<FGuid, FMetasoundFrontendLiteral>& Pair)
+							{
+								return FMetasoundFrontendClassInputDefault(Pair.Key, Pair.Value);
+							});
+
+							FMetasoundFrontendClassInput ClassInput = FGraphBuilder::CreateUniqueClassInput(*OutAsset.GetOwningAsset(), VertexParams, InputDefaults, &Breadcrumb.MemberName);
 							if (const FMetasoundFrontendNode* NewNode = Builder.AddGraphInput(ClassInput))
 							{
 								Input = Graph.FindOrAddInput(NewNode->GetID());
diff --git a/Engine/Plugins/Runtime/Metasound/Source/MetasoundEditor/Private/MetasoundEditorGraph.cpp b/Engine/Plugins/Runtime/Metasound/Source/MetasoundEditor/Private/MetasoundEditorGraph.cpp
--- a/Engine/Plugins/Runtime/Metasound/Source/MetasoundEditor/Private/MetasoundEditorGraph.cpp
+++ b/Engine/Plugins/Runtime/Metasound/Source/MetasoundEditor/Private/MetasoundEditorGraph.cpp
@@ -16,6 +16,7 @@
 #include "MetasoundEditorGraphValidation.h"
 #include "MetasoundEditorModule.h"
 #include "MetasoundEditorSubsystem.h"
+#include "MetasoundFrontendDataTypeRegistry.h"
 #include "MetasoundFrontendDocumentBuilder.h"
 #include "MetasoundFrontendNodeTemplateRegistry.h"
 #include "MetasoundFrontendSearchEngine.h"
@@ -50,10 +51,10 @@ namespace Metasound::Editor
 			return *FString::Printf(TEXT("Member_%s"), *FGuid::NewGuid().ToString());
 		}
 
-		void OnLiteralChanged(UMetasoundEditorGraphMember& InMember, EPropertyChangeType::Type InChangeType)
+		void OnLiteralChanged(UMetasoundEditorGraphMember& InMember, const FGuid* InPageID, EPropertyChangeType::Type InChangeType)
 		{
 			constexpr bool bPostTransaction = false;
-			InMember.UpdateFrontendDefaultLiteral(bPostTransaction);
+			InMember.UpdateFrontendDefaultLiteral(bPostTransaction, InPageID);
 
 			const bool bCommitChange = InChangeType != EPropertyChangeType::Interactive;
 			if (bCommitChange)
@@ -69,50 +70,56 @@ namespace Metasound::Editor
 			}
 		}
 
-		void SetAndUpdatePreviewLiteral(FMetaSoundFrontendDocumentBuilder& Builder, UMetasoundEditorGraph* MetaSoundGraph, FName MemberName, UMetasoundEditorGraphMemberDefaultLiteral& Literal)
+		// Avoids member literal setting the node literal if its not required (which in turn
+		// avoids 'Reset To Default' action from being enabled when the default is equal)
+		void SetOrClearIfLiteralMatchesNodeVertexDefault(FMetaSoundFrontendDocumentBuilder& InBuilder, const FMetasoundFrontendVertexHandle& VertexHandle, const FMetasoundFrontendLiteral& InDefaultLiteral)
 		{
-			Builder.SetGraphInputDefault(MemberName, Literal.GetDefault());
-			if (GEditor)
-			{
-				if (MetaSoundGraph && MetaSoundGraph->IsPreviewing())
-				{
-					UAudioComponent* PreviewComponent = GEditor->GetPreviewAudioComponent();
-					check(PreviewComponent);
-
-					if (TScriptInterface<IAudioParameterControllerInterface> ParamInterface = PreviewComponent)
-					{
-						Literal.UpdatePreviewInstance(MemberName, ParamInterface);
-					}
-				}
-			}
-		}
-
-		void SetLiteralOrClearIfMatchesDefault(Frontend::FInputHandle& InInputHandle, const FMetasoundFrontendLiteral& InDefaultLiteral)
-		{
-			using namespace Frontend;
+			const FMetasoundFrontendVertex* Vertex = InBuilder.FindNodeInput(VertexHandle.NodeID, VertexHandle.VertexID);
+			check(Vertex);
 
-			// Avoids member literal setting the node literal if its not required (which in turn
-			// avoids 'Reset To Default' action from being enabled when the default is equal)
 			bool bClearLiteral = false;
-			if (const FMetasoundFrontendLiteral* ClassDefault = InInputHandle->GetClassDefaultLiteral())
+			const TArray<FMetasoundFrontendClassInputDefault>* ClassDefaults = InBuilder.FindNodeClassInputDefaults(VertexHandle.NodeID, Vertex->Name);
+			if (ClassDefaults)
 			{
-				bClearLiteral = ClassDefault->IsEqual(InDefaultLiteral);
+				const FGuid PageID = Engine::FDocumentBuilderRegistry::GetChecked().ResolveTargetPageID(*ClassDefaults);
+				auto MatchesPageID = [&PageID](const FMetasoundFrontendClassInputDefault& InputDefault) { return InputDefault.PageID == PageID; };
+				if (const FMetasoundFrontendClassInputDefault* ClassDefault = ClassDefaults->FindByPredicate(MatchesPageID))
+				{
+					bClearLiteral = ClassDefault->Literal.IsEqual(InDefaultLiteral);
+				}
 			}
 
 			if (!bClearLiteral)
 			{
 				FMetasoundFrontendLiteral DefaultTypeLiteral;
-				DefaultTypeLiteral.SetFromLiteral(IDataTypeRegistry::Get().CreateDefaultLiteral(InInputHandle->GetDataType()));
+				DefaultTypeLiteral.SetFromLiteral(Frontend::IDataTypeRegistry::Get().CreateDefaultLiteral(Vertex->TypeName));
 				bClearLiteral = InDefaultLiteral.IsEqual(DefaultTypeLiteral);
 			}
 
 			if (bClearLiteral)
 			{
-				InInputHandle->ClearLiteral();
+				InBuilder.RemoveNodeInputDefault(VertexHandle.NodeID, VertexHandle.VertexID);
 			}
 			else
 			{
-				InInputHandle->SetLiteral(InDefaultLiteral);
+				InBuilder.SetNodeInputDefault(VertexHandle.NodeID, VertexHandle.VertexID, InDefaultLiteral);
+			}
+		}
+
+		void UpdatePreviewParameter(UMetasoundEditorGraph* MetaSoundGraph, FName MemberName, UMetasoundEditorGraphMemberDefaultLiteral& Literal)
+		{
+			if (GEditor)
+			{
+				if (MetaSoundGraph && MetaSoundGraph->IsPreviewing())
+				{
+					UAudioComponent* PreviewComponent = GEditor->GetPreviewAudioComponent();
+					check(PreviewComponent);
+
+					if (TScriptInterface<IAudioParameterControllerInterface> ParamInterface = PreviewComponent)
+					{
+						Literal.UpdatePreviewInstance(MemberName, ParamInterface);
+					}
+				}
 			}
 		}
 	} // namespace GraphPrivate
@@ -188,28 +195,25 @@ void UMetasoundEditorGraphMember::PostEditUndo()
 
 	if (IsValid(this))
 	{
-		SetDataType(TypeName);
-
 		constexpr bool bPostTransaction = false;
+		SetDataType(TypeName, bPostTransaction);
 		UpdateFrontendDefaultLiteral(bPostTransaction);
-
-		if (UMetasoundEditorGraph* Graph = GetOwningGraph())
-		{
-			if (IsValid(Graph))
-			{
-				UObject* MetaSound = Graph->GetMetasound();
-				if (IsValid(MetaSound))
-				{
-					FMetasoundAssetBase* MetasoundAsset = IMetasoundUObjectRegistry::Get().GetObjectAsAssetBase(MetaSound);
-					check(MetasoundAsset);
-					MetasoundAsset->GetModifyContext().SetDocumentModified();
-				}
-			}
-		}
 	}
 }
 #endif // WITH_EDITOR
 
+bool UMetasoundEditorGraphMember::Synchronize()
+{
+	bool bModified = false;
+	if (!Literal)
+	{
+		bModified = true;
+		InitializeLiteral();
+	}
+
+	return bModified;
+}
+
 void UMetasoundEditorGraphVertex::InitMember(FName InDataType, const FMetasoundFrontendLiteral& InDefaultLiteral, FGuid InNodeID, FMetasoundFrontendClassName&& InClassName)
 {
 	TypeName = InDataType;
@@ -572,12 +576,44 @@ bool UMetasoundEditorGraphVertex::CanRename(const FText& InNewText, FText& OutEr
 	return bIsNameValid;
 }
 
-const UMetasoundEditorGraphMember* UMetasoundEditorGraphMemberDefaultLiteral::FindMember() const
+bool UMetasoundEditorGraphVertex::Synchronize()
+{
+	bool bModified = Super::Synchronize();
+
+	if (const FMetasoundFrontendClassVertex* Vertex = GetFrontendClassVertex(); ensure(Vertex))
+	{
+		if (TypeName != Vertex->TypeName)
+		{
+			bModified = true;
+			TypeName = Vertex->TypeName;
+
+			InitializeLiteral();
+		}
+
+		FMetaSoundFrontendDocumentBuilder& Builder = GetFrontendBuilderChecked();
+		if (const FMetasoundFrontendNode* Node = Builder.FindNode(NodeID); ensure(Node))
+		{
+			if (const FMetasoundFrontendClass* Class = Builder.FindDependency(Node->ClassID); ensure(Class))
+			{
+				const FMetasoundFrontendClassName& FrontendClassName = Class->Metadata.GetClassName();
+				if (ClassName != FrontendClassName)
+				{
+					bModified = true;
+					ClassName = FrontendClassName;
+				}
+			}
+		}
+	}
+
+	return bModified;
+}
+
+UMetasoundEditorGraphMember* UMetasoundEditorGraphMemberDefaultLiteral::FindMember() const
 {
 	using namespace Metasound;
 
 	const FMetasoundAssetBase& MetaSound = Editor::FGraphBuilder::GetOutermostConstMetaSoundChecked(*this);
-	if (const UMetasoundEditorGraph* Graph = Cast<UMetasoundEditorGraph>(MetaSound.GetGraph()))
+	if (UMetasoundEditorGraph* Graph = Cast<UMetasoundEditorGraph>(MetaSound.GetGraph()))
 	{
 		return Graph->FindMember(MemberID);
 	}
@@ -585,18 +621,50 @@ const UMetasoundEditorGraphMember* UMetasoundEditorGraphMemberDefaultLiteral::Fi
 	return nullptr;
 }
 
-UMetasoundEditorGraphMember* UMetasoundEditorGraphMemberDefaultLiteral::FindMember()
+void UMetasoundEditorGraphMemberDefaultLiteral::ForceRefresh()
 {
-	using namespace Metasound;
+}
 
-	FMetasoundAssetBase* MetasoundAsset = IMetasoundUObjectRegistry::Get().GetObjectAsAssetBase(GetOutermostObject());
-	check(MetasoundAsset);
-	if (UMetasoundEditorGraph* Graph = Cast<UMetasoundEditorGraph>(MetasoundAsset->GetGraph()))
+FName UMetasoundEditorGraphMemberDefaultLiteral::GetDataType() const
+{
+	return { };
+}
+
+EMetasoundFrontendLiteralType UMetasoundEditorGraphMemberDefaultLiteral::GetLiteralType() const
+{
+	return EMetasoundFrontendLiteralType::None;
+}
+
+void UMetasoundEditorGraphMemberDefaultLiteral::InitDefault(const FGuid& InPageID)
+{
+}
+
+void UMetasoundEditorGraphMemberDefaultLiteral::IterateDefaults(TFunctionRef<void(const FGuid&, FMetasoundFrontendLiteral)> Iter) const
+{
+	FMetasoundFrontendLiteral Literal;
+	if (TryFindDefault(Literal))
 	{
-		return Graph->FindMember(MemberID);
+		Iter(Metasound::Frontend::DefaultPageID, MoveTemp(Literal));
 	}
+}
 
-	return nullptr;
+bool UMetasoundEditorGraphMemberDefaultLiteral::RemoveDefault(const FGuid& InPageID)
+{
+	return false;
+}
+
+void UMetasoundEditorGraphMemberDefaultLiteral::ResetDefaults()
+{
+}
+
+void UMetasoundEditorGraphMemberDefaultLiteral::SetFromLiteral(const FMetasoundFrontendLiteral& InLiteral, const FGuid& InPageID)
+{
+}
+
+bool UMetasoundEditorGraphMemberDefaultLiteral::TryFindDefault(FMetasoundFrontendLiteral& OutLiteral, const FGuid* InPageID) const
+{
+	OutLiteral = { };
+	return true;
 }
 
 #if WITH_EDITOR
@@ -604,15 +672,28 @@ void UMetasoundEditorGraphMemberDefaultLiteral::PostEditChangeProperty(FProperty
 {
 	if (UMetasoundEditorGraphMember* Member = FindMember())
 	{
-		Metasound::Editor::GraphPrivate::OnLiteralChanged(*Member, InPropertyChangedEvent.ChangeType);
+		Metasound::Editor::GraphPrivate::OnLiteralChanged(*Member, nullptr, InPropertyChangedEvent.ChangeType);
 	}
 }
 
 void UMetasoundEditorGraphMemberDefaultLiteral::PostEditChangeChainProperty(FPropertyChangedChainEvent& InPropertyChangedEvent)
 {
+	if (const FEditPropertyChain::TDoubleLinkedListNode* MemberNode = InPropertyChangedEvent.PropertyChain.GetActiveMemberNode())
+	{
+		if (const FProperty* ChildProperty = MemberNode->GetValue())
+		{
+			const FName ChildPropertyName = ChildProperty->GetFName();
+			if (ChildPropertyName.IsEqual(GetDefaultsPropertyName()) || ChildPropertyName.IsEqual(GetDefaultPropertyName()))
+			{
+				ResolvePageDefaults();
+				SortPageDefaults();
+			}
+		}
+	}
+
 	if (UMetasoundEditorGraphMember* Member = FindMember())
 	{
-		Metasound::Editor::GraphPrivate::OnLiteralChanged(*Member, InPropertyChangedEvent.ChangeType);
+		Metasound::Editor::GraphPrivate::OnLiteralChanged(*Member, nullptr, InPropertyChangedEvent.ChangeType);
 	}
 }
 
@@ -631,11 +712,36 @@ void UMetasoundEditorGraphMemberDefaultLiteral::PostEditUndo()
 	if (UMetasoundEditorGraphMember* Member = FindMember())
 	{
 		Member->UpdateFrontendDefaultLiteral(bPostTransaction);
-		FGraphBuilder::GetOutermostMetaSoundChecked(*this).GetModifyContext().SetDocumentModified();
 	}
 }
 #endif // WITH_EDITOR
 
+bool UMetasoundEditorGraphMemberDefaultLiteral::TryGetPreviewPageID(FGuid& OutPreviewPageID) const
+{
+	using namespace Metasound::Engine;
+
+	const UMetasoundEditorGraphMember* Member = FindMember();
+	if (!ensure(Member))
+	{
+		OutPreviewPageID = Metasound::Frontend::DefaultPageID;
+		return false;
+	}
+
+	const FMetaSoundFrontendDocumentBuilder& Builder = Member->GetFrontendBuilderChecked();
+	if (const FMetasoundFrontendClassInput* Input = Builder.FindGraphInput(Member->GetMemberName()))
+	{
+		const FGuid PageID = FDocumentBuilderRegistry::GetChecked().ResolveTargetPageID(*Input);
+		if (PageID == Builder.GetBuildPageID())
+		{
+			OutPreviewPageID = PageID;
+			return true;
+		}
+	}
+
+	OutPreviewPageID = Metasound::Frontend::DefaultPageID;
+	return false;
+}
+
 Metasound::Editor::ENodeSection UMetasoundEditorGraphInput::GetSectionID() const 
 {
 	return Metasound::Editor::ENodeSection::Inputs;
@@ -656,7 +762,7 @@ Metasound::Frontend::FNodeHandle UMetasoundEditorGraphInput::AddNodeHandle(const
 	FMetasoundAssetBase* MetaSoundAsset = Metasound::IMetasoundUObjectRegistry::Get().GetObjectAsAssetBase(&MetaSound);
 	check(MetaSoundAsset);
 
-	FMetasoundFrontendClassInput ClassInput = FGraphBuilder::CreateUniqueClassInput(MetaSound, InParams, nullptr, &InName);
+	FMetasoundFrontendClassInput ClassInput = FGraphBuilder::CreateUniqueClassInput(MetaSound, InParams, { }, &InName);
 	return MetaSoundAsset->GetRootGraphHandle()->AddInputVertex(ClassInput);
 }
 
@@ -683,6 +789,11 @@ const FText& UMetasoundEditorGraphInput::GetGraphMemberLabel() const
 	return Label;
 }
 
+const FMetasoundFrontendClassVertex* UMetasoundEditorGraphInput::GetFrontendClassVertex() const
+{
+	return GetFrontendBuilderChecked().FindGraphInput(GetMemberName());
+}
+
 FText UMetasoundEditorGraphInput::GetDescription() const
 {
 	const FMetaSoundFrontendDocumentBuilder& Builder = GetFrontendBuilderChecked();
@@ -729,6 +840,11 @@ TArray<UMetasoundEditorGraphMemberNode*> UMetasoundEditorGraphInput::GetNodes()
 	return Nodes;
 }
 
+bool UMetasoundEditorGraphInput::IsDefaultPaged() const
+{
+	return true;
+}
+
 bool UMetasoundEditorGraphInput::IsInterfaceMember(FMetasoundFrontendInterface* OutInterface) const
 {
 	FMetasoundFrontendInterface Interface;
@@ -779,8 +895,8 @@ void UMetasoundEditorGraphInput::SetSortOrderIndex(int32 InSortOrderIndex)
 
 void UMetasoundEditorGraphInput::ResetToClassDefault()
 {
+	using namespace Metasound;
 	using namespace Metasound::Editor;
-	using namespace Metasound::Frontend;
 
 	if (ensure(Literal))
 	{
@@ -789,12 +905,13 @@ void UMetasoundEditorGraphInput::ResetToClassDefault()
 		Builder.CastDocumentObjectChecked<UObject>().Modify();
 		Literal->Modify();
 
-		FMetasoundFrontendLiteral DefaultLiteral;
-		DefaultLiteral.SetFromLiteral(IDataTypeRegistry::Get().CreateDefaultLiteral(GetDataType()));
-		Literal->SetFromLiteral(DefaultLiteral);
-
 		const FName MemberName = GetMemberName();
-		GraphPrivate::SetAndUpdatePreviewLiteral(Builder, GetOwningGraph(), MemberName, *Literal);
+		Builder.ResetGraphInputDefault(MemberName);
+
+		constexpr bool bPostTransaction = false;
+		UpdateFrontendDefaultLiteral(bPostTransaction);
+
+		GraphPrivate::UpdatePreviewParameter(GetOwningGraph(), MemberName, *Literal);
 	}
 }
 
@@ -813,13 +930,20 @@ void UMetasoundEditorGraphInput::SetDataType(FName InNewType, bool bPostTransact
 			Graph->Modify();
 			Modify();
 
-			EMetaSoundBuilderResult Result = EMetaSoundBuilderResult::Failed;
-			Graph->GetBuilderChecked().SetGraphInputDataType(GetMemberName(), InNewType, Result);
-			ensure(Result == EMetaSoundBuilderResult::Succeeded);
+			FMetaSoundFrontendDocumentBuilder& Builder = GetFrontendBuilderChecked();
+			const bool bSuccess = Builder.SetGraphInputDataType(GetMemberName(), InNewType);
+			ensure(bSuccess);
 
 			// Cached TypeName here must be set prior to re-initializing literal below
 			TypeName = InNewType;
-			ClassName = GetConstNodeHandle()->GetClassMetadata().GetClassName();
+
+			if (const FMetasoundFrontendNode* Node = Builder.FindNode(NodeID); ensure(Node))
+			{
+				if (const FMetasoundFrontendClass* Dependency = Builder.FindDependency(Node->ClassID); ensure(Dependency))
+				{
+					ClassName = Dependency->Metadata.GetClassName();
+				}
+			}
 
 			InitializeLiteral();
 
@@ -867,7 +991,19 @@ void UMetasoundEditorGraphInput::SetMemberName(const FName& InNewName, bool bPos
 	Super::SetMemberName(InNewName, bPostTransaction);
 }
 
-void UMetasoundEditorGraphInput::UpdateFrontendDefaultLiteral(bool bPostTransaction)
+bool UMetasoundEditorGraphInput::Synchronize()
+{
+	bool bModified = Super::Synchronize();
+
+	if (ensure(Literal))
+	{
+		bModified |= Literal->Synchronize();
+	}
+
+	return bModified;
+}
+
+void UMetasoundEditorGraphInput::UpdateFrontendDefaultLiteral(bool bPostTransaction, const FGuid* InPageID)
 {
 	using namespace Metasound::Editor;
 
@@ -875,12 +1011,30 @@ void UMetasoundEditorGraphInput::UpdateFrontendDefaultLiteral(bool bPostTransact
 	{
 		FMetaSoundFrontendDocumentBuilder& Builder = GetFrontendBuilderChecked();
 
-		const FScopedTransaction Transaction(LOCTEXT("Set Input Default", "Set MetaSound Input Default"), bPostTransaction);
+		const FScopedTransaction Transaction(FText::Format(LOCTEXT("Set Input Defaults", "Set MetaSound Input '{0}' Default(s)"), GetDisplayName()), bPostTransaction);
 		Builder.CastDocumentObjectChecked<UObject>().Modify();
 		Literal->Modify();
 
 		const FName MemberName = GetMemberName();
-		GraphPrivate::SetAndUpdatePreviewLiteral(Builder, GetOwningGraph(), MemberName, *Literal);
+		if (InPageID)
+		{
+			FMetasoundFrontendLiteral Default;
+			if (Literal->TryFindDefault(Default, InPageID))
+			{
+				Builder.SetGraphInputDefault(MemberName, Default, InPageID);
+			}
+		}
+		else
+		{
+			TArray<FMetasoundFrontendClassInputDefault> NewDefaults;
+			Literal->IterateDefaults([&NewDefaults](const FGuid& PageID, FMetasoundFrontendLiteral Default)
+			{
+				NewDefaults.Add(FMetasoundFrontendClassInputDefault(PageID, MoveTemp(Default)));
+			});
+			Builder.SetGraphInputDefaults(MemberName, MoveTemp(NewDefaults));
+		}
+
+		GraphPrivate::UpdatePreviewParameter(GetOwningGraph(), MemberName, *Literal);
 	}
 }
 
@@ -935,6 +1089,11 @@ FText UMetasoundEditorGraphOutput::GetDescription() const
 	return { };
 }
 
+const FMetasoundFrontendClassVertex* UMetasoundEditorGraphOutput::GetFrontendClassVertex() const
+{
+	return GetFrontendBuilderChecked().FindGraphOutput(GetMemberName());
+}
+
 int32 UMetasoundEditorGraphOutput::GetSortOrderIndex() const
 {
 	using namespace Metasound::Editor;
@@ -1009,23 +1168,16 @@ void UMetasoundEditorGraphOutput::ResetToClassDefault()
 	using namespace Metasound::Editor;
 	using namespace Metasound::Frontend;
 
+	FMetaSoundFrontendDocumentBuilder& Builder = GetFrontendBuilderChecked();
 	TArray<UMetasoundEditorGraphMemberNode*> Nodes = GetNodes();
 	for (UMetasoundEditorGraphMemberNode* Node : Nodes)
 	{
-		TArray<FInputHandle> Inputs = Node->GetNodeHandle()->GetInputs();
-		if (ensure(!Inputs.IsEmpty()))
+		TArray<const FMetasoundFrontendVertex*> InputVertices = Builder.FindNodeInputs(Node->GetNodeID());
+		if (ensure(InputVertices.Num() == 1))
 		{
-			Inputs.Last()->ClearLiteral();
+			Builder.RemoveNodeInputDefault(Node->GetNodeID(), InputVertices.Last()->VertexID);
 		}
 	}
-
-	FMetasoundFrontendLiteral DefaultLiteral;
-	DefaultLiteral.SetFromLiteral(IDataTypeRegistry::Get().CreateDefaultLiteral(GetDataType()));
-
-	Literal->Modify();
-	Literal->SetFromLiteral(DefaultLiteral);
-
-	FGraphBuilder::GetOutermostMetaSoundChecked(*this).GetModifyContext().AddMemberIDsModified({ GetMemberID() });
 }
 
 void UMetasoundEditorGraphOutput::SetDataType(FName InNewType, bool bPostTransaction)
@@ -1082,7 +1234,7 @@ void UMetasoundEditorGraphOutput::SetVertexAccessType(EMetasoundFrontendVertexAc
 	}
 }
 
-void UMetasoundEditorGraphOutput::UpdateFrontendDefaultLiteral(bool bPostTransaction)
+void UMetasoundEditorGraphOutput::UpdateFrontendDefaultLiteral(bool bPostTransaction, const FGuid* InPageID)
 {
 	using namespace Metasound;
 	using namespace Metasound::Editor;
@@ -1105,18 +1257,18 @@ void UMetasoundEditorGraphOutput::UpdateFrontendDefaultLiteral(bool bPostTransac
 		return;
 	}
 
-	const FScopedTransaction Transaction(LOCTEXT("Set Output Default", "Set MetaSound Output Default"), bPostTransaction);
-	Metasound->Modify();
-
-	const FMetasoundFrontendLiteral DefaultLiteral = Literal->GetDefault();
-
-	TArray<UMetasoundEditorGraphMemberNode*> Nodes = GetNodes();
-	for (UMetasoundEditorGraphMemberNode* Node : Nodes)
+	FMetasoundFrontendLiteral DefaultLiteral;
+	if (ensure(Literal->TryFindDefault(DefaultLiteral)))
 	{
-		TArray<FInputHandle> Inputs = Node->GetNodeHandle()->GetInputs();
-		if (ensure(!Inputs.IsEmpty()))
+		const FScopedTransaction Transaction(FText::Format(LOCTEXT("Set Output Default", "Set MetaSound Output '{0}' Default"), GetDisplayName()), bPostTransaction);
+		Metasound->Modify();
+
+		if (const FMetasoundFrontendNode* FrontendNode = GetFrontendNode())
 		{
-			GraphPrivate::SetLiteralOrClearIfMatchesDefault(Inputs.Last(), DefaultLiteral);
+			FMetaSoundFrontendDocumentBuilder& Builder = GetFrontendBuilderChecked();
+			const FGuid& VertexID = FrontendNode->Interface.Inputs.Last().VertexID;
+			const FMetasoundFrontendVertexHandle VertexHandle(NodeID, VertexID);
+			GraphPrivate::SetOrClearIfLiteralMatchesNodeVertexDefault(Builder, VertexHandle, DefaultLiteral);
 		}
 	}
 }
@@ -1138,6 +1290,54 @@ Metasound::Editor::ENodeSection UMetasoundEditorGraphOutput::GetSectionID() cons
 	return Metasound::Editor::ENodeSection::Outputs;
 }
 
+bool UMetasoundEditorGraphOutput::Synchronize()
+{
+	using namespace Metasound;
+
+	bool bModified = Super::Synchronize();
+
+	FMetaSoundFrontendDocumentBuilder& Builder = GetFrontendBuilderChecked();
+	const FName MemberName = GetMemberName();
+	if (const FMetasoundFrontendClassOutput* ClassOutput = Builder.FindGraphOutput(MemberName))
+	{
+		if (ensure(Literal))
+		{
+			TOptional<FMetasoundFrontendLiteral> NewDefault;
+			if (const FMetasoundFrontendNode* OutputNode = Builder.FindGraphOutputNode(GetMemberName()); ensure(OutputNode))
+			{
+				FMetasoundFrontendLiteral DefaultLiteral;
+				Literal->TryFindDefault(DefaultLiteral);
+				if (!OutputNode->InputLiterals.IsEmpty())
+				{
+					const FMetasoundFrontendVertexLiteral& VertexLiteral = OutputNode->InputLiterals.Last();
+					if (!VertexLiteral.Value.IsEqual(DefaultLiteral))
+					{
+						NewDefault = VertexLiteral.Value;
+					}
+				}
+				else
+				{
+					FMetasoundFrontendLiteral TypeDefault;
+					TypeDefault.SetFromLiteral(Frontend::IDataTypeRegistry::Get().CreateDefaultLiteral(TypeName));
+					if (!TypeDefault.IsEqual(DefaultLiteral))
+					{
+						NewDefault = TypeDefault;
+					}
+				}
+			}
+
+			if (NewDefault.IsSet())
+			{
+				bModified = true;
+				Literal->ResetDefaults();
+				Literal->SetFromLiteral(*NewDefault);
+			}
+		}
+	}
+
+	return bModified;
+}
+
 void UMetasoundEditorGraphVariable::InitMember(FName InDataType, const FMetasoundFrontendLiteral& InDefaultLiteral, FGuid InVariableID)
 {
 	TypeName = InDataType;
@@ -1565,44 +1765,91 @@ void UMetasoundEditorGraphVariable::ResetToClassDefault()
 	FGraphBuilder::GetOutermostMetaSoundChecked(*this).GetModifyContext().AddMemberIDsModified({ GetMemberID() });
 }
 
-void UMetasoundEditorGraphVariable::UpdateFrontendDefaultLiteral(bool bPostTransaction)
+bool UMetasoundEditorGraphVariable::Synchronize()
 {
 	using namespace Metasound;
-	using namespace Metasound::Editor;
-	using namespace Metasound::Frontend;
-	using namespace Metasound::VariableNames;
 
-	UObject* Metasound = nullptr;
-	UMetasoundEditorGraph* MetaSoundGraph = GetOwningGraph();
-	if (ensure(MetaSoundGraph))
+	bool bModified = Super::Synchronize();
+
+	FMetaSoundFrontendDocumentBuilder& Builder = GetFrontendBuilderChecked();
+	const FMetasoundFrontendGraph& Graph = Builder.FindConstBuildGraphChecked();
+	const FName MemberName = GetMemberName();
+	auto IsVariable = [&MemberName] (const FMetasoundFrontendVariable& Var) { return Var.Name == MemberName; };
+	if (const FMetasoundFrontendVariable* Variable = Graph.Variables.FindByPredicate(IsVariable))
 	{
-		Metasound = MetaSoundGraph->GetMetasound();
+		if (ensure(Literal))
+		{
+			TOptional<FMetasoundFrontendLiteral> NewDefault;
+			if (const FMetasoundFrontendNode* MutatorNode = Builder.FindNode(Variable->MutatorNodeID); ensure(MutatorNode))
+			{
+				FMetasoundFrontendLiteral DefaultLiteral;
+				Literal->TryFindDefault(DefaultLiteral);
+				if (!MutatorNode->InputLiterals.IsEmpty())
+				{
+					const FMetasoundFrontendVertexLiteral& VertexLiteral = MutatorNode->InputLiterals.Last();
+					if (!VertexLiteral.Value.IsEqual(DefaultLiteral))
+					{
+						NewDefault = VertexLiteral.Value;
+					}
+				}
+				else
+				{
+					FMetasoundFrontendLiteral TypeDefault;
+					TypeDefault.SetFromLiteral(Frontend::IDataTypeRegistry::Get().CreateDefaultLiteral(TypeName));
+					if (!TypeDefault.IsEqual(DefaultLiteral))
+					{
+						NewDefault = TypeDefault;
+					}
+				}
+			}
+
+			if (NewDefault.IsSet())
+			{
+				bModified = true;
+				Literal->ResetDefaults();
+				Literal->SetFromLiteral(*NewDefault);
+			}
+		}
 	}
 
-	if (!ensure(Metasound))
+	return bModified;
+}
+
+void UMetasoundEditorGraphVariable::UpdateFrontendDefaultLiteral(bool bPostTransaction, const FGuid* InPageID)
+{
+	using namespace Metasound;
+	using namespace Metasound::Editor;
+	using namespace Metasound::Frontend;
+	using namespace Metasound::VariableNames;
+
+	if (!ensure(Literal))
 	{
 		return;
 	}
 
-	if (!ensure(Literal))
+	FMetaSoundFrontendDocumentBuilder& Builder = GetFrontendBuilderChecked();
+	const FMetasoundFrontendVariable* Variable = Builder.FindGraphVariable(GetMemberName());
+	if (!ensure(Variable))
 	{
 		return;
 	}
 
 	const FScopedTransaction Transaction(LOCTEXT("Set Variable Default", "Set MetaSound Variable Default"), bPostTransaction);
-	Metasound->Modify();
+	Builder.CastDocumentObjectChecked<UObject>().Modify();
 
-	const FMetasoundFrontendLiteral DefaultLiteral = Literal->GetDefault();
-	FVariableHandle VariableHandle = GetVariableHandle();
-	VariableHandle->SetLiteral(DefaultLiteral);
+	FMetasoundFrontendLiteral DefaultLiteral;
+	Literal->TryFindDefault(DefaultLiteral, &Builder.GetBuildPageID());
 
-	FNodeHandle MutatorNode = VariableHandle->FindMutatorNode();
-	if (MutatorNode->IsValid())
+	// Swap once variable support is added to document builder API
+	Builder.SetGraphVariableDefault(GetMemberName(), DefaultLiteral);
+
+	if (const FMetasoundFrontendNode* MutatorNode = Builder.FindNode(Variable->MutatorNodeID))
 	{
-		FInputHandle InputHandle = MutatorNode->GetInputWithVertexName(METASOUND_GET_PARAM_NAME(InputData));
-		if (ensure(InputHandle->IsValid()))
+		const FMetasoundFrontendVertex* Input = MutatorNode->Interface.Inputs.FindByPredicate([](const FMetasoundFrontendVertex& Vertex) { return Vertex.Name == METASOUND_GET_PARAM_NAME(InputData); });
+		if (ensure(Input))
 		{
-			GraphPrivate::SetLiteralOrClearIfMatchesDefault(InputHandle, DefaultLiteral);
+			const FMetasoundFrontendVertexHandle VertexHandle(MutatorNode->GetID(), Input->VertexID);
+			GraphPrivate::SetOrClearIfLiteralMatchesNodeVertexDefault(Builder, VertexHandle, DefaultLiteral);
 		}
 	}
 }
diff --git a/Engine/Plugins/Runtime/Metasound/Source/MetasoundEditor/Private/MetasoundEditorGraphBuilder.cpp b/Engine/Plugins/Runtime/Metasound/Source/MetasoundEditor/Private/MetasoundEditorGraphBuilder.cpp
--- a/Engine/Plugins/Runtime/Metasound/Source/MetasoundEditor/Private/MetasoundEditorGraphBuilder.cpp
+++ b/Engine/Plugins/Runtime/Metasound/Source/MetasoundEditor/Private/MetasoundEditorGraphBuilder.cpp
@@ -39,6 +39,7 @@
 #include "MetasoundFrontendSearchEngine.h"
 #include "MetasoundFrontendTransform.h"
 #include "MetasoundLiteral.h"
+#include "MetasoundSettings.h"
 #include "MetasoundTime.h"
 #include "MetasoundUObjectRegistry.h"
 #include "MetasoundVariableNodes.h"
@@ -1380,7 +1381,7 @@ namespace Metasound
 			return true;
 		}
 
-		FMetasoundFrontendClassInput FGraphBuilder::CreateUniqueClassInput(UObject& InMetaSound, const FCreateNodeVertexParams& InParams, const FMetasoundFrontendLiteral* InDefaultValue, const FName* InNameBase)
+		FMetasoundFrontendClassInput FGraphBuilder::CreateUniqueClassInput(UObject& InMetaSound, const FCreateNodeVertexParams& InParams, const TArray<FMetasoundFrontendClassInputDefault>& InDefaultLiterals, const FName* InNameBase)
 		{
 			FMetasoundFrontendClassInput ClassInput;
 			ClassInput.Name = GenerateUniqueNameByClassType(InMetaSound, EMetasoundFrontendClassType::Input, InNameBase ? InNameBase->ToString() : TEXT("Input"));
@@ -1391,14 +1392,28 @@ namespace Metasound
 			// Can be unset if attempting to mirror parameters from a reroute, so default to reference
 			ClassInput.AccessType = InParams.AccessType == EMetasoundFrontendVertexAccessType::Unset ? EMetasoundFrontendVertexAccessType::Reference : InParams.AccessType;
 
-			if (InDefaultValue)
+			// Should always have at least one value
+			if (InDefaultLiterals.IsEmpty())
 			{
-				ClassInput.InitDefault(*InDefaultValue);
+				ClassInput.InitDefault();
 			}
 			else
 			{
-				Metasound::FLiteral Literal = Frontend::IDataTypeRegistry::Get().CreateDefaultLiteral(InParams.DataType);
-				ClassInput.InitDefault().SetFromLiteral(Literal);
+				if (const UMetaSoundSettings* Settings = GetDefault<UMetaSoundSettings>())
+				{
+					TSet<FGuid> ValidPageIDs;
+					Settings->IteratePageSettings([&ValidPageIDs](const FMetaSoundPageSettings& PageSetting)
+					{
+						ValidPageIDs.Add(PageSetting.UniqueId);
+					});
+					for (const FMetasoundFrontendClassInputDefault& InputDefault : InDefaultLiterals)
+					{
+						if (ValidPageIDs.Contains(InputDefault.PageID))
+						{
+							ClassInput.AddDefault(InputDefault.PageID) = InputDefault.Literal;
+						}
+					}
+				}
 			}
 
 			return ClassInput;
@@ -2719,6 +2734,11 @@ namespace Metasound
 				}
 			}
 
+			for (TObjectPtr<UMetasoundEditorGraphInput> Input : OutGraph.Inputs)
+			{
+				bEditorGraphModified |= Input->Synchronize();
+			}
+
 			// Collect all editor graph outputs with corresponding frontend outputs. 
 			GraphHandle->IterateConstNodes([&](FConstNodeHandle NodeHandle)
 			{
@@ -2751,87 +2771,10 @@ namespace Metasound
 				}
 			}
 
-			UMetasoundEditorGraphMember* Member = nullptr;
-
-			auto SynchronizeMemberDataType = [&](UMetasoundEditorGraphVertex& InVertex)
-			{
-				FConstNodeHandle NodeHandle = InVertex.GetConstNodeHandle();
-				TArray<FConstInputHandle> InputHandles = NodeHandle->GetConstInputs();
-				if (ensure(InputHandles.Num() == 1))
-				{
-					FConstInputHandle InputHandle = InputHandles.Last();
-					const FName NewDataType = InputHandle->GetDataType();
-					if (InVertex.GetDataType() != NewDataType)
-					{
-						constexpr bool bIncludeNamespace = true;
-						FText NodeDisplayName = FGraphBuilder::GetDisplayName(*NodeHandle, bIncludeNamespace);
-						UE_LOG(LogMetasoundEditor, Verbose, TEXT("Synchronizing Member '%s': Updating DataType to '%s'."), *NodeDisplayName.ToString(), *NewDataType.ToString());
-
-						FMetasoundFrontendLiteral DefaultLiteral;
-						DefaultLiteral.SetFromLiteral(IDataTypeRegistry::Get().CreateDefaultLiteral(NewDataType));
-						if (const FMetasoundFrontendLiteral* InputLiteral = InputHandle->GetLiteral())
-						{
-							DefaultLiteral = *InputLiteral;
-						}
-
-						InVertex.ClassName = NodeHandle->GetClassMetadata().GetClassName();
-
-						constexpr bool bPostTransaction = false;
-						InVertex.SetDataType(NewDataType, bPostTransaction);
-
-						if (DefaultLiteral.IsValid())
-						{
-							UMetasoundEditorGraphMemberDefaultLiteral* Literal = InVertex.GetLiteral();
-							Literal->SetFromLiteral(DefaultLiteral);
-							bEditorGraphModified = true;
-						}
-					}
-				}
-			};
-
-			// Synchronize data types & default values for input nodes.
-			GraphHandle->IterateConstNodes([&](FConstNodeHandle NodeHandle)
+			for (TObjectPtr<UMetasoundEditorGraphOutput> Output : OutGraph.Outputs)
 			{
-				if (UMetasoundEditorGraphInput* Input = OutGraph.FindInput(NodeHandle->GetID()))
-				{
-					SynchronizeMemberDataType(*Input);
-
-					if (UMetasoundEditorGraphMemberDefaultLiteral* Literal = Input->GetLiteral())
-					{
-						const FName NodeName = NodeHandle->GetNodeName();
-						const FGuid VertexID = GraphHandle->GetVertexIDForInputVertex(NodeName);
-						FMetasoundFrontendLiteral DefaultLiteral = GraphHandle->GetDefaultInput(VertexID);
-						if (!DefaultLiteral.IsEqual(Literal->GetDefault()))
-						{
-							if (DefaultLiteral.GetType() != EMetasoundFrontendLiteralType::None)
-							{
-								UE_LOG(LogMetasoundEditor, Verbose, TEXT("Synchronizing default value to '%s' for input '%s'"), *DefaultLiteral.ToString(), *NodeName.ToString());
-								Literal->SetFromLiteral(DefaultLiteral);
-								bEditorGraphModified = true;
-							}
-						}
-					}
-					else
-					{
-						Input->InitializeLiteral();
-					}
-				}
-			}, EMetasoundFrontendClassType::Input);
-
-			// Synchronize data types of output nodes.
-			GraphHandle->IterateConstNodes([&](FConstNodeHandle NodeHandle)
-			{
-				if (UMetasoundEditorGraphOutput* Output = OutGraph.FindOutput(NodeHandle->GetID()))
-				{
-					SynchronizeMemberDataType(*Output);
-
-					// Fix up members with no literal set
-					if (!Output->GetLiteral())
-					{
-						Output->InitializeLiteral();
-					}
-				}
-			}, EMetasoundFrontendClassType::Output);
+				bEditorGraphModified |= Output->Synchronize();
+			}
 
 			// Enforce all variables exists in the edgraph
 			TSet<FGuid> VariableIDs;
diff --git a/Engine/Plugins/Runtime/Metasound/Source/MetasoundEditor/Private/MetasoundEditorGraphBuilder.h b/Engine/Plugins/Runtime/Metasound/Source/MetasoundEditor/Private/MetasoundEditorGraphBuilder.h
--- a/Engine/Plugins/Runtime/Metasound/Source/MetasoundEditor/Private/MetasoundEditorGraphBuilder.h
+++ b/Engine/Plugins/Runtime/Metasound/Source/MetasoundEditor/Private/MetasoundEditorGraphBuilder.h
@@ -180,7 +180,7 @@ namespace Metasound
 			static FMetasoundFrontendClassInput CreateUniqueClassInput(
 				UObject& InMetaSound,
 				const FCreateNodeVertexParams& InVertexParams,
-				const FMetasoundFrontendLiteral* InDefaultValue = nullptr,
+				const TArray<FMetasoundFrontendClassInputDefault>& InDefaultLiterals = { },
 				const FName* InNameBase = nullptr);
 
 			// Creates a unique class output with the given default data. Output is not assigned a NodeID.
diff --git a/Engine/Plugins/Runtime/Metasound/Source/MetasoundEditor/Private/MetasoundEditorGraphMemberDefaults.cpp b/Engine/Plugins/Runtime/Metasound/Source/MetasoundEditor/Private/MetasoundEditorGraphMemberDefaults.cpp
--- a/Engine/Plugins/Runtime/Metasound/Source/MetasoundEditor/Private/MetasoundEditorGraphMemberDefaults.cpp
+++ b/Engine/Plugins/Runtime/Metasound/Source/MetasoundEditor/Private/MetasoundEditorGraphMemberDefaults.cpp
@@ -2,22 +2,26 @@
 #include "MetasoundEditorGraphMemberDefaults.h"
 
 #include "Algo/Count.h"
+#include "Algo/Sort.h"
 #include "Algo/Transform.h"
 #include "AudioDefines.h"
 #include "AudioParameterControllerInterface.h"
 #include "EdGraph/EdGraphNode.h"
 #include "GraphEditorSettings.h"
 #include "MetasoundDataReference.h"
+#include "MetasoundDocumentBuilderRegistry.h"
 #include "MetasoundEditorGraph.h"
 #include "MetasoundEditorGraphBuilder.h"
 #include "MetasoundEditorGraphNode.h"
 #include "MetasoundEditorSettings.h"
 #include "MetasoundFrontend.h"
 #include "MetasoundFrontendController.h"
+#include "MetasoundFrontendDataTypeRegistry.h"
 #include "MetasoundFrontendDocument.h"
 #include "MetasoundFrontendLiteral.h"
 #include "MetasoundFrontendRegistries.h"
 #include "MetasoundPrimitives.h"
+#include "MetasoundSettings.h"
 #include "Misc/Guid.h"
 #include "Sound/SoundWave.h"
 #include "UObject/ObjectMacros.h"
@@ -25,463 +29,1491 @@
 
 #include UE_INLINE_GENERATED_CPP_BY_NAME(MetasoundEditorGraphMemberDefaults)
 
-namespace Metasound
+namespace Metasound::Editor
 {
-	namespace Editor
+	namespace MemberDefaultsPrivate
 	{
-		namespace MemberDefaultsPrivate
+		TMap<FName, const FMetaSoundPageSettings*> GetPageSettingsByName()
 		{
-			template <typename TPODType>
-			void ConvertLiteral(const FMetasoundFrontendLiteral& InLiteral, TPODType& OutValue)
+			TMap<FName, const FMetaSoundPageSettings*> PageMap;
+			if (const UMetaSoundSettings* Settings = GetDefault<UMetaSoundSettings>())
 			{
-				InLiteral.TryGet(OutValue);
+				Settings->IteratePageSettings([&PageMap](const FMetaSoundPageSettings& PageSettings)
+				{
+					PageMap.Add(PageSettings.Name, &PageSettings);
+				});
+			}
+			return PageMap;
+		}
+
+		TMap<FGuid, const FMetaSoundPageSettings*> GetPageSettingsByID()
+		{
+			TMap<FGuid, const FMetaSoundPageSettings*> PageMap;
+			if (const UMetaSoundSettings* Settings = GetDefault<UMetaSoundSettings>())
+			{
+				Settings->IteratePageSettings([&PageMap](const FMetaSoundPageSettings& PageSettings)
+				{
+					PageMap.Add(PageSettings.UniqueId, &PageSettings);
+				});
+			}
+			return PageMap;
+		}
+
+		template <typename LiteralType>
+		void InitDefault(UMetasoundEditorGraphMemberDefaultLiteral& EditorLiteral, const FGuid& InPageID)
+		{
+			FMetasoundFrontendLiteral DefaultLiteral;
+			DefaultLiteral.SetFromLiteral(Frontend::IDataTypeRegistry::Get().CreateDefaultLiteral(GetMetasoundDataTypeName<LiteralType>()));
+			EditorLiteral.SetFromLiteral(DefaultLiteral, InPageID);
+		}
+
+		template <typename PageDefaultType, typename MemberType, typename LiteralType = MemberType>
+		void IterateDefaults(TFunctionRef<void(const FGuid&, FMetasoundFrontendLiteral)> Iter, TFunctionRef<LiteralType(const MemberType&)> MemberToLiteral, const TArray<PageDefaultType>& Defaults)
+		{
+			for (const PageDefaultType& PageDefault : Defaults)
+			{
+				FMetasoundFrontendLiteral Value;
+				Value.Set(MemberToLiteral(PageDefault.Value));
+				Iter(PageDefault.PageID, MoveTemp(Value));
+			}
+		}
+
+		template <typename PageDefaultType>
+		void IterateDefaults(TFunctionRef<void(const FGuid&, FMetasoundFrontendLiteral)> Iter, const TArray<PageDefaultType>& Defaults)
+		{
+			for (const PageDefaultType& PageDefault : Defaults)
+			{
+				FMetasoundFrontendLiteral Value;
+				Value.Set(PageDefault.Value);
+				Iter(PageDefault.PageID, MoveTemp(Value));
+			}
+		}
+
+		template <typename PageDefaultType>
+		bool RemoveDefault(const FGuid& InPageID, TArray<PageDefaultType>& OutDefaults)
+		{
+			return OutDefaults.RemoveAllSwap([&InPageID](const FMetasoundEditorMemberPageDefault& PageDefault)
+			{
+				return PageDefault.PageID == InPageID;
+			}) > 0;
+		}
+
+		template <typename PageDefaultType>
+		void ResolvePageDefaults(TArray<PageDefaultType>& OutPageDefaults)
+		{
+			TMap<FName, const FMetaSoundPageSettings*> PageSettingsByName = MemberDefaultsPrivate::GetPageSettingsByName();
+			if (OutPageDefaults.IsEmpty() || PageSettingsByName.IsEmpty() || PageSettingsByName.Num() == 1)
+			{
+				OutPageDefaults.SetNum(1);
+				OutPageDefaults.Last().PageName = Frontend::DefaultPageName;
+				OutPageDefaults.Last().PageID = Frontend::DefaultPageID;
+				return;
+			}
+
+			// Find duplicates to be removed (this can happen if page renamed to an existing page by user.
+			// This will result in a page with the same name but mismatched page ID to given name) and missing
+			// page names.
+			TArray<FName> MissingPages;
+			TArray<FName> DuplicatePages;
+			{
+				PageSettingsByName.GetKeys(MissingPages);
+				for (int32 Index = OutPageDefaults.Num() - 1; Index >= 0; --Index)
+				{
+					const FName PageName = OutPageDefaults[Index].PageName;
+					if (MissingPages.Remove(PageName) == 0)
+					{
+						DuplicatePages.AddUnique(PageName);
+					}
+				}
+			}
+
+			auto InitPageDefault = [&PageSettingsByName](FName NewPageName, FMetasoundEditorMemberPageDefault& OutDefault)
+			{
+				const FMetaSoundPageSettings* NewPageSettings = PageSettingsByName.FindChecked(NewPageName);
+				check(NewPageSettings);
+				OutDefault.PageID = NewPageSettings->UniqueId;
+				OutDefault.PageName = NewPageSettings->Name;
+			};
+
+			for (int32 Index = OutPageDefaults.Num() - 1; Index >= 0; --Index)
+			{
+				FMetasoundEditorMemberPageDefault& Default = OutPageDefaults[Index];
+
+				// If user selected assigned a new name, resolve the page ID to be that of the given name
+				// and remove other value if it existed
+				if (const FMetaSoundPageSettings* PageSettings = PageSettingsByName.FindRef(Default.PageName))
+				{
+					if (Default.PageID == PageSettings->UniqueId)
+					{
+						const bool bIsDuplicate = DuplicatePages.RemoveSwap(Default.PageName, EAllowShrinking::No) > 0;
+						if (bIsDuplicate)
+						{
+							OutPageDefaults.RemoveAtSwap(Index, EAllowShrinking::No);
+						}
+					}
+					else
+					{
+						Default.PageID = PageSettings->UniqueId;
+					}
+				}
+
+				// Otherwise if user added new default entry with the new entry ID, give
+				// it a valid page ID & name.  If no additional valid pages are implementable,
+				// remove it.
+				else if (Default.PageID == FMetasoundEditorMemberPageDefault::GetNewEntryID())
+				{
+					if (MissingPages.IsEmpty())
+					{
+						OutPageDefaults.RemoveAtSwap(Index, EAllowShrinking::No);
+					}
+					else
+					{
+						const FName NewPageName = MissingPages.Pop();
+						InitPageDefault(NewPageName, Default);
+					}
+				}
+			}
+
+			// Must always contains at least default value in editor
+			if (MissingPages.Contains(Frontend::DefaultPageName))
+			{
+				PageDefaultType NewDefault;
+				InitPageDefault(Frontend::DefaultPageName, NewDefault);
+				OutPageDefaults.Add(MoveTemp(NewDefault));
+			}
+
+			OutPageDefaults.Shrink();
+		}
+
+		template <typename PageDefaultType>
+		void SortPageDefaults(TArray<PageDefaultType>& OutPageDefaults)
+		{
+			TMap<FGuid, int32> PageIDToSortOrder;
+			const UMetaSoundSettings* Settings = ::GetDefault<UMetaSoundSettings>();
+			check(Settings);
+
+			constexpr bool bReverse = true;
+			int32 Index = Settings->GetProjectPageSettings().Num(); // no "-1" as default settings are not in project page settings, which are subsequently iterated over.
+			Settings->IteratePageSettings([&PageIDToSortOrder, &Index](const FMetaSoundPageSettings& PageSettings)
+			{
+				PageIDToSortOrder.Add(PageSettings.UniqueId, Index);
+				--Index;
+			}, bReverse);
+
+			Algo::SortBy(OutPageDefaults, [&PageIDToSortOrder, &OutPageDefaults](const PageDefaultType& PageDefault) -> int32
+			{
+				if (const int32* SortValue = PageIDToSortOrder.Find(PageDefault.PageID))
+				{
+					return *SortValue;
+				}
+				return FMath::Max(OutPageDefaults.Num() + 1, PageIDToSortOrder.Num() + 1);
+			});
+		}
+
+		template <typename PageDefaultType, typename MemberType>
+		void SetFromLiteral(const TFunctionRef<void(MemberType&)> SetMemberFunc, const FGuid& InPageID, TArray<PageDefaultType>& OutPageDefaults)
+		{
+			auto MatchesID = [&InPageID](const PageDefaultType& Default) { return Default.PageID == InPageID; };
+			if (PageDefaultType* DefaultEntry = OutPageDefaults.FindByPredicate(MatchesID))
+			{
+				SetMemberFunc(DefaultEntry->Value);
+			}
+			else
+			{
+				PageDefaultType PageDefault(InPageID);
+				const UMetaSoundSettings* Settings = ::GetDefault<UMetaSoundSettings>();
+				check(Settings);
+				if (const FMetaSoundPageSettings* PageSettings = Settings->FindPageSettings(InPageID))
+				{
+					PageDefault.PageName = PageSettings->Name;
+					SetMemberFunc(PageDefault.Value);
+					OutPageDefaults.Add(MoveTemp(PageDefault));
+					SortPageDefaults(OutPageDefaults);
+				}
+			}
+		}
+
+		template <typename PageDefaultType, typename MemberType>
+		void SetFromLiteral(const FMetasoundFrontendLiteral& InLiteral, const FGuid& InPageID, TArray<PageDefaultType>& OutPageDefaults)
+		{
+			auto SetMemberFunc = [&InLiteral](MemberType& MemberValue)
+			{
+				InLiteral.TryGet(MemberValue);
+			};
+			SetFromLiteral<PageDefaultType, MemberType>(SetMemberFunc, InPageID, OutPageDefaults);
+		}
+
+		template <typename PageDefaultType, typename MemberType, typename LiteralType>
+		bool SynchronizePageDefault(
+			UMetasoundEditorGraphMember* Member,
+			TFunctionRef<MemberType(const LiteralType&)> LiteralToMember,
+			TFunctionRef<LiteralType(const MemberType&)> MemberToLiteral,
+			TArray<PageDefaultType>& OutPageDefaults)
+		{
+			using namespace Frontend;
+
+			bool bModified = false;
+			if (ensure(Member))
+			{
+				FMetaSoundFrontendDocumentBuilder& Builder = Member->GetFrontendBuilderChecked();
+				if (const FMetasoundFrontendClassInput* ClassInput = Builder.FindGraphInput(Member->GetMemberName()))
+				{
+					TMap<FGuid, const FMetasoundFrontendLiteral*> FrontendDefaults;
+					Algo::Transform(ClassInput->GetDefaults(), FrontendDefaults, [](const FMetasoundFrontendClassInputDefault& Default)
+					{
+						return TPair<FGuid, const FMetasoundFrontendLiteral*>(Default.PageID, &Default.Literal);
+					});
+
+					TMap<FGuid, const FMetaSoundPageSettings*> PageSettingsByID = MemberDefaultsPrivate::GetPageSettingsByID();
+					FMetasoundFrontendLiteral Value;
+					for (int32 Index = OutPageDefaults.Num() - 1; Index >= 0; --Index)
+					{
+						PageDefaultType& Default = OutPageDefaults[Index];
+						if (const FMetaSoundPageSettings* PageSettings = PageSettingsByID.FindRef(Default.PageID))
+						{
+							if (const FMetasoundFrontendLiteral* FrontendLiteral = FrontendDefaults.FindRef(Default.PageID))
+							{
+								FMetasoundFrontendLiteral TestLiteral;
+								if (FrontendLiteral->GetType() == EMetasoundFrontendLiteralType::None)
+								{
+									TestLiteral.SetFromLiteral(IDataTypeRegistry::Get().CreateDefaultLiteral(ClassInput->TypeName));
+								}
+								else
+								{
+									TestLiteral = *FrontendLiteral;
+								}
+								Value.Set(MemberToLiteral(Default.Value));
+								if (!TestLiteral.IsEqual(Value))
+								{
+									bModified = true;
+
+									LiteralType NewValue;
+									ensure(TestLiteral.TryGet(NewValue));
+									Default.Value = LiteralToMember(NewValue);
+								}
+
+								if (Default.PageName != PageSettings->Name)
+								{
+									bModified = true;
+									Default.PageName = PageSettings->Name;
+								}
+
+								FrontendDefaults.Remove(Default.PageID);
+							}
+							else
+							{
+								OutPageDefaults.RemoveAtSwap(Index, EAllowShrinking::No);
+							}
+						}
+						else
+						{
+							OutPageDefaults.RemoveAtSwap(Index, EAllowShrinking::No);
+						}
+					}
+
+					for (const TPair<FGuid, const FMetasoundFrontendLiteral*>& Pair : FrontendDefaults)
+					{
+						if (const FMetaSoundPageSettings* PageSettings = PageSettingsByID.FindRef(Pair.Key))
+						{
+							bModified = true;
+							PageDefaultType NewPageDefault;
+							NewPageDefault.PageID = Pair.Key;
+							NewPageDefault.PageName = PageSettings->Name;
+
+							LiteralType NewValue;
+							check(Pair.Value);
+							ensure(Pair.Value->TryGet(NewValue));
+
+							NewPageDefault.Value = LiteralToMember(NewValue);
+							OutPageDefaults.Add(MoveTemp(NewPageDefault));
+						}
+					}
+
+					OutPageDefaults.Shrink();
+				}
+			}
+
+			if (bModified)
+			{
+				MemberDefaultsPrivate::SortPageDefaults<PageDefaultType>(OutPageDefaults);
+			}
+
+			return bModified;
+		}
+
+		template <typename PageDefaultType, typename MemberType, typename LiteralType = MemberType>
+		bool SynchronizePageDefault(UMetasoundEditorGraphMember* Member, TArray<PageDefaultType>& OutPageDefaults)
+		{
+			auto LiteralToMember = [](const LiteralType& Literal) -> MemberType { return Literal; };
+			auto MemberToLiteral = [](const MemberType& Literal) -> LiteralType { return Literal; };
+			return SynchronizePageDefault<PageDefaultType, MemberType, LiteralType>(Member, LiteralToMember, MemberToLiteral, OutPageDefaults);
+		}
+
+		template <typename PageDefaultType, typename MemberType, typename LiteralType>
+		bool TryFindDefault(const TArray<PageDefaultType>& PageDefaults, const FGuid* InPageID, TFunctionRef<LiteralType(const MemberType&)> MemberToLiteral, FMetasoundFrontendLiteral& OutLiteral)
+		{
+			const FGuid PageID = InPageID ? *InPageID : Metasound::Frontend::DefaultPageID;
+			auto IsPageDefault = [&PageID](const PageDefaultType& PageDefault) { return PageDefault.PageID == PageID; };
+			if (const PageDefaultType* PageDefaultPtr = PageDefaults.FindByPredicate(IsPageDefault))
+			{
+				OutLiteral.Set(MemberToLiteral(PageDefaultPtr->Value));
+				return true;
 			}
 
-			template <typename TPODType, typename TLiteralType = TPODType>
-			void ConvertLiteralToArray(const FMetasoundFrontendLiteral& InLiteral, TArray<TLiteralType>& OutArray)
-			{
-				OutArray.Reset();
-				TArray<TPODType> Values;
-				InLiteral.TryGet(Values);
-				Algo::Transform(Values, OutArray, [](const TPODType& InValue) { return TLiteralType{ InValue }; });
-			}
-		}
-	}
+			OutLiteral = { };
+			return false;
+		}
+
+		template <typename PageDefaultType, typename MemberType, typename LiteralType = MemberType>
+		bool TryFindDefault(const TArray<PageDefaultType>& PageDefaults, const FGuid* InPageID, FMetasoundFrontendLiteral& OutLiteral)
+		{
+			auto MemberToLiteral = [](const MemberType& Member) -> LiteralType { return Member; };
+			return TryFindDefault<PageDefaultType, MemberType, LiteralType>(PageDefaults, InPageID, MemberToLiteral, OutLiteral);
+		}
+	} // namespace MemberDefaultsPrivate
+} // Metasound::Editor
+
+EMetasoundFrontendLiteralType UMetasoundEditorGraphMemberDefaultBool::GetLiteralType() const
+{
+	return EMetasoundFrontendLiteralType::Boolean;
+}
+
+void UMetasoundEditorGraphMemberDefaultBool::InitDefault(const FGuid& InPageID)
+{
+	using namespace Metasound::Editor;
+	MemberDefaultsPrivate::InitDefault<bool>(*this, InPageID);
+}
+
+void UMetasoundEditorGraphMemberDefaultBool::IterateDefaults(TFunctionRef<void(const FGuid&, FMetasoundFrontendLiteral)> Iter) const
+{
+	using namespace Metasound::Editor;
+	auto MemberToLiteral = [](const FMetasoundEditorGraphMemberDefaultBoolRef& Ref) { return Ref.Value; };
+	MemberDefaultsPrivate::IterateDefaults<FMetasoundEditorMemberPageDefaultBool, FMetasoundEditorGraphMemberDefaultBoolRef, bool>(Iter, MemberToLiteral, Defaults);
+}
+
+bool UMetasoundEditorGraphMemberDefaultBool::RemoveDefault(const FGuid& InPageID)
+{
+	using namespace Metasound::Editor;
+	return MemberDefaultsPrivate::RemoveDefault(InPageID, Defaults);
+}
+
+void UMetasoundEditorGraphMemberDefaultBool::ResetDefaults()
+{
+	Defaults.Empty();
+	InitDefault(Metasound::Frontend::DefaultPageID);
+}
+
+void UMetasoundEditorGraphMemberDefaultBool::ResolvePageDefaults()
+{
+	using namespace Metasound::Editor;
+	MemberDefaultsPrivate::ResolvePageDefaults(Defaults);
+}
+
+void UMetasoundEditorGraphMemberDefaultBool::SetFromLiteral(const FMetasoundFrontendLiteral& InLiteral, const FGuid& InPageID)
+{
+	using namespace Metasound::Editor;
+	auto SetLiteral = [&InLiteral](FMetasoundEditorGraphMemberDefaultBoolRef& OutRef)
+	{
+		bool bValue = false;
+		InLiteral.TryGet(bValue);
+		OutRef.Value = bValue;
+	};
+	MemberDefaultsPrivate::SetFromLiteral<FMetasoundEditorMemberPageDefaultBool, FMetasoundEditorGraphMemberDefaultBoolRef>(SetLiteral, InPageID, Defaults);
+}
+
+void UMetasoundEditorGraphMemberDefaultBool::SortPageDefaults()
+{
+	using namespace Metasound::Editor;
+	MemberDefaultsPrivate::SortPageDefaults(Defaults);
+}
+
+bool UMetasoundEditorGraphMemberDefaultBool::Synchronize()
+{
+	using namespace Metasound::Editor;
+	auto LiteralToMember = [](bool bValue)
+	{
+		FMetasoundEditorGraphMemberDefaultBoolRef Ref;
+		Ref.Value = bValue;
+		return Ref;
+	};
+
+	auto MemberToLiteral = [](const FMetasoundEditorGraphMemberDefaultBoolRef& Ref)
+	{
+		return Ref.Value;
+	};
+
+	return MemberDefaultsPrivate::SynchronizePageDefault<FMetasoundEditorMemberPageDefaultBool, FMetasoundEditorGraphMemberDefaultBoolRef, bool>(
+		FindMember(),
+		LiteralToMember,
+		MemberToLiteral,
+		Defaults);
+}
+
+bool UMetasoundEditorGraphMemberDefaultBool::TryFindDefault(FMetasoundFrontendLiteral& OutLiteral, const FGuid* InPageID) const
+{
+	using namespace Metasound::Editor;
+	auto MemberToLiteral = [](const FMetasoundEditorGraphMemberDefaultBoolRef& Ref) { return Ref.Value; };
+	return MemberDefaultsPrivate::TryFindDefault<FMetasoundEditorMemberPageDefaultBool, FMetasoundEditorGraphMemberDefaultBoolRef, bool>(
+		Defaults,
+		InPageID,
+		MemberToLiteral,
+		OutLiteral);
+}
+
+void UMetasoundEditorGraphMemberDefaultBool::UpdatePreviewInstance(const Metasound::FVertexName& InParameterName, TScriptInterface<IAudioParameterControllerInterface>& InParameterInterface) const
+{
+	FGuid PageID;
+	if (TryGetPreviewPageID(PageID))
+	{
+		InParameterInterface->SetBoolParameter(InParameterName, GetDefaultAs<bool>(PageID));
+	}
+}
+
+EMetasoundFrontendLiteralType UMetasoundEditorGraphMemberDefaultBoolArray::GetLiteralType() const
+{
+	return EMetasoundFrontendLiteralType::BooleanArray;
+}
+
+void UMetasoundEditorGraphMemberDefaultBoolArray::InitDefault(const FGuid& InPageID)
+{
+	using namespace Metasound::Editor;
+	MemberDefaultsPrivate::InitDefault<TArray<bool>>(*this, InPageID);
+}
+
+void UMetasoundEditorGraphMemberDefaultBoolArray::IterateDefaults(TFunctionRef<void(const FGuid&, FMetasoundFrontendLiteral)> Iter) const
+{
+	using namespace Metasound::Editor;
+	auto MemberToLiteral = [](const TArray<FMetasoundEditorGraphMemberDefaultBoolRef>& MemberValues)
+	{
+		TArray<bool> LiteralValues;
+		Algo::Transform(MemberValues, LiteralValues, [](const FMetasoundEditorGraphMemberDefaultBoolRef& Ref) { return Ref.Value; });
+		return LiteralValues;
+	};
+	MemberDefaultsPrivate::IterateDefaults<FMetasoundEditorMemberPageDefaultBoolArray, TArray<FMetasoundEditorGraphMemberDefaultBoolRef>, TArray<bool>>(Iter, MemberToLiteral, Defaults);
+}
+
+bool UMetasoundEditorGraphMemberDefaultBoolArray::RemoveDefault(const FGuid& InPageID)
+{
+	using namespace Metasound::Editor;
+	return MemberDefaultsPrivate::RemoveDefault(InPageID, Defaults);
+}
+
+void UMetasoundEditorGraphMemberDefaultBoolArray::ResetDefaults()
+{
+	Defaults.Empty();
+	InitDefault(Metasound::Frontend::DefaultPageID);
+}
+
+void UMetasoundEditorGraphMemberDefaultBoolArray::ResolvePageDefaults()
+{
+	using namespace Metasound::Editor;
+	MemberDefaultsPrivate::ResolvePageDefaults(Defaults);
+}
+
+void UMetasoundEditorGraphMemberDefaultBoolArray::SetFromLiteral(const FMetasoundFrontendLiteral& InLiteral, const FGuid& InPageID)
+{
+	using namespace Metasound::Editor;
+
+	auto SetLiteral = [&InLiteral](TArray<FMetasoundEditorGraphMemberDefaultBoolRef>& OutRefs)
+	{
+		OutRefs.Reset();
+		TArray<bool> Values;
+		InLiteral.TryGet(Values);
+		Algo::Transform(Values, OutRefs, [](bool Value)
+		{
+			FMetasoundEditorGraphMemberDefaultBoolRef Ref;
+			Ref.Value = Value;
+			return Ref;
+		});
+	};
+	MemberDefaultsPrivate::SetFromLiteral<FMetasoundEditorMemberPageDefaultBoolArray, TArray<FMetasoundEditorGraphMemberDefaultBoolRef>>(SetLiteral, InPageID, Defaults);
+}
+
+void UMetasoundEditorGraphMemberDefaultBoolArray::SortPageDefaults()
+{
+	using namespace Metasound::Editor;
+	MemberDefaultsPrivate::SortPageDefaults<FMetasoundEditorMemberPageDefaultBoolArray>(Defaults);
+}
+
+bool UMetasoundEditorGraphMemberDefaultBoolArray::Synchronize()
+{
+	using namespace Metasound::Editor;
+
+	auto LiteralToMember = [](const TArray<bool>& Values)
+	{
+		TArray<FMetasoundEditorGraphMemberDefaultBoolRef> Defaults;
+		Algo::Transform(Values, Defaults, [](bool Value) { FMetasoundEditorGraphMemberDefaultBoolRef Ref; Ref.Value = Value; return Ref; });
+		return Defaults;
+	};
+
+	auto MemberToLiteral = [](const TArray<FMetasoundEditorGraphMemberDefaultBoolRef>& Values)
+	{
+		TArray<bool> Defaults;
+		Algo::Transform(Values, Defaults, [](const FMetasoundEditorGraphMemberDefaultBoolRef& Ref) { return Ref.Value; });
+		return Defaults;
+	};
+
+	return MemberDefaultsPrivate::SynchronizePageDefault<FMetasoundEditorMemberPageDefaultBoolArray, TArray<FMetasoundEditorGraphMemberDefaultBoolRef>, TArray<bool>>(
+		FindMember(),
+		LiteralToMember,
+		MemberToLiteral,
+		Defaults);
+}
+
+bool UMetasoundEditorGraphMemberDefaultBoolArray::TryFindDefault(FMetasoundFrontendLiteral& OutLiteral, const FGuid* InPageID) const
+{
+	using namespace Metasound::Editor;
+	auto MemberToLiteral = [](const TArray<FMetasoundEditorGraphMemberDefaultBoolRef>& Refs)
+	{
+		TArray<bool> Values;
+		Algo::Transform(Refs, Values, [](const FMetasoundEditorGraphMemberDefaultBoolRef& Ref) { return Ref.Value; });
+		return Values;
+	};
+	return MemberDefaultsPrivate::TryFindDefault<FMetasoundEditorMemberPageDefaultBoolArray, TArray<FMetasoundEditorGraphMemberDefaultBoolRef>, TArray<bool>>(Defaults, InPageID, MemberToLiteral, OutLiteral);
+}
+
+void UMetasoundEditorGraphMemberDefaultBoolArray::UpdatePreviewInstance(const Metasound::FVertexName& InParameterName, TScriptInterface<IAudioParameterControllerInterface>& InParameterInterface) const
+{
+	FGuid PageID;
+	if (TryGetPreviewPageID(PageID))
+	{
+		InParameterInterface->SetBoolArrayParameter(InParameterName, GetDefaultAs<TArray<bool>>(PageID));
+	}
+}
+
+void UMetasoundEditorGraphMemberDefaultFloat::ClampDefaults()
+{
+	using namespace Metasound::Editor;
+
+	for (FMetasoundEditorMemberPageDefaultFloat& PageDefault : Defaults)
+	{
+		const float ClampedValue = FMath::Clamp(PageDefault.Value, Range.X, Range.Y);
+		if (!FMath::IsNearlyEqual(PageDefault.Value, ClampedValue))
+		{
+			PageDefault.Value = ClampedValue;
+			OnDefaultValueChanged.Broadcast(PageDefault.PageID, PageDefault.Value);
+		}
+	}
+}
+
+float UMetasoundEditorGraphMemberDefaultFloat::GetDefault()
+{
+	return GetDefaultAs<float>();
+}
+
+EMetasoundFrontendLiteralType UMetasoundEditorGraphMemberDefaultFloat::GetLiteralType() const
+{
+	return EMetasoundFrontendLiteralType::Float;
+}
+
+void UMetasoundEditorGraphMemberDefaultFloat::InitDefault(const FGuid& InPageID)
+{
+	using namespace Metasound::Editor;
+	MemberDefaultsPrivate::InitDefault<float>(*this, InPageID);
+}
+
+void UMetasoundEditorGraphMemberDefaultFloat::Initialize()
+{
+	if (const UMetasoundEditorSettings* EditorSettings = ::GetDefault<UMetasoundEditorSettings>())
+	{
+		WidgetType = EditorSettings->DefaultInputWidgetType;
+	}
+}
+
+void UMetasoundEditorGraphMemberDefaultFloat::IterateDefaults(TFunctionRef<void(const FGuid&, FMetasoundFrontendLiteral)> Iter) const
+{
+	using namespace Metasound::Editor;
+	MemberDefaultsPrivate::IterateDefaults<FMetasoundEditorMemberPageDefaultFloat>(Iter, Defaults);
+}
+
+void UMetasoundEditorGraphMemberDefaultFloat::ForceRefresh()
+{
+	using namespace Metasound::Editor;
+
+	OnRangeChanged.Broadcast(Range);
+	for (FMetasoundEditorMemberPageDefaultFloat& PageDefault : Defaults)
+	{
+		const float ClampedFloat = FMath::Clamp(PageDefault.Value, Range.X, Range.Y);
+		PageDefault.Value = ClampedFloat;
+
+	}
+
+	// If set from literal, we force the default value to be the literal's value which may require the range to be fixed up 
+	SetInitialRange();
+}
+
+void UMetasoundEditorGraphMemberDefaultFloat::UpdatePreviewInstance(const Metasound::FVertexName& InParameterName, TScriptInterface<IAudioParameterControllerInterface>& InParameterInterface) const
+{
+	FGuid PageID;
+	if (TryGetPreviewPageID(PageID))
+	{
+		InParameterInterface->SetFloatParameter(InParameterName, GetDefaultAs<float>(PageID));
+	}
+}
+
+void UMetasoundEditorGraphMemberDefaultFloat::PostEditChangeChainProperty(FPropertyChangedChainEvent& PropertyChangedEvent)
+{
+	using namespace Metasound::Editor;
+
+	const FName PropertyName = PropertyChangedEvent.GetPropertyName();
+	if (PropertyName.IsEqual(GET_MEMBER_NAME_CHECKED(UMetasoundEditorGraphMemberDefaultFloat, WidgetType)) ||
+		PropertyName.IsEqual(GET_MEMBER_NAME_CHECKED(UMetasoundEditorGraphMemberDefaultFloat, WidgetUnitValueType)))
+	{
+		// Update VolumeWidgetDecibelRange based on current range (it might be stale)
+		if (WidgetUnitValueType == EAudioUnitsValueType::Volume && VolumeWidgetUseLinearOutput)
+		{
+			VolumeWidgetDecibelRange = FVector2D(Audio::ConvertToDecibels(Range.X), Audio::ConvertToDecibels(Range.Y));
+		}
+		else if (WidgetUnitValueType == EAudioUnitsValueType::Frequency)
+		{
+			SetRange(FVector2D(MIN_FILTER_FREQUENCY, MAX_FILTER_FREQUENCY));
+		}
+		else
+		{
+			SetInitialRange();
+		}
+
+		// If the widget type is changed to none, we need to refresh clamping the value or not, since if the widget was a slider before, the value was clamped
+		OnClampChanged.Broadcast(ClampDefault);
+	}
+	else if (PropertyName.IsEqual(GET_MEMBER_NAME_CHECKED(UMetasoundEditorGraphMemberDefaultFloat, ClampDefault)))
+	{
+		SetInitialRange();
+		OnClampChanged.Broadcast(ClampDefault);
+	}
+	else if (PropertyName.IsEqual(GET_MEMBER_NAME_CHECKED(UMetasoundEditorGraphMemberDefaultFloat, VolumeWidgetUseLinearOutput)))
+	{
+		for (FMetasoundEditorMemberPageDefaultFloat& PageDefault : Defaults)
+		{
+			if (VolumeWidgetUseLinearOutput)
+			{
+				// Range and default are currently in dB, need to change to linear
+				float DbDefault = PageDefault.Value;
+				VolumeWidgetDecibelRange = Range;
+				Range = FVector2D(Audio::ConvertToLinear(VolumeWidgetDecibelRange.X), Audio::ConvertToLinear(VolumeWidgetDecibelRange.Y));
+				PageDefault.Value = Audio::ConvertToLinear(DbDefault);
+			}
+			else
+			{
+				// Range and default are currently linear, need to change to dB
+				Range = VolumeWidgetDecibelRange;
+				PageDefault.Value = Audio::ConvertToDecibels(PageDefault.Value);
+			}
+		}
+	}
+	else
+	{
+		FName ChildPropertyName;
+		if (const FEditPropertyChain::TDoubleLinkedListNode* MemberNode = PropertyChangedEvent.PropertyChain.GetActiveMemberNode())
+		{
+			if (const FProperty* ChildProperty = MemberNode->GetValue())
+			{
+				ChildPropertyName = ChildProperty->GetFName();
+			}
+		}
+
+		if (ChildPropertyName.IsEqual(GET_MEMBER_NAME_CHECKED(UMetasoundEditorGraphMemberDefaultFloat, Range)))
+		{
+			if (PropertyChangedEvent.ChangeType != EPropertyChangeType::Interactive)
+			{
+				if (WidgetType != EMetasoundMemberDefaultWidget::None &&
+					WidgetUnitValueType == EAudioUnitsValueType::Volume)
+				{
+					if (VolumeWidgetUseLinearOutput)
+					{
+						Range.X = FMath::Max(Range.X, Audio::ConvertToLinear(SAudioVolumeRadialSlider::MinDbValue));
+						Range.Y = FMath::Min(Range.Y, Audio::ConvertToLinear(SAudioVolumeRadialSlider::MaxDbValue));
+						if (Range.X > Range.Y)
+						{
+							Range.Y = Range.X;
+						}
+						VolumeWidgetDecibelRange = FVector2D(Audio::ConvertToDecibels(Range.X), Audio::ConvertToDecibels(Range.Y));
+						OnRangeChanged.Broadcast(VolumeWidgetDecibelRange);
+						ClampDefaults();
+					}
+				}
+				else
+				{
+					// if Range.X > Range.Y, set Range.Y to Range.X
+					if (Range.X > Range.Y)
+					{
+						SetRange(FVector2D(Range.X, FMath::Max(Range.X, Range.Y)));
+					}
+					else
+					{
+						ForceRefresh();
+					}
+				}
+			}
+		}
+		else if (ChildPropertyName.IsEqual(GET_MEMBER_NAME_CHECKED(UMetasoundEditorGraphMemberDefaultFloat, VolumeWidgetDecibelRange)))
+		{
+			if (PropertyChangedEvent.ChangeType != EPropertyChangeType::Interactive)
+			{
+				VolumeWidgetDecibelRange.X = FMath::Max(VolumeWidgetDecibelRange.X, SAudioVolumeRadialSlider::MinDbValue);
+				VolumeWidgetDecibelRange.Y = FMath::Min(VolumeWidgetDecibelRange.Y, SAudioVolumeRadialSlider::MaxDbValue);
+				SetRange(FVector2D(Audio::ConvertToLinear(VolumeWidgetDecibelRange.X), Audio::ConvertToLinear(VolumeWidgetDecibelRange.Y)));
+			}
+		}
+	}
+
+	// Only update member on non-interactive changes to avoid refreshing the details panel mid-update
+	if (PropertyChangedEvent.ChangeType != EPropertyChangeType::Interactive)
+	{
+		UMetasoundEditorGraphMember* Member = FindMember();
+		if (ensure(Member))
+		{
+			FMetasoundFrontendDocumentModifyContext& ModifyContext = FGraphBuilder::GetOutermostMetaSoundChecked(*Member).GetModifyContext();
+
+			// Mark all nodes as modified to refresh them on synchronization.  This ensures all corresponding widgets get updated.
+			const TArray<UMetasoundEditorGraphMemberNode*> MemberNodes = Member->GetNodes();
+			TSet<FGuid> NodesToRefresh;
+			Algo::Transform(MemberNodes, NodesToRefresh, [](const UMetasoundEditorGraphMemberNode* MemberNode) { return MemberNode->GetNodeID(); });
+
+			// MemberID not marked as modified as this causes detail trees to collapse.
+			ModifyContext.AddNodeIDsModified({ NodesToRefresh });
+		}
+	}
+
+	Super::PostEditChangeChainProperty(PropertyChangedEvent);
+}
+
+bool UMetasoundEditorGraphMemberDefaultFloat::RemoveDefault(const FGuid& InPageID)
+{
+	using namespace Metasound::Editor;
+	return MemberDefaultsPrivate::RemoveDefault(InPageID, Defaults);
+}
+
+void UMetasoundEditorGraphMemberDefaultFloat::ResetDefaults()
+{
+	Defaults.Empty();
+	InitDefault(Metasound::Frontend::DefaultPageID);
+}
+
+void UMetasoundEditorGraphMemberDefaultFloat::ResolvePageDefaults()
+{
+	using namespace Metasound::Editor;
+	MemberDefaultsPrivate::ResolvePageDefaults(Defaults);
+}
+
+void UMetasoundEditorGraphMemberDefaultFloat::SetFromLiteral(const FMetasoundFrontendLiteral& InLiteral, const FGuid& InPageID)
+{
+	using namespace Metasound::Editor;
+	MemberDefaultsPrivate::SetFromLiteral<FMetasoundEditorMemberPageDefaultFloat, float>(InLiteral, InPageID, Defaults);
+
+	// If set from literal, we force the default value to be the literal's value which may require the range to be fixed up
+	SetInitialRange();
+}
+
+void UMetasoundEditorGraphMemberDefaultFloat::SortPageDefaults()
+{
+	using namespace Metasound::Editor;
+	MemberDefaultsPrivate::SortPageDefaults<FMetasoundEditorMemberPageDefaultFloat>(Defaults);
+}
+
+bool UMetasoundEditorGraphMemberDefaultFloat::Synchronize()
+{
+	using namespace Metasound::Editor;
+	return MemberDefaultsPrivate::SynchronizePageDefault<FMetasoundEditorMemberPageDefaultFloat, float>(FindMember(), Defaults);
+}
+
+void UMetasoundEditorGraphMemberDefaultFloat::SetInitialRange()
+{
+	// If value is within current range, keep it, otherwise set range to something reasonable
+	bool bDefaultsInRange = true;
+	float Min = TNumericLimits<float>::Min();
+	float Max = TNumericLimits<float>::Max();
+	for (const FMetasoundEditorMemberPageDefaultFloat& PageDefault : Defaults)
+	{
+		bDefaultsInRange &= PageDefault.Value >= Range.X && PageDefault.Value <= Range.Y;
+		Min = FMath::Max(PageDefault.Value, Min);
+		Max = FMath::Min(PageDefault.Value, Max);
+	}
+
+	if (!bDefaultsInRange)
+	{
+		if (Min > Max)
+		{
+			Swap(Min, Max);
+		}
+
+		if (FMath::IsNearlyEqual(Min, 0.0f) && FMath::IsNearlyEqual(Max, 0.0f))
+		{
+			SetRange(FVector2D(0.0f, 1.0f));
+		}
+		else if (FMath::IsNearlyEqual(Min, Max))
+		{
+			if (Min > 0.0f)
+			{
+				SetRange(FVector2D(0.0f, Min));
+			}
+			else
+			{
+				SetRange(FVector2D(Min, 0.0f));
+			}
+		}
+		else
+		{
+			SetRange(FVector2D(FMath::Min(0.0f, Min), FMath::Max(0.0f, Max)));
+		}
+	}
+}
+
+bool UMetasoundEditorGraphMemberDefaultFloat::TryFindDefault(FMetasoundFrontendLiteral& OutLiteral, const FGuid* InPageID) const
+{
+	using namespace Metasound::Editor;
+	return MemberDefaultsPrivate::TryFindDefault<FMetasoundEditorMemberPageDefaultFloat, float>(Defaults, InPageID, OutLiteral);
+}
+
+FVector2D UMetasoundEditorGraphMemberDefaultFloat::GetRange() const
+{
+	return Range;
+}
+
+void UMetasoundEditorGraphMemberDefaultFloat::SetRange(const FVector2D InRange)
+{
+	using namespace Metasound::Editor;
+
+	if (!(Range - InRange).IsNearlyZero())
+	{
+		Range = InRange;
+		OnRangeChanged.Broadcast(InRange);
+		ClampDefaults();
+	}
+}
+
+EMetasoundFrontendLiteralType UMetasoundEditorGraphMemberDefaultFloatArray::GetLiteralType() const
+{
+	return EMetasoundFrontendLiteralType::FloatArray;
+}
+
+void UMetasoundEditorGraphMemberDefaultFloatArray::InitDefault(const FGuid& InPageID)
+{
+	using namespace Metasound::Editor;
+	MemberDefaultsPrivate::InitDefault<TArray<float>>(*this, InPageID);
+
+}
+
+void UMetasoundEditorGraphMemberDefaultFloatArray::IterateDefaults(TFunctionRef<void(const FGuid&, FMetasoundFrontendLiteral)> Iter) const
+{
+	using namespace Metasound::Editor;
+	MemberDefaultsPrivate::IterateDefaults(Iter, Defaults);
+}
+
+bool UMetasoundEditorGraphMemberDefaultFloatArray::RemoveDefault(const FGuid& InPageID)
+{
+	using namespace Metasound::Editor;
+	return MemberDefaultsPrivate::RemoveDefault(InPageID, Defaults);
+}
+
+void UMetasoundEditorGraphMemberDefaultFloatArray::ResetDefaults()
+{
+	Defaults.Empty();
+	InitDefault(Metasound::Frontend::DefaultPageID);
+}
+
+void UMetasoundEditorGraphMemberDefaultFloatArray::ResolvePageDefaults()
+{
+	using namespace Metasound::Editor;
+	MemberDefaultsPrivate::ResolvePageDefaults(Defaults);
+}
+
+void UMetasoundEditorGraphMemberDefaultFloatArray::SetFromLiteral(const FMetasoundFrontendLiteral& InLiteral, const FGuid& InPageID)
+{
+	using namespace Metasound::Editor;
+	MemberDefaultsPrivate::SetFromLiteral<FMetasoundEditorMemberPageDefaultFloatArray, TArray<float>>(InLiteral, InPageID, Defaults);
+}
+
+void UMetasoundEditorGraphMemberDefaultFloatArray::SortPageDefaults()
+{
+	using namespace Metasound::Editor;
+	MemberDefaultsPrivate::SortPageDefaults<FMetasoundEditorMemberPageDefaultFloatArray>(Defaults);
+}
+
+bool UMetasoundEditorGraphMemberDefaultFloatArray::Synchronize()
+{
+	using namespace Metasound::Editor;
+	return MemberDefaultsPrivate::SynchronizePageDefault<FMetasoundEditorMemberPageDefaultFloatArray, TArray<float>>(FindMember(), Defaults);
+}
+
+bool UMetasoundEditorGraphMemberDefaultFloatArray::TryFindDefault(FMetasoundFrontendLiteral& OutLiteral, const FGuid* InPageID) const
+{
+	using namespace Metasound::Editor;
+	return MemberDefaultsPrivate::TryFindDefault<FMetasoundEditorMemberPageDefaultFloatArray, TArray<float>>(Defaults, InPageID, OutLiteral);
+}
+
+void UMetasoundEditorGraphMemberDefaultFloatArray::UpdatePreviewInstance(const Metasound::FVertexName& InParameterName, TScriptInterface<IAudioParameterControllerInterface>& InParameterInterface) const
+{
+	FGuid PageID;
+	if (TryGetPreviewPageID(PageID))
+	{
+		InParameterInterface->SetFloatArrayParameter(InParameterName, GetDefaultAs<TArray<float>>(PageID));
+	}
+}
+
+EMetasoundFrontendLiteralType UMetasoundEditorGraphMemberDefaultInt::GetLiteralType() const
+{
+	return EMetasoundFrontendLiteralType::Integer;
+}
+
+void UMetasoundEditorGraphMemberDefaultInt::InitDefault(const FGuid& InPageID)
+{
+	using namespace Metasound::Editor;
+	MemberDefaultsPrivate::InitDefault<int32>(*this, InPageID);
+
 }
 
-FMetasoundFrontendLiteral UMetasoundEditorGraphMemberDefaultBool::GetDefault() const 
+void UMetasoundEditorGraphMemberDefaultInt::IterateDefaults(TFunctionRef<void(const FGuid&, FMetasoundFrontendLiteral)> Iter) const
 {
-	FMetasoundFrontendLiteral Literal;
-	Literal.Set(Default.Value);
-	return Literal;
+	using namespace Metasound::Editor;
+	auto MemberToLiteral = [](const FMetasoundEditorGraphMemberDefaultIntRef& Ref) { return Ref.Value; };
+	MemberDefaultsPrivate::IterateDefaults<FMetasoundEditorMemberPageDefaultInt, FMetasoundEditorGraphMemberDefaultIntRef, int32>(Iter, MemberToLiteral, Defaults);
 }
 
-EMetasoundFrontendLiteralType UMetasoundEditorGraphMemberDefaultBool::GetLiteralType() const 
+bool UMetasoundEditorGraphMemberDefaultInt::RemoveDefault(const FGuid& InPageID)
 {
-	return EMetasoundFrontendLiteralType::Boolean;
+	using namespace Metasound::Editor;
+	return MemberDefaultsPrivate::RemoveDefault(InPageID, Defaults);
 }
 
-void UMetasoundEditorGraphMemberDefaultBool::SetFromLiteral(const FMetasoundFrontendLiteral& InLiteral)
+void UMetasoundEditorGraphMemberDefaultInt::ResetDefaults()
 {
-	Metasound::Editor::MemberDefaultsPrivate::ConvertLiteral<bool>(InLiteral, Default.Value);
+	Defaults.Empty();
+	InitDefault(Metasound::Frontend::DefaultPageID);
 }
 
-void UMetasoundEditorGraphMemberDefaultBool::UpdatePreviewInstance(const Metasound::FVertexName& InParameterName, TScriptInterface<IAudioParameterControllerInterface>& InParameterInterface) const
+void UMetasoundEditorGraphMemberDefaultInt::ResolvePageDefaults()
 {
-	InParameterInterface->SetBoolParameter(InParameterName, Default.Value);
+	using namespace Metasound::Editor;
+	MemberDefaultsPrivate::ResolvePageDefaults(Defaults);
 }
 
-void UMetasoundEditorGraphMemberDefaultBool::SetDefault(const bool InDefault)
+void UMetasoundEditorGraphMemberDefaultInt::SetFromLiteral(const FMetasoundFrontendLiteral& InLiteral, const FGuid& InPageID)
 {
-	if (Default.Value != InDefault)
+	using namespace Metasound::Editor;
+	auto SetLiteral = [&InLiteral](FMetasoundEditorGraphMemberDefaultIntRef& OutRef)
 	{
-		Default.Value = InDefault;
-		OnDefaultStateChanged.Broadcast(InDefault);
-	}
+		int32 Value = 0;
+		InLiteral.TryGet(Value);
+		OutRef.Value = Value;
+	};
+	MemberDefaultsPrivate::SetFromLiteral<FMetasoundEditorMemberPageDefaultInt, FMetasoundEditorGraphMemberDefaultIntRef>(SetLiteral, InPageID, Defaults);
 }
 
-const bool UMetasoundEditorGraphMemberDefaultBool::GetDefault()
+void UMetasoundEditorGraphMemberDefaultInt::SortPageDefaults()
 {
-	return Default.Value;
+	using namespace Metasound::Editor;
+	MemberDefaultsPrivate::SortPageDefaults(Defaults);
 }
 
-FMetasoundFrontendLiteral UMetasoundEditorGraphMemberDefaultBoolArray::GetDefault() const
+bool UMetasoundEditorGraphMemberDefaultInt::Synchronize()
 {
-	TArray<bool> BoolArray;
-	Algo::Transform(Default, BoolArray, [](const FMetasoundEditorGraphMemberDefaultBoolRef& InValue) { return InValue.Value; });
+	using namespace Metasound::Editor;
+	auto LiteralToMember = [](int32 bValue)
+	{
+		FMetasoundEditorGraphMemberDefaultIntRef Ref;
+		Ref.Value = bValue;
+		return Ref;
+	};
 
-	FMetasoundFrontendLiteral Literal;
-	Literal.Set(BoolArray);
-	return Literal;
+	auto MemberToLiteral = [](const FMetasoundEditorGraphMemberDefaultIntRef& Ref)
+	{
+		return Ref.Value;
+	};
+
+	return MemberDefaultsPrivate::SynchronizePageDefault<FMetasoundEditorMemberPageDefaultInt, FMetasoundEditorGraphMemberDefaultIntRef, int32>(
+		FindMember(),
+		LiteralToMember,
+		MemberToLiteral,
+		Defaults);
 }
 
-EMetasoundFrontendLiteralType UMetasoundEditorGraphMemberDefaultBoolArray::GetLiteralType() const
+bool UMetasoundEditorGraphMemberDefaultInt::TryFindDefault(FMetasoundFrontendLiteral& OutLiteral, const FGuid* InPageID) const
 {
-	return EMetasoundFrontendLiteralType::BooleanArray;
+	using namespace Metasound::Editor;
+	auto MemberToLiteral = [](const FMetasoundEditorGraphMemberDefaultIntRef& Ref) { return Ref.Value; };
+	return MemberDefaultsPrivate::TryFindDefault<FMetasoundEditorMemberPageDefaultInt, FMetasoundEditorGraphMemberDefaultIntRef, int32>(
+		Defaults,
+		InPageID,
+		MemberToLiteral,
+		OutLiteral);
 }
 
-void UMetasoundEditorGraphMemberDefaultBoolArray::SetFromLiteral(const FMetasoundFrontendLiteral& InLiteral)
+void UMetasoundEditorGraphMemberDefaultInt::UpdatePreviewInstance(const Metasound::FVertexName& InParameterName, TScriptInterface<IAudioParameterControllerInterface>& InParameterInterface) const
 {
-	Metasound::Editor::MemberDefaultsPrivate::ConvertLiteralToArray<bool, FMetasoundEditorGraphMemberDefaultBoolRef>(InLiteral, Default);
+	FGuid PageID;
+	if (TryGetPreviewPageID(PageID))
+	{
+		InParameterInterface->SetIntParameter(InParameterName, GetDefaultAs<int32>(PageID));
+	}
 }
 
-void UMetasoundEditorGraphMemberDefaultBoolArray::UpdatePreviewInstance(const Metasound::FVertexName& InParameterName, TScriptInterface<IAudioParameterControllerInterface>& InParameterInterface) const
+EMetasoundFrontendLiteralType UMetasoundEditorGraphMemberDefaultIntArray::GetLiteralType() const
 {
-	TArray<bool> BoolArray;
-	Algo::Transform(Default, BoolArray, [](const FMetasoundEditorGraphMemberDefaultBoolRef& InValue) { return InValue.Value; });
-	InParameterInterface->SetBoolArrayParameter(InParameterName, BoolArray);
+	return EMetasoundFrontendLiteralType::IntegerArray;
 }
 
-FMetasoundFrontendLiteral UMetasoundEditorGraphMemberDefaultInt::GetDefault() const
+void UMetasoundEditorGraphMemberDefaultIntArray::InitDefault(const FGuid& InPageID)
 {
-	FMetasoundFrontendLiteral Literal;
-	Literal.Set(Default.Value);
-	return Literal;
+	using namespace Metasound::Editor;
+	MemberDefaultsPrivate::InitDefault<TArray<int32>>(*this, InPageID);
 }
 
-EMetasoundFrontendLiteralType UMetasoundEditorGraphMemberDefaultInt::GetLiteralType() const
+void UMetasoundEditorGraphMemberDefaultIntArray::IterateDefaults(TFunctionRef<void(const FGuid&, FMetasoundFrontendLiteral)> Iter) const
 {
-	return EMetasoundFrontendLiteralType::Integer;
+	using namespace Metasound::Editor;
+	auto MemberToLiteral = [](const TArray<FMetasoundEditorGraphMemberDefaultIntRef>& MemberValues)
+	{
+		TArray<int32> LiteralValues;
+		Algo::Transform(MemberValues, LiteralValues, [](const FMetasoundEditorGraphMemberDefaultIntRef& Ref) { return Ref.Value; });
+		return LiteralValues;
+	};
+	MemberDefaultsPrivate::IterateDefaults<FMetasoundEditorMemberPageDefaultIntArray, TArray<FMetasoundEditorGraphMemberDefaultIntRef>, TArray<int32>>(Iter, MemberToLiteral, Defaults);
+}
+
+bool UMetasoundEditorGraphMemberDefaultIntArray::RemoveDefault(const FGuid& InPageID)
+{
+	using namespace Metasound::Editor;
+	return MemberDefaultsPrivate::RemoveDefault(InPageID, Defaults);
 }
 
-void UMetasoundEditorGraphMemberDefaultInt::SetFromLiteral(const FMetasoundFrontendLiteral& InLiteral)
+void UMetasoundEditorGraphMemberDefaultIntArray::ResetDefaults()
 {
-	Metasound::Editor::MemberDefaultsPrivate::ConvertLiteral<int32>(InLiteral, Default.Value);
+	Defaults.Empty();
+	InitDefault(Metasound::Frontend::DefaultPageID);
 }
 
-void UMetasoundEditorGraphMemberDefaultInt::UpdatePreviewInstance(const Metasound::FVertexName& InParameterName, TScriptInterface<IAudioParameterControllerInterface>& InParameterInterface) const
+void UMetasoundEditorGraphMemberDefaultIntArray::ResolvePageDefaults()
 {
-	InParameterInterface->SetIntParameter(InParameterName, Default.Value);
+	using namespace Metasound::Editor;
+	MemberDefaultsPrivate::ResolvePageDefaults(Defaults);
 }
 
-FMetasoundFrontendLiteral UMetasoundEditorGraphMemberDefaultIntArray::GetDefault() const
+void UMetasoundEditorGraphMemberDefaultIntArray::SetFromLiteral(const FMetasoundFrontendLiteral& InLiteral, const FGuid& InPageID)
 {
-	TArray<int32> IntArray;
-	Algo::Transform(Default, IntArray, [](const FMetasoundEditorGraphMemberDefaultIntRef& InValue) { return InValue.Value; });
+	using namespace Metasound::Editor;
 
-	FMetasoundFrontendLiteral Literal;
-	Literal.Set(IntArray);
+	auto SetLiteral = [&InLiteral](TArray<FMetasoundEditorGraphMemberDefaultIntRef>& OutRefs)
+	{
+		OutRefs.Reset();
+		TArray<int32> Values;
+		InLiteral.TryGet(Values);
+		Algo::Transform(Values, OutRefs, [](const int32& Value)
+			{
+				FMetasoundEditorGraphMemberDefaultIntRef Ref;
+				Ref.Value = Value;
+				return Ref;
+			});
+	};
+	MemberDefaultsPrivate::SetFromLiteral<FMetasoundEditorMemberPageDefaultIntArray, TArray<FMetasoundEditorGraphMemberDefaultIntRef>>(SetLiteral, InPageID, Defaults);
+}
 
-	return Literal;
+void UMetasoundEditorGraphMemberDefaultIntArray::SortPageDefaults()
+{
+	using namespace Metasound::Editor;
+	MemberDefaultsPrivate::SortPageDefaults<FMetasoundEditorMemberPageDefaultIntArray>(Defaults);
 }
 
-EMetasoundFrontendLiteralType UMetasoundEditorGraphMemberDefaultIntArray::GetLiteralType() const
+bool UMetasoundEditorGraphMemberDefaultIntArray::Synchronize()
 {
-	return EMetasoundFrontendLiteralType::IntegerArray;
+	using namespace Metasound::Editor;
+
+	auto LiteralToMember = [](const TArray<int32>& Values)
+	{
+		TArray<FMetasoundEditorGraphMemberDefaultIntRef> Defaults;
+		Algo::Transform(Values, Defaults, [](const int32& Value) { FMetasoundEditorGraphMemberDefaultIntRef Ref; Ref.Value = Value; return Ref; });
+		return Defaults;
+	};
+
+	auto MemberToLiteral = [](const TArray<FMetasoundEditorGraphMemberDefaultIntRef>& Values)
+	{
+		TArray<int32> Defaults;
+		Algo::Transform(Values, Defaults, [](const FMetasoundEditorGraphMemberDefaultIntRef& Ref) { return Ref.Value; });
+		return Defaults;
+	};
+
+	return MemberDefaultsPrivate::SynchronizePageDefault<FMetasoundEditorMemberPageDefaultIntArray, TArray<FMetasoundEditorGraphMemberDefaultIntRef>, TArray<int32>>(
+		FindMember(),
+		LiteralToMember,
+		MemberToLiteral,
+		Defaults);
 }
 
-void UMetasoundEditorGraphMemberDefaultIntArray::SetFromLiteral(const FMetasoundFrontendLiteral& InLiteral)
+bool UMetasoundEditorGraphMemberDefaultIntArray::TryFindDefault(FMetasoundFrontendLiteral& OutLiteral, const FGuid* InPageID) const
 {
-	Metasound::Editor::MemberDefaultsPrivate::ConvertLiteralToArray<int32, FMetasoundEditorGraphMemberDefaultIntRef>(InLiteral, Default);
+	using namespace Metasound::Editor;
+	auto MemberToLiteral = [](const TArray<FMetasoundEditorGraphMemberDefaultIntRef>& Refs) -> TArray<int32>
+	{
+		TArray<int32> Values;
+		Algo::Transform(Refs, Values, [](const FMetasoundEditorGraphMemberDefaultIntRef& Ref) { return Ref.Value; });
+		return Values;
+	};
+	return MemberDefaultsPrivate::TryFindDefault<FMetasoundEditorMemberPageDefaultIntArray, TArray<FMetasoundEditorGraphMemberDefaultIntRef>, TArray<int32>>(Defaults, InPageID, MemberToLiteral, OutLiteral);
 }
 
 void UMetasoundEditorGraphMemberDefaultIntArray::UpdatePreviewInstance(const Metasound::FVertexName& InParameterName, TScriptInterface<IAudioParameterControllerInterface>& InParameterInterface) const
 {
-	TArray<int32> IntArray;
-	Algo::Transform(Default, IntArray, [](const FMetasoundEditorGraphMemberDefaultIntRef& InValue) { return InValue.Value; });
-	InParameterInterface->SetIntArrayParameter(InParameterName, IntArray);
+	FGuid PageID;
+	if (TryGetPreviewPageID(PageID))
+	{
+		InParameterInterface->SetIntArrayParameter(InParameterName, GetDefaultAs<TArray<int32>>(PageID));
+	}
 }
 
-FMetasoundFrontendLiteral UMetasoundEditorGraphMemberDefaultFloat::GetDefault() const
+const FGuid& FMetasoundEditorMemberPageDefault::GetNewEntryID()
 {
-	FMetasoundFrontendLiteral Literal;
-	Literal.Set(Default);
-	return Literal;
+	static const FGuid NewEntryID = FGuid::NewGuid();
+	return NewEntryID;
 }
 
-EMetasoundFrontendLiteralType UMetasoundEditorGraphMemberDefaultFloat::GetLiteralType() const
+EMetasoundFrontendLiteralType UMetasoundEditorGraphMemberDefaultString::GetLiteralType() const
 {
-	return EMetasoundFrontendLiteralType::Float;
+	return EMetasoundFrontendLiteralType::String;
 }
 
-void UMetasoundEditorGraphMemberDefaultFloat::Initialize()
+void UMetasoundEditorGraphMemberDefaultString::InitDefault(const FGuid& InPageID)
 {
-	if (const UMetasoundEditorSettings* EditorSettings = ::GetDefault<UMetasoundEditorSettings>())
+	using namespace Metasound::Editor;
+	MemberDefaultsPrivate::InitDefault<FString>(*this, InPageID);
+}
+
+void UMetasoundEditorGraphMemberDefaultString::IterateDefaults(TFunctionRef<void(const FGuid&, FMetasoundFrontendLiteral)> Iter) const
+{
+	using namespace Metasound::Editor;
+	MemberDefaultsPrivate::IterateDefaults<FMetasoundEditorMemberPageDefaultString>(Iter, Defaults);
+}
+
+void UMetasoundEditorGraphMemberDefaultString::UpdatePreviewInstance(const Metasound::FVertexName& InParameterName, TScriptInterface<IAudioParameterControllerInterface>& InParameterInterface) const
+{
+	FGuid PageID;
+	if (TryGetPreviewPageID(PageID))
 	{
-		WidgetType = EditorSettings->DefaultInputWidgetType;
+		InParameterInterface->SetStringParameter(InParameterName, GetDefaultAs<FString>(PageID));
 	}
 }
 
-void UMetasoundEditorGraphMemberDefaultFloat::ForceRefresh()
+bool UMetasoundEditorGraphMemberDefaultString::RemoveDefault(const FGuid& InPageID)
 {
-	OnRangeChanged.Broadcast(Range);
-	SetDefault(FMath::Clamp(Default, Range.X, Range.Y));
+	using namespace Metasound::Editor;
+	return MemberDefaultsPrivate::RemoveDefault(InPageID, Defaults);
 }
 
-void UMetasoundEditorGraphMemberDefaultFloat::SetFromLiteral(const FMetasoundFrontendLiteral& InLiteral)
+void UMetasoundEditorGraphMemberDefaultString::ResetDefaults()
 {
-	Metasound::Editor::MemberDefaultsPrivate::ConvertLiteral<float>(InLiteral, Default);
+	Defaults.Empty();
+	InitDefault(Metasound::Frontend::DefaultPageID);
+}
 
-	// If set from literal, we force the default value to be the literal's value which may require the range to be fixed up 
-	SetInitialRange();
+void UMetasoundEditorGraphMemberDefaultString::ResolvePageDefaults()
+{
+	using namespace Metasound::Editor;
+	MemberDefaultsPrivate::ResolvePageDefaults(Defaults);
 }
 
-void UMetasoundEditorGraphMemberDefaultFloat::UpdatePreviewInstance(const Metasound::FVertexName& InParameterName, TScriptInterface<IAudioParameterControllerInterface>& InParameterInterface) const
+void UMetasoundEditorGraphMemberDefaultString::SetFromLiteral(const FMetasoundFrontendLiteral& InLiteral, const FGuid& InPageID)
 {
-	InParameterInterface->SetFloatParameter(InParameterName, Default);
+	using namespace Metasound::Editor;
+	MemberDefaultsPrivate::SetFromLiteral<FMetasoundEditorMemberPageDefaultString, FString>(InLiteral, InPageID, Defaults);
 }
 
-void UMetasoundEditorGraphMemberDefaultFloat::PostEditChangeChainProperty(FPropertyChangedChainEvent& PropertyChangedEvent)
+void UMetasoundEditorGraphMemberDefaultString::SortPageDefaults()
 {
 	using namespace Metasound::Editor;
+	MemberDefaultsPrivate::SortPageDefaults<FMetasoundEditorMemberPageDefaultString>(Defaults);
+}
 
-	if (PropertyChangedEvent.GetPropertyName().IsEqual(GET_MEMBER_NAME_CHECKED(UMetasoundEditorGraphMemberDefaultFloat, Default)))
-	{
-		SetDefault(Default);
-	}
-	else if (PropertyChangedEvent.GetPropertyName().IsEqual(GET_MEMBER_NAME_CHECKED(UMetasoundEditorGraphMemberDefaultFloat, WidgetType)) ||
-		PropertyChangedEvent.GetPropertyName().IsEqual(GET_MEMBER_NAME_CHECKED(UMetasoundEditorGraphMemberDefaultFloat, WidgetUnitValueType)))
-	{
-		// Update VolumeWidgetDecibelRange based on current range (it might be stale)
-		if (WidgetUnitValueType == EAudioUnitsValueType::Volume && VolumeWidgetUseLinearOutput)
-		{
-			VolumeWidgetDecibelRange = FVector2D(Audio::ConvertToDecibels(Range.X), Audio::ConvertToDecibels(Range.Y));
-		}
-		else if (WidgetUnitValueType == EAudioUnitsValueType::Frequency)
-		{
-			SetRange(FVector2D(MIN_FILTER_FREQUENCY, MAX_FILTER_FREQUENCY));
-		}
-		else
-		{
-			SetInitialRange();
-		}
+bool UMetasoundEditorGraphMemberDefaultString::Synchronize()
+{
+	using namespace Metasound::Editor;
+	return MemberDefaultsPrivate::SynchronizePageDefault<FMetasoundEditorMemberPageDefaultString, FString>(FindMember(), Defaults);
+}
 
-		// If the widget type is changed to none, we need to refresh clamping the value or not, since if the widget was a slider before, the value was clamped
-		OnClampChanged.Broadcast(ClampDefault);
-	}
-	else if (PropertyChangedEvent.PropertyChain.GetActiveMemberNode()->GetValue()->GetFName().IsEqual(GET_MEMBER_NAME_CHECKED(UMetasoundEditorGraphMemberDefaultFloat, Range)))
-	{
-		if (PropertyChangedEvent.ChangeType != EPropertyChangeType::Interactive)
-		{
-			if (WidgetType != EMetasoundMemberDefaultWidget::None &&
-				WidgetUnitValueType == EAudioUnitsValueType::Volume)
-			{
-				if (VolumeWidgetUseLinearOutput)
-				{
-					Range.X = FMath::Max(Range.X, Audio::ConvertToLinear(SAudioVolumeRadialSlider::MinDbValue));
-					Range.Y = FMath::Min(Range.Y, Audio::ConvertToLinear(SAudioVolumeRadialSlider::MaxDbValue));
-					if (Range.X > Range.Y)
-					{
-						Range.Y = Range.X;
-					}
-					VolumeWidgetDecibelRange = FVector2D(Audio::ConvertToDecibels(Range.X), Audio::ConvertToDecibels(Range.Y));
-					OnRangeChanged.Broadcast(VolumeWidgetDecibelRange);
-					SetDefault(FMath::Clamp(Default, Range.X, Range.Y));
-				}
-			}
-			else
-			{
-				// if Range.X > Range.Y, set Range.Y to Range.X
-				if (Range.X > Range.Y)
-				{
-					SetRange(FVector2D(Range.X, FMath::Max(Range.X, Range.Y)));
-				}
-				else
-				{
-					ForceRefresh();
-				}
-			}
-		}
-	}
-	else if (PropertyChangedEvent.GetPropertyName().IsEqual(GET_MEMBER_NAME_CHECKED(UMetasoundEditorGraphMemberDefaultFloat, ClampDefault)))
-	{
-		SetInitialRange();
-		OnClampChanged.Broadcast(ClampDefault);
-	}
-	else if (PropertyChangedEvent.PropertyChain.GetActiveMemberNode()->GetValue()->GetFName().IsEqual(GET_MEMBER_NAME_CHECKED(UMetasoundEditorGraphMemberDefaultFloat, VolumeWidgetDecibelRange)))
-	{
-		if (PropertyChangedEvent.ChangeType != EPropertyChangeType::Interactive)
-		{
-			VolumeWidgetDecibelRange.X = FMath::Max(VolumeWidgetDecibelRange.X, SAudioVolumeRadialSlider::MinDbValue);
-			VolumeWidgetDecibelRange.Y = FMath::Min(VolumeWidgetDecibelRange.Y, SAudioVolumeRadialSlider::MaxDbValue);
-			SetRange(FVector2D(Audio::ConvertToLinear(VolumeWidgetDecibelRange.X), Audio::ConvertToLinear(VolumeWidgetDecibelRange.Y)));
-		}
-	}
-	else if (PropertyChangedEvent.GetPropertyName().IsEqual(GET_MEMBER_NAME_CHECKED(UMetasoundEditorGraphMemberDefaultFloat, VolumeWidgetUseLinearOutput)))
-	{
-		if (VolumeWidgetUseLinearOutput)
-		{
-			// Range and default are currently in dB, need to change to linear
-			float DbDefault = Default;
-			VolumeWidgetDecibelRange = Range;
-			Range = FVector2D(Audio::ConvertToLinear(VolumeWidgetDecibelRange.X), Audio::ConvertToLinear(VolumeWidgetDecibelRange.Y));
-			Default = Audio::ConvertToLinear(DbDefault);
-		}
-		else
-		{
-			// Range and default are currently linear, need to change to dB
-			Range = VolumeWidgetDecibelRange;
-			Default = Audio::ConvertToDecibels(Default);
-		}
-	}
+bool UMetasoundEditorGraphMemberDefaultString::TryFindDefault(FMetasoundFrontendLiteral& OutLiteral, const FGuid* InPageID) const
+{
+	using namespace Metasound::Editor;
+	return MemberDefaultsPrivate::TryFindDefault<FMetasoundEditorMemberPageDefaultString, FString>(Defaults, InPageID, OutLiteral);
+}
 
-	// Only update member on non-interactive changes to avoid refreshing the details panel mid-update
-	if (PropertyChangedEvent.ChangeType != EPropertyChangeType::Interactive)
-	{
-		UMetasoundEditorGraphMember* Member = FindMember();
-		if (ensure(Member))
-		{
-			FMetasoundFrontendDocumentModifyContext& ModifyContext = FGraphBuilder::GetOutermostMetaSoundChecked(*Member).GetModifyContext();
-			ModifyContext.AddMemberIDsModified({ Member->GetMemberID() });
+EMetasoundFrontendLiteralType UMetasoundEditorGraphMemberDefaultStringArray::GetLiteralType() const
+{
+	return EMetasoundFrontendLiteralType::StringArray;
+}
 
-			// Mark all nodes as modified to refresh them on synchronization.  This ensures all corresponding widgets get updated.
-			const TArray<UMetasoundEditorGraphMemberNode*> MemberNodes = Member->GetNodes();
-			TSet<FGuid> NodesToRefresh;
-			Algo::Transform(MemberNodes, NodesToRefresh, [](const UMetasoundEditorGraphMemberNode* MemberNode) { return MemberNode->GetNodeID(); });
-			ModifyContext.AddNodeIDsModified({ NodesToRefresh });
-		}
-	}
+void UMetasoundEditorGraphMemberDefaultStringArray::InitDefault(const FGuid& InPageID)
+{
+	using namespace Metasound::Editor;
+	MemberDefaultsPrivate::InitDefault<TArray<FString>>(*this, InPageID);
+}
 
-	Super::PostEditChangeChainProperty(PropertyChangedEvent);
+void UMetasoundEditorGraphMemberDefaultStringArray::IterateDefaults(TFunctionRef<void(const FGuid&, FMetasoundFrontendLiteral)> Iter) const
+{
+	using namespace Metasound::Editor;
+	MemberDefaultsPrivate::IterateDefaults(Iter, Defaults);
 }
 
-void UMetasoundEditorGraphMemberDefaultFloat::SetDefault(const float InDefault)
+bool UMetasoundEditorGraphMemberDefaultStringArray::RemoveDefault(const FGuid& InPageID)
 {
-	if (!FMath::IsNearlyEqual(Default, InDefault))
-	{
-		Default = InDefault;
-		OnDefaultValueChanged.Broadcast(InDefault);
-	}
+	using namespace Metasound::Editor;
+	return MemberDefaultsPrivate::RemoveDefault(InPageID, Defaults);
 }
 
-void UMetasoundEditorGraphMemberDefaultFloat::SetInitialRange()
+void UMetasoundEditorGraphMemberDefaultStringArray::ResetDefaults()
 {
-	// If value is within current range, keep it, otherwise set range to something reasonable 
-	if (!(Default >= Range.X && Default <= Range.Y))
-	{
-		if (FMath::IsNearlyEqual(Default, 0.0f))
-		{
-			SetRange(FVector2D(0.0f, 1.0f));
-		}
-		else
-		{
-			SetRange(FVector2D(FMath::Min(0.0f, Default), FMath::Max(0.0f, Default)));
-		}
-	}
+	Defaults.Empty();
+	InitDefault(Metasound::Frontend::DefaultPageID);
 }
 
-float UMetasoundEditorGraphMemberDefaultFloat::GetDefault()
+void UMetasoundEditorGraphMemberDefaultStringArray::ResolvePageDefaults()
 {
-	return Default;
+	using namespace Metasound::Editor;
+	MemberDefaultsPrivate::ResolvePageDefaults(Defaults);
 }
 
-FVector2D UMetasoundEditorGraphMemberDefaultFloat::GetRange()
+void UMetasoundEditorGraphMemberDefaultStringArray::SetFromLiteral(const FMetasoundFrontendLiteral& InLiteral, const FGuid& InPageID)
 {
-	return Range;
+	using namespace Metasound::Editor;
+	MemberDefaultsPrivate::SetFromLiteral<FMetasoundEditorMemberPageDefaultStringArray, TArray<FString>>(InLiteral, InPageID, Defaults);
 }
 
-void UMetasoundEditorGraphMemberDefaultFloat::SetRange(const FVector2D InRange)
+void UMetasoundEditorGraphMemberDefaultStringArray::SortPageDefaults()
 {
-	if (!(Range - InRange).IsNearlyZero())
+	using namespace Metasound::Editor;
+	MemberDefaultsPrivate::SortPageDefaults<FMetasoundEditorMemberPageDefaultStringArray>(Defaults);
+}
+
+bool UMetasoundEditorGraphMemberDefaultStringArray::Synchronize()
+{
+	using namespace Metasound::Editor;
+	return MemberDefaultsPrivate::SynchronizePageDefault<FMetasoundEditorMemberPageDefaultStringArray, TArray<FString>>(FindMember(), Defaults);
+}
+
+bool UMetasoundEditorGraphMemberDefaultStringArray::TryFindDefault(FMetasoundFrontendLiteral& OutLiteral, const FGuid* InPageID) const
+{
+	using namespace Metasound::Editor;
+	return MemberDefaultsPrivate::TryFindDefault<FMetasoundEditorMemberPageDefaultStringArray, TArray<FString>>(Defaults, InPageID, OutLiteral);
+}
+
+void UMetasoundEditorGraphMemberDefaultStringArray::UpdatePreviewInstance(const Metasound::FVertexName& InParameterName, TScriptInterface<IAudioParameterControllerInterface>& InParameterInterface) const
+{
+	FGuid PageID;
+	if (TryGetPreviewPageID(PageID))
 	{
-		Range = InRange;
-		OnRangeChanged.Broadcast(InRange);
-		SetDefault(FMath::Clamp(Default, Range.X, Range.Y));
+		InParameterInterface->SetStringArrayParameter(InParameterName, GetDefaultAs<TArray<FString>>(PageID));
 	}
 }
 
-FMetasoundFrontendLiteral UMetasoundEditorGraphMemberDefaultFloatArray::GetDefault() const
+EMetasoundFrontendLiteralType UMetasoundEditorGraphMemberDefaultObject::GetLiteralType() const
 {
-	FMetasoundFrontendLiteral Literal;
-	Literal.Set(Default);
-	return Literal;
+	return EMetasoundFrontendLiteralType::UObject;
 }
 
-EMetasoundFrontendLiteralType UMetasoundEditorGraphMemberDefaultFloatArray::GetLiteralType() const
+void UMetasoundEditorGraphMemberDefaultObject::InitDefault(const FGuid& InPageID)
 {
-	return EMetasoundFrontendLiteralType::FloatArray;
+	// Can't use template as the object type's MetaSound proxy DataType is unknown
+	FMetasoundFrontendLiteral DefaultLiteral;
+	constexpr UObject* NullObject = nullptr;
+	DefaultLiteral.Set(NullObject);
+	SetFromLiteral(DefaultLiteral, InPageID);
 }
 
-void UMetasoundEditorGraphMemberDefaultFloatArray::SetFromLiteral(const FMetasoundFrontendLiteral& InLiteral)
+void UMetasoundEditorGraphMemberDefaultObject::IterateDefaults(TFunctionRef<void(const FGuid&, FMetasoundFrontendLiteral)> Iter) const
 {
-	Metasound::Editor::MemberDefaultsPrivate::ConvertLiteralToArray<float>(InLiteral, Default);
+	using namespace Metasound::Editor;
+	auto MemberToLiteral = [](const FMetasoundEditorGraphMemberDefaultObjectRef& Ref) -> UObject* { return Ref.Object; };
+	MemberDefaultsPrivate::IterateDefaults<FMetasoundEditorMemberPageDefaultObjectRef, FMetasoundEditorGraphMemberDefaultObjectRef, UObject*>(Iter, MemberToLiteral, Defaults);
 }
 
-void UMetasoundEditorGraphMemberDefaultFloatArray::UpdatePreviewInstance(const Metasound::FVertexName& InParameterName, TScriptInterface<IAudioParameterControllerInterface>& InParameterInterface) const
+bool UMetasoundEditorGraphMemberDefaultObject::RemoveDefault(const FGuid& InPageID)
 {
-	InParameterInterface->SetFloatArrayParameter(InParameterName, Default);
+	using namespace Metasound::Editor;
+	return MemberDefaultsPrivate::RemoveDefault(InPageID, Defaults);
 }
 
-FMetasoundFrontendLiteral UMetasoundEditorGraphMemberDefaultString::GetDefault() const
+void UMetasoundEditorGraphMemberDefaultObject::ResetDefaults()
 {
-	FMetasoundFrontendLiteral Literal;
-	Literal.Set(Default);
-	return Literal;
+	Defaults.Empty();
+	InitDefault(Metasound::Frontend::DefaultPageID);
 }
 
-EMetasoundFrontendLiteralType UMetasoundEditorGraphMemberDefaultString::GetLiteralType() const
+void UMetasoundEditorGraphMemberDefaultObject::ResolvePageDefaults()
 {
-	return EMetasoundFrontendLiteralType::String;
+	using namespace Metasound::Editor;
+	MemberDefaultsPrivate::ResolvePageDefaults(Defaults);
 }
 
-void UMetasoundEditorGraphMemberDefaultString::SetFromLiteral(const FMetasoundFrontendLiteral& InLiteral)
+void UMetasoundEditorGraphMemberDefaultObject::SetFromLiteral(const FMetasoundFrontendLiteral& InLiteral, const FGuid& InPageID)
 {
-	Metasound::Editor::MemberDefaultsPrivate::ConvertLiteral<FString>(InLiteral, Default);
+	using namespace Metasound::Editor;
+	auto SetLiteral = [&InLiteral](FMetasoundEditorGraphMemberDefaultObjectRef& OutRef)
+	{
+		UObject* Value = nullptr;
+		InLiteral.TryGet(Value);
+		OutRef.Object = Value;
+	};
+	MemberDefaultsPrivate::SetFromLiteral<FMetasoundEditorMemberPageDefaultObjectRef, FMetasoundEditorGraphMemberDefaultObjectRef>(SetLiteral, InPageID, Defaults);
 }
 
-void UMetasoundEditorGraphMemberDefaultString::UpdatePreviewInstance(const Metasound::FVertexName& InParameterName, TScriptInterface<IAudioParameterControllerInterface>& InParameterInterface) const
+void UMetasoundEditorGraphMemberDefaultObject::SortPageDefaults()
 {
-	InParameterInterface->SetStringParameter(InParameterName, Default);
+	using namespace Metasound::Editor;
+	MemberDefaultsPrivate::SortPageDefaults(Defaults);
 }
 
-FMetasoundFrontendLiteral UMetasoundEditorGraphMemberDefaultStringArray::GetDefault() const
+bool UMetasoundEditorGraphMemberDefaultObject::Synchronize()
 {
-	FMetasoundFrontendLiteral Literal;
-	Literal.Set(Default);
-	return Literal;
+	using namespace Metasound::Editor;
+	auto LiteralToMember = [](UObject* Value)
+	{
+		FMetasoundEditorGraphMemberDefaultObjectRef Ref;
+		Ref.Object = Value;
+		return Ref;
+	};
+
+	auto MemberToLiteral = [](const FMetasoundEditorGraphMemberDefaultObjectRef& Value) -> UObject*
+	{
+		return Value.Object;
+	};
+
+	return MemberDefaultsPrivate::SynchronizePageDefault<FMetasoundEditorMemberPageDefaultObjectRef, FMetasoundEditorGraphMemberDefaultObjectRef, UObject*>(
+		FindMember(),
+		LiteralToMember,
+		MemberToLiteral,
+		Defaults);
 }
 
-EMetasoundFrontendLiteralType UMetasoundEditorGraphMemberDefaultStringArray::GetLiteralType() const
+bool UMetasoundEditorGraphMemberDefaultObject::TryFindDefault(FMetasoundFrontendLiteral& OutLiteral, const FGuid* InPageID) const
 {
-	return EMetasoundFrontendLiteralType::StringArray;
+	using namespace Metasound::Editor;
+	auto MemberToLiteral = [](const FMetasoundEditorGraphMemberDefaultObjectRef& Value) -> UObject*
+	{
+		return Value.Object;
+	};
+	return MemberDefaultsPrivate::TryFindDefault<FMetasoundEditorMemberPageDefaultObjectRef, FMetasoundEditorGraphMemberDefaultObjectRef, UObject*>(
+		Defaults,
+		InPageID,
+		MemberToLiteral,
+		OutLiteral);
 }
 
-void UMetasoundEditorGraphMemberDefaultStringArray::SetFromLiteral(const FMetasoundFrontendLiteral& InLiteral)
+void UMetasoundEditorGraphMemberDefaultObject::UpdatePreviewInstance(const Metasound::FVertexName& InParameterName, TScriptInterface<IAudioParameterControllerInterface>& InParameterInterface) const
 {
-	Metasound::Editor::MemberDefaultsPrivate::ConvertLiteralToArray<FString>(InLiteral, Default);
+	FGuid PageID;
+	if (TryGetPreviewPageID(PageID))
+	{
+		InParameterInterface->SetObjectParameter(InParameterName, GetDefaultAs<UObject*>(PageID));
+	}
 }
 
-void UMetasoundEditorGraphMemberDefaultStringArray::UpdatePreviewInstance(const Metasound::FVertexName& InParameterName, TScriptInterface<IAudioParameterControllerInterface>& InParameterInterface) const
+EMetasoundFrontendLiteralType UMetasoundEditorGraphMemberDefaultObjectArray::GetLiteralType() const
 {
-	InParameterInterface->SetStringArrayParameter(InParameterName, Default);
+	return EMetasoundFrontendLiteralType::UObjectArray;
 }
 
-FMetasoundFrontendLiteral UMetasoundEditorGraphMemberDefaultObject::GetDefault() const
+void UMetasoundEditorGraphMemberDefaultObjectArray::InitDefault(const FGuid& InPageID)
 {
-	FMetasoundFrontendLiteral Literal;
-	Literal.Set(Default.Object);
-	return Literal;
+	// Can't use template as the object type's MetaSound proxy DataType is unknown
+	FMetasoundFrontendLiteral DefaultLiteral;
+	DefaultLiteral.Set(TArray<UObject*>());
+	SetFromLiteral(DefaultLiteral, InPageID);
 }
 
-EMetasoundFrontendLiteralType UMetasoundEditorGraphMemberDefaultObject::GetLiteralType() const
+void UMetasoundEditorGraphMemberDefaultObjectArray::IterateDefaults(TFunctionRef<void(const FGuid&, FMetasoundFrontendLiteral)> Iter) const
 {
-	return EMetasoundFrontendLiteralType::UObject;
+	using namespace Metasound::Editor;
+	auto MemberToLiteral = [](const TArray<FMetasoundEditorGraphMemberDefaultObjectRef>& MemberValues)
+	{
+		TArray<UObject*> LiteralValues;
+		Algo::Transform(MemberValues, LiteralValues, [](const FMetasoundEditorGraphMemberDefaultObjectRef& Ref) { return Ref.Object; });
+		return LiteralValues;
+	};
+	MemberDefaultsPrivate::IterateDefaults<FMetasoundEditorMemberPageDefaultObjectArray, TArray<FMetasoundEditorGraphMemberDefaultObjectRef>, TArray<UObject*>>(Iter, MemberToLiteral, Defaults);
 }
 
-void UMetasoundEditorGraphMemberDefaultObject::SetFromLiteral(const FMetasoundFrontendLiteral& InLiteral)
+bool UMetasoundEditorGraphMemberDefaultObjectArray::RemoveDefault(const FGuid& InPageID)
 {
-	UObject* Object = nullptr;
-	// If default constructed (type is None), then we don't need to try to get the object, and just return nullptr
-	if (InLiteral.GetType() != EMetasoundFrontendLiteralType::None)
-	{
-		ensure(InLiteral.TryGet(Object));
-	}
-	Default.Object = Object;
+	using namespace Metasound::Editor;
+	return MemberDefaultsPrivate::RemoveDefault(InPageID, Defaults);
 }
 
-void UMetasoundEditorGraphMemberDefaultObject::UpdatePreviewInstance(const Metasound::FVertexName& InParameterName, TScriptInterface<IAudioParameterControllerInterface>& InParameterInterface) const
+void UMetasoundEditorGraphMemberDefaultObjectArray::ResetDefaults()
+{
+	Defaults.Empty();
+	InitDefault(Metasound::Frontend::DefaultPageID);
+}
+
+void UMetasoundEditorGraphMemberDefaultObjectArray::ResolvePageDefaults()
 {
-	// TODO. We need proxy object here safely.
+	using namespace Metasound::Editor;
+	MemberDefaultsPrivate::ResolvePageDefaults(Defaults);
 }
 
-FMetasoundFrontendLiteral UMetasoundEditorGraphMemberDefaultObjectArray::GetDefault() const
+void UMetasoundEditorGraphMemberDefaultObjectArray::SetFromLiteral(const FMetasoundFrontendLiteral& InLiteral, const FGuid& InPageID)
 {
-	TArray<UObject*> ObjectArray;
-	Algo::Transform(Default, ObjectArray, [](const FMetasoundEditorGraphMemberDefaultObjectRef& InValue) { return InValue.Object; });
+	using namespace Metasound::Editor;
 
-	FMetasoundFrontendLiteral Literal;
-	Literal.Set(ObjectArray);
-	return Literal;
+	auto SetLiteral = [&InLiteral](TArray<FMetasoundEditorGraphMemberDefaultObjectRef>& OutRefs)
+	{
+		OutRefs.Reset();
+		TArray<UObject*> Values;
+		InLiteral.TryGet(Values);
+		Algo::Transform(Values, OutRefs, [](UObject* Value)
+		{
+			FMetasoundEditorGraphMemberDefaultObjectRef Ref;
+			Ref.Object = Value;
+			return Ref;
+		});
+	};
+	MemberDefaultsPrivate::SetFromLiteral<FMetasoundEditorMemberPageDefaultObjectArray, TArray<FMetasoundEditorGraphMemberDefaultObjectRef>>(SetLiteral, InPageID, Defaults);
 }
 
-EMetasoundFrontendLiteralType UMetasoundEditorGraphMemberDefaultObjectArray::GetLiteralType() const
+void UMetasoundEditorGraphMemberDefaultObjectArray::SortPageDefaults()
 {
-	return EMetasoundFrontendLiteralType::UObjectArray;
+	using namespace Metasound::Editor;
+	MemberDefaultsPrivate::SortPageDefaults<FMetasoundEditorMemberPageDefaultObjectArray>(Defaults);
 }
 
-void UMetasoundEditorGraphMemberDefaultObjectArray::SetFromLiteral(const FMetasoundFrontendLiteral& InLiteral)
+bool UMetasoundEditorGraphMemberDefaultObjectArray::Synchronize()
 {
-	TArray<UObject*> ObjectArray;
-	ensure(InLiteral.TryGet(ObjectArray));
+	using namespace Metasound::Editor;
+
+	auto LiteralToMember = [](const TArray<UObject*>& Values)
+	{
+		TArray<FMetasoundEditorGraphMemberDefaultObjectRef> Defaults;
+		Algo::Transform(Values, Defaults, [](UObject* Value) { FMetasoundEditorGraphMemberDefaultObjectRef Ref; Ref.Object = Value; return Ref; });
+		return Defaults;
+	};
+
+	auto MemberToLiteral = [](const TArray<FMetasoundEditorGraphMemberDefaultObjectRef>& Values)
+	{
+		TArray<UObject*> Defaults;
+		Algo::Transform(Values, Defaults, [](const FMetasoundEditorGraphMemberDefaultObjectRef& Ref) { return Ref.Object; });
+		return Defaults;
+	};
+
+	return MemberDefaultsPrivate::SynchronizePageDefault<FMetasoundEditorMemberPageDefaultObjectArray, TArray<FMetasoundEditorGraphMemberDefaultObjectRef>, TArray<UObject*>>(
+		FindMember(),
+		LiteralToMember,
+		MemberToLiteral,
+		Defaults);
+}
 
-	Default.Reset();
-	Algo::Transform(ObjectArray, Default, [](UObject* InValue) { return FMetasoundEditorGraphMemberDefaultObjectRef { InValue }; });
+bool UMetasoundEditorGraphMemberDefaultObjectArray::TryFindDefault(FMetasoundFrontendLiteral& OutLiteral, const FGuid* InPageID) const
+{
+	using namespace Metasound::Editor;
+	auto MemberToLiteral = [](const TArray<FMetasoundEditorGraphMemberDefaultObjectRef>& Refs)
+	{
+		TArray<UObject*> Values;
+		Algo::Transform(Refs, Values, [](const FMetasoundEditorGraphMemberDefaultObjectRef& Ref) { return Ref.Object; });
+		return Values;
+	};
+	return MemberDefaultsPrivate::TryFindDefault<FMetasoundEditorMemberPageDefaultObjectArray, TArray<FMetasoundEditorGraphMemberDefaultObjectRef>, TArray<UObject*>>(Defaults, InPageID, MemberToLiteral, OutLiteral);
 }
 
 void UMetasoundEditorGraphMemberDefaultObjectArray::UpdatePreviewInstance(const Metasound::FVertexName& InParameterName, TScriptInterface<IAudioParameterControllerInterface>& InParameterInterface) const
 {
-	TArray<UObject*> ObjectArray;
-	Algo::Transform(Default, ObjectArray, [](const FMetasoundEditorGraphMemberDefaultObjectRef& InValue) { return InValue.Object; });
-	// TODO. We need proxy object here safely.
+	FGuid PageID;
+	if (TryGetPreviewPageID(PageID))
+	{
+		InParameterInterface->SetObjectArrayParameter(InParameterName, GetDefaultAs<TArray<UObject*>>(PageID));
+	}
 }
-
diff --git a/Engine/Plugins/Runtime/Metasound/Source/MetasoundEditor/Private/MetasoundEditorGraphNode.cpp b/Engine/Plugins/Runtime/Metasound/Source/MetasoundEditor/Private/MetasoundEditorGraphNode.cpp
--- a/Engine/Plugins/Runtime/Metasound/Source/MetasoundEditor/Private/MetasoundEditorGraphNode.cpp
+++ b/Engine/Plugins/Runtime/Metasound/Source/MetasoundEditor/Private/MetasoundEditorGraphNode.cpp
@@ -988,7 +988,9 @@ void UMetasoundEditorGraphOutputNode::CacheBreadcrumb()
 
 		if (const UMetasoundEditorGraphMemberDefaultLiteral* Literal = Output->GetLiteral())
 		{
-			Breadcrumb.DefaultLiterals.Add(Metasound::Frontend::DefaultPageID, Literal->GetDefault());
+			FMetasoundFrontendLiteral DefaultLiteral;
+			Literal->TryFindDefault(DefaultLiteral);
+			Breadcrumb.DefaultLiterals.Add(Metasound::Frontend::DefaultPageID, DefaultLiteral);
 		}
 	}
 }
diff --git a/Engine/Plugins/Runtime/Metasound/Source/MetasoundEditor/Private/MetasoundEditorGraphNodeVisualization.h b/Engine/Plugins/Runtime/Metasound/Source/MetasoundEditor/Private/MetasoundEditorGraphNodeVisualization.h
--- a/Engine/Plugins/Runtime/Metasound/Source/MetasoundEditor/Private/MetasoundEditorGraphNodeVisualization.h
+++ b/Engine/Plugins/Runtime/Metasound/Source/MetasoundEditor/Private/MetasoundEditorGraphNodeVisualization.h
@@ -91,7 +91,9 @@ namespace Metasound::Editor
 
 							if (const UMetasoundEditorGraphMemberDefaultLiteral* MemberDefaultLiteral = Member->GetLiteral())
 							{
-								const FMetasoundFrontendLiteral DefaultLiteral = MemberDefaultLiteral->GetDefault();
+								FMetaSoundFrontendDocumentBuilder& Builder = Member->GetFrontendBuilderChecked();
+								FMetasoundFrontendLiteral DefaultLiteral;
+								MemberDefaultLiteral->TryFindDefault(DefaultLiteral, &Builder.GetBuildPageID());
 								if (DefaultLiteral.TryGet(OutValue))
 								{
 									return true;
diff --git a/Engine/Plugins/Runtime/Metasound/Source/MetasoundEditor/Private/MetasoundEditorGraphSchema.cpp b/Engine/Plugins/Runtime/Metasound/Source/MetasoundEditor/Private/MetasoundEditorGraphSchema.cpp
--- a/Engine/Plugins/Runtime/Metasound/Source/MetasoundEditor/Private/MetasoundEditorGraphSchema.cpp
+++ b/Engine/Plugins/Runtime/Metasound/Source/MetasoundEditor/Private/MetasoundEditorGraphSchema.cpp
@@ -416,7 +416,7 @@ PRAGMA_ENABLE_DEPRECATION_WARNINGS
 				return nullptr;
 			}
 
-			const FMetasoundFrontendVertex* InputVertex = Builder.FindNodeInput(InputVertexHandle.NodeID, InputVertexHandle.VertexID);;
+			const FMetasoundFrontendVertex* InputVertex = Builder.FindNodeInput(InputVertexHandle.NodeID, InputVertexHandle.VertexID);
 			if (!ensure(InputVertex))
 			{
 				return nullptr;
@@ -426,14 +426,25 @@ PRAGMA_ENABLE_DEPRECATION_WARNINGS
 			ParentMetasound.Modify();
 			MetasoundGraph->Modify();
 
-			FMetasoundFrontendLiteral DefaultValue;
-			FGraphBuilder::GetPinLiteral(*FromPin, DefaultValue);
+			TArray<FMetasoundFrontendClassInputDefault> DefaultLiterals;
+			if (const FMetasoundFrontendVertexLiteral* VertexLiteral = Builder.FindNodeInputDefault(InputVertexHandle.NodeID, InputVertexHandle.VertexID))
+			{
+				DefaultLiterals.Add_GetRef(Builder.GetBuildPageID()).Literal = VertexLiteral->Value;
+			}
+			if (DefaultLiterals.IsEmpty())
+			{
+				if (const TArray<FMetasoundFrontendClassInputDefault>* ClassDefaults = Builder.FindNodeClassInputDefaults(InputVertexHandle.NodeID, InputVertex->Name))
+				{
+					DefaultLiterals = *ClassDefaults;
+				}
+			}
+
 			const FCreateNodeVertexParams VertexParams =
 			{
 				InputVertex->TypeName,
 				Builder.GetNodeInputAccessType(InputVertexHandle.NodeID, InputVertex->VertexID)
 			};
-			FMetasoundFrontendClassInput ClassInput = FGraphBuilder::CreateUniqueClassInput(ParentMetasound, VertexParams, &DefaultValue, &InputVertex->Name);
+			FMetasoundFrontendClassInput ClassInput = FGraphBuilder::CreateUniqueClassInput(ParentMetasound, VertexParams, DefaultLiterals, &InputVertex->Name);
 			if (const FMetasoundFrontendNode* NewNode = Builder.AddGraphInput(ClassInput))
 			{
 				UMetasoundEditorGraphInput* Input = MetasoundGraph->FindOrAddInput(NewNode->GetID());
@@ -580,7 +591,6 @@ PRAGMA_ENABLE_DEPRECATION_WARNINGS
 			return nullptr;
 		}
 
-
 		UEdGraphNode* PromoteToMutatorVariable(UEdGraph* ParentGraph, UEdGraphPin* FromPin, const FVector2D InLocation, bool bSelectNewNode)
 		{
 			using namespace Metasound;
diff --git a/Engine/Plugins/Runtime/Metasound/Source/MetasoundEditor/Private/MetasoundEditorModule.cpp b/Engine/Plugins/Runtime/Metasound/Source/MetasoundEditor/Private/MetasoundEditorModule.cpp
--- a/Engine/Plugins/Runtime/Metasound/Source/MetasoundEditor/Private/MetasoundEditorModule.cpp
+++ b/Engine/Plugins/Runtime/Metasound/Source/MetasoundEditor/Private/MetasoundEditorModule.cpp
@@ -629,9 +629,16 @@ namespace Metasound
 					FOnGetPropertyTypeCustomizationInstance::CreateLambda([]() { return MakeShared<FMetasoundMemberDefaultObjectDetailCustomization>(); }));
 
 				LiteralCustomizationFactories.Add(UMetasoundEditorGraphMemberDefaultLiteral::StaticClass(), MakeUnique<FMetasoundDefaultLiteralCustomizationFactory>());
-				LiteralCustomizationFactories.Add(UMetasoundEditorGraphMemberDefaultFloat::StaticClass(), MakeUnique<FMetasoundFloatLiteralCustomizationFactory>());
 				LiteralCustomizationFactories.Add(UMetasoundEditorGraphMemberDefaultBool::StaticClass(), MakeUnique<FMetasoundBoolLiteralCustomizationFactory>());
+				LiteralCustomizationFactories.Add(UMetasoundEditorGraphMemberDefaultBoolArray::StaticClass(), MakeUnique<FMetasoundDefaultLiteralCustomizationFactory>());
+				LiteralCustomizationFactories.Add(UMetasoundEditorGraphMemberDefaultFloat::StaticClass(), MakeUnique<FMetasoundFloatLiteralCustomizationFactory>());
+				LiteralCustomizationFactories.Add(UMetasoundEditorGraphMemberDefaultFloatArray::StaticClass(), MakeUnique<FMetasoundDefaultLiteralCustomizationFactory>());
+				LiteralCustomizationFactories.Add(UMetasoundEditorGraphMemberDefaultInt::StaticClass(), MakeUnique<FMetasoundDefaultLiteralCustomizationFactory>());
+				LiteralCustomizationFactories.Add(UMetasoundEditorGraphMemberDefaultIntArray::StaticClass(), MakeUnique<FMetasoundDefaultLiteralCustomizationFactory>());
+				LiteralCustomizationFactories.Add(UMetasoundEditorGraphMemberDefaultObject::StaticClass(), MakeUnique<FMetasoundDefaultLiteralCustomizationFactory>());
 				LiteralCustomizationFactories.Add(UMetasoundEditorGraphMemberDefaultObjectArray::StaticClass(), MakeUnique<FMetasoundObjectArrayLiteralCustomizationFactory>());
+				LiteralCustomizationFactories.Add(UMetasoundEditorGraphMemberDefaultString::StaticClass(), MakeUnique<FMetasoundDefaultLiteralCustomizationFactory>());
+				LiteralCustomizationFactories.Add(UMetasoundEditorGraphMemberDefaultStringArray::StaticClass(), MakeUnique<FMetasoundDefaultLiteralCustomizationFactory>());
 
 				StyleSet = MakeShared<FSlateStyle>();
 
diff --git a/Engine/Plugins/Runtime/Metasound/Source/MetasoundEditor/Private/MetasoundNodeDetailCustomization.cpp b/Engine/Plugins/Runtime/Metasound/Source/MetasoundEditor/Private/MetasoundNodeDetailCustomization.cpp
--- a/Engine/Plugins/Runtime/Metasound/Source/MetasoundEditor/Private/MetasoundNodeDetailCustomization.cpp
+++ b/Engine/Plugins/Runtime/Metasound/Source/MetasoundEditor/Private/MetasoundNodeDetailCustomization.cpp
@@ -273,52 +273,63 @@ namespace Metasound
 			}
 		}
 
-		TArray<IDetailPropertyRow*> FMetasoundFloatLiteralCustomization::CustomizeLiteral(UMetasoundEditorGraphMemberDefaultLiteral& InLiteral, IDetailLayoutBuilder& InDetailLayout)
+		void FMetasoundFloatLiteralCustomization::CustomizeDefaults(UMetasoundEditorGraphMemberDefaultLiteral& InLiteral, IDetailLayoutBuilder& InDetailLayout)
 		{
 			check(DefaultCategoryBuilder);
 
-			UMetasoundEditorGraphMemberDefaultFloat* DefaultFloat = Cast<UMetasoundEditorGraphMemberDefaultFloat>(&InLiteral);
-			if (!ensure(DefaultFloat))
+			TArray<UObject*> EditLiterals;
+			if (UMetasoundEditorGraphMemberDefaultFloat* CastLiteral = Cast<UMetasoundEditorGraphMemberDefaultFloat>(&InLiteral); ensure(CastLiteral))
 			{
-				return { };
+				EditLiterals = { CastLiteral };
+				FloatLiteral = CastLiteral;
 			}
-			FloatLiteral = DefaultFloat;
-
-			TArray<IDetailPropertyRow*> DefaultRows;
-			TSharedPtr<IPropertyHandle> DefaultValueHandle;
-			IDetailPropertyRow* Row = DefaultCategoryBuilder->AddExternalObjectProperty(TArray<UObject*>({ DefaultFloat }), UMetasoundEditorGraphMemberDefaultFloat::GetDefaultPropertyName());
-			if (ensure(Row))
+			else
 			{
-				DefaultRows.Add(Row);
-				DefaultValueHandle = Row->GetPropertyHandle();
+				FloatLiteral.Reset();
+				return;
 			}
 
-			// Apply the clamp range to the default value if not using a widget and ClampDefault is true
-			const bool bUsingWidget = DefaultFloat->WidgetType != EMetasoundMemberDefaultWidget::None;
-			const bool bShouldClampDefaultValue = bUsingWidget || (!bUsingWidget && DefaultFloat->ClampDefault);
+			FMetasoundDefaultLiteralCustomizationBase::CustomizeDefaults(InLiteral, InDetailLayout);
+
+			TAttribute<EVisibility> DefaultVisibility = GetDefaultVisibility();
 
-			IDetailPropertyRow* ClampRow = DefaultCategoryBuilder->AddExternalObjectProperty(TArray<UObject*>({ DefaultFloat }), GET_MEMBER_NAME_CHECKED(UMetasoundEditorGraphMemberDefaultFloat, ClampDefault));
-			if (ensure(ClampRow))
+			IDetailCategoryBuilder& EditorOptionsBuilder = InDetailLayout.EditCategory("DefaultEditorOptions");
+			auto AddOptionPropRow = [&](FName PropertyName)
 			{
-				DefaultRows.Add(ClampRow);
+				IDetailPropertyRow* Row = EditorOptionsBuilder.AddExternalObjectProperty(EditLiterals, PropertyName);
+				if (ensure(Row))
+				{
+					Row->Visibility(DefaultVisibility);
+				}
+				return Row;
+			};
 
-				if (DefaultValueHandle.IsValid())
+			if (IDetailPropertyRow* ClampRow = AddOptionPropRow(GET_MEMBER_NAME_CHECKED(UMetasoundEditorGraphMemberDefaultFloat, ClampDefault)))
+			{
+				// Apply the clamp range to the default value if not using a widget and ClampDefault is true
+				const bool bUsingWidget = FloatLiteral->WidgetType != EMetasoundMemberDefaultWidget::None;
+				const bool bShouldClampDefaultValue = bUsingWidget || (!bUsingWidget && FloatLiteral->ClampDefault);
+				ClampRow->Visibility(DefaultVisibility);
+				for (TSharedPtr<IPropertyHandle>& DefaultValueHandle : DefaultProperties)
 				{
-					if (bShouldClampDefaultValue)
-					{
-						FVector2D Range = DefaultFloat->GetRange();
-						DefaultValueHandle->SetInstanceMetaData("ClampMin", FString::Printf(TEXT("%f"), Range.X));
-						DefaultValueHandle->SetInstanceMetaData("ClampMax", FString::Printf(TEXT("%f"), Range.Y));
-					}
-					else // Stop clamping
+					if (DefaultValueHandle.IsValid())
 					{
-						DefaultValueHandle->SetInstanceMetaData("ClampMin", "");
-						DefaultValueHandle->SetInstanceMetaData("ClampMax", "");
+						if (bShouldClampDefaultValue)
+						{
+							FVector2D Range = FloatLiteral->GetRange();
+							DefaultValueHandle->SetInstanceMetaData("ClampMin", FString::Printf(TEXT("%f"), Range.X));
+							DefaultValueHandle->SetInstanceMetaData("ClampMax", FString::Printf(TEXT("%f"), Range.Y));
+						}
+						else // Stop clamping
+						{
+							DefaultValueHandle->SetInstanceMetaData("ClampMin", "");
+							DefaultValueHandle->SetInstanceMetaData("ClampMax", "");
+						}
 					}
 				}
 
-				DefaultFloat->OnClampChanged.Remove(OnClampChangedDelegateHandle);
-				OnClampChangedDelegateHandle = DefaultFloat->OnClampChanged.AddLambda([this](bool ClampInput)
+				FloatLiteral->OnClampChanged.Remove(OnClampChangedDelegateHandle);
+				OnClampChangedDelegateHandle = FloatLiteral->OnClampChanged.AddLambda([this](bool ClampInput)
 				{
 					if (FloatLiteral.IsValid())
 					{
@@ -333,11 +344,7 @@ namespace Metasound
 
 				if (bShouldClampDefaultValue)
 				{
-					IDetailPropertyRow* RangeRow = DefaultCategoryBuilder->AddExternalObjectProperty(TArray<UObject*>({ DefaultFloat }), GET_MEMBER_NAME_CHECKED(UMetasoundEditorGraphMemberDefaultFloat, Range));
-					if (ensure(RangeRow))
-					{
-						DefaultRows.Add(RangeRow);
-					}
+					AddOptionPropRow(GET_MEMBER_NAME_CHECKED(UMetasoundEditorGraphMemberDefaultFloat, Range));
 				}
 			}
 
@@ -354,45 +361,43 @@ namespace Metasound
 			// add input widget properties
 			if (bShowWidgetOptions)
 			{
-				IDetailCategoryBuilder& WidgetCategoryBuilder = InDetailLayout.EditCategory("EditorOptions");
-				DefaultRows.Add(WidgetCategoryBuilder.AddExternalObjectProperty(TArray<UObject*>({ DefaultFloat }), GET_MEMBER_NAME_CHECKED(UMetasoundEditorGraphMemberDefaultFloat, WidgetType)));
-				DefaultRows.Add(WidgetCategoryBuilder.AddExternalObjectProperty(TArray<UObject*>({ DefaultFloat }), GET_MEMBER_NAME_CHECKED(UMetasoundEditorGraphMemberDefaultFloat, WidgetOrientation)));
-				DefaultRows.Add(WidgetCategoryBuilder.AddExternalObjectProperty(TArray<UObject*>({ DefaultFloat }), GET_MEMBER_NAME_CHECKED(UMetasoundEditorGraphMemberDefaultFloat, WidgetUnitValueType)));
-				if (DefaultFloat->WidgetType != EMetasoundMemberDefaultWidget::None && DefaultFloat->WidgetUnitValueType == EAudioUnitsValueType::Volume)
+				AddOptionPropRow(GET_MEMBER_NAME_CHECKED(UMetasoundEditorGraphMemberDefaultFloat, WidgetType));
+				AddOptionPropRow(GET_MEMBER_NAME_CHECKED(UMetasoundEditorGraphMemberDefaultFloat, WidgetOrientation));
+				AddOptionPropRow(GET_MEMBER_NAME_CHECKED(UMetasoundEditorGraphMemberDefaultFloat, WidgetUnitValueType));
+				if (FloatLiteral->WidgetType != EMetasoundMemberDefaultWidget::None && FloatLiteral->WidgetUnitValueType == EAudioUnitsValueType::Volume)
 				{
-					DefaultRows.Add(WidgetCategoryBuilder.AddExternalObjectProperty(TArray<UObject*>({ DefaultFloat }), GET_MEMBER_NAME_CHECKED(UMetasoundEditorGraphMemberDefaultFloat, VolumeWidgetUseLinearOutput)));
-					if (DefaultFloat->VolumeWidgetUseLinearOutput)
+					AddOptionPropRow(GET_MEMBER_NAME_CHECKED(UMetasoundEditorGraphMemberDefaultFloat, VolumeWidgetUseLinearOutput));
+					if (FloatLiteral->VolumeWidgetUseLinearOutput)
 					{
-						DefaultRows.Add(WidgetCategoryBuilder.AddExternalObjectProperty(TArray<UObject*>({ DefaultFloat }), GET_MEMBER_NAME_CHECKED(UMetasoundEditorGraphMemberDefaultFloat, VolumeWidgetDecibelRange)));
+						AddOptionPropRow(GET_MEMBER_NAME_CHECKED(UMetasoundEditorGraphMemberDefaultFloat, VolumeWidgetDecibelRange));
 					}
 				}
 			}
-
-			return DefaultRows;
 		}
 
 		FMetasoundBoolLiteralCustomization::~FMetasoundBoolLiteralCustomization()
 		{
 		}
 
-		TArray<IDetailPropertyRow*> FMetasoundBoolLiteralCustomization::CustomizeLiteral(UMetasoundEditorGraphMemberDefaultLiteral& InLiteral, IDetailLayoutBuilder& InDetailLayout)
+		void FMetasoundBoolLiteralCustomization::CustomizeDefaults(UMetasoundEditorGraphMemberDefaultLiteral& InLiteral, IDetailLayoutBuilder& InDetailLayout)
 		{
 			check(DefaultCategoryBuilder);
 
+			DefaultProperties.Reset();
+
 			UMetasoundEditorGraphMemberDefaultBool* DefaultBool = Cast<UMetasoundEditorGraphMemberDefaultBool>(&InLiteral);
 			if (!ensure(DefaultBool))
 			{
-				return { };
+				return;
 			}
 			BoolLiteral = DefaultBool;
 
-			TArray<IDetailPropertyRow*> DefaultRows;
-			TSharedPtr<IPropertyHandle> DefaultValueHandle;
+			TArray<TSharedPtr<IPropertyHandle>> DefaultValueHandles;
 			IDetailPropertyRow* Row = DefaultCategoryBuilder->AddExternalObjectProperty(TArray<UObject*>({ DefaultBool }), GET_MEMBER_NAME_CHECKED(UMetasoundEditorGraphMemberDefaultBool, Default));
 			if (ensure(Row))
 			{
-				DefaultRows.Add(Row);
-				DefaultValueHandle = Row->GetPropertyHandle();
+				DefaultValueHandles = { Row->GetPropertyHandle() };
+				Row->Visibility(GetDefaultVisibility());
 			}
 
 			// Enable widget options for editable inputs only 
@@ -409,7 +414,7 @@ namespace Metasound
 			if (bShowWidgetOptions)
 			{
 				Frontend::FDataTypeRegistryInfo DataTypeInfo;
-				MemberCustomizationPrivate::GetDataTypeFromElementPropertyHandle(DefaultValueHandle, DataTypeInfo);
+				MemberCustomizationPrivate::GetDataTypeFromElementPropertyHandle(DefaultValueHandles.Last(), DataTypeInfo);
 
 				const UMetasoundEditorSettings* EditorSettings = GetDefault<UMetasoundEditorSettings>();
 				check(EditorSettings)
@@ -418,96 +423,111 @@ namespace Metasound
 				{
 					if (MemberCustomizationPrivate::GetPrimitiveTypeName(DataTypeInfo) != Metasound::GetMetasoundDataTypeName<Metasound::FTrigger>())
 					{
-						IDetailCategoryBuilder& WidgetCategoryBuilder = InDetailLayout.EditCategory("EditorOptions");
-						DefaultRows.Add(WidgetCategoryBuilder.AddExternalObjectProperty(TArray<UObject*>({ DefaultBool }), GET_MEMBER_NAME_CHECKED(UMetasoundEditorGraphMemberDefaultBool, WidgetType)));
+						IDetailCategoryBuilder& EditorOptionsBuilder = InDetailLayout.EditCategory("DefaultEditorOptions");
+						if (IDetailPropertyRow* TypeRow = EditorOptionsBuilder.AddExternalObjectProperty(TArray<UObject*>({ DefaultBool }), GET_MEMBER_NAME_CHECKED(UMetasoundEditorGraphMemberDefaultBool, WidgetType)))
+						{
+							DefaultValueHandles.Add(Row->GetPropertyHandle());
+							TypeRow->Visibility(GetDefaultVisibility());
+						}
 					}
 				}
 			}
-
-			return DefaultRows;
 		}
 
-		TArray<IDetailPropertyRow*> FMetasoundObjectArrayLiteralCustomization::CustomizeLiteral(UMetasoundEditorGraphMemberDefaultLiteral& InLiteral, IDetailLayoutBuilder& InDetailLayout)
+		void FMetasoundObjectArrayLiteralCustomization::CustomizeDefaults(UMetasoundEditorGraphMemberDefaultLiteral& InLiteral, IDetailLayoutBuilder& InDetailLayout)
 		{
 			check(DefaultCategoryBuilder);
 
-			TSharedPtr<IPropertyHandle> DefaultValueHandle;
-			IDetailPropertyRow* Row = DefaultCategoryBuilder->AddExternalObjectProperty(TArray<UObject*>({ &InLiteral }), GET_MEMBER_NAME_CHECKED(UMetasoundEditorGraphMemberDefaultObjectArray, Default));
-			if (ensure(Row))
+			FOnDefaultPageRowAdded OnDefaultPageRowAdded = [this, &InDetailLayout](IDetailPropertyRow& ValueRow, TSharedRef<IPropertyHandle> PageDefaultProperty)
 			{
-				DefaultValueHandle = Row->GetPropertyHandle();
-			}
-
-			constexpr bool bShowChildren = true;
-			Row->ShowPropertyButtons(false)
-			.CustomWidget(bShowChildren)
-			.NameContent()
-			[
-				DefaultValueHandle->CreatePropertyNameWidget()
-			]
-			.ValueContent()
-			[
-				SNew(SAssetDropTarget)
-				.bSupportsMultiDrop(true)
-				.OnAreAssetsAcceptableForDropWithReason_Lambda([this, DefaultValueHandle](TArrayView<FAssetData> InAssets, FText& OutReason)
+				TSharedPtr<IPropertyHandle> ValueProperty = PageDefaultProperty->GetChildHandle(GET_MEMBER_NAME_CHECKED(FMetasoundEditorMemberPageDefaultObjectArray, Value));
+				if (ValueProperty.IsValid())
 				{
-					Frontend::FDataTypeRegistryInfo DataTypeInfo;
-					const bool bMemberFound = MemberCustomizationPrivate::GetDataTypeFromElementPropertyHandle(DefaultValueHandle, DataTypeInfo);
-					bool bCanDrop = bMemberFound;
-					if (UClass* ProxyGenClass = DataTypeInfo.ProxyGeneratorClass; bCanDrop && bMemberFound)
+					TSharedPtr<IPropertyHandleArray> ArrayProperty = ValueProperty->AsArray();
+					if (ArrayProperty.IsValid())
 					{
-						bCanDrop = true;
-						for (const FAssetData& AssetData : InAssets)
+						uint32 NumElements = 0;
+						ArrayProperty->GetNumElements(NumElements);
+						for (uint32 Index = 0; Index < NumElements; ++Index)
 						{
-							if (UClass* Class = AssetData.GetClass())
+							TSharedPtr<IPropertyHandle> ElementProperty = ArrayProperty->GetElement(Index);
+							if (ElementProperty.IsValid())
 							{
+								TSharedPtr<IPropertyHandle> DefaultValueProperty = ElementProperty->GetChildHandle(GET_MEMBER_NAME_CHECKED(FMetasoundEditorGraphMemberDefaultObjectRef, Object));
+
+								IDetailPropertyRow& Row = DefaultCategoryBuilder->AddProperty(ElementProperty);
+								constexpr bool bShowChildren = false;
+								Row.ShowPropertyButtons(false)
+								.CustomWidget(bShowChildren);
+								(*Row.CustomNameWidget())
+								[
+									ElementProperty->CreatePropertyNameWidget()
+								];
+								(*Row.CustomValueWidget())
+								[
+									SNew(SAssetDropTarget)
+									.bSupportsMultiDrop(true)
+									.OnAreAssetsAcceptableForDropWithReason_Lambda([this, DefaultValueProperty](TArrayView<FAssetData> InAssets, FText& OutReason)
+									{
+										Frontend::FDataTypeRegistryInfo DataTypeInfo;
+										const bool bMemberFound = MemberCustomizationPrivate::GetDataTypeFromElementPropertyHandle(DefaultValueProperty, DataTypeInfo);
+										bool bCanDrop = bMemberFound;
+										if (UClass* ProxyGenClass = DataTypeInfo.ProxyGeneratorClass; bCanDrop && bMemberFound)
+										{
+											bCanDrop = true;
+											for (const FAssetData& AssetData : InAssets)
+											{
+												if (UClass* Class = AssetData.GetClass())
+												{
 PRAGMA_DISABLE_DEPRECATION_WARNINGS
-								const IMetasoundEditorModule& EditorModule = FModuleManager::GetModuleChecked<IMetasoundEditorModule>("MetaSoundEditor");
-								if (EditorModule.IsExplicitProxyClass(*DataTypeInfo.ProxyGeneratorClass))
-								{
-									bCanDrop &= Class == DataTypeInfo.ProxyGeneratorClass;
-									continue;
-								}
+													const IMetasoundEditorModule& EditorModule = FModuleManager::GetModuleChecked<IMetasoundEditorModule>("MetaSoundEditor");
+													if (EditorModule.IsExplicitProxyClass(*DataTypeInfo.ProxyGeneratorClass))
+													{
+														bCanDrop &= Class == DataTypeInfo.ProxyGeneratorClass;
+														continue;
+													}
 PRAGMA_ENABLE_DEPRECATION_WARNINGS
-								if (DataTypeInfo.bIsExplicit)
-								{
-									bCanDrop &= Class == DataTypeInfo.ProxyGeneratorClass;
-								}
-								else
-								{
-									bCanDrop &= Class->IsChildOf(DataTypeInfo.ProxyGeneratorClass);
-								}
+													if (DataTypeInfo.bIsExplicit)
+													{
+														bCanDrop &= Class == DataTypeInfo.ProxyGeneratorClass;
+													}
+													else
+													{
+														bCanDrop &= Class->IsChildOf(DataTypeInfo.ProxyGeneratorClass);
+													}
+												}
+											}
+										}
+
+										return bCanDrop;
+									})
+									.OnAssetsDropped_Lambda([this, ArrayProperty](const FDragDropEvent& DragDropEvent, TArrayView<FAssetData> InAssets)
+									{
+										if (ArrayProperty.IsValid())
+										{
+											FScopedTransaction Transaction(LOCTEXT("DragDropInputAssets", "Drop Asset(s) on MetaSound Input"));
+											for (const FAssetData& AssetData : InAssets)
+											{
+												uint32 AddIndex = INDEX_NONE;
+												ArrayProperty->GetNumElements(AddIndex);
+												ArrayProperty->AddItem();
+												TSharedPtr<IPropertyHandle> ElementHandle = ArrayProperty->GetElement(static_cast<int32>(AddIndex));
+												TSharedPtr<IPropertyHandle> ObjectHandle = ElementHandle->GetChildHandle(GET_MEMBER_NAME_CHECKED(FMetasoundEditorGraphMemberDefaultObjectRef, Object));
+												ObjectHandle->SetValue(AssetData.GetAsset());
+											}
+										}
+									})
+									[
+										DefaultValueProperty->CreatePropertyValueWidget()
+									]
+								];
+								Row.Visibility(GetDefaultVisibility());
 							}
 						}
 					}
-
-					return bCanDrop;
-				})
-				.OnAssetsDropped_Lambda([this, DefaultValueHandle](const FDragDropEvent& DragDropEvent, TArrayView<FAssetData> InAssets)
-				{
-					if (DefaultValueHandle.IsValid())
-					{
-						TSharedPtr<IPropertyHandleArray> ArrayHandle = DefaultValueHandle->AsArray();
-						if (ensure(ArrayHandle.IsValid()))
-						{
-							for (const FAssetData& AssetData : InAssets)
-							{
-								uint32 AddIndex = INDEX_NONE;
-								ArrayHandle->GetNumElements(AddIndex);
-								ArrayHandle->AddItem();
-								TSharedPtr<IPropertyHandle> ElementHandle = ArrayHandle->GetElement(static_cast<int32>(AddIndex));
-								TSharedPtr<IPropertyHandle> ObjectHandle = ElementHandle->GetChildHandle(GET_MEMBER_NAME_CHECKED(FMetasoundEditorGraphMemberDefaultObjectRef, Object));
-								ObjectHandle->SetValue(AssetData.GetAsset());
-							}
-						}
-					}
-				})
-				[
-					DefaultValueHandle->CreatePropertyValueWidget()
-				]
-			];
-
-			return { Row };
+				}
+			};
+			FMetasoundDefaultLiteralCustomizationBase::CustomizePageDefaultRows(InLiteral, InDetailLayout, &OnDefaultPageRowAdded);
 		}
 
 		FText FMetasoundMemberDefaultBoolDetailCustomization::GetPropertyNameOverride() const
@@ -786,13 +806,12 @@ PRAGMA_ENABLE_DEPRECATION_WARNINGS
 			if (ElementPropertyHandle.IsValid())
 			{
 				TSharedPtr<IPropertyHandle> ParentProperty = ElementPropertyHandle->GetParentHandle();
-				while (ParentProperty.IsValid() && ParentProperty->GetProperty() != nullptr)
+				if (ParentProperty.IsValid() && ParentProperty->GetProperty() != nullptr)
 				{
 					ParentPropertyHandleArray = ParentProperty->AsArray();
 					if (ParentPropertyHandleArray.IsValid())
 					{
 						ElementPropertyHandle = ParentProperty;
-						break;
 					}
 				}
 			}
@@ -1102,13 +1121,11 @@ PRAGMA_ENABLE_DEPRECATION_WARNINGS
 			}
 		}
 
-		TArray<IDetailPropertyRow*> FMetasoundMemberDetailCustomization::CustomizeDefaultCategory(IDetailLayoutBuilder& InDetailLayout)
+		void FMetasoundMemberDetailCustomization::CustomizeDefaultCategory(IDetailLayoutBuilder& InDetailLayout)
 		{
-			TArray<IDetailPropertyRow*> DefaultPropertyRows;
-
 			if (!GraphMember.IsValid())
 			{
-				return DefaultPropertyRows;
+				return;
 			}
 
 			UpdateRenameDelegate(*GraphMember);
@@ -1120,33 +1137,32 @@ PRAGMA_ENABLE_DEPRECATION_WARNINGS
 
 				IDetailCategoryBuilder& DefaultCategoryBuilder = GetDefaultCategoryBuilder(InDetailLayout);
 				IMetasoundEditorModule& EditorModule = FModuleManager::GetModuleChecked<IMetasoundEditorModule>("MetaSoundEditor");
-				TUniquePtr<FMetasoundDefaultLiteralCustomizationBase> LiteralCustomization = EditorModule.CreateMemberDefaultLiteralCustomization(*MemberClass, DefaultCategoryBuilder);
+				LiteralCustomization = EditorModule.CreateMemberDefaultLiteralCustomization(*MemberClass, DefaultCategoryBuilder);
+
+				TAttribute<EVisibility> Visibility = TAttribute<EVisibility>::CreateLambda([this]()
+				{
+					return GetDefaultVisibility();
+				});
+
 				if (LiteralCustomization.IsValid())
 				{
-					DefaultPropertyRows = LiteralCustomization->CustomizeLiteral(*MemberDefaultLiteral, InDetailLayout);
+					LiteralCustomization->SetDefaultVisibility(Visibility);
+					LiteralCustomization->SetEnabled(GetEnabled());
+					LiteralCustomization->SetResetOverride(GetResetOverride());
+					LiteralCustomization->CustomizeDefaults(*MemberDefaultLiteral, InDetailLayout);
 				}
 				else
 				{
 					IDetailPropertyRow* DefaultPropertyRow = DefaultCategoryBuilder.AddExternalObjectProperty(TArray<UObject*>({ MemberDefaultLiteral }), "Default");
-					ensureMsgf(DefaultPropertyRow, TEXT("Class '%s' missing expected 'Default' member."
+					if (ensureMsgf(DefaultPropertyRow, TEXT("Class '%s' missing expected 'Default' member."
 						"Either add/rename default member or register customization to display default value/opt out appropriately."),
-						*MemberClass->GetName());
-					DefaultPropertyRows.Add(DefaultPropertyRow);
-				}
-			}
-
-			for (IDetailPropertyRow* Row : DefaultPropertyRows)
-			{
-				if (ensure(Row))
-				{
-					Row->Visibility(TAttribute<EVisibility>::CreateLambda([this]()
+						*MemberClass->GetName()))
 					{
-						return GetDefaultVisibility();
-					}));
+						DefaultPropertyRow->Visibility(Visibility);
+						DefaultPropertyRow->IsEnabled(GetEnabled());
+					}
 				}
 			}
-
-			return DefaultPropertyRows;
 		}
 
 		void FMetasoundMemberDetailCustomization::CustomizeGeneralCategory(IDetailLayoutBuilder& InDetailLayout)
@@ -1724,31 +1740,27 @@ PRAGMA_ENABLE_DEPRECATION_WARNINGS
 				];
 			}
 
-			TArray<IDetailPropertyRow*> DefaultPropertyRows = CustomizeDefaultCategory(InDetailLayout);
-
-			if (bIsPreset && !bIsDefaultConstructed && !bIsTriggerDataType)
+			if (bIsPreset)
 			{
-				const UMetasoundEditorGraphInput* Input = Cast<UMetasoundEditorGraphInput>(MemberDefaultLiteral->FindMember());
-				if (Input)
+				if (!bIsDefaultConstructed && !bIsTriggerDataType)
 				{
-					auto PropertyEnabled = TAttribute<bool>::CreateLambda([this] { return !GetInputInheritsDefault(); });
-					for (IDetailPropertyRow* DefaultPropertyRow : DefaultPropertyRows)
+					const UMetasoundEditorGraphInput* Input = Cast<UMetasoundEditorGraphInput>(MemberDefaultLiteral->FindMember());
+					if (Input)
 					{
-						DefaultPropertyRow->EditCondition(PropertyEnabled, { });
-						FResetToDefaultOverride ResetOverride = FResetToDefaultOverride::Create(
+						Enabled = TAttribute<bool>::CreateLambda([this] { return !GetInputInheritsDefault(); });
+						ResetOverride = FResetToDefaultOverride::Create(
 							FIsResetToDefaultVisible::CreateLambda([this](TSharedPtr<IPropertyHandle> /* PropertyHandle */) { return !GetInputInheritsDefault(); }),
 							FResetToDefaultHandler::CreateLambda([this](TSharedPtr<IPropertyHandle> /* PropertyHandle */) { SetInputInheritsDefault(); }));
-						DefaultPropertyRow->OverrideResetToDefault(ResetOverride);
 					}
 				}
 			}
-			else if (!bIsPreset)
+			else
 			{
 				// Make default value uneditable while playing for constructor inputs
 				const UMetasoundEditorGraphInput* Input = Cast<UMetasoundEditorGraphInput>(MemberDefaultLiteral->FindMember());
 				if (Input)
 				{
-					auto PropertyEnabled = TAttribute<bool>::CreateLambda([this, Input]
+					Enabled = TAttribute<bool>::CreateLambda([this, Input]
 					{
 						if (Input->GetVertexAccessType() == EMetasoundFrontendVertexAccessType::Value)
 						{
@@ -1760,12 +1772,15 @@ PRAGMA_ENABLE_DEPRECATION_WARNINGS
 						}
 						return true;
 					});
-					for (IDetailPropertyRow* DefaultPropertyRow : DefaultPropertyRows)
-					{
-						DefaultPropertyRow->EditCondition(PropertyEnabled, { });
-					}
 				}
 			}
+
+			CustomizeDefaultCategory(InDetailLayout);
+		}
+
+		TAttribute<bool> FMetasoundInputDetailCustomization::GetEnabled() const
+		{
+			return Enabled;
 		}
 
 		bool FMetasoundInputDetailCustomization::IsDefaultEditable() const
diff --git a/Engine/Plugins/Runtime/Metasound/Source/MetasoundEditor/Private/MetasoundNodeDetailCustomization.h b/Engine/Plugins/Runtime/Metasound/Source/MetasoundEditor/Private/MetasoundNodeDetailCustomization.h
--- a/Engine/Plugins/Runtime/Metasound/Source/MetasoundEditor/Private/MetasoundNodeDetailCustomization.h
+++ b/Engine/Plugins/Runtime/Metasound/Source/MetasoundEditor/Private/MetasoundNodeDetailCustomization.h
@@ -185,7 +185,7 @@ namespace Metasound
 			}
 			virtual ~FMetasoundFloatLiteralCustomization();
 
-			virtual TArray<IDetailPropertyRow*> CustomizeLiteral(UMetasoundEditorGraphMemberDefaultLiteral& InLiteral, IDetailLayoutBuilder& InDetailLayout) override;
+			virtual void CustomizeDefaults(UMetasoundEditorGraphMemberDefaultLiteral& InLiteral, IDetailLayoutBuilder& InDetailLayout) override;
 		};
 
 		class FMetasoundBoolLiteralCustomization : public FMetasoundDefaultLiteralCustomizationBase
@@ -199,7 +199,7 @@ namespace Metasound
 			}
 			virtual ~FMetasoundBoolLiteralCustomization();
 
-			virtual TArray<IDetailPropertyRow*> CustomizeLiteral(UMetasoundEditorGraphMemberDefaultLiteral& InLiteral, IDetailLayoutBuilder& InDetailLayout) override;
+			virtual void CustomizeDefaults(UMetasoundEditorGraphMemberDefaultLiteral& InLiteral, IDetailLayoutBuilder& InDetailLayout) override;
 		};
 
 		// Customization to support drag-and-drop of Proxy UObject types on underlying members that are structs.
@@ -214,7 +214,7 @@ namespace Metasound
 
 			virtual ~FMetasoundObjectArrayLiteralCustomization() = default;
 
-			virtual TArray<IDetailPropertyRow*> CustomizeLiteral(UMetasoundEditorGraphMemberDefaultLiteral& InLiteral, IDetailLayoutBuilder& InDetailLayout) override;
+			virtual void CustomizeDefaults(UMetasoundEditorGraphMemberDefaultLiteral& InLiteral, IDetailLayoutBuilder& InDetailLayout) override;
 		};
 
 		class FMetasoundDefaultLiteralCustomizationFactory : public IMemberDefaultLiteralCustomizationFactory
@@ -235,7 +235,7 @@ namespace Metasound
 				return TUniquePtr<FMetasoundDefaultLiteralCustomizationBase>(new FMetasoundFloatLiteralCustomization(DefaultCategoryBuilder));
 			}
 		};
-		
+
 		// Customization to support bool widgets (ex. Buttons)
 		class FMetasoundBoolLiteralCustomizationFactory : public IMemberDefaultLiteralCustomizationFactory
 		{
@@ -337,11 +337,17 @@ namespace Metasound
 
 			void UpdateRenameDelegate(UMetasoundEditorGraphMember& InMember);
 			void CacheMemberData(IDetailLayoutBuilder& InDetailLayout);
+
+			virtual void CustomizeDefaultCategory(IDetailLayoutBuilder& InDetailLayout);
 			virtual void CustomizeGeneralCategory(IDetailLayoutBuilder& InDetailLayout);
-			virtual TArray<IDetailPropertyRow*> CustomizeDefaultCategory(IDetailLayoutBuilder& InDetailLayout);
 
 			virtual EVisibility GetDefaultVisibility() const { return EVisibility::Visible; }
+			virtual TAttribute<bool> GetEnabled() const { return { }; }
+			virtual const TOptional<FResetToDefaultOverride>& GetResetOverride() const { return ResetOverride; }
+
+			UE_DEPRECATED(5.5, "Use GetEnabled attribute instead")
 			virtual bool IsDefaultEditable() const { return true; }
+
 			virtual bool IsInterfaceMember() const { return false; }
 
 			// IDetailCustomization interface
@@ -363,9 +369,14 @@ namespace Metasound
 			TSharedPtr<SEditableTextBox> NameEditableTextBox;
 			FMetasoundDataTypeSelector DataTypeSelector;
 
+			TOptional<FResetToDefaultOverride> ResetOverride;
+
 			bool bIsNameInvalid = false;
 
 			FDelegateHandle RenameRequestedHandle;
+
+		private:
+			TUniquePtr<FMetasoundDefaultLiteralCustomizationBase> LiteralCustomization;
 		};
 
 		class FMetasoundVertexDetailCustomization : public FMetasoundMemberDetailCustomization
@@ -400,12 +411,15 @@ namespace Metasound
 			virtual ~FMetasoundInputDetailCustomization() = default;
 			
 			virtual void CustomizeDetails(IDetailLayoutBuilder& InDetailLayout) override;
+			virtual TAttribute<bool> GetEnabled() const override;
 			virtual bool IsDefaultEditable() const override;
 
 		private:
 			bool GetInputInheritsDefault() const;
 			void SetInputInheritsDefault();
 			void ClearInputInheritsDefault();
+
+			TAttribute<bool> Enabled;
 		};
 
 		class FMetasoundVariableDetailCustomization : public FMetasoundMemberDetailCustomization
diff --git a/Engine/Plugins/Runtime/Metasound/Source/MetasoundEditor/Private/SMetasoundGraphNode.cpp b/Engine/Plugins/Runtime/Metasound/Source/MetasoundEditor/Private/SMetasoundGraphNode.cpp
--- a/Engine/Plugins/Runtime/Metasound/Source/MetasoundEditor/Private/SMetasoundGraphNode.cpp
+++ b/Engine/Plugins/Runtime/Metasound/Source/MetasoundEditor/Private/SMetasoundGraphNode.cpp
@@ -18,6 +18,7 @@
 #include "KismetPins/SGraphPinObject.h"
 #include "KismetPins/SGraphPinString.h"
 #include "MetasoundBuilderSubsystem.h"
+#include "MetasoundDocumentBuilderRegistry.h"
 #include "MetasoundEditorGraph.h"
 #include "MetasoundEditorGraphBuilder.h"
 #include "MetasoundEditorGraphInputNode.h"
@@ -59,6 +60,7 @@
 #include "Widgets/SWidget.h"
 #include "Widgets/Text/SInlineEditableTextBlock.h"
 
+
 #define LOCTEXT_NAMESPACE "MetasoundEditor"
 
 namespace Metasound
@@ -276,17 +278,11 @@ namespace Metasound
 			];
 		}
 
-		UMetasoundEditorGraphNode& SMetaSoundGraphNode::GetMetaSoundNode()
+		UMetasoundEditorGraphNode& SMetaSoundGraphNode::GetMetaSoundNode() const
 		{
 			return *CastChecked<UMetasoundEditorGraphNode>(GraphNode);
 		}
 
-		const UMetasoundEditorGraphNode& SMetaSoundGraphNode::GetMetaSoundNode() const
-		{
-			check(GraphNode);
-			return *Cast<UMetasoundEditorGraphNode>(GraphNode);
-		}
-
 		TSharedPtr<SGraphPin> SMetaSoundGraphNode::CreatePinWidget(UEdGraphPin* InPin) const
 		{
 			using namespace Frontend;
@@ -749,19 +745,35 @@ namespace Metasound
 			return nullptr;
 		}
 
-		UMetasoundEditorGraphMemberNode* SMetaSoundGraphNode::GetMetaSoundMemberNode()
+		TAttribute<bool> SMetaSoundGraphNode::GetInputWidgetEnabled() const
+		{
+			return TAttribute<bool>::Create([this]()
+			{
+				if (const UMetasoundEditorGraphMemberNode* Node = GetMetaSoundMemberNode())
+				{
+					return Node->EnableInteractWidgets();
+				}
+				return false;
+			});
+		}
+
+		UMetasoundEditorGraphMemberNode* SMetaSoundGraphNode::GetMetaSoundMemberNode() const
 		{
 			return Cast<UMetasoundEditorGraphMemberNode>(&GetMetaSoundNode());
 		}
 
-		TSharedRef<SWidget> SMetaSoundGraphNode::CreateNodeContentArea()
+		TSharedPtr<SWidget> SMetaSoundGraphNode::CreateInputNodeContentArea(const FMetaSoundFrontendDocumentBuilder& InBuilder, TSharedRef<SHorizontalBox> ContentBox)
 		{
-			using namespace Frontend;
+			using namespace Engine;
 
-			FConstNodeHandle NodeHandle = GetMetaSoundNode().GetConstNodeHandle();
-			const FMetasoundFrontendClassStyleDisplay& StyleDisplay = NodeHandle->GetClassStyle().Display;
-			TSharedPtr<SHorizontalBox> ContentBox = SNew(SHorizontalBox);
-			TSharedPtr<SWidget> OuterContentBox; // currently only used for input float nodes to accommodate the input widget
+			TWeakObjectPtr<UMetasoundEditorGraphInput> GraphMember;
+			{
+				GraphMember = Cast<UMetasoundEditorGraphInput>(GetMetaSoundMember());
+				if (!GraphMember.IsValid())
+				{
+					return { };
+				}
+			}
 
 			bool bUseAudioMaterialWidgets = false;
 			if (const UMetasoundEditorSettings* EditorSettings = GetDefault<UMetasoundEditorSettings>())
@@ -769,195 +781,368 @@ namespace Metasound
 				bUseAudioMaterialWidgets = EditorSettings->bUseAudioMaterialWidgets;
 			}
 
-			// If input node, check if custom widget required
-			bool bShowContentWidget = false;
-			if (UMetasoundEditorGraphInput* GraphMember = Cast<UMetasoundEditorGraphInput>(GetMetaSoundMember()))
+			const UMetasoundEditorGraph* OwningGraph = GraphMember->GetOwningGraph();
+			if (!OwningGraph || !OwningGraph->IsEditable() || GraphMember->GetVertexAccessType() == EMetasoundFrontendVertexAccessType::Unset)
 			{
-				const UMetasoundEditorGraph* OwningGraph = GraphMember->GetOwningGraph();
-				if (OwningGraph && OwningGraph->IsEditable() && GraphMember->GetVertexAccessType() != EMetasoundFrontendVertexAccessType::Unset)
+				return { };
+			}
+
+			TSharedPtr<SWidget> OuterContentBox;
+
+			TWeakObjectPtr<UMetasoundEditorGraphMemberDefaultFloat> DefaultFloat = Cast<UMetasoundEditorGraphMemberDefaultFloat>(GraphMember->GetLiteral());
+			if (DefaultFloat.IsValid() && DefaultFloat->WidgetType != EMetasoundMemberDefaultWidget::None)
+			{
+				constexpr float WidgetPadding = 3.0f;
+				static const FVector2D SliderDesiredSizeVertical = FVector2D(30.0f, 250.0f);
+				static const FVector2D RadialSliderDesiredSize = FVector2D(56.0f, 87.0f);
+
+				auto OnValueChangedLambda = [DefaultFloat, GraphMember, this](float Value)
 				{
-					UMetasoundEditorGraphMemberDefaultFloat* DefaultFloat = Cast<UMetasoundEditorGraphMemberDefaultFloat>(GraphMember->GetLiteral());
-					if (DefaultFloat && DefaultFloat->WidgetType != EMetasoundMemberDefaultWidget::None)
+					if (DefaultFloat.IsValid() && FloatInputWidget.IsValid())
 					{
-						constexpr float WidgetPadding = 3.0f;
-						static const FVector2D SliderDesiredSizeVertical = FVector2D(30.0f, 250.0f);
-						static const FVector2D RadialSliderDesiredSize = FVector2D(56.0f, 87.0f);
+						if (!bIsInputWidgetTransacting)
+						{
+							GEditor->BeginTransaction(LOCTEXT("MetasoundGraphNode_MetasoundSetInputDefault", "Set MetaSound Input Default"));
+							bIsInputWidgetTransacting = true;
+						}
 
-						bShowContentWidget = true;
+						FMetaSoundFrontendDocumentBuilder& Builder = GraphMember->GetFrontendBuilderChecked();
+						Builder.CastDocumentObjectChecked<UObject>().Modify();
+						DefaultFloat->Modify();
 
-						auto OnValueChangedLambda = [DefaultFloat, GraphMember, this](float Value)
-						{
-							if (InputWidget.IsValid())
-							{
-								if (!bIsInputWidgetTransacting)
-								{
-									GEditor->BeginTransaction(LOCTEXT("MetasoundGraphNode_MetasoundSetInputDefault", "Set MetaSound Input Default"));
-									bIsInputWidgetTransacting = true;
-								}
-								GraphMember->GetOwningGraph()->GetMetasound()->Modify();
-								DefaultFloat->Modify();
+						FMetasoundFrontendLiteral OutputLiteral;
+						OutputLiteral.Set(FloatInputWidget->GetOutputValue(Value));
 
-								constexpr bool bPostTransaction = true;
-								float Output = InputWidget->GetOutputValue(Value);
-								DefaultFloat->SetDefault(Output);
-								GraphMember->UpdateFrontendDefaultLiteral(bPostTransaction);
-							}
-						};
+						const FGuid BuildPageID = Builder.GetBuildPageID();
+						DefaultFloat->SetFromLiteral(OutputLiteral, BuildPageID);
 
-						auto OnValueCommittedLambda = [DefaultFloat, GraphMember, this](float Value)
-						{
-							if (InputWidget.IsValid())
-							{
-								bool bPostTransaction = false;
-								float Output = InputWidget->GetOutputValue(Value);
-								DefaultFloat->SetDefault(Output);
+						constexpr bool bPostTransaction = true;
+						GraphMember->UpdateFrontendDefaultLiteral(bPostTransaction, &BuildPageID);
+					}
+				};
 
-								if (bIsInputWidgetTransacting)
-								{
-									GEditor->EndTransaction();
-									bIsInputWidgetTransacting = false;
-								}
-								else
-								{
-									bPostTransaction = true;
-									UE_LOG(LogMetaSound, Warning, TEXT("Unmatched MetaSound editor widget transaction."));
-								}
+				auto OnValueCommittedLambda = [DefaultFloat, GraphMember, this](float Value)
+				{
+					if (DefaultFloat.IsValid() && FloatInputWidget.IsValid())
+					{
+						bool bPostTransaction = false;
+						FMetasoundFrontendLiteral OutputLiteral;
+						OutputLiteral.Set(FloatInputWidget->GetOutputValue(Value));
+						const FGuid BuildPageID = GraphMember->GetFrontendBuilderChecked().GetBuildPageID();
+						DefaultFloat->SetFromLiteral(OutputLiteral, BuildPageID);
+						DefaultFloat->OnDefaultValueChanged.Broadcast(BuildPageID, Value);
+
+						if (bIsInputWidgetTransacting)
+						{
+							GEditor->EndTransaction();
+							bIsInputWidgetTransacting = false;
+						}
+						else
+						{
+							bPostTransaction = true;
+							UE_LOG(LogMetaSound, Warning, TEXT("Unmatched MetaSound editor widget transaction."));
+						}
 
-								GraphMember->UpdateFrontendDefaultLiteral(bPostTransaction);
+						GraphMember->UpdateFrontendDefaultLiteral(bPostTransaction);
 
-								if (UMetasoundEditorGraph* Graph = GraphMember->GetOwningGraph())
-								{
-									FGraphBuilder::GetOutermostMetaSoundChecked(*Graph).GetModifyContext().AddMemberIDsModified({ GraphMember->GetMemberID() });
-								}
-							}
-						};
+						if (UMetasoundEditorGraph* Graph = GraphMember->GetOwningGraph())
+						{
+							FGraphBuilder::GetOutermostMetaSoundChecked(*Graph).GetModifyContext().AddMemberIDsModified({ GraphMember->GetMemberID() });
+						}
+					}
+				};
 
-						if (DefaultFloat->WidgetType == EMetasoundMemberDefaultWidget::Slider)
+				if (DefaultFloat->WidgetType == EMetasoundMemberDefaultWidget::Slider)
+				{
+					if (bUseAudioMaterialWidgets)
+					{
+						const UMetasoundEditorSettings* EditorSettings = GetDefault<UMetasoundEditorSettings>();
+						check(EditorSettings)
+
+						SAssignNew(FloatInputWidget, SAudioMaterialLabeledSlider)
+							.Owner(GraphMember->GetOwningGraph())
+							.Style(EditorSettings->GetSliderStyle())
+							.AudioUnitsValueType(DefaultFloat->WidgetUnitValueType)
+							.bUseLinearOutput(DefaultFloat->VolumeWidgetUseLinearOutput)
+							.OnValueChanged_Lambda(OnValueChangedLambda)
+							.OnValueCommitted_Lambda(OnValueCommittedLambda);
+					}
+					else
+					{
+						// Create slider 
+						if (DefaultFloat->WidgetUnitValueType == EAudioUnitsValueType::Frequency)
+						{
+							SAssignNew(FloatInputWidget, SAudioFrequencySlider)
+								.OnValueChanged_Lambda(OnValueChangedLambda)
+								.OnValueCommitted_Lambda(OnValueCommittedLambda);
+						}
+						else if (DefaultFloat->WidgetUnitValueType == EAudioUnitsValueType::Volume)
 						{
-							if (bUseAudioMaterialWidgets)
+							SAssignNew(FloatInputWidget, SAudioVolumeSlider)
+								.OnValueChanged_Lambda(OnValueChangedLambda)
+								.OnValueCommitted_Lambda(OnValueCommittedLambda);
+							StaticCastSharedPtr<SAudioVolumeSlider>(FloatInputWidget)->SetUseLinearOutput(DefaultFloat->VolumeWidgetUseLinearOutput);
+						}
+						else
+						{
+							SAssignNew(FloatInputWidget, SAudioSlider)
+								.OnValueChanged_Lambda(OnValueChangedLambda)
+								.OnValueCommitted_Lambda(OnValueCommittedLambda);
+							FloatInputWidget->SetShowUnitsText(false);
+						}
+					}
+					// Slider layout 
+					if (DefaultFloat->WidgetOrientation == Orient_Vertical)
+					{
+						SAssignNew(OuterContentBox, SVerticalBox)
+							+ SVerticalBox::Slot()
+							.HAlign(HAlign_Right)
+							.VAlign(VAlign_Center)
+							.AutoHeight()
+							[
+								ContentBox
+							]
+						+ SVerticalBox::Slot()
+							.HAlign(HAlign_Fill)
+							.VAlign(VAlign_Top)
+							.Padding(WidgetPadding, 0.0f, WidgetPadding, WidgetPadding)
+							.AutoHeight()
+							[
+								FloatInputWidget.ToSharedRef()
+							];
+
+							if (!bUseAudioMaterialWidgets)
 							{
-								const UMetasoundEditorSettings* EditorSettings = GetDefault<UMetasoundEditorSettings>();
-								check(EditorSettings)
-
-								SAssignNew(InputWidget, SAudioMaterialLabeledSlider)
-									.Owner(GraphMember->GetOwningGraph())
-									.Style(EditorSettings->GetSliderStyle())
-									.AudioUnitsValueType(DefaultFloat->WidgetUnitValueType)
-									.bUseLinearOutput(DefaultFloat->VolumeWidgetUseLinearOutput)
-									.OnValueChanged_Lambda(OnValueChangedLambda)
-									.OnValueCommitted_Lambda(OnValueCommittedLambda);
+								FloatInputWidget->SetDesiredSizeOverride(SliderDesiredSizeVertical);
 							}
-							else
-							{
-								// Create slider 
-								if (DefaultFloat->WidgetUnitValueType == EAudioUnitsValueType::Frequency)
-								{
-									SAssignNew(InputWidget, SAudioFrequencySlider)
-										.OnValueChanged_Lambda(OnValueChangedLambda)
-										.OnValueCommitted_Lambda(OnValueCommittedLambda);
-								}
-								else if (DefaultFloat->WidgetUnitValueType == EAudioUnitsValueType::Volume)
-								{
-									SAssignNew(InputWidget, SAudioVolumeSlider)
-										.OnValueChanged_Lambda(OnValueChangedLambda)
-										.OnValueCommitted_Lambda(OnValueCommittedLambda);
-									StaticCastSharedPtr<SAudioVolumeSlider>(InputWidget)->SetUseLinearOutput(DefaultFloat->VolumeWidgetUseLinearOutput);
-								}
-								else
-								{
-									SAssignNew(InputWidget, SAudioSlider)
-										.OnValueChanged_Lambda(OnValueChangedLambda)
-										.OnValueCommitted_Lambda(OnValueCommittedLambda);
-									InputWidget->SetShowUnitsText(false);
-								}
-							}
-							// Slider layout 
-							if (DefaultFloat->WidgetOrientation == Orient_Vertical)
+					}
+					else // horizontal orientation
+					{
+						UMetasoundEditorGraphMemberNode* MemberNode = GetMetaSoundMemberNode();
+						TSharedPtr<SWidget> Slot1;
+						TSharedPtr<SWidget> Slot2;
+						if (MemberNode->IsA<UMetasoundEditorGraphInputNode>())
+						{
+							Slot1 = FloatInputWidget;
+							Slot2 = ContentBox;
+						}
+						else
+						{
+							Slot1 = ContentBox;
+							Slot2 = FloatInputWidget;
+						}
+
+						SAssignNew(OuterContentBox, SHorizontalBox)
+							+ SHorizontalBox::Slot()
+							.HAlign(HAlign_Fill)
+							.VAlign(VAlign_Center)
+							.Padding(WidgetPadding, 0.0f, WidgetPadding, 0.0f)
+							.AutoWidth()
+							[
+								Slot1.ToSharedRef()
+							]
+						+ SHorizontalBox::Slot()
+							.HAlign(HAlign_Center)
+							.VAlign(VAlign_Fill)
+							.AutoWidth()
+							[
+								Slot2.ToSharedRef()
+							];
+
+							if (!bUseAudioMaterialWidgets)
 							{
-								SAssignNew(OuterContentBox, SVerticalBox)
-									+ SVerticalBox::Slot()
-									.HAlign(HAlign_Right)
-									.VAlign(VAlign_Center)
-									.AutoHeight()
-									[
-										ContentBox.ToSharedRef()
-									]
-								+ SVerticalBox::Slot()
-									.HAlign(HAlign_Fill)
-									.VAlign(VAlign_Top)
-									.Padding(WidgetPadding, 0.0f, WidgetPadding, WidgetPadding)
-									.AutoHeight()
-									[
-										InputWidget.ToSharedRef()
-									];
-
-									if (!bUseAudioMaterialWidgets)
-									{
-										InputWidget->SetDesiredSizeOverride(SliderDesiredSizeVertical);
-									}
+								FloatInputWidget->SetDesiredSizeOverride(FVector2D(SliderDesiredSizeVertical.Y, SliderDesiredSizeVertical.X));
 							}
-							else // horizontal orientation
-							{
-								UMetasoundEditorGraphMemberNode* MemberNode = GetMetaSoundMemberNode();
-								TSharedPtr<SWidget> Slot1;
-								TSharedPtr<SWidget> Slot2;
-								if (MemberNode->IsA<UMetasoundEditorGraphInputNode>())
-								{
-									Slot1 = InputWidget;
-									Slot2 = ContentBox;
-								}
-								else
-								{
-									Slot1 = ContentBox;
-									Slot2 = InputWidget;
-								}
+						}
 
-								SAssignNew(OuterContentBox, SHorizontalBox)
-									+ SHorizontalBox::Slot()
-									.HAlign(HAlign_Fill)
-									.VAlign(VAlign_Center)
-									.Padding(WidgetPadding, 0.0f, WidgetPadding, 0.0f)
-									.AutoWidth()
-									[
-										Slot1.ToSharedRef()
-									]
-								+ SHorizontalBox::Slot()
-									.HAlign(HAlign_Center)
-									.VAlign(VAlign_Fill)
-									.AutoWidth()
-									[
-										Slot2.ToSharedRef()
-									];
-
-									if (!bUseAudioMaterialWidgets)
-									{
-										InputWidget->SetDesiredSizeOverride(FVector2D(SliderDesiredSizeVertical.Y, SliderDesiredSizeVertical.X));
-									}
-								}
+						if (bUseAudioMaterialWidgets)
+						{
+							// safe downcast because the ptr was just assigned above 
+							StaticCastSharedPtr<SAudioMaterialLabeledSlider>(FloatInputWidget)->SetOrientation(DefaultFloat->WidgetOrientation);
+						}
+						else
+						{
+							// safe downcast because the ptr was just assigned above 
+							StaticCastSharedPtr<SAudioSliderBase>(FloatInputWidget)->SetOrientation(DefaultFloat->WidgetOrientation);
+						}
+				}
+				else if (DefaultFloat->WidgetType == EMetasoundMemberDefaultWidget::RadialSlider)
+				{
+					auto OnRadialSliderMouseCaptureBeginLambda = [this]()
+					{
+						if (!bIsInputWidgetTransacting)
+						{
+							GEditor->BeginTransaction(LOCTEXT("MetasoundSetRadialSliderInputDefault", "Set MetaSound Input Default"));
+							bIsInputWidgetTransacting = true;
+						}
+					};
 
-								if (bUseAudioMaterialWidgets)
-								{
-									// safe downcast because the ptr was just assigned above 
-									StaticCastSharedPtr<SAudioMaterialLabeledSlider>(InputWidget)->SetOrientation(DefaultFloat->WidgetOrientation);
-								}
-								else
-								{
-									// safe downcast because the ptr was just assigned above 
-									StaticCastSharedPtr<SAudioSliderBase>(InputWidget)->SetOrientation(DefaultFloat->WidgetOrientation);
-								}
+					auto OnRadialSliderMouseCaptureEndLambda = [this]()
+					{
+						if (bIsInputWidgetTransacting)
+						{
+							GEditor->EndTransaction();
+							bIsInputWidgetTransacting = false;
+						}
+						else
+						{
+							UE_LOG(LogMetaSound, Warning, TEXT("Unmatched MetaSound editor widget transaction."));
+						}
+					};
+
+					if (bUseAudioMaterialWidgets)
+					{
+						const UMetasoundEditorSettings* EditorSettings = GetDefault<UMetasoundEditorSettings>();
+						check(EditorSettings)
+
+						SAssignNew(FloatInputWidget, SAudioMaterialLabeledKnob)
+							.Owner(GraphMember->GetOwningGraph())
+							.Style(EditorSettings->GetKnobStyle())
+							.OnValueChanged_Lambda(OnValueChangedLambda)
+							.AudioUnitsValueType(DefaultFloat->WidgetUnitValueType)
+							.bUseLinearOutput(DefaultFloat->VolumeWidgetUseLinearOutput)
+							.OnMouseCaptureBegin_Lambda(OnRadialSliderMouseCaptureBeginLambda)
+							.OnMouseCaptureEnd_Lambda(OnRadialSliderMouseCaptureEndLambda);
+					}
+					else
+					{
+						// Create slider 
+						if (DefaultFloat->WidgetUnitValueType == EAudioUnitsValueType::Frequency)
+						{
+							SAssignNew(FloatInputWidget, SAudioFrequencyRadialSlider)
+								.OnValueChanged_Lambda(OnValueChangedLambda)
+								.OnMouseCaptureBegin_Lambda(OnRadialSliderMouseCaptureBeginLambda)
+								.OnMouseCaptureEnd_Lambda(OnRadialSliderMouseCaptureEndLambda);
+						}
+						else if (DefaultFloat->WidgetUnitValueType == EAudioUnitsValueType::Volume)
+						{
+							SAssignNew(FloatInputWidget, SAudioVolumeRadialSlider)
+								.OnValueChanged_Lambda(OnValueChangedLambda)
+								.OnMouseCaptureBegin_Lambda(OnRadialSliderMouseCaptureBeginLambda)
+								.OnMouseCaptureEnd_Lambda(OnRadialSliderMouseCaptureEndLambda);
+							StaticCastSharedPtr<SAudioVolumeRadialSlider>(FloatInputWidget)->SetUseLinearOutput(DefaultFloat->VolumeWidgetUseLinearOutput);
+						}
+						else
+						{
+							SAssignNew(FloatInputWidget, SAudioRadialSlider)
+								.OnValueChanged_Lambda(OnValueChangedLambda)
+								.OnMouseCaptureBegin_Lambda(OnRadialSliderMouseCaptureBeginLambda)
+								.OnMouseCaptureEnd_Lambda(OnRadialSliderMouseCaptureEndLambda);
+							FloatInputWidget->SetShowUnitsText(false);
+						}
+					}
+					// Only vertical layout for radial slider
+					SAssignNew(OuterContentBox, SVerticalBox)
+						+ SVerticalBox::Slot()
+						.HAlign(HAlign_Right)
+						.VAlign(VAlign_Center)
+						.AutoHeight()
+						[
+							ContentBox
+						]
+					+ SVerticalBox::Slot()
+						.HAlign(HAlign_Fill)
+						.VAlign(VAlign_Top)
+						.Padding(WidgetPadding, 0.0f, WidgetPadding, WidgetPadding)
+						.AutoHeight()
+						[
+							FloatInputWidget.ToSharedRef()
+						];
+						FloatInputWidget->SetDesiredSizeOverride(RadialSliderDesiredSize);
+				}
+
+				const FMetasoundFrontendClassInput* ClassInput = InBuilder.FindGraphInput(GraphMember->GetMemberName());
+				FGuid ResolvedPageID = Frontend::DefaultPageID;
+				if (ensure(ClassInput))
+				{
+					ResolvedPageID = Engine::FDocumentBuilderRegistry::GetChecked().ResolveTargetPageID(*ClassInput);
+				}
+
+				FloatInputWidget->SetOutputRange(DefaultFloat->GetRange());
+				FloatInputWidget->SetUnitsTextReadOnly(true);
+				FloatInputWidget->SetSliderValue(FloatInputWidget->GetSliderValue(DefaultFloat->GetDefaultAs<float>(ResolvedPageID)));
+				FloatInputWidget->SetEnabled(GetInputWidgetEnabled());
+
+				// Setup & clear delegate if necessary (ex. if was just saved)
+				if (InputSliderOnValueChangedDelegateHandle.IsValid())
+				{
+					DefaultFloat->OnDefaultValueChanged.Remove(InputSliderOnValueChangedDelegateHandle);
+					InputSliderOnValueChangedDelegateHandle.Reset();
+				}
+
+				InputSliderOnValueChangedDelegateHandle = DefaultFloat->OnDefaultValueChanged.AddLambda([this, Widget = FloatInputWidget](const FGuid& PageID, float Value)
+				{
+					using namespace Metasound::Engine;
+
+					if (Widget.IsValid())
+					{
+						UMetasoundEditorGraphNode& EdNode = GetMetaSoundNode();
+						UObject& MetaSound = EdNode.GetMetasoundChecked();
+						UMetaSoundBuilderBase& Builder = FDocumentBuilderRegistry::GetChecked().FindOrBeginBuilding(MetaSound);
+
+						if (Builder.GetConstBuilder().GetBuildPageID() == PageID)
+						{
+							const float SliderValue = Widget->GetSliderValue(Value);
+							Widget->SetSliderValue(SliderValue);
 						}
-						else if (DefaultFloat->WidgetType == EMetasoundMemberDefaultWidget::RadialSlider)
+					}
+				});
+
+				if (InputSliderOnRangeChangedDelegateHandle.IsValid())
+				{
+					DefaultFloat->OnRangeChanged.Remove(InputSliderOnRangeChangedDelegateHandle);
+					InputSliderOnRangeChangedDelegateHandle.Reset();
+				}
+
+				InputSliderOnRangeChangedDelegateHandle = DefaultFloat->OnRangeChanged.AddLambda([Widget = FloatInputWidget](FVector2D Range)
+				{
+					if (Widget.IsValid())
+					{
+						Widget->SetOutputRange(Range);
+					}
+				});
+			}
+			else
+			{
+				TWeakObjectPtr<UMetasoundEditorGraphMemberDefaultBool> DefaultBool = Cast<UMetasoundEditorGraphMemberDefaultBool>(GraphMember->GetLiteral());
+				if (DefaultBool.IsValid())
+				{
+					if (bUseAudioMaterialWidgets)
+					{
+						bool bIsNotTriggerNode = GraphMember->GetDataType() != Metasound::GetMetasoundDataTypeName<Metasound::FTrigger>();
+
+						if ((bIsNotTriggerNode && DefaultBool.IsValid()) && DefaultBool->WidgetType != EMetasoundBoolMemberDefaultWidget::None)
 						{
-							auto OnRadialSliderMouseCaptureBeginLambda = [this]()
+							constexpr float WidgetPadding = 3.0f;
+							static const FVector2D ButtonDesiredSize = FVector2D(56.0f, 87.0f);
+
+							auto OnBoolValueChangedLambda = [DefaultBool, GraphMember, this](bool Value)
 							{
-								if (!bIsInputWidgetTransacting)
+								if (DefaultBool.IsValid() && MaterialButtonWidget.IsValid())
 								{
-									GEditor->BeginTransaction(LOCTEXT("MetasoundSetRadialSliderInputDefault", "Set MetaSound Input Default"));
-									bIsInputWidgetTransacting = true;
+									if (!bIsInputWidgetTransacting)
+									{
+										GEditor->BeginTransaction(FText::Format(LOCTEXT("MetasoundGraphNode_MetasoundSetBoolInput", "Set MetaSound Input '{0}' Boolean Default"), GraphMember->GetDisplayName()));
+										bIsInputWidgetTransacting = true;
+									}
+
+									FMetaSoundFrontendDocumentBuilder& Builder = GraphMember->GetFrontendBuilderChecked();
+									Builder.CastDocumentObjectChecked<UObject>().Modify();
+									DefaultBool->Modify();
+
+									FMetasoundFrontendLiteral OutputLiteral;
+									OutputLiteral.Set(FloatInputWidget->GetOutputValue(Value));
+
+									const FGuid BuildPageID = Builder.GetBuildPageID();
+									DefaultBool->SetFromLiteral(OutputLiteral, BuildPageID);
+									DefaultBool->OnDefaultStateChanged.Broadcast(Value, BuildPageID);
+
+									constexpr bool bPostTransaction = true;
+									GraphMember->UpdateFrontendDefaultLiteral(bPostTransaction, &BuildPageID);
 								}
 							};
 
-							auto OnRadialSliderMouseCaptureEndLambda = [this]()
+							auto OnBoolMouseCaptureEndLambda = [this]()
 							{
 								if (bIsInputWidgetTransacting)
 								{
@@ -970,207 +1155,75 @@ namespace Metasound
 								}
 							};
 
-							if (bUseAudioMaterialWidgets)
-							{
-								const UMetasoundEditorSettings* EditorSettings = GetDefault<UMetasoundEditorSettings>();
-								check(EditorSettings)
-
-								SAssignNew(InputWidget, SAudioMaterialLabeledKnob)
-									.Owner(GraphMember->GetOwningGraph())
-									.Style(EditorSettings->GetKnobStyle())
-									.OnValueChanged_Lambda(OnValueChangedLambda)
-									.AudioUnitsValueType(DefaultFloat->WidgetUnitValueType)
-									.bUseLinearOutput(DefaultFloat->VolumeWidgetUseLinearOutput)
-									.OnMouseCaptureBegin_Lambda(OnRadialSliderMouseCaptureBeginLambda)
-									.OnMouseCaptureEnd_Lambda(OnRadialSliderMouseCaptureEndLambda);
-							}
-							else
-							{
-								// Create slider 
-								if (DefaultFloat->WidgetUnitValueType == EAudioUnitsValueType::Frequency)
-								{
-									SAssignNew(InputWidget, SAudioFrequencyRadialSlider)
-										.OnValueChanged_Lambda(OnValueChangedLambda)
-										.OnMouseCaptureBegin_Lambda(OnRadialSliderMouseCaptureBeginLambda)
-										.OnMouseCaptureEnd_Lambda(OnRadialSliderMouseCaptureEndLambda);
-								}
-								else if (DefaultFloat->WidgetUnitValueType == EAudioUnitsValueType::Volume)
-								{
-									SAssignNew(InputWidget, SAudioVolumeRadialSlider)
-										.OnValueChanged_Lambda(OnValueChangedLambda)
-										.OnMouseCaptureBegin_Lambda(OnRadialSliderMouseCaptureBeginLambda)
-										.OnMouseCaptureEnd_Lambda(OnRadialSliderMouseCaptureEndLambda);
-									StaticCastSharedPtr<SAudioVolumeRadialSlider>(InputWidget)->SetUseLinearOutput(DefaultFloat->VolumeWidgetUseLinearOutput);
-								}
-								else
-								{
-									SAssignNew(InputWidget, SAudioRadialSlider)
-										.OnValueChanged_Lambda(OnValueChangedLambda)
-										.OnMouseCaptureBegin_Lambda(OnRadialSliderMouseCaptureBeginLambda)
-										.OnMouseCaptureEnd_Lambda(OnRadialSliderMouseCaptureEndLambda);
-									InputWidget->SetShowUnitsText(false);
-								}
-							}
-							// Only vertical layout for radial slider
-							SAssignNew(OuterContentBox, SVerticalBox)
-								+ SVerticalBox::Slot()
-								.HAlign(HAlign_Right)
-								.VAlign(VAlign_Center)
-								.AutoHeight()
-								[
-									ContentBox.ToSharedRef()
-								]
-							+ SVerticalBox::Slot()
-								.HAlign(HAlign_Fill)
-								.VAlign(VAlign_Top)
-								.Padding(WidgetPadding, 0.0f, WidgetPadding, WidgetPadding)
-								.AutoHeight()
-								[
-									InputWidget.ToSharedRef()
-								];
-								InputWidget->SetDesiredSizeOverride(RadialSliderDesiredSize);
-						}
-
-						InputWidget->SetOutputRange(DefaultFloat->GetRange());
-						InputWidget->SetUnitsTextReadOnly(true);
-						InputWidget->SetSliderValue(InputWidget->GetSliderValue(DefaultFloat->GetDefault()));
-						InputWidget->SetEnabled(TAttribute<bool>::Create([this]()
-						{
-							if (UMetasoundEditorGraphMemberNode* Node = GetMetaSoundMemberNode())
-							{
-								return Node->EnableInteractWidgets();
-							}
-							return false;
-						}));
+							const UMetasoundEditorSettings* EditorSettings = GetDefault<UMetasoundEditorSettings>();
+							check(EditorSettings);
 
-						// Setup & clear delegate if necessary (ex. if was just saved)
-						if (InputSliderOnValueChangedDelegateHandle.IsValid())
-						{
-							DefaultFloat->OnDefaultValueChanged.Remove(InputSliderOnValueChangedDelegateHandle);
-							InputSliderOnValueChangedDelegateHandle.Reset();
-						}
+							SAssignNew(MaterialButtonWidget, SAudioMaterialButton)
+								.AudioMaterialButtonStyle(EditorSettings->GetButtonStyle())
+								.OnBooleanValueChanged_Lambda(OnBoolValueChangedLambda)
+								.OnMouseCaptureEnd_Lambda(OnBoolMouseCaptureEndLambda)
+								.bIsPressedAttribute(DefaultBool->GetDefaultAs<bool>());
 
-						InputSliderOnValueChangedDelegateHandle = DefaultFloat->OnDefaultValueChanged.AddLambda([Widget = InputWidget](float Value)
-						{
-							if (Widget.IsValid())
-							{
-								const float SliderValue = Widget->GetSliderValue(Value);
-								Widget->SetSliderValue(SliderValue);
-							}
-						});
+							SAssignNew(OuterContentBox, SVerticalBox)
+							+ SVerticalBox::Slot()
+							.HAlign(HAlign_Right)
+							.VAlign(VAlign_Center)
+							.AutoHeight()
+							[
+								ContentBox
+							]
+							+ SVerticalBox::Slot()
+							.HAlign(HAlign_Fill)
+							.VAlign(VAlign_Top)
+							.Padding(WidgetPadding, 0.0f, WidgetPadding, WidgetPadding)
+							.AutoHeight()
+							[
+								MaterialButtonWidget.ToSharedRef()
+							];
 
-						if (InputSliderOnRangeChangedDelegateHandle.IsValid())
-						{
-							DefaultFloat->OnRangeChanged.Remove(InputSliderOnRangeChangedDelegateHandle);
-							InputSliderOnRangeChangedDelegateHandle.Reset();
-						}
+							MaterialButtonWidget->SetDesiredSizeOverride(ButtonDesiredSize);
+							MaterialButtonWidget->SetEnabled(GetInputWidgetEnabled());
 
-						InputSliderOnRangeChangedDelegateHandle = DefaultFloat->OnRangeChanged.AddLambda([Widget = InputWidget](FVector2D Range)
-						{
-							if (Widget.IsValid())
+							// Setup & clear delegate if necessary (ex. if was just saved)
+							if (InputButtonOnStateChangedDelegateHandle.IsValid())
 							{
-								Widget->SetOutputRange(Range);
+								DefaultBool->OnDefaultStateChanged.Remove(InputButtonOnStateChangedDelegateHandle);
+								InputButtonOnStateChangedDelegateHandle.Reset();
 							}
-						});
-					}
-					else if (UMetasoundEditorGraphMemberDefaultBool* DefaultBool = Cast<UMetasoundEditorGraphMemberDefaultBool>(GraphMember->GetLiteral()))
-					{
-						if (bUseAudioMaterialWidgets)
-						{
-							bool bIsNotTriggerNode = GraphMember->GetDataType() != Metasound::GetMetasoundDataTypeName<Metasound::FTrigger>();
 
-							if ((bIsNotTriggerNode && IsValid(DefaultBool)) && DefaultBool->WidgetType != EMetasoundBoolMemberDefaultWidget::None)
+							InputButtonOnStateChangedDelegateHandle = DefaultBool->OnDefaultStateChanged.AddLambda([this, Widget = MaterialButtonWidget](bool bValue, const FGuid& InPageID)
 							{
-								bShowContentWidget = true;
-								constexpr float WidgetPadding = 3.0f;
-								static const FVector2D ButtonDesiredSize = FVector2D(56.0f, 87.0f);
-
-								auto OnboolValueChangedLambda = [DefaultBool, GraphMember, this](bool Value)
-									{
-										if (MaterialButtonWidget.IsValid())
-										{
-											if (!bIsInputWidgetTransacting)
-											{
-												GEditor->BeginTransaction(LOCTEXT("MetasoundGraphNode_MetasoundSetInputDefault", "Set MetaSound Input Default"));
-												bIsInputWidgetTransacting = true;
-											}
-											GraphMember->GetOwningGraph()->GetMetasound()->Modify();
-											DefaultBool->Modify();
-
-											constexpr bool bPostTransaction = true;
-											DefaultBool->SetDefault(Value);
-											GraphMember->UpdateFrontendDefaultLiteral(bPostTransaction);
-										}
-									};
-
-								auto OnboolMouseCaptureEndLambda = [this]()
-									{
-										if (bIsInputWidgetTransacting)
-										{
-											GEditor->EndTransaction();
-											bIsInputWidgetTransacting = false;
-										}
-										else
-										{
-											UE_LOG(LogMetaSound, Warning, TEXT("Unmatched MetaSound editor widget transaction."));
-										}
-									};
-
-								const UMetasoundEditorSettings* EditorSettings = GetDefault<UMetasoundEditorSettings>();
-								check(EditorSettings)
-
-									SAssignNew(MaterialButtonWidget, SAudioMaterialButton)
-									.AudioMaterialButtonStyle(EditorSettings->GetButtonStyle())
-									.OnBooleanValueChanged_Lambda(OnboolValueChangedLambda)
-									.OnMouseCaptureEnd_Lambda(OnboolMouseCaptureEndLambda)
-									.bIsPressedAttribute(DefaultBool->GetDefault());
-
-								SAssignNew(OuterContentBox, SVerticalBox)
-									+ SVerticalBox::Slot()
-									.HAlign(HAlign_Right)
-									.VAlign(VAlign_Center)
-									.AutoHeight()
-									[
-										ContentBox.ToSharedRef()
-									]
-									+ SVerticalBox::Slot()
-									.HAlign(HAlign_Fill)
-									.VAlign(VAlign_Top)
-									.Padding(WidgetPadding, 0.0f, WidgetPadding, WidgetPadding)
-									.AutoHeight()
-									[
-										MaterialButtonWidget.ToSharedRef()										
-									];
-
-								MaterialButtonWidget->SetDesiredSizeOverride(ButtonDesiredSize);
-								MaterialButtonWidget->SetEnabled(TAttribute<bool>::Create([this]()
-									{
-										if (UMetasoundEditorGraphMemberNode* Node = GetMetaSoundMemberNode())
-										{
-											return Node->EnableInteractWidgets();
-										}
-										return false;
-									}));
-
-								// Setup & clear delegate if necessary (ex. if was just saved)
-								if (InputButtonOnStateChangedDelegateHandle.IsValid())
+								UMetasoundEditorGraphNode& EdNode = GetMetaSoundNode();
+								UObject& MetaSound = EdNode.GetMetasoundChecked();
+								UMetaSoundBuilderBase& Builder = FDocumentBuilderRegistry::GetChecked().FindOrBeginBuilding(MetaSound);
+								const FGuid& BuildPageID = Builder.GetConstBuilder().GetBuildPageID();
+								if (Widget.IsValid() && BuildPageID == InPageID)
 								{
-									DefaultBool->OnDefaultStateChanged.Remove(InputButtonOnStateChangedDelegateHandle);
-									InputButtonOnStateChangedDelegateHandle.Reset();
+									Widget->SetPressedState(bValue);
 								}
-
-								InputButtonOnStateChangedDelegateHandle = DefaultBool->OnDefaultStateChanged.AddLambda([Widget = MaterialButtonWidget](bool Value)
-									{
-										if (Widget.IsValid())
-										{
-											Widget->SetPressedState(Value);
-										}
-									});
-							}
+							});
 						}
 					}
 				}
 			}
+
+			return OuterContentBox;
+		}
+
+		TSharedRef<SWidget> SMetaSoundGraphNode::CreateNodeContentArea()
+		{
+			using namespace Engine;
+			using namespace Frontend;
+
+			FConstNodeHandle NodeHandle = GetMetaSoundNode().GetConstNodeHandle();
+			const FMetasoundFrontendClassStyleDisplay& StyleDisplay = NodeHandle->GetClassStyle().Display;
+			TSharedRef<SHorizontalBox> ContentBox = SNew(SHorizontalBox);
+
+			UMetasoundEditorGraphNode& EdNode = GetMetaSoundNode();
+			UObject& MetaSound = EdNode.GetMetasoundChecked();
+			UMetaSoundBuilderBase& Builder = FDocumentBuilderRegistry::GetChecked().FindOrBeginBuilding(MetaSound);
+
+			TSharedPtr<SWidget> InputContent = CreateInputNodeContentArea(Builder.GetConstBuilder(), ContentBox);
 	
 			// Gives more space for user to grab a bit easier as variables do not have any title area nor icon
 			const float GrabPadding = IsVariableMutator() ? 28.0f : 0.0f;
@@ -1216,7 +1269,7 @@ namespace Metasound
 				.VAlign(VAlign_Fill)
 				.Padding(FMargin(0,3))
 				[
-					(bShowContentWidget ? OuterContentBox : ContentBox).ToSharedRef()
+					InputContent.IsValid() ? InputContent.ToSharedRef() : ContentBox
 				];
 		}
 
diff --git a/Engine/Plugins/Runtime/Metasound/Source/MetasoundEditor/Private/SMetasoundGraphNode.h b/Engine/Plugins/Runtime/Metasound/Source/MetasoundEditor/Private/SMetasoundGraphNode.h
--- a/Engine/Plugins/Runtime/Metasound/Source/MetasoundEditor/Private/SMetasoundGraphNode.h
+++ b/Engine/Plugins/Runtime/Metasound/Source/MetasoundEditor/Private/SMetasoundGraphNode.h
@@ -64,26 +64,31 @@ namespace Metasound::Editor
 		FLinearColor GetNodeTitleColorOverride() const;
 
 		FName GetLiteralDataType() const;
-		UMetasoundEditorGraphNode& GetMetaSoundNode();
-		const UMetasoundEditorGraphNode& GetMetaSoundNode() const;
+		UMetasoundEditorGraphNode& GetMetaSoundNode() const;
 
 	public:
 		static void ExecuteTrigger(UMetasoundEditorGraphMemberDefaultLiteral& Literal);
 		static TSharedRef<SWidget> CreateTriggerSimulationWidget(UMetasoundEditorGraphMemberDefaultLiteral& Literal, TAttribute<EVisibility>&& InVisibility, TAttribute<bool>&& InEnablement, const FText* InToolTip = nullptr);
 
 	private:
+		// Returns attribute that returns whether or not the input widget should be enabled.
+		TAttribute<bool> GetInputWidgetEnabled() const;
+
 		// If this node represents a graph member node, returns corresponding member.
 		UMetasoundEditorGraphMember* GetMetaSoundMember();
 
 		// If this node represents a graph member node, returns cast node.
-		UMetasoundEditorGraphMemberNode* GetMetaSoundMemberNode();
+		UMetasoundEditorGraphMemberNode* GetMetaSoundMemberNode() const;
 
 		TAttribute<EVisibility> GetSimulationVisibilityAttribute() const;
 
+		// Input node should be moved to own implementation
+		TSharedPtr<SWidget> CreateInputNodeContentArea(const FMetaSoundFrontendDocumentBuilder& InBuilder, TSharedRef<SHorizontalBox> ContentBox);
+
 		// Slider widget for float input
-		TSharedPtr<SAudioInputWidget> InputWidget;
+		TSharedPtr<SAudioInputWidget> FloatInputWidget;
 
-		//Button Widget for bool input.
+		// Button Widget for bool input.
 		TSharedPtr<SAudioMaterialButton> MaterialButtonWidget;
 
 		// Handle for on state changed delegate for Button 
diff --git a/Engine/Plugins/Runtime/Metasound/Source/MetasoundEditor/Public/MetasoundDefaultLiteralCustomization.h b/Engine/Plugins/Runtime/Metasound/Source/MetasoundEditor/Public/MetasoundDefaultLiteralCustomization.h
new file mode 100644
--- /dev/null
+++ b/Engine/Plugins/Runtime/Metasound/Source/MetasoundEditor/Public/MetasoundDefaultLiteralCustomization.h
@@ -0,0 +1,75 @@
+// Copyright Epic Games, Inc. All Rights Reserved.
+#pragma once
+
+#include "DetailCategoryBuilder.h"
+#include "IDetailPropertyRow.h"
+#include "MetasoundFrontendDataTypeRegistry.h"
+#include "MetasoundFrontendLiteral.h"
+#include "MetasoundFrontendRegistries.h"
+#include "Modules/ModuleInterface.h"
+#include "PropertyHandle.h"
+#include "Templates/Function.h"
+
+// Forward Declarations
+class IDetailPropertyRow;
+class IPropertyHandle;
+class SSearchableComboBox;
+class UEdGraphPin;
+class UMetasoundEditorGraph;
+class UMetasoundEditorGraphMemberDefaultLiteral;
+class UMetasoundEditorGraphNode;
+
+
+namespace Metasound::Editor
+{
+	class METASOUNDEDITOR_API FMetasoundDefaultLiteralCustomizationBase
+	{
+	protected:
+		IDetailCategoryBuilder* DefaultCategoryBuilder = nullptr;
+
+	public:
+		using FOnDefaultPageRowAdded = TFunction<void(IDetailPropertyRow& /* Value Row */, TSharedRef<IPropertyHandle> /* Page Property Handle */)>;
+
+		FMetasoundDefaultLiteralCustomizationBase(IDetailCategoryBuilder& InDefaultCategoryBuilder);
+		virtual ~FMetasoundDefaultLiteralCustomizationBase();
+
+		virtual void CustomizeDefaults(UMetasoundEditorGraphMemberDefaultLiteral& InLiteral, IDetailLayoutBuilder& InDetailLayout);
+
+		UE_DEPRECATED(5.5, "Use CustomizeDefaults instead and provide returned customized handles")
+		virtual TArray<IDetailPropertyRow*> CustomizeLiteral(UMetasoundEditorGraphMemberDefaultLiteral& InLiteral, IDetailLayoutBuilder& InDetailLayout);
+
+		virtual TAttribute<EVisibility> GetDefaultVisibility() const;
+		virtual TAttribute<bool> GetEnabled() const;
+
+		virtual void SetDefaultVisibility(TAttribute<EVisibility> VisibilityAttribute);
+		virtual void SetEnabled(TAttribute<bool> EnableAttribute);
+		virtual void SetResetOverride(const TOptional<FResetToDefaultOverride>& InResetOverride);
+
+	protected:
+		void CustomizePageDefaultRows(UMetasoundEditorGraphMemberDefaultLiteral& InLiteral, IDetailLayoutBuilder& InDetailLayout, FOnDefaultPageRowAdded* OnPageRowAdded = nullptr);
+
+		TArray<TSharedPtr<IPropertyHandle>> DefaultProperties;
+
+	private:
+		TSharedRef<SWidget> BuildPageDefaultNameWidget(UMetasoundEditorGraphMemberDefaultLiteral& Literal, TSharedRef<IPropertyHandle> ElementProperty);
+		void BuildPageDefaultComboBox(UMetasoundEditorGraphMemberDefaultLiteral& Literal, FText RowName);
+		void UpdatePagePickerNames(TWeakObjectPtr<UMetasoundEditorGraphMemberDefaultLiteral> LiteralPtr);
+
+		TArray<TSharedPtr<FString>> AddablePageStringNames;
+		TSet<FName> ImplementedPageNames;
+		TSharedPtr<SSearchableComboBox> PageDefaultComboBox;
+		FDelegateHandle OnPageSettingsUpdatedHandle;
+
+		TAttribute<bool> Enabled;
+		TAttribute<EVisibility> Visibility;
+		TOptional<FResetToDefaultOverride> ResetOverride;
+	};
+
+	class METASOUNDEDITOR_API IMemberDefaultLiteralCustomizationFactory
+	{
+	public:
+		virtual ~IMemberDefaultLiteralCustomizationFactory() = default;
+
+		virtual TUniquePtr<FMetasoundDefaultLiteralCustomizationBase> CreateLiteralCustomization(IDetailCategoryBuilder& DefaultCategoryBuilder) const = 0;
+	};
+} // namespace Metasound::Editor
diff --git a/Engine/Plugins/Runtime/Metasound/Source/MetasoundEditor/Public/MetasoundEditorGraph.h b/Engine/Plugins/Runtime/Metasound/Source/MetasoundEditor/Public/MetasoundEditorGraph.h
--- a/Engine/Plugins/Runtime/Metasound/Source/MetasoundEditor/Public/MetasoundEditorGraph.h
+++ b/Engine/Plugins/Runtime/Metasound/Source/MetasoundEditor/Public/MetasoundEditorGraph.h
@@ -52,13 +52,37 @@ class METASOUNDEDITOR_API UMetasoundEditorGraphMemberDefaultLiteral : public UMe
 	GENERATED_BODY()
 
 public:
+	// Returns the default literal casted as the given template type. Asserts if
+	// literal is not provided type and enforces type is supported by literal.
+	template <typename T>
+	T GetDefaultAs(const FGuid& InPageID = Metasound::Frontend::DefaultPageID) const
+	{
+		T Value;
+		FMetasoundFrontendLiteral Val;
+		bool bSuccess = TryFindDefault(Val, &InPageID);
+		checkf(bSuccess, TEXT("Value assigned to PageID not found"));
+		bSuccess = Val.TryGet(Value);
+		checkf(bSuccess, TEXT("Literal type not supported"));
+		return Value;
+	}
+
+	virtual void ForceRefresh() override;
+	virtual FName GetDataType() const;
+	virtual EMetasoundFrontendLiteralType GetLiteralType() const override;
+	virtual void InitDefault(const FGuid& InPageID = Metasound::Frontend::DefaultPageID);
+	virtual void Initialize() { };
+	virtual void IterateDefaults(TFunctionRef<void(const FGuid&, FMetasoundFrontendLiteral)> Iter) const;
+	virtual bool RemoveDefault(const FGuid& InPageID);
+	virtual void ResetDefaults();
+	virtual void SetFromLiteral(const FMetasoundFrontendLiteral& InLiteral, const FGuid& InPageID = Metasound::Frontend::DefaultPageID) override;
+	virtual bool TryFindDefault(FMetasoundFrontendLiteral& OutLiteral, const FGuid* InPageID = nullptr) const;
+
 	virtual void UpdatePreviewInstance(const Metasound::FVertexName& InParameterName, TScriptInterface<IAudioParameterControllerInterface>& InParameterInterface) const
 	{
 	}
 
 	// Called when literal is initialized for the first time to allow for setting
 	// initial editor-only fields based on context within editor/document model.
-	virtual void Initialize() { };
 
 #if WITH_EDITOR
 	virtual void PostEditChangeProperty(FPropertyChangedEvent& InPropertyChangedEvent) override;
@@ -66,8 +90,21 @@ public:
 	virtual void PostEditUndo() override;
 #endif // WITH_EDITOR
 
-	const UMetasoundEditorGraphMember* FindMember() const;
-	UMetasoundEditorGraphMember* FindMember();
+	// Synchronizes local transient editor-only member data with changes made to the
+	// associated Frontend Document member. Returns whether or not object was modified.
+	virtual bool Synchronize() { return false; };
+
+	UMetasoundEditorGraphMember* FindMember() const;
+
+	static FName GetDefaultPropertyName()
+	{
+		return "Default";
+	}
+
+	static FName GetDefaultsPropertyName()
+	{
+		return "Defaults";
+	}
 
 	UE_DEPRECATED(5.5, "Due to serialization optimization, literals now inherit from LiteralMetadata and are no longer parented under members. Use FindMember instead")
 	const UMetasoundEditorGraphMember* GetParentMember() const
@@ -80,6 +117,11 @@ public:
 	{
 		return Cast<UMetasoundEditorGraphMember>(GetOuter());
 	}
+
+protected:
+	virtual void ResolvePageDefaults() { }
+	virtual void SortPageDefaults() { }
+	bool TryGetPreviewPageID(FGuid& OutPreviewPageID) const;
 };
 
 /** UMetasoundEditorGraphMember is a base class for non-node graph level members 
@@ -136,7 +178,10 @@ public:
 	/** Update the frontend with the given member's default UObject value.
 	 * @param bPostTransaction - Post as editor transaction if true
 	 */
-	virtual void UpdateFrontendDefaultLiteral(bool bPostTransaction) PURE_VIRTUAL(UMetasoundEditorGraphMember::UpdateFrontendDefaultLiteral, );
+	virtual void UpdateFrontendDefaultLiteral(bool bPostTransaction, const FGuid* InPageID = nullptr) PURE_VIRTUAL(UMetasoundEditorGraphMember::UpdateFrontendDefaultLiteral, );
+
+	// Synchronizes cached data with the frontend representation on the represented document.
+	virtual bool Synchronize();
 
 	FMetaSoundFrontendDocumentBuilder& GetFrontendBuilderChecked() const;
 
@@ -151,6 +196,9 @@ public:
 	/* Whether this member can be renamed. */
 	virtual bool CanRename() const PURE_VIRTUAL(UMetasoundEditorGraphMember::CanRename, return false;);
 
+	/** Whether the displayed default supports paged values or is a single default value (i.e. characterized per page). */
+	virtual bool IsDefaultPaged() const { return false; }
+
 #if WITH_EDITOR
 	virtual void PostEditUndo() override;
 #endif // WITH_EDITOR
@@ -222,6 +270,7 @@ public:
 	virtual void SetDescription(const FText& InDescription, bool bPostTransaction) override;
 	virtual void SetMemberName(const FName& InNewName, bool bPostTransaction) override;
 	virtual void SetDisplayName(const FText& InNewName, bool bPostTransaction) override;
+	virtual bool Synchronize() override;
 	/* ~End UMetasoundEditorGraphMember interface */
 
 	/** Version of interface membership, or invalid version if not an interface member. */
@@ -233,6 +282,8 @@ public:
 	/** Returns the Metasound class type of the associated node */
 	virtual EMetasoundFrontendClassType GetClassType() const PURE_VIRTUAL(UMetasoundEditorGraphVertex::GetClassType, return EMetasoundFrontendClassType::Invalid; )
 
+	virtual const FMetasoundFrontendClassVertex* GetFrontendClassVertex() const PURE_VIRTUAL(UMetasoundEditorGraphVertex::GetFrontendClassVertex, return nullptr; );
+
 	/** Returns the SortOrderIndex assigned to this member. */
 	virtual int32 GetSortOrderIndex() const PURE_VIRTUAL(UMetasoundEditorGraphVertex::GetSortOrderIndex, return 0; )
 
@@ -270,17 +321,20 @@ class METASOUNDEDITOR_API UMetasoundEditorGraphInput : public UMetasoundEditorGr
 
 public:
 	virtual FText GetDescription() const override;
+	virtual const FMetasoundFrontendClassVertex* GetFrontendClassVertex() const override;
 	virtual int32 GetSortOrderIndex() const override;
 	virtual void SetSortOrderIndex(int32 InSortOrderIndex) override;
 	virtual TArray<UMetasoundEditorGraphMemberNode*> GetNodes() const override;
 
 	virtual const FText& GetGraphMemberLabel() const override;
+	virtual bool IsDefaultPaged() const override;
 	virtual bool IsInterfaceMember(FMetasoundFrontendInterface* OutInterface = nullptr) const override;
 	virtual void ResetToClassDefault() override;
 	virtual void SetDataType(FName InNewType, bool bPostTransaction) override;
 	virtual void SetVertexAccessType(EMetasoundFrontendVertexAccessType InNewAccessType, bool bPostTransaction) override;
 	virtual void SetMemberName(const FName& InNewName, bool bPostTransaction) override;
-	virtual void UpdateFrontendDefaultLiteral(bool bPostTransaction) override;
+	virtual bool Synchronize() override;
+	virtual void UpdateFrontendDefaultLiteral(bool bPostTransaction, const FGuid* InPageID = nullptr) override;
 	virtual EMetasoundFrontendVertexAccessType GetVertexAccessType() const override;
 
 protected:
@@ -302,6 +356,7 @@ class METASOUNDEDITOR_API UMetasoundEditorGraphOutput : public UMetasoundEditorG
 
 public:
 	virtual FText GetDescription() const override;
+	virtual const FMetasoundFrontendClassVertex* GetFrontendClassVertex() const override;
 	virtual int32 GetSortOrderIndex() const override;
 	virtual bool IsInterfaceMember(FMetasoundFrontendInterface* OutInterface = nullptr) const override;
 	virtual void SetSortOrderIndex(int32 InSortOrderIndex) override;
@@ -309,7 +364,8 @@ public:
 	virtual void ResetToClassDefault() override;
 	virtual void SetDataType(FName InNewType, bool bPostTransaction) override;
 	virtual void SetVertexAccessType(EMetasoundFrontendVertexAccessType InNewAccessType, bool bPostTransaction) override;
-	virtual void UpdateFrontendDefaultLiteral(bool bPostTransaction) override;
+	virtual bool Synchronize() override;
+	virtual void UpdateFrontendDefaultLiteral(bool bPostTransaction, const FGuid* InPageID = nullptr) override;
 	virtual EMetasoundFrontendVertexAccessType GetVertexAccessType() const override;
 
 protected:
@@ -350,7 +406,8 @@ public:
 	virtual void SetDisplayName(const FText& InNewName, bool bPostTransaction) override;
 
 	virtual void ResetToClassDefault() override;
-	virtual void UpdateFrontendDefaultLiteral(bool bPostTransaction) override;
+	virtual bool Synchronize() override;
+	virtual void UpdateFrontendDefaultLiteral(bool bPostTransaction, const FGuid* InPageID = nullptr) override;
 
 	virtual TArray<UMetasoundEditorGraphMemberNode*> GetNodes() const override;
 
diff --git a/Engine/Plugins/Runtime/Metasound/Source/MetasoundEditor/Public/MetasoundEditorGraphMemberDefaults.h b/Engine/Plugins/Runtime/Metasound/Source/MetasoundEditor/Public/MetasoundEditorGraphMemberDefaults.h
--- a/Engine/Plugins/Runtime/Metasound/Source/MetasoundEditor/Public/MetasoundEditorGraphMemberDefaults.h
+++ b/Engine/Plugins/Runtime/Metasound/Source/MetasoundEditor/Public/MetasoundEditorGraphMemberDefaults.h
@@ -19,6 +19,7 @@
 #include "MetasoundFrontendLiteral.h"
 #include "MetasoundFrontendRegistries.h"
 #include "MetasoundPrimitives.h"
+#include "MetasoundSettings.h"
 #include "MetasoundVertex.h"
 #include "Misc/Guid.h"
 #include "SAudioRadialSlider.h"
@@ -29,107 +30,234 @@
 #include "MetasoundEditorGraphMemberDefaults.generated.h"
 
 
-// Broken out to be able to customize and swap enum behavior for basic integer literal behavior
+// Editor-only page default for more desirable customization behavior in representing
+// Frontend Literal value. Should never be serialized as generation is non-deterministic.
+USTRUCT()
+struct FMetasoundEditorMemberPageDefault
+{
+	GENERATED_BODY()
+
+	FMetasoundEditorMemberPageDefault() = default;
+	FMetasoundEditorMemberPageDefault(const FGuid& InPageID)
+		: PageID(InPageID)
+	{
+	}
+
+	// Unique ID used to determine if this default is a newly created and uninitialized member
+	// of a given collection.  ID is transient per application cycle, so non-deterministic
+	// between editor sessions.
+	static const FGuid& GetNewEntryID();
+
+	// Selectable PageName
+	UPROPERTY(EditAnywhere, Category = DefaultValue, meta = (NoResetToDefault, GetOptions = "MetasoundEngine.MetaSoundSettings.GetPageNames"))
+	FName PageName;
+
+	// Used for hash and mirrors document-stored value. Defaults to random value to allow for assignment in
+	// post-edit change. Allows for name collisions if user is amidst renaming or rebasing values
+	UPROPERTY(Transient, VisibleAnywhere, Category = DefaultValue)
+	FGuid PageID = GetNewEntryID();
+};
+
+
+// For bool input widget
+DECLARE_MULTICAST_DELEGATE_TwoParams(FOnMetasoundBoolStateChangedEvent, bool, const FGuid& /* PageID */);
+
+// Broken out to be able to customize and swap enum behavior for boolean literal behavior (ex. for triggers)
 USTRUCT()
 struct FMetasoundEditorGraphMemberDefaultBoolRef
 {
 	GENERATED_BODY()
 
-	UPROPERTY(EditAnywhere, Category = DefaultValue)
+	UPROPERTY(EditAnywhere, Category = DefaultValue, meta = (DisplayName = "Default"))
 	bool Value = false;
 };
 
-// For bool input widget
-DECLARE_MULTICAST_DELEGATE_OneParam(FOnMetasoundBoolStateChangedEvent, bool);
+
+USTRUCT()
+struct FMetasoundEditorMemberPageDefaultBool : public FMetasoundEditorMemberPageDefault
+{
+	GENERATED_BODY()
+
+	UPROPERTY(EditAnywhere, Category = DefaultValue, meta = (DisplayName = "Default"))
+	FMetasoundEditorGraphMemberDefaultBoolRef Value;
+};
+
 
 UCLASS(MinimalAPI)
 class UMetasoundEditorGraphMemberDefaultBool : public UMetasoundEditorGraphMemberDefaultLiteral
 {
 	GENERATED_BODY()
 
+private:
+	UPROPERTY(EditAnywhere, Transient, Category = DefaultValue, meta = (NoResetToDefault, NoElementDuplicate, EditFixedOrder))
+	TArray<FMetasoundEditorMemberPageDefaultBool> Defaults;
+
 public:
 	virtual ~UMetasoundEditorGraphMemberDefaultBool() = default;
 
-	UPROPERTY(EditAnywhere, Category = DefaultValue)
+	UPROPERTY(meta = (Deprecated = 5.5, DeprecationMessage = "Default is no longer serialized and is privately managed to support per-page default values"))
 	FMetasoundEditorGraphMemberDefaultBoolRef Default;
 
 	UPROPERTY(EditAnywhere, Category = Widget, meta = (DisplayName = "Widget"))
 	EMetasoundBoolMemberDefaultWidget WidgetType = EMetasoundBoolMemberDefaultWidget::None;
 
-	FOnMetasoundBoolStateChangedEvent OnDefaultStateChanged;
-
-	virtual FMetasoundFrontendLiteral GetDefault() const override;
 	virtual EMetasoundFrontendLiteralType GetLiteralType() const override;
-	virtual void SetFromLiteral(const FMetasoundFrontendLiteral& InLiteral) override;
+	virtual void InitDefault(const FGuid& InPageID = Metasound::Frontend::DefaultPageID) override;
+	virtual void IterateDefaults(TFunctionRef<void(const FGuid&, FMetasoundFrontendLiteral)> Iter) const override;
+	virtual bool RemoveDefault(const FGuid& InPageID) override;
+	virtual void ResetDefaults() override;
+	virtual void SetFromLiteral(const FMetasoundFrontendLiteral& InLiteral, const FGuid& InPageID) override;
+	virtual bool Synchronize() override;
+	virtual bool TryFindDefault(FMetasoundFrontendLiteral& OutLiteral, const FGuid* InPageID = nullptr) const override;
 	virtual void UpdatePreviewInstance(const Metasound::FVertexName& InParameterName, TScriptInterface<IAudioParameterControllerInterface>& InParameterInterface) const override;
 
-	void SetDefault(const bool InDefault);
-	const bool GetDefault();
+	UE_DEPRECATED(5.5, "Use SetFromLiteral instead and broadcast state change delegate explicitly where desired")
+	void SetDefault(const bool InDefault) { }
+
+	FOnMetasoundBoolStateChangedEvent OnDefaultStateChanged;
+
+protected:
+	virtual void ResolvePageDefaults() override;
+	virtual void SortPageDefaults() override;
+};
+
+
+USTRUCT()
+struct FMetasoundEditorMemberPageDefaultBoolArray : public FMetasoundEditorMemberPageDefault
+{
+	GENERATED_BODY()
+
+	UPROPERTY(EditAnywhere, Category = DefaultValue, meta = (DisplayName = "Default"))
+	TArray<FMetasoundEditorGraphMemberDefaultBoolRef> Value;
 };
 
+
 UCLASS(MinimalAPI)
 class UMetasoundEditorGraphMemberDefaultBoolArray : public UMetasoundEditorGraphMemberDefaultLiteral
 {
 	GENERATED_BODY()
 
+private:
+	UPROPERTY(EditAnywhere, Transient, Category = DefaultValue, meta = (NoResetToDefault, NoElementDuplicate, EditFixedOrder))
+	TArray<FMetasoundEditorMemberPageDefaultBoolArray> Defaults;
+
 public:
 	virtual ~UMetasoundEditorGraphMemberDefaultBoolArray() = default;
 
-	UPROPERTY(EditAnywhere, Category = DefaultValue)
+	UPROPERTY(meta = (Deprecated = 5.5, DeprecationMessage = "Default is no longer serialized and is privately managed to support per-page default values"))
 	TArray<FMetasoundEditorGraphMemberDefaultBoolRef> Default;
 
-	virtual FMetasoundFrontendLiteral GetDefault() const override;
 	virtual EMetasoundFrontendLiteralType GetLiteralType() const override;
-	virtual void SetFromLiteral(const FMetasoundFrontendLiteral& InLiteral) override;
+	virtual void InitDefault(const FGuid& InPageID = Metasound::Frontend::DefaultPageID) override;
+	virtual void IterateDefaults(TFunctionRef<void(const FGuid&, FMetasoundFrontendLiteral)> Iter) const override;
+	virtual bool RemoveDefault(const FGuid& InPageID) override;
+	virtual void ResetDefaults() override;
+	virtual void SetFromLiteral(const FMetasoundFrontendLiteral& InLiteral, const FGuid& InPageID) override;
+	virtual bool Synchronize() override;
+	virtual bool TryFindDefault(FMetasoundFrontendLiteral& OutLiteral, const FGuid* InPageID = nullptr) const override;
 	virtual void UpdatePreviewInstance(const Metasound::FVertexName& InParameterName, TScriptInterface<IAudioParameterControllerInterface>& InParameterInterface) const override;
+
+protected:
+	virtual void ResolvePageDefaults() override;
+	virtual void SortPageDefaults() override;
 };
 
+
 // Broken out to be able to customize and swap enum behavior for basic integer literal behavior
 USTRUCT()
 struct FMetasoundEditorGraphMemberDefaultIntRef
 {
 	GENERATED_BODY()
 
-	UPROPERTY(EditAnywhere, Category = DefaultValue)
+	UPROPERTY(EditAnywhere, Category = DefaultValue, meta = (DisplayName = "Default"))
 	int32 Value = 0;
 };
 
+
+USTRUCT()
+struct FMetasoundEditorMemberPageDefaultInt : public FMetasoundEditorMemberPageDefault
+{
+	GENERATED_BODY()
+
+	UPROPERTY(EditAnywhere, Category = DefaultValue)
+	FMetasoundEditorGraphMemberDefaultIntRef Value;
+};
+
+
 UCLASS(MinimalAPI)
 class UMetasoundEditorGraphMemberDefaultInt : public UMetasoundEditorGraphMemberDefaultLiteral
 {
 	GENERATED_BODY()
 
+private:
+	UPROPERTY(EditAnywhere, Transient, Category = DefaultValue, meta = (NoResetToDefault, NoElementDuplicate, EditFixedOrder))
+	TArray<FMetasoundEditorMemberPageDefaultInt> Defaults;
+
 public:
 	virtual ~UMetasoundEditorGraphMemberDefaultInt() = default;
 
-	UPROPERTY(EditAnywhere, Category = DefaultValue)
+	UPROPERTY(meta = (Deprecated = 5.5, DeprecationMessage = "Default is no longer serialized and is privately managed to support per-page default values"))
 	FMetasoundEditorGraphMemberDefaultIntRef Default;
 
-	virtual FMetasoundFrontendLiteral GetDefault() const override;
 	virtual EMetasoundFrontendLiteralType GetLiteralType() const override;
-	virtual void SetFromLiteral(const FMetasoundFrontendLiteral& InLiteral) override;
+	virtual void InitDefault(const FGuid& InPageID = Metasound::Frontend::DefaultPageID) override;
+	virtual void IterateDefaults(TFunctionRef<void(const FGuid&, FMetasoundFrontendLiteral)> Iter) const override;
+	virtual bool RemoveDefault(const FGuid& InPageID) override;
+	virtual void ResetDefaults() override;
+	virtual void SetFromLiteral(const FMetasoundFrontendLiteral& InLiteral, const FGuid& InPageID) override;
+	virtual bool Synchronize() override;
+	virtual bool TryFindDefault(FMetasoundFrontendLiteral& OutLiteral, const FGuid* InPageID = nullptr) const override;
 	virtual void UpdatePreviewInstance(const Metasound::FVertexName& InParameterName, TScriptInterface<IAudioParameterControllerInterface>& InParameterInterface) const override;
+
+protected:
+	virtual void ResolvePageDefaults() override;
+	virtual void SortPageDefaults() override;
+};
+
+
+USTRUCT()
+struct FMetasoundEditorMemberPageDefaultIntArray : public FMetasoundEditorMemberPageDefault
+{
+	GENERATED_BODY()
+
+	UPROPERTY(EditAnywhere, Category = DefaultValue, meta = (DisplayName = "Default"))
+	TArray<FMetasoundEditorGraphMemberDefaultIntRef> Value;
 };
 
+
 UCLASS(MinimalAPI)
 class UMetasoundEditorGraphMemberDefaultIntArray : public UMetasoundEditorGraphMemberDefaultLiteral
 {
 	GENERATED_BODY()
 
+private:
+	UPROPERTY(EditAnywhere, Transient, Category = DefaultValue, meta = (NoResetToDefault, NoElementDuplicate, EditFixedOrder))
+	TArray<FMetasoundEditorMemberPageDefaultIntArray> Defaults;
+
 public:
 	virtual ~UMetasoundEditorGraphMemberDefaultIntArray() = default;
 
-	UPROPERTY(EditAnywhere, Category = DefaultValue)
+	UPROPERTY(meta = (Deprecated = 5.5, DeprecationMessage = "Default is no longer serialized and is privately managed to support per-page default values"))
 	TArray<FMetasoundEditorGraphMemberDefaultIntRef> Default;
 
-	virtual FMetasoundFrontendLiteral GetDefault() const override;
 	virtual EMetasoundFrontendLiteralType GetLiteralType() const override;
-	virtual void SetFromLiteral(const FMetasoundFrontendLiteral& InLiteral) override;
+	virtual void InitDefault(const FGuid& InPageID = Metasound::Frontend::DefaultPageID) override;
+	virtual void IterateDefaults(TFunctionRef<void(const FGuid&, FMetasoundFrontendLiteral)> Iter) const override;
+	virtual bool RemoveDefault(const FGuid& InPageID) override;
+	virtual void ResetDefaults() override;
+	virtual void SetFromLiteral(const FMetasoundFrontendLiteral& InLiteral, const FGuid& InPageID) override;
+	virtual bool Synchronize() override;
+	virtual bool TryFindDefault(FMetasoundFrontendLiteral& OutLiteral, const FGuid* InPageID = nullptr) const override;
 	virtual void UpdatePreviewInstance(const Metasound::FVertexName& InParameterName, TScriptInterface<IAudioParameterControllerInterface>& InParameterInterface) const override;
+
+protected:
+	virtual void ResolvePageDefaults() override;
+	virtual void SortPageDefaults() override;
 };
 
+
 // For input widget
-DECLARE_MULTICAST_DELEGATE_OneParam(FOnMetasoundInputValueChangedEvent, float);
+DECLARE_MULTICAST_DELEGATE_TwoParams(FOnMetasoundInputValueChangedEvent, FGuid /* PageID */, float /* Value */);
 DECLARE_MULTICAST_DELEGATE_OneParam(FOnMetasoundRangeChangedEvent, FVector2D);
 DECLARE_MULTICAST_DELEGATE_OneParam(FOnMetasoundInputClampDefaultChangedEvent, bool);
 
@@ -142,15 +270,29 @@ enum class UE_DEPRECATED(5.5, "EMetasoundMemberDefaultWidgetValueType is depreca
 	Volume
 };
 
+
+USTRUCT()
+struct FMetasoundEditorMemberPageDefaultFloat : public FMetasoundEditorMemberPageDefault
+{
+	GENERATED_BODY()
+
+	UPROPERTY(EditAnywhere, Category = DefaultValue, meta = (DisplayName = "Default"))
+	float Value = 0.0f;
+};
+
+
 UCLASS(MinimalAPI)
 class UMetasoundEditorGraphMemberDefaultFloat : public UMetasoundEditorGraphMemberDefaultLiteral
 {
 	GENERATED_BODY()
 
 private:
-	UPROPERTY(EditAnywhere, Category = DefaultValue)
+	UPROPERTY(meta = (Deprecated = 5.5, DeprecationMessage = "Default is no longer serialized and is privately managed to support per-page default values"))
 	float Default = 0.f;
 
+	UPROPERTY(EditAnywhere, Transient, Category = DefaultValue, meta = (NoResetToDefault, NoElementDuplicate, EditFixedOrder))
+	TArray<FMetasoundEditorMemberPageDefaultFloat> Defaults;
+
 public:
 	virtual ~UMetasoundEditorGraphMemberDefaultFloat() = default;
 
@@ -185,48 +327,93 @@ public:
 	UPROPERTY(EditAnywhere, Category = Widget, meta = (DisplayName = "Range in dB", EditCondition = "VolumeWidgetUseLinearOutput", EditConditionHides))
 	FVector2D VolumeWidgetDecibelRange = FVector2D(SAudioVolumeRadialSlider::MinDbValue, 0.0f);
 
-	static FName GetDefaultPropertyName()
-	{
-		return GET_MEMBER_NAME_CHECKED(UMetasoundEditorGraphMemberDefaultFloat, Default);
-	}
-
 	FOnMetasoundInputValueChangedEvent OnDefaultValueChanged;
 	FOnMetasoundRangeChangedEvent OnRangeChanged;
 	FOnMetasoundInputClampDefaultChangedEvent OnClampChanged;
 
 	virtual void ForceRefresh() override;
-	virtual FMetasoundFrontendLiteral GetDefault() const override;
 	virtual EMetasoundFrontendLiteralType GetLiteralType() const override;
+	virtual void InitDefault(const FGuid& InPageID = Metasound::Frontend::DefaultPageID) override;
 	virtual void Initialize() override;
-	virtual void SetFromLiteral(const FMetasoundFrontendLiteral& InLiteral) override;
+	virtual void IterateDefaults(TFunctionRef<void(const FGuid&, FMetasoundFrontendLiteral)> Iter) const override;
+	virtual bool RemoveDefault(const FGuid& InPageID) override;
+	virtual void ResetDefaults() override;
+	virtual void SetFromLiteral(const FMetasoundFrontendLiteral& InLiteral, const FGuid& InPageID) override;
+	virtual bool Synchronize() override;
+	virtual bool TryFindDefault(FMetasoundFrontendLiteral& OutLiteral, const FGuid* InPageID = nullptr) const override;
 	virtual void UpdatePreviewInstance(const Metasound::FVertexName& InParameterName, TScriptInterface<IAudioParameterControllerInterface>& InParameterInterface) const override;
 	virtual void PostEditChangeChainProperty(FPropertyChangedChainEvent& InPropertyChangedEvent) override;
 
-	void SetDefault(const float InDefault);
+	UE_DEPRECATED(5.5, "Use SetFromLiteral instead")
+	void SetDefault(const float InDefault) { }
+
 	// Set Range to reasonable limit given current Default value
 	void SetInitialRange();
+
+	UE_DEPRECATED(5.5, "Use TryFindDefault or GetDefaultAs template call provided from parent class 'UMetaSoundFrontendMemberMetadata' providing an applicable page ID instead")
 	float GetDefault();
-	FVector2D GetRange();
+
+	FVector2D GetRange() const;
 	void SetRange(const FVector2D InRange);
+
+protected:
+	void ClampDefaults();
+	virtual void ResolvePageDefaults() override;
+	virtual void SortPageDefaults() override;
+};
+
+
+USTRUCT()
+struct FMetasoundEditorMemberPageDefaultFloatArray : public FMetasoundEditorMemberPageDefault
+{
+	GENERATED_BODY()
+
+	UPROPERTY(EditAnywhere, Category = DefaultValue, meta = (DisplayName = "Default"))
+	TArray<float> Value;
 };
 
+
 UCLASS(MinimalAPI)
 class UMetasoundEditorGraphMemberDefaultFloatArray : public UMetasoundEditorGraphMemberDefaultLiteral
 {
 	GENERATED_BODY()
 
+private:
+	UPROPERTY(EditAnywhere, Transient, Category = DefaultValue, meta = (NoResetToDefault, NoElementDuplicate, EditFixedOrder))
+	TArray<FMetasoundEditorMemberPageDefaultFloatArray> Defaults;
+
 public:
 	virtual ~UMetasoundEditorGraphMemberDefaultFloatArray() = default;
 
-	UPROPERTY(EditAnywhere, Category = DefaultValue)
+	UPROPERTY(meta = (Deprecated = 5.5, DeprecationMessage = "Default is no longer serialized and is privately managed to support per-page default values"))
 	TArray<float> Default;
 
-	virtual FMetasoundFrontendLiteral GetDefault() const override;
 	virtual EMetasoundFrontendLiteralType GetLiteralType() const override;
-	virtual void SetFromLiteral(const FMetasoundFrontendLiteral& InLiteral) override;
+	virtual void InitDefault(const FGuid& InPageID = Metasound::Frontend::DefaultPageID) override;
+	virtual void IterateDefaults(TFunctionRef<void(const FGuid&, FMetasoundFrontendLiteral)> Iter) const override;
+	virtual bool RemoveDefault(const FGuid& InPageID) override;
+	virtual void ResetDefaults() override;
+	virtual void SetFromLiteral(const FMetasoundFrontendLiteral& InLiteral, const FGuid& InPageID) override;
+	virtual bool Synchronize() override;
+	virtual bool TryFindDefault(FMetasoundFrontendLiteral& OutLiteral, const FGuid* InPageID = nullptr) const override;
 	virtual void UpdatePreviewInstance(const Metasound::FVertexName& InParameterName, TScriptInterface<IAudioParameterControllerInterface>& InParameterInterface) const override;
+
+protected:
+	virtual void ResolvePageDefaults() override;
+	virtual void SortPageDefaults() override;
+};
+
+
+USTRUCT()
+struct FMetasoundEditorMemberPageDefaultString : public FMetasoundEditorMemberPageDefault
+{
+	GENERATED_BODY()
+
+	UPROPERTY(EditAnywhere, Category = DefaultValue, meta = (DisplayName = "Default"))
+	FString Value;
 };
 
+
 UCLASS(MinimalAPI)
 class UMetasoundEditorGraphMemberDefaultString : public UMetasoundEditorGraphMemberDefaultLiteral
 {
@@ -235,30 +422,67 @@ class UMetasoundEditorGraphMemberDefaultString : public UMetasoundEditorGraphMem
 public:
 	virtual ~UMetasoundEditorGraphMemberDefaultString() = default;
 
-	UPROPERTY(EditAnywhere, Category = DefaultValue)
+	UPROPERTY(meta = (Deprecated = 5.5, DeprecationMessage = "Default is no longer serialized and is privately managed to support per-page default values"))
 	FString Default;
 
-	virtual FMetasoundFrontendLiteral GetDefault() const override;
+private:
+	UPROPERTY(EditAnywhere, Transient, Category = DefaultValue, meta = (NoResetToDefault, NoElementDuplicate, EditFixedOrder))
+	TArray<FMetasoundEditorMemberPageDefaultString> Defaults;
+
 	virtual EMetasoundFrontendLiteralType GetLiteralType() const override;
-	virtual void SetFromLiteral(const FMetasoundFrontendLiteral& InLiteral) override;
+	virtual void InitDefault(const FGuid& InPageID = Metasound::Frontend::DefaultPageID) override;
+	virtual void IterateDefaults(TFunctionRef<void(const FGuid&, FMetasoundFrontendLiteral)> Iter) const override;
+	virtual bool RemoveDefault(const FGuid& InPageID) override;
+	virtual void ResetDefaults() override;
+	virtual void SetFromLiteral(const FMetasoundFrontendLiteral& InLiteral, const FGuid& InPageID) override;
+	virtual bool Synchronize() override;
+	virtual bool TryFindDefault(FMetasoundFrontendLiteral& OutLiteral, const FGuid* InPageID = nullptr) const override;
 	virtual void UpdatePreviewInstance(const Metasound::FVertexName& InParameterName, TScriptInterface<IAudioParameterControllerInterface>& InParameterInterface) const override;
+
+protected:
+	virtual void ResolvePageDefaults() override;
+	virtual void SortPageDefaults() override;
+};
+
+
+USTRUCT()
+struct FMetasoundEditorMemberPageDefaultStringArray : public FMetasoundEditorMemberPageDefault
+{
+	GENERATED_BODY()
+
+	UPROPERTY(EditAnywhere, Category = DefaultValue, meta = (DisplayName = "Default"))
+	TArray<FString> Value;
 };
 
+
 UCLASS(MinimalAPI)
 class UMetasoundEditorGraphMemberDefaultStringArray : public UMetasoundEditorGraphMemberDefaultLiteral
 {
 	GENERATED_BODY()
 
+private:
+	UPROPERTY(EditAnywhere, Transient, Category = DefaultValue, meta = (NoResetToDefault, NoElementDuplicate, EditFixedOrder))
+	TArray<FMetasoundEditorMemberPageDefaultStringArray> Defaults;
+
 public:
 	virtual ~UMetasoundEditorGraphMemberDefaultStringArray() = default;
 
-	UPROPERTY(EditAnywhere, Category = DefaultValue)
+	UPROPERTY(meta = (Deprecated = 5.5, DeprecationMessage = "Default is no longer serialized and is privately managed to support per-page default values"))
 	TArray<FString> Default;
 
-	virtual FMetasoundFrontendLiteral GetDefault() const override;
 	virtual EMetasoundFrontendLiteralType GetLiteralType() const override;
-	virtual void SetFromLiteral(const FMetasoundFrontendLiteral& InLiteral) override;
+	virtual void InitDefault(const FGuid& InPageID = Metasound::Frontend::DefaultPageID) override;
+	virtual void IterateDefaults(TFunctionRef<void(const FGuid&, FMetasoundFrontendLiteral)> Iter) const override;
+	virtual bool RemoveDefault(const FGuid& InPageID) override;
+	virtual void ResetDefaults() override;
+	virtual void SetFromLiteral(const FMetasoundFrontendLiteral& InLiteral, const FGuid& InPageID) override;
+	virtual bool Synchronize() override;
+	virtual bool TryFindDefault(FMetasoundFrontendLiteral& OutLiteral, const FGuid* InPageID = nullptr) const override;
 	virtual void UpdatePreviewInstance(const Metasound::FVertexName& InParameterName, TScriptInterface<IAudioParameterControllerInterface>& InParameterInterface) const override;
+
+protected:
+	virtual void ResolvePageDefaults() override;
+	virtual void SortPageDefaults() override;
 };
 
 // Broken out to be able to customize and swap AllowedClass based on provided object proxy
@@ -267,10 +491,19 @@ struct FMetasoundEditorGraphMemberDefaultObjectRef
 {
 	GENERATED_BODY()
 
-	UPROPERTY(EditAnywhere, Category = DefaultValue)
+	UPROPERTY(EditAnywhere, Category = DefaultValue, meta = (DisplayName = "Default"))
 	TObjectPtr<UObject> Object = nullptr;
 };
 
+USTRUCT()
+struct FMetasoundEditorMemberPageDefaultObjectRef : public FMetasoundEditorMemberPageDefault
+{
+	GENERATED_BODY()
+
+	UPROPERTY(EditAnywhere, Category = DefaultValue, meta = (ShowOnlyInnerProperties))
+	FMetasoundEditorGraphMemberDefaultObjectRef Value;
+};
+
 UCLASS(MinimalAPI)
 class UMetasoundEditorGraphMemberDefaultObject : public UMetasoundEditorGraphMemberDefaultLiteral
 {
@@ -279,13 +512,34 @@ class UMetasoundEditorGraphMemberDefaultObject : public UMetasoundEditorGraphMem
 public:
 	virtual ~UMetasoundEditorGraphMemberDefaultObject() = default;
 
-	UPROPERTY(EditAnywhere, Category = DefaultValue)
+	UPROPERTY(meta = (Deprecated = 5.5, DeprecationMessage = "Default is no longer serialized and is privately managed to support per-page default values"))
 	FMetasoundEditorGraphMemberDefaultObjectRef Default;
 
-	virtual FMetasoundFrontendLiteral GetDefault() const override;
+	UPROPERTY(EditAnywhere, Transient, Category = DefaultValue, meta = (NoResetToDefault, NoElementDuplicate, EditFixedOrder))
+	TArray<FMetasoundEditorMemberPageDefaultObjectRef> Defaults;
+
 	virtual EMetasoundFrontendLiteralType GetLiteralType() const override;
-	virtual void SetFromLiteral(const FMetasoundFrontendLiteral& InLiteral) override;
+	virtual void InitDefault(const FGuid& InPageID = Metasound::Frontend::DefaultPageID) override;
+	virtual void IterateDefaults(TFunctionRef<void(const FGuid&, FMetasoundFrontendLiteral)> Iter) const override;
+	virtual bool RemoveDefault(const FGuid& InPageID) override;
+	virtual void ResetDefaults() override;
+	virtual void SetFromLiteral(const FMetasoundFrontendLiteral& InLiteral, const FGuid& InPageID) override;
+	virtual bool Synchronize() override;
+	virtual bool TryFindDefault(FMetasoundFrontendLiteral& OutLiteral, const FGuid* InPageID = nullptr) const override;
 	virtual void UpdatePreviewInstance(const Metasound::FVertexName& InParameterName, TScriptInterface<IAudioParameterControllerInterface>& InParameterInterface) const override;
+
+protected:
+	virtual void ResolvePageDefaults() override;
+	virtual void SortPageDefaults() override;
+};
+
+USTRUCT()
+struct FMetasoundEditorMemberPageDefaultObjectArray : public FMetasoundEditorMemberPageDefault
+{
+	GENERATED_BODY()
+
+	UPROPERTY(EditAnywhere, Category = DefaultValue, meta = (DisplayName = "Default"))
+	TArray<FMetasoundEditorGraphMemberDefaultObjectRef> Value;
 };
 
 UCLASS(MinimalAPI)
@@ -293,14 +547,24 @@ class UMetasoundEditorGraphMemberDefaultObjectArray : public UMetasoundEditorGra
 {
 	GENERATED_BODY()
 
+	UPROPERTY(EditAnywhere, Transient, Category = DefaultValue, meta = (NoResetToDefault, NoElementDuplicate, EditFixedOrder))
+	TArray<FMetasoundEditorMemberPageDefaultObjectArray> Defaults;
+
 public:
 	virtual ~UMetasoundEditorGraphMemberDefaultObjectArray() = default;
 
-	UPROPERTY(EditAnywhere, Category = DefaultValue)
+	UPROPERTY(meta = (Deprecated = 5.5, DeprecationMessage = "Default is no longer serialized and is privately managed to support per-page default values"))
 	TArray<FMetasoundEditorGraphMemberDefaultObjectRef> Default;
 
-	virtual FMetasoundFrontendLiteral GetDefault() const override;
 	virtual EMetasoundFrontendLiteralType GetLiteralType() const override;
-	virtual void SetFromLiteral(const FMetasoundFrontendLiteral& InLiteral) override;
+	virtual void InitDefault(const FGuid& InPageID = Metasound::Frontend::DefaultPageID) override;
+	virtual void IterateDefaults(TFunctionRef<void(const FGuid&, FMetasoundFrontendLiteral)> Iter) const override;
+	virtual bool RemoveDefault(const FGuid& InPageID) override;
+	virtual void ResetDefaults() override;
+	virtual void ResolvePageDefaults() override;
+	virtual void SetFromLiteral(const FMetasoundFrontendLiteral& InLiteral, const FGuid& InPageID) override;
+	virtual void SortPageDefaults() override;
+	virtual bool Synchronize() override;
+	virtual bool TryFindDefault(FMetasoundFrontendLiteral& OutLiteral, const FGuid* InPageID = nullptr) const override;
 	virtual void UpdatePreviewInstance(const Metasound::FVertexName& InParameterName, TScriptInterface<IAudioParameterControllerInterface>& InParameterInterface) const override;
 };
\ No newline at end of file
diff --git a/Engine/Plugins/Runtime/Metasound/Source/MetasoundEditor/Public/MetasoundEditorModule.h b/Engine/Plugins/Runtime/Metasound/Source/MetasoundEditor/Public/MetasoundEditorModule.h
--- a/Engine/Plugins/Runtime/Metasound/Source/MetasoundEditor/Public/MetasoundEditorModule.h
+++ b/Engine/Plugins/Runtime/Metasound/Source/MetasoundEditor/Public/MetasoundEditorModule.h
@@ -4,6 +4,8 @@
 #include "DetailCategoryBuilder.h"
 #include "EdGraph/EdGraphNode.h"
 #include "EdGraph/EdGraphPin.h"
+#include "IDetailPropertyRow.h"
+#include "MetasoundDefaultLiteralCustomization.h"
 #include "MetasoundFrontendDataTypeRegistry.h"
 #include "MetasoundFrontendLiteral.h"
 #include "MetasoundFrontendRegistries.h"
@@ -15,10 +17,13 @@
 // Forward Declarations
 class IDetailLayoutBuilder;
 class IDetailPropertyRow;
+class IPropertyHandle;
+class SSearchableComboBox;
 class UEdGraphPin;
 class UMetasoundEditorGraph;
 class UMetasoundEditorGraphMemberDefaultLiteral;
 class UMetasoundEditorGraphNode;
+
 namespace Metasound::Engine 
 {
 	enum class EAssetScanStatus : uint8;
@@ -50,32 +55,6 @@ namespace Metasound
 			const FSlateBrush* PinDisconnectedIcon = nullptr;
 		};
 
-		class METASOUNDEDITOR_API FMetasoundDefaultLiteralCustomizationBase
-		{
-		protected:
-			IDetailCategoryBuilder* DefaultCategoryBuilder = nullptr;
-
-		public:
-			FMetasoundDefaultLiteralCustomizationBase(IDetailCategoryBuilder& InDefaultCategoryBuilder)
-				: DefaultCategoryBuilder(&InDefaultCategoryBuilder)
-			{
-			}
-
-			virtual ~FMetasoundDefaultLiteralCustomizationBase() = default;
-
-			// Customizes the given literal for the provided DetailLayoutBuilder.
-			// @return the DetailPropertyRow created for the default parameter set by this customization.
-			virtual TArray<IDetailPropertyRow*> CustomizeLiteral(UMetasoundEditorGraphMemberDefaultLiteral& InLiteral, IDetailLayoutBuilder& InDetailLayout) { return { }; };
-		};
-
-		class METASOUNDEDITOR_API IMemberDefaultLiteralCustomizationFactory
-		{
-		public:
-			virtual ~IMemberDefaultLiteralCustomizationFactory() = default;
-
-			virtual TUniquePtr<FMetasoundDefaultLiteralCustomizationBase> CreateLiteralCustomization(IDetailCategoryBuilder& DefaultCategoryBuilder) const = 0;
-		};
-
 
 		struct FCreateGraphNodeVisualizationWidgetParams
 		{
diff --git a/Engine/Plugins/Runtime/Metasound/Source/MetasoundFrontend/Private/MetasoundFrontendDocument.cpp b/Engine/Plugins/Runtime/Metasound/Source/MetasoundFrontend/Private/MetasoundFrontendDocument.cpp
--- a/Engine/Plugins/Runtime/Metasound/Source/MetasoundFrontend/Private/MetasoundFrontendDocument.cpp
+++ b/Engine/Plugins/Runtime/Metasound/Source/MetasoundFrontend/Private/MetasoundFrontendDocument.cpp
@@ -831,6 +831,33 @@ FMetasoundFrontendClassMetadata FMetasoundFrontendClassMetadata::GenerateClassMe
 	return NewMetadata;
 }
 
+FMetasoundFrontendClassInputDefault::FMetasoundFrontendClassInputDefault(FMetasoundFrontendLiteral InLiteral)
+	: Literal(InLiteral)
+	, PageID(Metasound::Frontend::DefaultPageID)
+{
+}
+
+FMetasoundFrontendClassInputDefault::FMetasoundFrontendClassInputDefault(const FGuid& InPageID, FMetasoundFrontendLiteral InLiteral)
+	: Literal(MoveTemp(InLiteral))
+	, PageID(InPageID)
+{
+}
+
+FMetasoundFrontendClassInputDefault::FMetasoundFrontendClassInputDefault(const FAudioParameter& InParameter)
+	: Literal(InParameter)
+{
+}
+
+bool FMetasoundFrontendClassInputDefault::IsFunctionalEquivalent(const FMetasoundFrontendClassInputDefault& InLHS, const FMetasoundFrontendClassInputDefault& InRHS)
+{
+	if (InLHS.PageID != InRHS.PageID)
+	{
+		return false;
+	}
+
+	return InLHS.Literal.IsEqual(InRHS.Literal);
+}
+
 FMetasoundFrontendClassInput::FMetasoundFrontendClassInput(const FMetasoundFrontendClassVertex& InOther)
 :	FMetasoundFrontendClassVertex(InOther)
 {
@@ -860,12 +887,33 @@ FMetasoundFrontendClassInput::FMetasoundFrontendClassInput(const Audio::FParamet
 #endif // WITH_EDITOR
 }
 
-FMetasoundFrontendLiteral& FMetasoundFrontendClassInput::AddDefault(const FGuid& InPageID)
+bool FMetasoundFrontendClassInput::IsFunctionalEquivalent(const FMetasoundFrontendClassInput& InLHS, const FMetasoundFrontendClassInput& InRHS)
 {
-#if WITH_EDITORONLY_DATA
-	checkf(InPageID == Metasound::Frontend::DefaultPageID, TEXT("Non-default page mutation is not supported without editor data loaded."));
-#endif // WITH_EDITORONLY_DATA
+	if (!FMetasoundFrontendClassVertex::IsFunctionalEquivalent(InLHS, InRHS))
+	{
+		return false;
+	}
 
+	const TArray<FMetasoundFrontendClassInputDefault>& LHSDefaults = InLHS.GetDefaults();
+	const TArray<FMetasoundFrontendClassInputDefault>& RHSDefaults = InRHS.GetDefaults();
+	if (LHSDefaults.Num() != RHSDefaults.Num())
+	{
+		return false;
+	}
+
+	for (int32 Index = 0; Index < LHSDefaults.Num(); ++Index)
+	{
+		if (!FMetasoundFrontendClassInputDefault::IsFunctionalEquivalent(LHSDefaults[Index], RHSDefaults[Index]))
+		{
+			return false;
+		}
+	}
+
+	return true;
+}
+
+FMetasoundFrontendLiteral& FMetasoundFrontendClassInput::AddDefault(const FGuid& InPageID)
+{
 	checkf(!ContainsDefault(InPageID), TEXT("Page default with given ID already exists"));
 	return Defaults.Add_GetRef(FMetasoundFrontendClassInputDefault { InPageID }).Literal;
 }
@@ -919,8 +967,19 @@ const TArray<FMetasoundFrontendClassInputDefault>& FMetasoundFrontendClassInput:
 
 FMetasoundFrontendLiteral& FMetasoundFrontendClassInput::InitDefault()
 {
+	using namespace Metasound::Frontend;
+
 	checkf(Defaults.IsEmpty(), TEXT("Default(s) already initialized"));
-	return Defaults.Add_GetRef(FMetasoundFrontendClassInputDefault { Metasound::Frontend::DefaultPageID }).Literal;
+	FMetasoundFrontendLiteral& NewLiteral = Defaults.Add_GetRef(FMetasoundFrontendClassInputDefault
+	{
+		::Metasound::Frontend::DefaultPageID
+	}).Literal;
+
+	if (IDataTypeRegistry::Get().IsRegistered(TypeName))
+	{
+		NewLiteral.SetFromLiteral(IDataTypeRegistry::Get().CreateDefaultLiteral(TypeName));
+	}
+	return NewLiteral;
 }
 
 void FMetasoundFrontendClassInput::InitDefault(FMetasoundFrontendLiteral InitLiteral)
@@ -956,9 +1015,7 @@ void FMetasoundFrontendClassInput::ResetDefaults()
 	using namespace Metasound::Frontend;
 
 	Defaults.Reset();
-	FMetasoundFrontendLiteral TypeDefault;
-	TypeDefault.SetFromLiteral(IDataTypeRegistry::Get().CreateDefaultLiteral(TypeName));
-	InitDefault(MoveTemp(TypeDefault));
+	InitDefault();
 	Defaults.Shrink();
 }
 
@@ -1173,13 +1230,6 @@ TArray<FMetasoundFrontendGraph>& FMetasoundFrontendGraphClass::IPropertyVersionT
 	return GraphClass.PagedGraphs;
 }
 
-void UMetaSoundFrontendMemberMetadata::IterateDefaults(TFunctionRef<void(const FGuid&, FMetasoundFrontendLiteral)> Iter) const
-{
-	// TODO: Replace with follow-up change using page-addressed defaults
-	Iter(Metasound::Frontend::DefaultPageID, GetDefault());
-}
-
-
 FMetasoundFrontendVersionNumber FMetasoundFrontendDocument::GetMaxVersion()
 {
 	return Metasound::Frontend::GetMaxDocumentVersion();
diff --git a/Engine/Plugins/Runtime/Metasound/Source/MetasoundFrontend/Private/MetasoundFrontendNodeController.cpp b/Engine/Plugins/Runtime/Metasound/Source/MetasoundFrontend/Private/MetasoundFrontendNodeController.cpp
--- a/Engine/Plugins/Runtime/Metasound/Source/MetasoundFrontend/Private/MetasoundFrontendNodeController.cpp
+++ b/Engine/Plugins/Runtime/Metasound/Source/MetasoundFrontend/Private/MetasoundFrontendNodeController.cpp
@@ -961,18 +961,12 @@ namespace Metasound
 				}
 			}
 
-			Algo::Transform(OutInterfaceUpdates.RegistryClass.Interface.Inputs, OutInterfaceUpdates.AddedInputs, [&](const FMetasoundFrontendClassInput& Input) { return &Input; });
+			Algo::Transform(OutInterfaceUpdates.RegistryClass.Interface.Inputs, OutInterfaceUpdates.AddedInputs, [](const FMetasoundFrontendClassInput& Input) { return &Input; });
 			for (const FMetasoundFrontendClassInput& Input : NodeClassInterface.Inputs)
 			{
-				const FMetasoundFrontendLiteral& NodeClassDefault = Input.FindConstDefaultChecked(Frontend::DefaultPageID);
-				auto IsEquivalent = [&NodeClassDefault, &Input](const FMetasoundFrontendClassInput* Iter)
+				auto IsEquivalent = [&Input](const FMetasoundFrontendClassInput* RegistryInput)
 				{
-					const bool bDefaultEquivalent = Iter->FindConstDefaultChecked(Frontend::DefaultPageID).IsEqual(NodeClassDefault);
-					if (bDefaultEquivalent)
-					{
-						return FMetasoundFrontendClassVertex::IsFunctionalEquivalent(Input, *Iter);
-					}
-					return false;
+					return FMetasoundFrontendClassInput::IsFunctionalEquivalent(Input, *RegistryInput);
 				};
 
 				const int32 Index = OutInterfaceUpdates.AddedInputs.FindLastByPredicate(IsEquivalent);
@@ -986,12 +980,13 @@ namespace Metasound
 				}
 			}
 
-			Algo::Transform(OutInterfaceUpdates.RegistryClass.Interface.Outputs, OutInterfaceUpdates.AddedOutputs, [&](const FMetasoundFrontendClassOutput& Output) { return &Output; });
+
+			Algo::Transform(OutInterfaceUpdates.RegistryClass.Interface.Outputs, OutInterfaceUpdates.AddedOutputs, [](const FMetasoundFrontendClassOutput& Output) { return &Output; });
 			for (const FMetasoundFrontendClassOutput& Output : NodeClassInterface.Outputs)
 			{
-				auto IsFunctionalEquivalent = [NodeClassOutput = &Output](const FMetasoundFrontendClassOutput* Iter)
+				auto IsFunctionalEquivalent = [&Output](const FMetasoundFrontendClassOutput* Iter)
 				{
-					return FMetasoundFrontendClassVertex::IsFunctionalEquivalent(*NodeClassOutput, *Iter);
+					return FMetasoundFrontendClassVertex::IsFunctionalEquivalent(Output, *Iter);
 				};
 
 				const int32 Index = OutInterfaceUpdates.AddedOutputs.FindLastByPredicate(IsFunctionalEquivalent);
diff --git a/Engine/Plugins/Runtime/Metasound/Source/MetasoundFrontend/Private/MetasoundFrontendTransform.cpp b/Engine/Plugins/Runtime/Metasound/Source/MetasoundFrontend/Private/MetasoundFrontendTransform.cpp
--- a/Engine/Plugins/Runtime/Metasound/Source/MetasoundFrontend/Private/MetasoundFrontendTransform.cpp
+++ b/Engine/Plugins/Runtime/Metasound/Source/MetasoundFrontend/Private/MetasoundFrontendTransform.cpp
@@ -867,7 +867,22 @@ namespace Metasound
 						ClassInput.NodeID = ExistingClassInput->NodeID;
 					}
 
-					FMetasoundFrontendLiteral& DefaultLiteral = ClassInput.InitDefault();
+					auto InheritDefaultsFromGraph = [&ClassInput, &NodeName](const FConstGraphHandle& GraphHandle)
+					{
+						const FGuid GraphVertexID = GraphHandle->GetVertexIDForInputVertex(NodeName);
+						if (const FMetasoundFrontendClassInput* GraphClassInput = GraphHandle->FindInputDescriptionWithVertexID(GraphVertexID))
+						{
+							GraphClassInput->IterateDefaults([&ClassInput](const FGuid& PageID, const FMetasoundFrontendLiteral& Literal)
+							{
+								ClassInput.AddDefault(PageID) = Literal;
+							});
+						}
+						else
+						{
+							ClassInput.InitDefault();
+						}
+					};
+
 					if (InPresetGraph->ContainsInputVertex(NodeName, ClassInput.TypeName))
 					{
 						// If the input vertex already exists in the parent graph,
@@ -875,22 +890,16 @@ namespace Metasound
 						// input names.
 						if (InPresetGraph->GetInputsInheritingDefault().Contains(NodeName))
 						{
-							const FGuid ReferencedVertexID = ReferencedGraph->GetVertexIDForInputVertex(NodeName);
-							DefaultLiteral = ReferencedGraph->GetDefaultInput(ReferencedVertexID);
+							InheritDefaultsFromGraph(ReferencedGraph);
 						}
 						else
 						{
-							FGuid VertexID = InPresetGraph->GetVertexIDForInputVertex(NodeName);
-							DefaultLiteral = InPresetGraph->GetDefaultInput(VertexID);
+							InheritDefaultsFromGraph(InPresetGraph);
 						}
 					}
 					else
 					{
-						// If the input vertex does not exist on the parent graph,
-						// then it is a new vertex and should use the default value
-						// of the referenced graph.
-						const FGuid ReferencedVertexID = ReferencedGraph->GetVertexIDForInputVertex(NodeName);
-						DefaultLiteral = ReferencedGraph->GetDefaultInput(ReferencedVertexID);
+						InheritDefaultsFromGraph(ReferencedGraph);
 					}
 
 					ClassInputs.Add(MoveTemp(ClassInput));
diff --git a/Engine/Plugins/Runtime/Metasound/Source/MetasoundFrontend/Public/MetasoundFrontendDocument.h b/Engine/Plugins/Runtime/Metasound/Source/MetasoundFrontend/Public/MetasoundFrontendDocument.h
--- a/Engine/Plugins/Runtime/Metasound/Source/MetasoundFrontend/Public/MetasoundFrontendDocument.h
+++ b/Engine/Plugins/Runtime/Metasound/Source/MetasoundFrontend/Public/MetasoundFrontendDocument.h
@@ -949,23 +949,11 @@ struct METASOUNDFRONTEND_API FMetasoundFrontendClassInputDefault
 	GENERATED_BODY()
 
 	FMetasoundFrontendClassInputDefault() = default;
+	FMetasoundFrontendClassInputDefault(FMetasoundFrontendLiteral InLiteral);
+	FMetasoundFrontendClassInputDefault(const FGuid& InPageID, FMetasoundFrontendLiteral InLiteral = { });
+	FMetasoundFrontendClassInputDefault(const FAudioParameter& InParameter);
 
-	FMetasoundFrontendClassInputDefault(const FMetasoundFrontendLiteral& InLiteral)
-		: Literal(InLiteral)
-		, PageID(Metasound::Frontend::DefaultPageID)
-	{
-	}
-
-	FMetasoundFrontendClassInputDefault(const FGuid& InPageID, FMetasoundFrontendLiteral InLiteral = { })
-		: Literal(MoveTemp(InLiteral))
-		, PageID(InPageID)
-	{
-	}
-
-	FMetasoundFrontendClassInputDefault(const FAudioParameter& InParameter)
-		: Literal(InParameter)
-	{
-	}
+	static bool IsFunctionalEquivalent(const FMetasoundFrontendClassInputDefault& InLHS, const FMetasoundFrontendClassInputDefault& InRHS);
 
 	UPROPERTY()
 	FMetasoundFrontendLiteral Literal;
@@ -989,6 +977,8 @@ struct METASOUNDFRONTEND_API FMetasoundFrontendClassInput : public FMetasoundFro
 	FMetasoundFrontendLiteral DefaultLiteral;
 #endif // WITH_EDITORONLY_DATA
 
+	static bool IsFunctionalEquivalent(const FMetasoundFrontendClassInput& InLHS, const FMetasoundFrontendClassInput& InRHS);
+
 private:
 	UPROPERTY(EditAnywhere, Category = Parameters)
 	TArray<FMetasoundFrontendClassInputDefault> Defaults;
@@ -1786,14 +1776,18 @@ class METASOUNDFRONTEND_API UMetaSoundFrontendMemberMetadata : public UObject
 {
 	GENERATED_BODY()
 
-#if WITH_EDITOR
 public:
+	UE_DEPRECATED(5.5, "Implementation moved to child editor class instead of compiled out (not required by Frontend representation")
 	virtual void ForceRefresh() { }
-	virtual FMetasoundFrontendLiteral GetDefault() const { return FMetasoundFrontendLiteral(); }
-	virtual void IterateDefaults(TFunctionRef<void(const FGuid&, FMetasoundFrontendLiteral)> Iter) const;
+
+	UE_DEPRECATED(5.5, "Default is no longer required to be stored or represented in metadata and may differ in paged or non-paged implementation")
+	FMetasoundFrontendLiteral GetDefault() const { return FMetasoundFrontendLiteral(); }
+
+	UE_DEPRECATED(5.5, "Implementation moved to child editor class instead of compiled out (not required by Frontend representation")
 	virtual EMetasoundFrontendLiteralType GetLiteralType() const { return EMetasoundFrontendLiteralType::None; }
-	virtual void SetFromLiteral(const FMetasoundFrontendLiteral& InLiteral) { }
-#endif // WITH_EDITOR
+
+	UE_DEPRECATED(5.5, "Default is no longer required to be stored or represented in metadata and may differ in paged or non-paged implementation")
+	virtual void SetFromLiteral(const FMetasoundFrontendLiteral& InLiteral, const FGuid& InPageID = Metasound::Frontend::DefaultPageID) { }
 
 #if WITH_EDITORONLY_DATA
 	UPROPERTY()
