Shadowman strikes again while some construction workers are being lazy on the job

#tests compiles
#rnx

diff --git a/Engine/Plugins/Runtime/Metasound/Source/MetasoundEditor/Private/MetasoundEditorGraph.cpp b/Engine/Plugins/Runtime/Metasound/Source/MetasoundEditor/Private/MetasoundEditorGraph.cpp
--- a/Engine/Plugins/Runtime/Metasound/Source/MetasoundEditor/Private/MetasoundEditorGraph.cpp
+++ b/Engine/Plugins/Runtime/Metasound/Source/MetasoundEditor/Private/MetasoundEditorGraph.cpp
@@ -1267,7 +1267,7 @@ void UMetasoundEditorGraphOutput::UpdateFrontendDefaultLiteral(bool bPostTransac
 		{
 			FMetaSoundFrontendDocumentBuilder& Builder = GetFrontendBuilderChecked();
 			const FGuid& VertexID = FrontendNode->Interface.Inputs.Last().VertexID;
-			const FMetasoundFrontendVertexHandle VertexHandle(NodeID, VertexID);
+			const FMetasoundFrontendVertexHandle VertexHandle { NodeID, VertexID };
 			GraphPrivate::SetOrClearIfLiteralMatchesNodeVertexDefault(Builder, VertexHandle, DefaultLiteral);
 		}
 	}
@@ -1848,7 +1848,7 @@ void UMetasoundEditorGraphVariable::UpdateFrontendDefaultLiteral(bool bPostTrans
 		const FMetasoundFrontendVertex* Input = MutatorNode->Interface.Inputs.FindByPredicate([](const FMetasoundFrontendVertex& Vertex) { return Vertex.Name == METASOUND_GET_PARAM_NAME(InputData); });
 		if (ensure(Input))
 		{
-			const FMetasoundFrontendVertexHandle VertexHandle(MutatorNode->GetID(), Input->VertexID);
+			const FMetasoundFrontendVertexHandle VertexHandle { MutatorNode->GetID(), Input->VertexID };
 			GraphPrivate::SetOrClearIfLiteralMatchesNodeVertexDefault(Builder, VertexHandle, DefaultLiteral);
 		}
 	}
diff --git a/Engine/Plugins/Runtime/Metasound/Source/MetasoundEditor/Private/MetasoundEditorGraphMemberDefaults.cpp b/Engine/Plugins/Runtime/Metasound/Source/MetasoundEditor/Private/MetasoundEditorGraphMemberDefaults.cpp
--- a/Engine/Plugins/Runtime/Metasound/Source/MetasoundEditor/Private/MetasoundEditorGraphMemberDefaults.cpp
+++ b/Engine/Plugins/Runtime/Metasound/Source/MetasoundEditor/Private/MetasoundEditorGraphMemberDefaults.cpp
@@ -541,16 +541,16 @@ bool UMetasoundEditorGraphMemberDefaultBoolArray::Synchronize()
 
 	auto LiteralToMember = [](const TArray<bool>& Values)
 	{
-		TArray<FMetasoundEditorGraphMemberDefaultBoolRef> Defaults;
-		Algo::Transform(Values, Defaults, [](bool Value) { FMetasoundEditorGraphMemberDefaultBoolRef Ref; Ref.Value = Value; return Ref; });
-		return Defaults;
+		TArray<FMetasoundEditorGraphMemberDefaultBoolRef> DefaultValues;
+		Algo::Transform(Values, DefaultValues, [](bool Value) { FMetasoundEditorGraphMemberDefaultBoolRef Ref; Ref.Value = Value; return Ref; });
+		return DefaultValues;
 	};
 
 	auto MemberToLiteral = [](const TArray<FMetasoundEditorGraphMemberDefaultBoolRef>& Values)
 	{
-		TArray<bool> Defaults;
-		Algo::Transform(Values, Defaults, [](const FMetasoundEditorGraphMemberDefaultBoolRef& Ref) { return Ref.Value; });
-		return Defaults;
+		TArray<bool> DefaultValues;
+		Algo::Transform(Values, DefaultValues, [](const FMetasoundEditorGraphMemberDefaultBoolRef& Ref) { return Ref.Value; });
+		return DefaultValues;
 	};
 
 	return MemberDefaultsPrivate::SynchronizePageDefault<FMetasoundEditorMemberPageDefaultBoolArray, TArray<FMetasoundEditorGraphMemberDefaultBoolRef>, TArray<bool>>(
@@ -1120,16 +1120,16 @@ bool UMetasoundEditorGraphMemberDefaultIntArray::Synchronize()
 
 	auto LiteralToMember = [](const TArray<int32>& Values)
 	{
-		TArray<FMetasoundEditorGraphMemberDefaultIntRef> Defaults;
-		Algo::Transform(Values, Defaults, [](const int32& Value) { FMetasoundEditorGraphMemberDefaultIntRef Ref; Ref.Value = Value; return Ref; });
-		return Defaults;
+		TArray<FMetasoundEditorGraphMemberDefaultIntRef> DefaultValues;
+		Algo::Transform(Values, DefaultValues, [](const int32& Value) { FMetasoundEditorGraphMemberDefaultIntRef Ref; Ref.Value = Value; return Ref; });
+		return DefaultValues;
 	};
 
 	auto MemberToLiteral = [](const TArray<FMetasoundEditorGraphMemberDefaultIntRef>& Values)
 	{
-		TArray<int32> Defaults;
-		Algo::Transform(Values, Defaults, [](const FMetasoundEditorGraphMemberDefaultIntRef& Ref) { return Ref.Value; });
-		return Defaults;
+		TArray<int32> DefaultValues;
+		Algo::Transform(Values, DefaultValues, [](const FMetasoundEditorGraphMemberDefaultIntRef& Ref) { return Ref.Value; });
+		return DefaultValues;
 	};
 
 	return MemberDefaultsPrivate::SynchronizePageDefault<FMetasoundEditorMemberPageDefaultIntArray, TArray<FMetasoundEditorGraphMemberDefaultIntRef>, TArray<int32>>(
@@ -1478,16 +1478,16 @@ bool UMetasoundEditorGraphMemberDefaultObjectArray::Synchronize()
 
 	auto LiteralToMember = [](const TArray<UObject*>& Values)
 	{
-		TArray<FMetasoundEditorGraphMemberDefaultObjectRef> Defaults;
-		Algo::Transform(Values, Defaults, [](UObject* Value) { FMetasoundEditorGraphMemberDefaultObjectRef Ref; Ref.Object = Value; return Ref; });
-		return Defaults;
+		TArray<FMetasoundEditorGraphMemberDefaultObjectRef> DefaultValues;
+		Algo::Transform(Values, DefaultValues, [](UObject* Value) { FMetasoundEditorGraphMemberDefaultObjectRef Ref; Ref.Object = Value; return Ref; });
+		return DefaultValues;
 	};
 
 	auto MemberToLiteral = [](const TArray<FMetasoundEditorGraphMemberDefaultObjectRef>& Values)
 	{
-		TArray<UObject*> Defaults;
-		Algo::Transform(Values, Defaults, [](const FMetasoundEditorGraphMemberDefaultObjectRef& Ref) { return Ref.Object; });
-		return Defaults;
+		TArray<UObject*> DefaultValues;
+		Algo::Transform(Values, DefaultValues, [](const FMetasoundEditorGraphMemberDefaultObjectRef& Ref) { return Ref.Object; });
+		return DefaultValues;
 	};
 
 	return MemberDefaultsPrivate::SynchronizePageDefault<FMetasoundEditorMemberPageDefaultObjectArray, TArray<FMetasoundEditorGraphMemberDefaultObjectRef>, TArray<UObject*>>(
diff --git a/Engine/Plugins/Runtime/Metasound/Source/MetasoundEditor/Public/MetasoundEditorGraphMemberDefaults.h b/Engine/Plugins/Runtime/Metasound/Source/MetasoundEditor/Public/MetasoundEditorGraphMemberDefaults.h
--- a/Engine/Plugins/Runtime/Metasound/Source/MetasoundEditor/Public/MetasoundEditorGraphMemberDefaults.h
+++ b/Engine/Plugins/Runtime/Metasound/Source/MetasoundEditor/Public/MetasoundEditorGraphMemberDefaults.h
@@ -78,6 +78,12 @@ struct FMetasoundEditorMemberPageDefaultBool : public FMetasoundEditorMemberPage
 {
 	GENERATED_BODY()
 
+	FMetasoundEditorMemberPageDefaultBool() = default;
+	FMetasoundEditorMemberPageDefaultBool(const FGuid& InPageID)
+		: FMetasoundEditorMemberPageDefault(InPageID)
+	{
+	}
+
 	UPROPERTY(EditAnywhere, Category = DefaultValue, meta = (DisplayName = "Default"))
 	FMetasoundEditorGraphMemberDefaultBoolRef Value;
 };
@@ -127,6 +133,12 @@ struct FMetasoundEditorMemberPageDefaultBoolArray : public FMetasoundEditorMembe
 {
 	GENERATED_BODY()
 
+	FMetasoundEditorMemberPageDefaultBoolArray() = default;
+	FMetasoundEditorMemberPageDefaultBoolArray(const FGuid& InPageID)
+		: FMetasoundEditorMemberPageDefault(InPageID)
+	{
+	}
+
 	UPROPERTY(EditAnywhere, Category = DefaultValue, meta = (DisplayName = "Default"))
 	TArray<FMetasoundEditorGraphMemberDefaultBoolRef> Value;
 };
@@ -179,6 +191,12 @@ struct FMetasoundEditorMemberPageDefaultInt : public FMetasoundEditorMemberPageD
 {
 	GENERATED_BODY()
 
+	FMetasoundEditorMemberPageDefaultInt() = default;
+	FMetasoundEditorMemberPageDefaultInt(const FGuid& InPageID)
+		: FMetasoundEditorMemberPageDefault(InPageID)
+	{
+	}
+
 	UPROPERTY(EditAnywhere, Category = DefaultValue)
 	FMetasoundEditorGraphMemberDefaultIntRef Value;
 };
@@ -220,6 +238,12 @@ struct FMetasoundEditorMemberPageDefaultIntArray : public FMetasoundEditorMember
 {
 	GENERATED_BODY()
 
+	FMetasoundEditorMemberPageDefaultIntArray() = default;
+	FMetasoundEditorMemberPageDefaultIntArray(const FGuid& InPageID)
+		: FMetasoundEditorMemberPageDefault(InPageID)
+	{
+	}
+
 	UPROPERTY(EditAnywhere, Category = DefaultValue, meta = (DisplayName = "Default"))
 	TArray<FMetasoundEditorGraphMemberDefaultIntRef> Value;
 };
@@ -276,6 +300,12 @@ struct FMetasoundEditorMemberPageDefaultFloat : public FMetasoundEditorMemberPag
 {
 	GENERATED_BODY()
 
+	FMetasoundEditorMemberPageDefaultFloat() = default;
+	FMetasoundEditorMemberPageDefaultFloat(const FGuid& InPageID)
+		: FMetasoundEditorMemberPageDefault(InPageID)
+	{
+	}
+
 	UPROPERTY(EditAnywhere, Category = DefaultValue, meta = (DisplayName = "Default"))
 	float Value = 0.0f;
 };
@@ -368,6 +398,12 @@ struct FMetasoundEditorMemberPageDefaultFloatArray : public FMetasoundEditorMemb
 {
 	GENERATED_BODY()
 
+	FMetasoundEditorMemberPageDefaultFloatArray() = default;
+	FMetasoundEditorMemberPageDefaultFloatArray(const FGuid& InPageID)
+		: FMetasoundEditorMemberPageDefault(InPageID)
+	{
+	}
+
 	UPROPERTY(EditAnywhere, Category = DefaultValue, meta = (DisplayName = "Default"))
 	TArray<float> Value;
 };
@@ -409,6 +445,12 @@ struct FMetasoundEditorMemberPageDefaultString : public FMetasoundEditorMemberPa
 {
 	GENERATED_BODY()
 
+	FMetasoundEditorMemberPageDefaultString() = default;
+	FMetasoundEditorMemberPageDefaultString(const FGuid& InPageID)
+		: FMetasoundEditorMemberPageDefault(InPageID)
+	{
+	}
+
 	UPROPERTY(EditAnywhere, Category = DefaultValue, meta = (DisplayName = "Default"))
 	FString Value;
 };
@@ -450,6 +492,12 @@ struct FMetasoundEditorMemberPageDefaultStringArray : public FMetasoundEditorMem
 {
 	GENERATED_BODY()
 
+	FMetasoundEditorMemberPageDefaultStringArray() = default;
+	FMetasoundEditorMemberPageDefaultStringArray(const FGuid& InPageID)
+		: FMetasoundEditorMemberPageDefault(InPageID)
+	{
+	}
+
 	UPROPERTY(EditAnywhere, Category = DefaultValue, meta = (DisplayName = "Default"))
 	TArray<FString> Value;
 };
@@ -500,6 +548,12 @@ struct FMetasoundEditorMemberPageDefaultObjectRef : public FMetasoundEditorMembe
 {
 	GENERATED_BODY()
 
+	FMetasoundEditorMemberPageDefaultObjectRef() = default;
+	FMetasoundEditorMemberPageDefaultObjectRef(const FGuid& InPageID)
+		: FMetasoundEditorMemberPageDefault(InPageID)
+	{
+	}
+
 	UPROPERTY(EditAnywhere, Category = DefaultValue, meta = (ShowOnlyInnerProperties))
 	FMetasoundEditorGraphMemberDefaultObjectRef Value;
 };
@@ -538,6 +592,12 @@ struct FMetasoundEditorMemberPageDefaultObjectArray : public FMetasoundEditorMem
 {
 	GENERATED_BODY()
 
+	FMetasoundEditorMemberPageDefaultObjectArray() = default;
+	FMetasoundEditorMemberPageDefaultObjectArray(const FGuid& InPageID)
+		: FMetasoundEditorMemberPageDefault(InPageID)
+	{
+	}
+
 	UPROPERTY(EditAnywhere, Category = DefaultValue, meta = (DisplayName = "Default"))
 	TArray<FMetasoundEditorGraphMemberDefaultObjectRef> Value;
 };
