Missed static analysis warning (because I was running client-only locally in submit tool like a chump)

#tests no but ran non-unity properly using editor this time so points for that?
#rnx

diff --git a/Engine/Plugins/Runtime/Metasound/Source/MetasoundEditor/Private/MetasoundNodeDetailCustomization.cpp b/Engine/Plugins/Runtime/Metasound/Source/MetasoundEditor/Private/MetasoundNodeDetailCustomization.cpp
--- a/Engine/Plugins/Runtime/Metasound/Source/MetasoundEditor/Private/MetasoundNodeDetailCustomization.cpp
+++ b/Engine/Plugins/Runtime/Metasound/Source/MetasoundEditor/Private/MetasoundNodeDetailCustomization.cpp
@@ -393,41 +393,41 @@ namespace Metasound
 			BoolLiteral = DefaultBool;
 
 			TArray<TSharedPtr<IPropertyHandle>> DefaultValueHandles;
-			IDetailPropertyRow* Row = DefaultCategoryBuilder->AddExternalObjectProperty(TArray<UObject*>({ DefaultBool }), GET_MEMBER_NAME_CHECKED(UMetasoundEditorGraphMemberDefaultBool, Default));
-			if (ensure(Row))
+			const FName DefaultName = GET_MEMBER_NAME_CHECKED(UMetasoundEditorGraphMemberDefaultBool, Default);
+			if (IDetailPropertyRow* Row = DefaultCategoryBuilder->AddExternalObjectProperty(TArray<UObject*>({ DefaultBool }), DefaultName); ensure(Row))
 			{
 				DefaultValueHandles = { Row->GetPropertyHandle() };
 				Row->Visibility(GetDefaultVisibility());
-			}
 
-			// Enable widget options for editable inputs only 
-			bool bShowWidgetOptions = false;
-			if (const UMetasoundEditorGraphInput* ParentMember = Cast <UMetasoundEditorGraphInput>(InLiteral.FindMember()))
-			{
-				if (const UMetasoundEditorGraph* OwningGraph = ParentMember->GetOwningGraph())
+				// Enable widget options for editable inputs only 
+				bool bShowWidgetOptions = false;
+				if (const UMetasoundEditorGraphInput* ParentMember = Cast <UMetasoundEditorGraphInput>(InLiteral.FindMember()))
 				{
-					bShowWidgetOptions = OwningGraph->IsEditable();
+					if (const UMetasoundEditorGraph* OwningGraph = ParentMember->GetOwningGraph())
+					{
+						bShowWidgetOptions = OwningGraph->IsEditable();
+					}
 				}
-			}
 
-			// add input widget properties
-			if (bShowWidgetOptions)
-			{
-				Frontend::FDataTypeRegistryInfo DataTypeInfo;
-				MemberCustomizationPrivate::GetDataTypeFromElementPropertyHandle(DefaultValueHandles.Last(), DataTypeInfo);
+				// add input widget properties
+				if (bShowWidgetOptions)
+				{
+					Frontend::FDataTypeRegistryInfo DataTypeInfo;
+					MemberCustomizationPrivate::GetDataTypeFromElementPropertyHandle(DefaultValueHandles.Last(), DataTypeInfo);
 
-				const UMetasoundEditorSettings* EditorSettings = GetDefault<UMetasoundEditorSettings>();
-				check(EditorSettings)
+					const UMetasoundEditorSettings* EditorSettings = GetDefault<UMetasoundEditorSettings>();
+					check(EditorSettings);
 
-				if (EditorSettings->bUseAudioMaterialWidgets)
-				{
-					if (MemberCustomizationPrivate::GetPrimitiveTypeName(DataTypeInfo) != Metasound::GetMetasoundDataTypeName<Metasound::FTrigger>())
+					if (EditorSettings->bUseAudioMaterialWidgets)
 					{
-						IDetailCategoryBuilder& EditorOptionsBuilder = InDetailLayout.EditCategory("DefaultEditorOptions");
-						if (IDetailPropertyRow* TypeRow = EditorOptionsBuilder.AddExternalObjectProperty(TArray<UObject*>({ DefaultBool }), GET_MEMBER_NAME_CHECKED(UMetasoundEditorGraphMemberDefaultBool, WidgetType)))
+						if (MemberCustomizationPrivate::GetPrimitiveTypeName(DataTypeInfo) != Metasound::GetMetasoundDataTypeName<Metasound::FTrigger>())
 						{
-							DefaultValueHandles.Add(Row->GetPropertyHandle());
-							TypeRow->Visibility(GetDefaultVisibility());
+							IDetailCategoryBuilder& EditorOptionsBuilder = InDetailLayout.EditCategory("DefaultEditorOptions");
+							if (IDetailPropertyRow* TypeRow = EditorOptionsBuilder.AddExternalObjectProperty(TArray<UObject*>({ DefaultBool }), GET_MEMBER_NAME_CHECKED(UMetasoundEditorGraphMemberDefaultBool, WidgetType)))
+							{
+								DefaultValueHandles.Add(Row->GetPropertyHandle());
+								TypeRow->Visibility(GetDefaultVisibility());
+							}
 						}
 					}
 				}
