- Rename 'Target' to 'Can'Target' for clarity in Page Settings
- Return copy instead of reference of platform target ids to avoid shared access violation when cooking pages
- Add trace info to monitor cook times of page data

#rnx

diff --git a/Engine/Plugins/Runtime/Metasound/Source/MetasoundEngine/Private/MetasoundDocumentBuilderRegistry.cpp b/Engine/Plugins/Runtime/Metasound/Source/MetasoundEngine/Private/MetasoundDocumentBuilderRegistry.cpp
--- a/Engine/Plugins/Runtime/Metasound/Source/MetasoundEngine/Private/MetasoundDocumentBuilderRegistry.cpp
+++ b/Engine/Plugins/Runtime/Metasound/Source/MetasoundEngine/Private/MetasoundDocumentBuilderRegistry.cpp
@@ -99,24 +99,27 @@ namespace Metasound::Engine
 #if WITH_EDITORONLY_DATA
 	bool FDocumentBuilderRegistry::CookPages(FName PlatformName, FMetaSoundFrontendDocumentBuilder& Builder) const
 	{
+		METASOUND_LLM_SCOPE;
+		METASOUND_TRACE_CPUPROFILER_EVENT_SCOPE(FDocumentBuilderRegistry::CookPages);
+
 		bool bModified = false;
 
 		const UMetaSoundSettings* Settings = GetDefault<UMetaSoundSettings>();
 		check(Settings);
-		const TArray<FGuid>& PlatformTargetPageIDs = Settings->GetCookedTargetPageIDs(PlatformName);
+		const TArray<FGuid> PlatformTargetPageIDs = Settings->GetCookedTargetPageIDs(PlatformName);
 		checkf(!PlatformTargetPageIDs.IsEmpty(), TEXT("Must have at least one targeted page ID to cook MetaSound."));
 
-		auto StripPageEntries = [&](TArray<FGuid>& PageIdsToResolve, TFunctionRef<bool(const FGuid&)> RemovePageItem)
+		auto StripPageEntries = [&](TArray<FGuid>& PageIdsToResolve, TSet<FGuid>& ResolveTargetScratch, TFunctionRef<bool(const FGuid&)> RemovePageItem)
 		{
-			TSet<FGuid> ResolvedTargetIDs;
+			ResolveTargetScratch.Reset();
 			for (const FGuid& TargetPage : PlatformTargetPageIDs)
 			{
 				const FGuid PageID = ResolveTargetPageIDInternal(*Settings, PageIdsToResolve, TargetPage, PlatformName);
-				ResolvedTargetIDs.Add(PageID);
+				ResolveTargetScratch.Add(PageID);
 			}
 
-			auto IsResolvedTarget = [&ResolvedTargetIDs](const FGuid& PageID) { return ResolvedTargetIDs.Contains(PageID); };
-			checkf(!ResolvedTargetIDs.IsEmpty(), TEXT("Failed to resolve any valid target IDs, which will leave serialized page array in invalid state."));
+			auto IsResolvedTarget = [&ResolveTargetScratch](const FGuid& PageID) { return ResolveTargetScratch.Contains(PageID); };
+			checkf(!ResolveTargetScratch.IsEmpty(), TEXT("Failed to resolve any valid target IDs, which will leave serialized page array in invalid state."));
 			PageIdsToResolve.RemoveAllSwap(IsResolvedTarget, EAllowShrinking::No);
 
 			for (const FGuid& PageID : PageIdsToResolve)
@@ -128,6 +131,7 @@ namespace Metasound::Engine
 		const FMetasoundFrontendDocument& Document = Builder.GetConstDocumentChecked();
 
 		TArray<FGuid> ResolvePageIDs;
+		TSet<FGuid> ResolvedTargetScratchIDs;
 		{ // Strip graphs
 			auto AddPageID = [&ResolvePageIDs](const FMetasoundFrontendGraph& Graph) { ResolvePageIDs.Add(Graph.PageID); };
 			Document.RootGraph.IterateGraphPages(AddPageID);
@@ -137,7 +141,7 @@ namespace Metasound::Engine
 				return Builder.RemoveGraphPage(InPageID);
 			};
 
-			StripPageEntries(ResolvePageIDs, RemoveGraphPage);
+			StripPageEntries(ResolvePageIDs, ResolvedTargetScratchIDs, RemoveGraphPage);
 
 			checkf(!Document.RootGraph.GetConstGraphPages().IsEmpty(),
 				TEXT("Document in MetaSound asset '%s' had all default values "
@@ -158,7 +162,7 @@ namespace Metasound::Engine
 					return Builder.RemoveGraphInputDefault(GraphInput.Name, InPageID);
 				};
 
-				StripPageEntries(ResolvePageIDs, RemoveDefault);
+				StripPageEntries(ResolvePageIDs, ResolvedTargetScratchIDs, RemoveDefault);
 				checkf(!GraphInput.GetDefaults().IsEmpty(),
 					TEXT("Input '%s' had all default values stripped leaving it in an invalid state. "
 					"Input must always have at least one default value"),
diff --git a/Engine/Plugins/Runtime/Metasound/Source/MetasoundEngine/Private/MetasoundSettings.cpp b/Engine/Plugins/Runtime/Metasound/Source/MetasoundEngine/Private/MetasoundSettings.cpp
--- a/Engine/Plugins/Runtime/Metasound/Source/MetasoundEngine/Private/MetasoundSettings.cpp
+++ b/Engine/Plugins/Runtime/Metasound/Source/MetasoundEngine/Private/MetasoundSettings.cpp
@@ -144,7 +144,7 @@ bool FMetaSoundPageSettings::GetExcludeFromCook(FName PlatformName) const
 TArray<FName> FMetaSoundPageSettings::GetTargetPlatforms() const
 {
 	TArray<FName> PlatformNames;
-	Algo::TransformIf(Target.PerPlatform, PlatformNames,
+	Algo::TransformIf(CanTarget.PerPlatform, PlatformNames,
 		[](const TPair<FName, bool>& Pair) { return Pair.Value; },
 		[](const TPair<FName, bool>& Pair) { return Pair.Key; });
 	return PlatformNames;
@@ -152,7 +152,7 @@ TArray<FName> FMetaSoundPageSettings::GetTargetPlatforms() const
 
 bool FMetaSoundPageSettings::PlatformCanTargetPage(FName PlatformName) const
 {
-	const bool bIsTargeted = Target.GetValueForPlatform(PlatformName);
+	const bool bIsTargeted = CanTarget.GetValueForPlatform(PlatformName);
 	return bIsTargeted;
 }
 
@@ -175,8 +175,8 @@ void UMetaSoundSettings::ConformPageSettings(bool bNotifyDefaultRenamed)
 		bTargetFound |= TargetPageName == Page.Name;
 	}
 
-	DefaultPageSettings.Target = PageSettings.IsEmpty();
-	if (DefaultPageSettings.Target.GetValue())
+	DefaultPageSettings.CanTarget = PageSettings.IsEmpty();
+	if (DefaultPageSettings.CanTarget.GetValue())
 	{
 		TargetPageName = DefaultPageSettings.Name;
 	}
@@ -252,33 +252,46 @@ TArray<FName> UMetaSoundSettings::GetAllPlatformNamesImplementingTargets() const
 #endif // WITH_EDITOR
 
 #if WITH_EDITORONLY_DATA
-const TArray<FGuid>& UMetaSoundSettings::GetCookedTargetPageIDs(FName PlatformName) const
+TArray<FGuid> UMetaSoundSettings::GetCookedTargetPageIDs(FName PlatformName) const
 {
 	FScopeLock Lock(&CookPlatformTargetCritSec);
+	return GetCookedTargetPageIDsInternal(PlatformName);
+}
+
+const TArray<FGuid>& UMetaSoundSettings::GetCookedTargetPageIDsInternal(FName PlatformName) const
+{
 	if (PlatformName != CookPlatformTargetPage)
 	{
 		CookPlatformTargetPage = PlatformName;
 		CookPlatformTargetPageIDs.Reset();
-		Algo::TransformIf(PageSettings, CookPlatformTargetPageIDs,
-		[&PlatformName](const FMetaSoundPageSettings& PageSetting)
+		auto CanTargetPage = [&PlatformName](const FMetaSoundPageSettings& PageSetting)
 		{
 #if WITH_EDITOR
 			return PageSetting.PlatformCanTargetPage(PlatformName);
 #else // !WITH_EDITOR
 			return true;
 #endif // !WITH_EDITOR
-		},
-		[](const FMetaSoundPageSettings& PageSetting)
+		};
+
+		auto GetID = [](const FMetaSoundPageSettings& PageSetting)
 		{
 			return PageSetting.UniqueId;
-		});
+		};
+
+		if (CanTargetPage(DefaultPageSettings))
+		{
+			CookPlatformTargetPageIDs.Add(DefaultPageSettings.UniqueId);
+		}
+
+		Algo::TransformIf(PageSettings, CookPlatformTargetPageIDs, CanTargetPage, GetID);
+
 
 		if (CookPlatformTargetPageIDs.IsEmpty())
 		{
 #if WITH_EDITOR
-			const bool bCanTargetDefault = DefaultPageSettings.Target.GetValueForPlatform(PlatformName);
+			const bool bCanTargetDefault = DefaultPageSettings.CanTarget.GetValueForPlatform(PlatformName);
 #else // !WITH_EDITOR
-			const bool bCanTargetDefault = DefaultPageSettings.Target.GetValue();
+			const bool bCanTargetDefault = DefaultPageSettings.CanTarget.GetValue();
 #endif // !WITH_EDITOR
 
 			if (!PageSettings.IsEmpty() && !bCanTargetDefault)
@@ -292,6 +305,16 @@ const TArray<FGuid>& UMetaSoundSettings::GetCookedTargetPageIDs(FName PlatformNa
 
 	return CookPlatformTargetPageIDs;
 }
+
+void UMetaSoundSettings::IterateCookedTargetPageIDs(FName PlatformName, TFunctionRef<void(const FGuid&)> Iter) const
+{
+	FScopeLock Lock(&CookPlatformTargetCritSec);
+	const TArray<FGuid>& CookPlatforms = GetCookedTargetPageIDsInternal(PlatformName);
+	for (const FGuid& CookPlatform : CookPlatforms)
+	{
+		Iter(CookPlatform);
+	}
+}
 #endif // WITH_EDITORONLY_DATA
 
 const FMetaSoundPageSettings& UMetaSoundSettings::GetTargetPageSettings() const
@@ -300,7 +323,7 @@ const FMetaSoundPageSettings& UMetaSoundSettings::GetTargetPageSettings() const
 
 	if (const FMetaSoundPageSettings* TargetSettings = FindPageSettings(TargetPage))
 	{
-		if (TargetSettings->Target.GetValue())
+		if (TargetSettings->CanTarget.GetValue())
 		{
 			return *TargetSettings;
 		}
@@ -310,7 +333,7 @@ const FMetaSoundPageSettings& UMetaSoundSettings::GetTargetPageSettings() const
 	// try and return any page setting set as a valid target.
 	for (const FMetaSoundPageSettings& Setting : PageSettings)
 	{
-		if (Setting.Target.GetValue())
+		if (Setting.CanTarget.GetValue())
 		{
 			return Setting;
 		}
@@ -360,7 +383,7 @@ bool UMetaSoundSettings::SetTargetPage(FName PageName)
 	{
 		if (TargetPageName != PageSetting->Name)
 		{
-			if (PageSetting->Target.GetValue())
+			if (PageSetting->CanTarget.GetValue())
 			{
 				TargetPageNameOverride = PageSetting->Name;
 				return true;
diff --git a/Engine/Plugins/Runtime/Metasound/Source/MetasoundEngine/Public/MetasoundSettings.h b/Engine/Plugins/Runtime/Metasound/Source/MetasoundEngine/Public/MetasoundSettings.h
--- a/Engine/Plugins/Runtime/Metasound/Source/MetasoundEngine/Public/MetasoundSettings.h
+++ b/Engine/Plugins/Runtime/Metasound/Source/MetasoundEngine/Public/MetasoundSettings.h
@@ -71,12 +71,17 @@ struct METASOUNDENGINE_API FMetaSoundPageSettings
 
 private:
 	// When true, page can be targeted for the assigned platform(s)/platform group(s).
-	UPROPERTY(EditAnywhere, config, Category = "Pages")
-	FPerPlatformBool Target = true;
+	// Is used to determine which asset page data (i.e. graphs and input defaults) can be
+	// excluded from cook if page is not explicitly omitted via 'ExcludeFromCook' setting.
+	UPROPERTY(EditAnywhere, Category = "Pages")
+	FPerPlatformBool CanTarget = true;
 
 #if WITH_EDITORONLY_DATA
-	// When true, exclude page data when cooking from the assigned platform(s)/platform group(s) (ignored if target is true for corresponding platform/group).
-	UPROPERTY(EditAnywhere, config, Category = "Pages")
+	// When true, exclude page data when cooking from the assigned platform(s)/platform group(s).
+	// If false, page data may or may not be included in cook depending on whether or not the given
+	// page data is required in order to ensure a value is always resolved for the cook platform target(s).
+	// (Ignored if CanTarget true for corresponding platform/group).
+	UPROPERTY(EditAnywhere, Category = "Pages")
 	FPerPlatformBool ExcludeFromCook = false;
 #endif //WITH_EDITORONLY_DATA
 
@@ -208,7 +213,11 @@ public:
 	const FMetaSoundPageSettings& GetDefaultPageSettings() const;
 
 #if WITH_EDITORONLY_DATA
-	const TArray<FGuid>& GetCookedTargetPageIDs(FName PlatformName) const;
+	// Returns PageIDs to be cooked to the given platform/platform group.
+	TArray<FGuid> GetCookedTargetPageIDs(FName PlatformName) const;
+
+	// Iterates PageIDs cooked to the given platform/platform group.
+	void IterateCookedTargetPageIDs(FName PlatformName, TFunctionRef<void(const FGuid&)> Iter) const;
 #endif // WITH_EDITORONLY_DATA
 
 #if WITH_EDITOR
@@ -255,8 +264,14 @@ public:
 	*/
 	UFUNCTION()
 	static TArray<FName> GetQualityNames();
+#endif // WITH_EDITOR
 
 private:
+#if WITH_EDITORONLY_DATA
+	const TArray<FGuid>& GetCookedTargetPageIDsInternal(FName PlatformName) const;
+#endif // WITH_EDITORONLY_DATA
+
+#if WITH_EDITOR
 	void ConformPageSettings(bool bNotifyDefaultRenamed);
 
 	virtual void PostEditChangeChainProperty(FPropertyChangedChainEvent& PropertyChangedEvent) override;
