Speculative fix for rare switch and ramp node crash
#rnx

diff --git a/Engine/Plugins/Runtime/Metasound/Source/MetasoundStandardNodes/Private/MetasoundSwitchAndRampNode.cpp b/Engine/Plugins/Runtime/Metasound/Source/MetasoundStandardNodes/Private/MetasoundSwitchAndRampNode.cpp
--- a/Engine/Plugins/Runtime/Metasound/Source/MetasoundStandardNodes/Private/MetasoundSwitchAndRampNode.cpp
+++ b/Engine/Plugins/Runtime/Metasound/Source/MetasoundStandardNodes/Private/MetasoundSwitchAndRampNode.cpp
@@ -156,7 +156,9 @@ namespace Metasound
 				//copy the rest of the audio as passthrough
 				int32 StartCopyFrame = FMath::Max(StartFrame, LastRampFrame);
 				if (StartCopyFrame < EndFrame)
+				{
 					FMemory::Memcpy(&OutputAudio[StartCopyFrame], &InputAudio[StartCopyFrame], sizeof(float) * (EndFrame - StartCopyFrame));
+				}
 
 				MostRecentOutputValue = OutputAudio[EndFrame - 1];
 			};
@@ -166,15 +168,21 @@ namespace Metasound
 				// OnTrigger
 				[this, &RunRampLambda](int32 StartFrame, int32 EndFrame)
 				{
-					float RampFromValue = MostRecentOutputValue;
-					float RampToValue = AudioIn->GetData()[StartFrame];
-					DiscontinuityAmount = RampFromValue - RampToValue;
+					// Only run ramp if there are samples to render out. Multiple triggers
+					// on a single frame result empty frame ranges (e.g. StartFrame == EndFrame)
+					// and the underlying DSP expects at least 1 renderable audio frame.
+					if (EndFrame > StartFrame)
+					{
+						float RampFromValue = MostRecentOutputValue;
+						float RampToValue = AudioIn->GetData()[StartFrame];
+						DiscontinuityAmount = RampFromValue - RampToValue;
 
-					float SmoothTimeSeconds = SmoothTime->GetSeconds();
-					NumRampSamples = FMath::Max(1, SampleRate * SmoothTimeSeconds);
-					RampSampleIndex = 0;
+						float SmoothTimeSeconds = SmoothTime->GetSeconds();
+						NumRampSamples = FMath::Max(1, SampleRate * SmoothTimeSeconds);
+						RampSampleIndex = 0;
 
-					RunRampLambda(StartFrame, EndFrame);
+						RunRampLambda(StartFrame, EndFrame);
+					}
 				}
 			);
 		}
