- Enable default page to not be targetable, allowing for removal at cook
- Add faux platform "editor" (ignores cook settings) and "default" (all platforms without specific target settings) as auditionable for faster workflow/more intuitive behavior
- Implement feedback to distinguish "audition page" (editor only concept) vs "target page" (runtime concept)
- Warn when auditioning  graphs that are not cooked/targeted for an assigned audition platform (as default input values cannot be resolved accordingly) in MetaSound Asset editor
- Add more editor/UX protections around having invalid settings where a platform has no target set

#rnx

diff --git a/Engine/Plugins/Runtime/Metasound/Source/MetasoundEditor/Private/MetasoundDefaultLiteralCustomization.cpp b/Engine/Plugins/Runtime/Metasound/Source/MetasoundEditor/Private/MetasoundDefaultLiteralCustomization.cpp
--- a/Engine/Plugins/Runtime/Metasound/Source/MetasoundEditor/Private/MetasoundDefaultLiteralCustomization.cpp
+++ b/Engine/Plugins/Runtime/Metasound/Source/MetasoundEditor/Private/MetasoundDefaultLiteralCustomization.cpp
@@ -12,6 +12,7 @@
 #include "MetasoundSettings.h"
 #include "PropertyCustomizationHelpers.h"
 #include "SSearchableComboBox.h"
+#include "UObject/WeakObjectPtrTemplates.h"
 
 #define LOCTEXT_NAMESPACE "MetaSoundEditor"
 
@@ -226,75 +227,109 @@ namespace Metasound::Editor
 		const FGuid PageID = LiteralCustomizationPrivate::GetGuidPropertyValue(ElementProperty->GetChildHandle(GET_MEMBER_NAME_CHECKED(FMetasoundEditorMemberPageDefault, PageID)));
 		FName PageName;
 		PageNameProperty->GetValue(PageName);
-		// Can't delete the default page, so just return name selector
-		if (PageID == Frontend::DefaultPageID)
-		{
-			return PageNameProperty->CreatePropertyValueWidget();
-		}
 
-		TWeakObjectPtr<UMetasoundEditorGraphMemberDefaultLiteral> LiteralPtr(&Literal);
+		TSharedRef<SHorizontalBox> NameBox = SNew(SHorizontalBox)
+			+ SHorizontalBox::Slot()
+			.AutoWidth()
+			.VAlign(VAlign_Center)
+			.Padding(2, 0)
+			[
+				PageNameProperty->CreatePropertyValueWidget()
+			];
 
-		UMetasoundEditorGraphMember* Member = LiteralPtr->FindMember();
+		const UMetasoundEditorGraphMember* Member = Literal.FindMember();
 		if (!ensure(Member))
 		{
 			return SNullWidget::NullWidget;
 		}
 
-		const FText RemoveDescription = FText::Format(LOCTEXT("RemovePageDefaultTransactionFormat", "Remove '{0}' Page '{1}' Default Value"), FText::FromName(Member->GetMemberName()), FText::FromName(PageName));
-		TSharedRef<SWidget> RemovePageDefaultButton = PropertyCustomizationHelpers::MakeDeleteButton(FSimpleDelegate::CreateLambda([this, RemoveDescription, PageID, LiteralPtr]()
+		if (Member->IsDefaultPaged())
 		{
-			using namespace Frontend;
-
-			if (!LiteralPtr.IsValid())
+			TAttribute<FText> ToolTip = LOCTEXT("MetaSound_ExecutingInputPageDefaultTooltip", "Initial default used upon executing preview.");
+			TAttribute<EVisibility> ExecVisibility = TAttribute<EVisibility>::CreateLambda([MemberPtr = TWeakObjectPtr<const UMetasoundEditorGraphMember>(Member), PageID]()
 			{
-				return;
-			}
+				if (const UMetasoundEditorGraphMember* GraphMember = MemberPtr.Get())
+				{
+					const FMetaSoundFrontendDocumentBuilder& Builder = GraphMember->GetFrontendBuilderChecked();
+					if (const FMetasoundFrontendClassInput* ClassInput = Builder.FindGraphInput(GraphMember->GetMemberName()))
+					{
+						const bool bIsPreviewing = Editor::IsPreviewingPageInputDefault(Builder, *ClassInput, PageID);
+						return bIsPreviewing ? EVisibility::Visible : EVisibility::Collapsed;
+					}
+					return EVisibility::Collapsed;
+				}
+
+				return EVisibility::Collapsed;
+			});
+			TSharedRef<SWidget> ExecImageWidget = SNew(SImage)
+				.Image(Style::CreateSlateIcon("MetasoundEditor.Page.Executing").GetIcon())
+				.ColorAndOpacity(FStyleColors::AccentGreen)
+				.Visibility(MoveTemp(ExecVisibility));
+			NameBox->AddSlot()
+				.FillWidth(1)
+				.HAlign(HAlign_Left)
+				.VAlign(VAlign_Center)
+				.Padding(2, 0)
+				[
+					ExecImageWidget
+				];
+		}
 
-			UMetasoundEditorGraphMember* Member = LiteralPtr->FindMember();
-			if (!Member)
+		// Can't delete the default page, so only show remove page field if non-default/project defined page.
+		if (PageID != Frontend::DefaultPageID)
+		{
+			TWeakObjectPtr<UMetasoundEditorGraphMemberDefaultLiteral> LiteralPtr(&Literal);
+
+			const FText RemoveDescription = FText::Format(LOCTEXT("RemovePageDefaultTransactionFormat", "Remove '{0}' Page '{1}' Default Value"), FText::FromName(Member->GetMemberName()), FText::FromName(PageName));
+			TSharedRef<SWidget> RemovePageDefaultButton = PropertyCustomizationHelpers::MakeDeleteButton(FSimpleDelegate::CreateLambda([this, RemoveDescription, PageID, LiteralPtr]()
 			{
-				return;
-			}
+				using namespace Frontend;
+
+				if (!LiteralPtr.IsValid())
+				{
+					return;
+				}
+
+				UMetasoundEditorGraphMember* Member = LiteralPtr->FindMember();
+				if (!Member)
+				{
+					return;
+				}
 
-			const FScopedTransaction Transaction(RemoveDescription);
+				const FScopedTransaction Transaction(RemoveDescription);
 
-			const UMetaSoundSettings* Settings = GetDefault<UMetaSoundSettings>();
-			check(Settings);
+				const UMetaSoundSettings* Settings = GetDefault<UMetaSoundSettings>();
+				check(Settings);
 
-			FMetaSoundFrontendDocumentBuilder& Builder = Member->GetFrontendBuilderChecked();
-			UObject& MetaSound = Builder.CastDocumentObjectChecked<UObject>();
-			MetaSound.Modify();
-			LiteralPtr->Modify();
+				FMetaSoundFrontendDocumentBuilder& Builder = Member->GetFrontendBuilderChecked();
+				UObject& MetaSound = Builder.CastDocumentObjectChecked<UObject>();
+				MetaSound.Modify();
+				LiteralPtr->Modify();
 
-			LiteralPtr->RemoveDefault(PageID);
+				LiteralPtr->RemoveDefault(PageID);
 
-			constexpr bool bPostTransaction = false;
-			Member->UpdateFrontendDefaultLiteral(bPostTransaction);
+				constexpr bool bPostTransaction = false;
+				Member->UpdateFrontendDefaultLiteral(bPostTransaction);
 
-			FMetasoundAssetBase& MetasoundAsset = Editor::FGraphBuilder::GetOutermostMetaSoundChecked(*LiteralPtr);
-			MetasoundAsset.GetModifyContext().AddMemberIDsModified({ Member->GetMemberID() });
-			UpdatePagePickerNames(LiteralPtr);
-			PageDefaultComboBox->RefreshOptions();
-			FGraphBuilder::RegisterGraphWithFrontend(MetaSound);
-		}),
-		RemoveDescription);
-
-		return SNew(SHorizontalBox)
-		+ SHorizontalBox::Slot()
-		.AutoWidth()
-		.VAlign(VAlign_Center)
-		.Padding(2, 0)
-		[
-			PageNameProperty->CreatePropertyValueWidget()
-		]
-		+ SHorizontalBox::Slot()
-		.FillWidth(1)
-		.HAlign(HAlign_Left)
-		.VAlign(VAlign_Center)
-		.Padding(2, 0)
-		[
-			RemovePageDefaultButton
-		];
+				FMetasoundAssetBase& MetasoundAsset = Editor::FGraphBuilder::GetOutermostMetaSoundChecked(*LiteralPtr);
+				MetasoundAsset.GetModifyContext().AddMemberIDsModified({ Member->GetMemberID() });
+				UpdatePagePickerNames(LiteralPtr);
+				PageDefaultComboBox->RefreshOptions();
+				FGraphBuilder::RegisterGraphWithFrontend(MetaSound);
+			}),
+			RemoveDescription);
+
+			NameBox->AddSlot()
+			.FillWidth(1)
+			.HAlign(HAlign_Left)
+			.VAlign(VAlign_Center)
+			.Padding(2, 0)
+			[
+				RemovePageDefaultButton
+			];
+		}
+
+		return NameBox;
 	}
 
 	void FMetasoundDefaultLiteralCustomizationBase::CustomizeDefaults(UMetasoundEditorGraphMemberDefaultLiteral& InLiteral, IDetailLayoutBuilder& InDetailLayout)
@@ -328,15 +363,9 @@ namespace Metasound::Editor
 			if (DefaultValueArray.IsValid())
 			{
 				DefaultValueArray->GetNumElements(NumElements);
-
-				bool bShowPageModifiers = Editor::PageEditorEnabled();
-				if (const UMetaSoundSettings* Settings = GetDefault<UMetaSoundSettings>())
-				{
-					// Don't show if project-specific defaults exist and/or project-level pages are implemented.
-					// If one of either, assumption is its the required default page.
-					bShowPageModifiers &= !Settings->GetProjectPageSettings().IsEmpty() || NumElements > 1;
-				}
-
+				const bool bHasProjectPageValues = NumElements > 1; // 1 value is only the default
+				constexpr bool bPresetCanEditPageValues = true;
+				const bool bShowPageModifiers = Editor::PageEditorEnabled(Member->GetFrontendBuilderChecked(), bHasProjectPageValues, bPresetCanEditPageValues);
 				if (bIsPagedDefault && bShowPageModifiers)
 				{
 					BuildPageDefaultComboBox(InLiteral, DefaultPageArrayHandle->GetPropertyDisplayName());
@@ -438,10 +467,13 @@ namespace Metasound::Editor
 			ImplementedGuids.Add(PageID);
 		});
 
-		auto GetStringName = [](const FMetaSoundPageSettings& Page) -> TSharedPtr<FString> { return MakeShared<FString>(Page.Name.ToString()); };
-		auto RemoveImplementedItem = [&ImplementedGuids](const FMetaSoundPageSettings& Page) { return !ImplementedGuids.Contains(Page.UniqueId); };
-		AddablePageStringNames.Add(GetStringName(Settings->GetDefaultPageSettings()));
-		Algo::TransformIf(Settings->GetProjectPageSettings(), AddablePageStringNames, RemoveImplementedItem, GetStringName);
+		Settings->IteratePageSettings([this, &ImplementedGuids](const FMetaSoundPageSettings& PageSettings)
+		{
+			if (!ImplementedGuids.Contains(PageSettings.UniqueId))
+			{
+				AddablePageStringNames.Add(MakeShared<FString>(PageSettings.Name.ToString()));
+			}
+		});
 
 		auto GetPageName = [&Settings](const FGuid& PageID)
 		{
diff --git a/Engine/Plugins/Runtime/Metasound/Source/MetasoundEditor/Private/MetasoundDetailCustomization.cpp b/Engine/Plugins/Runtime/Metasound/Source/MetasoundEditor/Private/MetasoundDetailCustomization.cpp
--- a/Engine/Plugins/Runtime/Metasound/Source/MetasoundEditor/Private/MetasoundDetailCustomization.cpp
+++ b/Engine/Plugins/Runtime/Metasound/Source/MetasoundEditor/Private/MetasoundDetailCustomization.cpp
@@ -17,6 +17,7 @@
 #include "MetasoundDocumentBuilderRegistry.h"
 #include "MetasoundEditor.h"
 #include "MetasoundEditorGraphBuilder.h"
+#include "MetasoundEditorModule.h"
 #include "MetasoundEditorSettings.h"
 #include "MetasoundEditorSubsystem.h"
 #include "MetasoundFrontend.h"
@@ -35,6 +36,7 @@
 #include "Styling/CoreStyle.h"
 #include "Styling/SlateColor.h"
 #include "Styling/SlateTypes.h"
+#include "Styling/StyleColors.h"
 #include "Templates/Casts.h"
 #include "Templates/SharedPointer.h"
 #include "UObject/WeakObjectPtr.h"
@@ -487,28 +489,38 @@ namespace Metasound::Editor
 		{
 			using namespace Frontend;
 
-			TSharedRef<SWidget> SelectButtonWidget = PropertyCustomizationHelpers::MakeUseSelectedButton(FSimpleDelegate::CreateLambda([this, InEntryPageID, InName]()
-			{
-				constexpr bool bOpenEditor = false; // Already focused by user action
-				const bool bFocusedPage = UMetaSoundEditorSubsystem::GetConstChecked().SetFocusedPage(*Builder.Get(), InEntryPageID, bOpenEditor);
-				if (bFocusedPage)
-				{
-					BuildPageName = InName;
-				}
-			}),
-			TAttribute<FText>::Create([this, InName]()
-			{
-				return BuildPageName == InName
-					? LOCTEXT("FocusedPageTooltip", "Currently focused page.")
-					: LOCTEXT("SetFocusedPageTooltip", "Sets the actively focused graph page of the MetaSound.");
-			}),
-			TAttribute<bool>::Create([this, InName]()
+			TSharedRef<SHorizontalBox> EntryWidget = SNew(SHorizontalBox);
+
+			// Page Focus
 			{
-				return BuildPageName != InName;
-			}));
+				TSharedRef<SWidget> SelectButtonWidget = PropertyCustomizationHelpers::MakeUseSelectedButton(
+					FSimpleDelegate::CreateLambda([this, PageID = InEntryPageID, InName]()
+					{
+						constexpr bool bOpenEditor = false; // Already focused by user action
+						UMetaSoundEditorSubsystem::GetConstChecked().SetFocusedPage(*Builder.Get(), PageID, bOpenEditor);
+						BuildPageName = InName;
+					}),
+					TAttribute<FText>::Create([this, InName]()
+					{
+						return BuildPageName == InName
+							? LOCTEXT("FocusedPageTooltip", "Currently focused page.")
+							: LOCTEXT("SetFocusedPageTooltip", "Sets the actively focused graph page of the MetaSound.");
+					}),
+					TAttribute<bool>::Create([this, InName]()
+					{
+						return BuildPageName != InName;
+					}));
 
-			TSharedRef<SHorizontalBox> EntryWidget = SNew(SHorizontalBox)
-				+ SHorizontalBox::Slot()
+				EntryWidget->AddSlot()
+					.Padding(2.0f)
+					.HAlign(HAlign_Center)
+					.VAlign(VAlign_Center)
+					.AutoWidth()[ SelectButtonWidget ];
+			}
+
+			// Page Name
+			{
+				EntryWidget->AddSlot()
 				.Padding(2.0f)
 				.HAlign(HAlign_Left)
 				.VAlign(VAlign_Center)
@@ -517,40 +529,62 @@ namespace Metasound::Editor
 					SNew(STextBlock)
 					.Text(FText::FromName(InName))
 				];
+			}
 
+			// Page Remove
 			if (!bIsDefault)
 			{
-				TSharedRef<SWidget> RemoveButtonWidget = PropertyCustomizationHelpers::MakeDeleteButton(FSimpleDelegate::CreateLambda([this, InName, InEntryPageID]()
-				{
-					using namespace Frontend;
-					const FScopedTransaction Transaction(FText::Format(LOCTEXT("RemovePageTransactionFormat", "Remove MetaSound Page '{0}'"), FText::FromName(InName)));
-					UObject& MetaSound = GetMetaSound();
-					MetaSound.Modify();
+				TSharedRef<SWidget> RemoveButtonWidget = PropertyCustomizationHelpers::MakeDeleteButton(
+					FSimpleDelegate::CreateLambda([this, InName, InEntryPageID]()
+					{
+						using namespace Frontend;
+						const FScopedTransaction Transaction(FText::Format(LOCTEXT("RemovePageTransactionFormat", "Remove MetaSound Page '{0}'"), FText::FromName(InName)));
+						UObject& MetaSound = GetMetaSound();
+						MetaSound.Modify();
 
-					// Removal may modify the builder's build page ID if it is the currently set value
-					Builder->Modify();
+						// Removal may modify the builder's build page ID if it is the currently set value
+						Builder->Modify();
+
+						const bool bGraphRemoved = Builder->GetBuilder().RemoveGraphPage(InEntryPageID);
+						if (bGraphRemoved)
+						{
+							UpdateItemNames();
+							ComboBox->RefreshOptions();
+							RebuildImplemented();
+						}
+					}), LOCTEXT("RemovePageGraphTooltip", "Removes the associated page graph from the MetaSound."));
+				EntryWidget->AddSlot()
+					.Padding(2.0f)
+					.HAlign(HAlign_Center)
+					.VAlign(VAlign_Center)
+					.AutoWidth()[ RemoveButtonWidget ];
+			}
 
-					const bool bGraphRemoved = Builder->GetBuilder().RemoveGraphPage(InEntryPageID);
-					if (bGraphRemoved)
+			// Page Playing Icon
+			{
+				TAttribute<FText> ToolTip = LOCTEXT("MetaSound_ExecutingPageGraphTooltip", "Currently executing graph.");
+				TAttribute<EVisibility> Visibility = TAttribute<EVisibility>::CreateLambda([this, PageID = InEntryPageID]()
+				{
+					if (Builder.IsValid())
 					{
-						UpdateItemNames();
-						ComboBox->RefreshOptions();
-						RebuildImplemented();
+						const bool bIsPreviewing = IsPreviewingPageGraph(Builder->GetConstBuilder(), PageID);
+						return bIsPreviewing ? EVisibility::Visible : EVisibility::Collapsed;
 					}
-				}), LOCTEXT("RemovePageTooltip2", "Removes the associated page from the MetaSound."));
+
+					return EVisibility::Collapsed;
+				});
+				TSharedRef<SWidget> ExecImageWidget = SNew(SImage)
+					.Image(Style::CreateSlateIcon("MetasoundEditor.Page.Executing").GetIcon())
+					.ColorAndOpacity(FStyleColors::AccentGreen)
+					.Visibility(MoveTemp(Visibility));
+
 				EntryWidget->AddSlot()
 					.Padding(2.0f)
 					.HAlign(HAlign_Center)
 					.VAlign(VAlign_Center)
-					.AutoWidth()[RemoveButtonWidget];
+					.AutoWidth()[ ExecImageWidget ];
 			}
 
-			EntryWidget->AddSlot()
-				.Padding(2.0f)
-				.HAlign(HAlign_Center)
-				.VAlign(VAlign_Center)
-				.AutoWidth() [ SelectButtonWidget ];
-
 			EntryWidget->SetEnabled(TAttribute<bool>::Create(TAttribute<bool>::FGetter::CreateSP(this, &FMetaSoundDetailCustomizationBase::IsGraphEditable)));
 			return EntryWidget;
 		};
@@ -602,11 +636,8 @@ namespace Metasound::Editor
 			else
 			{
 				constexpr bool bOpenEditor = false; // Already open/focused by user action
-				const bool bFocusedPage = UMetaSoundEditorSubsystem::GetConstChecked().SetFocusedPage(*Builder.Get(), PageID, bOpenEditor);
-				if (bFocusedPage)
-				{
-					BuildPageName = GetMissingPageName(PageID);
-				}
+				UMetaSoundEditorSubsystem::GetConstChecked().SetFocusedPage(*Builder.Get(), PageID, bOpenEditor);
+				BuildPageName = GetMissingPageName(PageID);
 			}
 		}
 		else
diff --git a/Engine/Plugins/Runtime/Metasound/Source/MetasoundEditor/Private/MetasoundEditor.cpp b/Engine/Plugins/Runtime/Metasound/Source/MetasoundEditor/Private/MetasoundEditor.cpp
--- a/Engine/Plugins/Runtime/Metasound/Source/MetasoundEditor/Private/MetasoundEditor.cpp
+++ b/Engine/Plugins/Runtime/Metasound/Source/MetasoundEditor/Private/MetasoundEditor.cpp
@@ -93,6 +93,7 @@
 #include "Widgets/SBoxPanel.h"
 #include "Widgets/SOverlay.h"
 #include "Widgets/SWindow.h"
+#include "Styling/StyleColors.h"
 
 #include UE_INLINE_GENERATED_CPP_BY_NAME(MetasoundEditor)
 
@@ -117,9 +118,82 @@ namespace Metasound
 				ECVF_Default);
 		} // namespace AssetEditorPrivate
 
-		bool PageEditorEnabled()
+		bool IsPreviewingMetaSound(const UObject& InMetaSound)
 		{
-			return AssetEditorPrivate::EnablePageEditor != 0;
+			using namespace Engine;
+			if (const UAudioComponent* PreviewComponent = GEditor->GetPreviewAudioComponent())
+			{
+				if (PreviewComponent->IsPlaying())
+				{
+					if (const USoundBase* Sound = PreviewComponent->Sound)
+					{
+						return Sound->GetUniqueID() == InMetaSound.GetUniqueID();
+					}
+				}
+			}
+
+			return false;
+		}
+
+		bool IsPreviewingPageInputDefault(const FMetaSoundFrontendDocumentBuilder& Builder, const FMetasoundFrontendClassInput& InClassInput, const FGuid& InPageID)
+		{
+			using namespace Engine;
+
+			const UObject& MetaSound = Builder.CastDocumentObjectChecked<const UObject>();
+			if (IsPreviewingMetaSound(MetaSound))
+			{
+				const FGuid TargetPageID = FDocumentBuilderRegistry::GetChecked().ResolveTargetPageID(InClassInput);
+				return TargetPageID == InPageID;
+			}
+
+			return false;
+		}
+
+		bool IsPreviewingPageGraph(const FMetaSoundFrontendDocumentBuilder& Builder, const FGuid& InPageID)
+		{
+			using namespace Engine;
+
+			const UObject& MetaSound = Builder.CastDocumentObjectChecked<const UObject>();
+			if (IsPreviewingMetaSound(MetaSound))
+			{
+				const FMetasoundFrontendGraphClass& GraphClass = Builder.GetConstDocumentChecked().RootGraph;
+				const FGuid TargetPageID = FDocumentBuilderRegistry::GetChecked().ResolveTargetPageID(GraphClass);
+				return TargetPageID == InPageID;
+			}
+
+			return false;
+		}
+
+		bool PageEditorEnabled(const FMetaSoundFrontendDocumentBuilder& Builder, bool bHasProjectPageValues, bool bPresetCanEditPageValues)
+		{
+			if (!AssetEditorPrivate::EnablePageEditor)
+			{
+				return false;
+			}
+
+			const UMetaSoundSettings* Settings = GetDefault<UMetaSoundSettings>();
+			if (!Settings)
+			{
+				return false;
+			}
+
+			if (Settings->GetProjectPageSettings().IsEmpty())
+			{
+				if (!bHasProjectPageValues)
+				{
+					return false;
+				}
+			}
+
+			if (!bPresetCanEditPageValues)
+			{
+				if (Builder.IsPreset())
+				{
+					return false;
+				}
+			}
+
+			return true;
 		}
 
 		namespace TabNamesPrivate
@@ -951,14 +1025,14 @@ namespace Metasound
 					InTabManager->RegisterTabSpawner(TabNamesPrivate::Pages, FOnSpawnTab::CreateLambda([this, InPagesDetails = PagesDetails](const FSpawnTabArgs&)
 					{
 						return SNew(SDockTab)
-							.Label(LOCTEXT("MetasoundPagesDetailsTitle", "Pages"))
-							[
-								InPagesDetails.ToSharedRef()
-							];
+						.Label(LOCTEXT("MetasoundPagesDetailsTitle", "Pages"))
+						[
+							InPagesDetails.ToSharedRef()
+						];
 					}), CanSpawnTab)
 					.SetDisplayName(LOCTEXT("PagesTab", "Pages"))
-						.SetGroup(WorkspaceMenuCategoryRef)
-						.SetIcon(FSlateIcon(FAppStyle::GetAppStyleSetName(), "Kismet.Tabs.Palette"));
+					.SetGroup(WorkspaceMenuCategoryRef)
+					.SetIcon(FSlateIcon(FAppStyle::GetAppStyleSetName(), "Kismet.Tabs.Palette"));
 				}
 			}
 
@@ -1130,6 +1204,16 @@ namespace Metasound
 			DocListener = MakeShared<FDocumentListener>(StaticCastSharedRef<FEditor>(AsShared()));
 			Builder->AddTransactionListener(DocListener->AsShared());
 
+			// Stat widgets are potentially intractable with transaction listener, so create then here
+			SAssignNew(FocusPageWidget, SPageStats)
+				.Visibility(EVisibility::HitTestInvisible);
+
+			SAssignNew(RenderPageWidget, SPageStats)
+				.Visibility(EVisibility::HitTestInvisible);
+
+			SAssignNew(RenderStatsWidget, SRenderStats)
+				.Visibility(EVisibility::HitTestInvisible);
+
 			if (FMetasoundAssetBase* MetaSoundAsset = IMetasoundUObjectRegistry::Get().GetObjectAsAssetBase(ObjectToEdit))
 			{
 				FMetaSoundFrontendDocumentBuilder& DocBuilder = Builder->GetBuilder();
@@ -1274,6 +1358,7 @@ namespace Metasound
 				UpdatePageInfo(false);
 			}
 
+			RefreshExecVisibility(Builder->GetConstBuilder().GetBuildPageID());
 			FSlateApplication::Get().SetUserFocus(0, MetasoundGraphEditor);
 		}
 
@@ -1340,23 +1425,16 @@ namespace Metasound
 
 		bool FEditor::ShowPageGraphDetails() const
 		{
-			const UMetaSoundSettings* Settings = GetDefault<UMetaSoundSettings>();
-			if (!Settings)
-			{
-				return false;
-			}
-
-			if (Settings->GetProjectPageSettings().IsEmpty())
-			{
-				return false;
-			}
-
-			if (!AssetEditorPrivate::EnablePageEditor)
+			if (Builder.IsValid())
 			{
-				return false;
+				const FMetaSoundFrontendDocumentBuilder& DocBuilder = Builder->GetConstBuilder();
+				const FMetasoundFrontendDocument& Document = DocBuilder.GetConstDocumentChecked();
+				const bool bLastGraph = Document.RootGraph.GetConstGraphPages().Num() == 1;
+				const bool bHasProjectPageValues = !bLastGraph && Document.RootGraph.FindConstGraph(Frontend::DefaultPageID) != nullptr;
+				return Editor::PageEditorEnabled(DocBuilder, bHasProjectPageValues);
 			}
 
-			return Builder.IsValid() && !Builder->IsPreset();
+			return false;
 		}
 
 		bool FEditor::GetBoundsForSelectedNodes(FSlateRect& Rect, float Padding)
@@ -1854,21 +1932,9 @@ namespace Metasound
 		TSharedRef<SWidget> FEditor::CreateAuditionMenuOptions()
 		{
 			TSharedPtr<FUICommandList> Commands = MakeShared<FUICommandList>();
-			FMenuBuilder MenuBuilder(true, Commands);
-			
-			constexpr bool bOpenOnClick = false;
-			constexpr bool bCloseAfterSelection = false;
-			const FText Label = LOCTEXT("PageAuditionSettingsSubMenu", "Pages");
-			const FText ToolTip = LOCTEXT("AuditionPageSettingsTooltip", "Settings related to auditioning pages");
-			MenuBuilder.AddSubMenu(Label, ToolTip,
-				FNewMenuDelegate::CreateLambda([this](FMenuBuilder& SubMenuBuilder)
-				{
-					CreateAuditionPageSubMenuOptions(SubMenuBuilder);
-				}),
-				bOpenOnClick,
-				FSlateIcon(),
-				bCloseAfterSelection
-			);
+			constexpr bool bShouldCloseWindowAfterMenuSelection = false;
+			FMenuBuilder MenuBuilder(bShouldCloseWindowAfterMenuSelection, Commands);
+			CreateAuditionPageSubMenuOptions(MenuBuilder);
 			return MenuBuilder.MakeWidget();
 		}
 
@@ -1880,7 +1946,7 @@ namespace Metasound
 				return;
 			}
 
-			MenuBuilder.BeginSection("SetAuditionPlatformSectionHeader", LOCTEXT("SetAuditionPlatformDescription", "Platform"));
+			MenuBuilder.BeginSection("SetAuditionPlatformSectionHeader", LOCTEXT("AuditionPlatformSectionName", "Audition Platform"));
 			{
 				auto CreatePlatformEntry = [this, &MenuBuilder](FName PlatformName, const FText& PlatformText)
 				{
@@ -1891,26 +1957,38 @@ namespace Metasound
 						{
 							EditorSettings->AuditionPlatform = PlatformName;
 							Stop();
+							SyncAuditionState();
 						}
 					});
-
-					SetPlatformAction.CanExecuteAction = FCanExecuteAction::CreateLambda([this, PlatformName]()
+					SetPlatformAction.GetActionCheckState = FGetActionCheckState::CreateLambda([this, PlatformName]()
 					{
 						if (const UMetasoundEditorSettings* EditorSettings = GetDefault<UMetasoundEditorSettings>())
 						{
-							return EditorSettings->AuditionPlatform != PlatformName;
+							if (EditorSettings->AuditionPlatform == PlatformName)
+							{
+								return ECheckBoxState::Checked;
+							}
 						}
 
-						return false;
+						return ECheckBoxState::Unchecked;
 					});
 
 					MenuBuilder.AddMenuEntry(PlatformText,
-						FText::Format(LOCTEXT("SetAuditionPlatformToolTip", "Sets the audition platform to '{0}'."), PlatformText),
+						FText::Format(LOCTEXT("SetAuditionPlatformToolTip", "Sets the page audition platform to '{0}'."), PlatformText),
 						FSlateIcon(),
-						SetPlatformAction);
+						SetPlatformAction,
+						{ },
+						EUserInterfaceActionType::RadioButton);
 				};
 
-				const TArray<FName> AuditionPlatforms = UMetasoundEditorSettings::GetAuditionPlatformNames();
+				TArray<FName> AuditionPlatforms = UMetasoundEditorSettings::GetAuditionPlatformNames();
+
+				// Protects against stale setting not showing after platform values are manipulated just for visibility
+				if (const UMetasoundEditorSettings* EditorSettings = GetDefault<UMetasoundEditorSettings>())
+				{
+					AuditionPlatforms.AddUnique(EditorSettings->AuditionPlatform);
+				}
+
 				for (const FName& PlatformName : AuditionPlatforms)
 				{
 					const FText PlatformText = FText::FromName(PlatformName);
@@ -1919,9 +1997,11 @@ namespace Metasound
 			}
 			MenuBuilder.EndSection();
 
-			MenuBuilder.BeginSection("SetAuditionTargetPageSectionHeader", LOCTEXT("SetAuditionTargetPageDescription", "Target Page"));
+			MenuBuilder.BeginSection("SetAuditionPageSectionHeader", LOCTEXT("SetAuditionPageDescription", "Audition Target Page"));
 			{
-				const FText FocusPageTooltip = LOCTEXT("EnableFocusTargetPageSwapTooltip", "Dynamically swap which page is targeted based on which page is focused");
+				const FText FocusPageTooltip = LOCTEXT("EnableFocusTargetPageSwapTooltip", "Attempts to set audition target page to focused graph page. "
+					"If page focus is modified and newly selected graph is not targetable for a selected audition platform (see 'Audition' or `MetaSound Editor Settings`), "
+					"will fail to pass validation for preview audition.");
 				MenuBuilder.AddWidget(
 					SNew(SCheckBox)
 					.OnCheckStateChanged_Lambda([this](ECheckBoxState State)
@@ -1933,7 +2013,8 @@ namespace Metasound
 								case ECheckBoxState::Checked:
 								{
 									EdSettings->AuditionPageMode = EAuditionPageMode::Focused;
-									SyncFocusedPage();
+									Stop();
+									SyncAuditionState();
 									break;
 								}
 
@@ -1960,7 +2041,7 @@ namespace Metasound
 						return ECheckBoxState::Unchecked;
 					})
 					.ToolTipText(FocusPageTooltip),
-					LOCTEXT("EnableFocusTargetPageGraphSwap", "Set Target To Focused Graph"),
+					LOCTEXT("EnableFocusTargetPageGraphSync", "Sync Audition Target With Focused Graph"),
 					true,
 					true,
 					FocusPageTooltip
@@ -1968,50 +2049,86 @@ namespace Metasound
 
 				auto TryAddPageEntry = [this, &MenuBuilder](const FMetaSoundPageSettings& PageSettings)
 				{
-					const FName AuditionTargetPage = PageSettings.Name;
+					const FName AuditionPage = PageSettings.Name;
 					const FText PageText = FText::FromName(PageSettings.Name);
 					FUIAction SetTargetPageAction;
-					SetTargetPageAction.ExecuteAction = FExecuteAction::CreateLambda([this, AuditionTargetPage]()
+
+					SetTargetPageAction.ExecuteAction = FExecuteAction::CreateLambda([this, AuditionPage]()
 					{
 						if (UMetasoundEditorSettings* EditorSettings = GetMutableDefault<UMetasoundEditorSettings>())
 						{
-							EditorSettings->AuditionTargetPage = AuditionTargetPage;
-							Stop();
+							if (EditorSettings->AuditionPage != AuditionPage)
+							{
+								EditorSettings->AuditionPage = AuditionPage;
+								Stop();
+								SyncAuditionState();
+							}
 						}
 					});
 
-					SetTargetPageAction.CanExecuteAction = FCanExecuteAction::CreateLambda([this, AuditionTargetPage]()
+					SetTargetPageAction.GetActionCheckState = FGetActionCheckState::CreateLambda([this, AuditionPage]()
 					{
 						if (const UMetasoundEditorSettings* EditorSettings = GetDefault<UMetasoundEditorSettings>())
 						{
-							const bool bUserAuditionMode = EditorSettings->AuditionPageMode == EAuditionPageMode::User;
-							if (bUserAuditionMode)
+							if (EditorSettings->AuditionPage == AuditionPage)
 							{
-								return EditorSettings->AuditionTargetPage != AuditionTargetPage;
+								return ECheckBoxState::Checked;
 							}
 						}
+
+						return ECheckBoxState::Unchecked;
+					});
+					SetTargetPageAction.CanExecuteAction = FCanExecuteAction::CreateLambda([this, AuditionPage]()
+					{
+						if (const UMetasoundEditorSettings* EdSettings = GetDefault<UMetasoundEditorSettings>())
+						{
+							return EdSettings->AuditionPageMode == EAuditionPageMode::User;
+						}
+
+						if (const UMetaSoundSettings* Settings = GetDefault<UMetaSoundSettings>())
+						{
+							return Settings->FindPageSettings(AuditionPage) != nullptr;
+						}
 						return false;
 					});
 
-					SetTargetPageAction.IsActionVisibleDelegate.BindLambda([PageName = PageSettings.Name]()
+					TAttribute<FText> TooltipAttribute = TAttribute<FText>::CreateLambda([this, AuditionPage, PageText]()
 					{
 						if (const UMetasoundEditorSettings* EditorSettings = GetDefault<UMetasoundEditorSettings>())
 						{
-							if (const UMetaSoundSettings* Settings = GetDefault<UMetaSoundSettings>())
+							if (EditorSettings->AuditionPlatform != UMetasoundEditorSettings::EditorAuditionPlatform)
 							{
-								if (const FMetaSoundPageSettings* PageSetting = Settings->FindPageSettings(PageName))
+								if (const UMetaSoundSettings* Settings = GetDefault<UMetaSoundSettings>())
 								{
-									return PageSetting->PlatformCanTargetPage(EditorSettings->AuditionPlatform);
+									if (const FMetaSoundPageSettings* PageSetting = Settings->FindPageSettings(AuditionPage))
+									{
+										if (!PageSetting->PlatformCanTargetPage(EditorSettings->AuditionPlatform))
+										{
+											return FText::Format(LOCTEXT("AuditionPageInvalidForPlatformToolTip", "Platform '{0}' does not target page '{1}'. See 'MetaSound' Project Settings"), FText::FromName(EditorSettings->AuditionPlatform), PageText);
+										}
+									}
+								}
+							}
+
+							const bool bUserAuditionMode = EditorSettings->AuditionPageMode == EAuditionPageMode::User;
+							if (bUserAuditionMode)
+							{
+								if (EditorSettings->AuditionPage != AuditionPage)
+								{
+									return FText::Format(LOCTEXT("SetAuditionPageToolTip", "Sets the user's editor AuditionPage setting to '{0}'."), PageText);
 								}
 							}
 						}
-						return false;
+
+						return FText();
 					});
 
 					MenuBuilder.AddMenuEntry(PageText,
-						FText::Format(LOCTEXT("SetAuditionPageToolTip", "Sets the audition target page to '{0}'."), PageText),
+						MoveTemp(TooltipAttribute),
 						FSlateIcon(),
-						SetTargetPageAction);
+						SetTargetPageAction,
+						{ },
+						EUserInterfaceActionType::RadioButton);
 				};
 
 				Settings->IteratePageSettings(TryAddPageEntry);
@@ -2319,6 +2436,11 @@ namespace Metasound
 
 		FText FEditor::GetGraphStatusDescription() const
 		{
+			if (!GraphStatusDescriptionOverride.IsEmpty())
+			{
+				return GraphStatusDescriptionOverride;
+			}
+
 			switch (HighestMessageSeverity)
 			{
 				case EMessageSeverity::Error:
@@ -2423,14 +2545,13 @@ namespace Metasound
 
 			if (USoundBase* MetaSoundToPlay = Cast<USoundBase>(GetMetasoundObject()))
 			{
-				HighestMessageSeverity = GetMetaSoundGraphChecked().GetHighestMessageSeverity();
+				SyncAuditionState();
+
 				if (HighestMessageSeverity == EMessageSeverity::Error)
 				{
 					return;
 				}
 
-				SyncFocusedPage();
-
 				// Even though the MetaSoundSource will attempt to register via InitResources
 				// later in this execution (and deeper in the stack), this call forces
 				// re-registering to make sure everything is up-to-date.
@@ -2578,16 +2699,35 @@ namespace Metasound
 			SetPreviewID(INDEX_NONE);
 		}
 
-		void FEditor::SyncFocusedPage() const
+		void FEditor::SyncAuditionState(bool bSetAuditionFocus)
 		{
-			if (const UMetaSoundSettings* Settings = GetDefault<UMetaSoundSettings>())
+			GraphStatusDescriptionOverride = { };
+			HighestMessageSeverity = GetMetaSoundGraphChecked().GetHighestMessageSeverity();
+
+			if (Builder.IsValid())
 			{
-				if (Builder.IsValid())
+				constexpr bool bOpenEditor = false; // Already Focused
+				constexpr bool bPostTransaction = false;
+				const FMetaSoundFrontendDocumentBuilder& DocBuilder = Builder->GetConstBuilder();
+				const FGuid BuildPageID = DocBuilder.GetBuildPageID();
+
+				UMetaSoundEditorSubsystem::GetChecked().SetFocusedPage(*Builder.Get(), BuildPageID, bOpenEditor, bPostTransaction);
+
+				if (const UMetasoundEditorSettings* EdSettings = GetDefault<UMetasoundEditorSettings>())
 				{
-					constexpr bool bOpenEditor = false; // Already Focused
-					constexpr bool bPostTransaction = false;
-					const FMetaSoundFrontendDocumentBuilder& DocBuilder = Builder->GetConstBuilder();
-					UMetaSoundEditorSubsystem::GetChecked().SetFocusedPage(*Builder.Get(), DocBuilder.GetBuildPageID(), bOpenEditor, bPostTransaction);
+					if (EdSettings->AuditionPlatform != UMetasoundEditorSettings::EditorAuditionPlatform)
+					{
+						if (!UMetaSoundEditorSubsystem::GetChecked().IsPageAuditionPlatformCookTarget(EdSettings->AuditionPage))
+						{
+							GraphStatusDescriptionOverride = LOCTEXT("InvalidAuditionPageWarning",
+								"Selected Audition Page in MetaSound Editor Settings is not a target page for the selectd 'Audition Platform'. "
+								"Execution may result in behavior that does not exhibit runtime behavior.");
+							if (HighestMessageSeverity > EMessageSeverity::Warning)
+							{
+								HighestMessageSeverity = EMessageSeverity::Warning;
+							}
+						}
+					}
 				}
 			}
 		}
@@ -2940,15 +3080,6 @@ namespace Metasound
 				.Visibility(EVisibility::HitTestInvisible)
 				.TextStyle(FAppStyle::Get(), "Graph.ZoomText")
 				.ColorAndOpacity(FLinearColor(1, 1, 1, 0.30f));
-
-			SAssignNew(FocusPageWidget, SPageStats)
-				.Visibility(EVisibility::HitTestInvisible);
-
-			SAssignNew(RenderPageWidget, SPageStats)
-				.Visibility(EVisibility::HitTestInvisible);
-
-			SAssignNew(RenderStatsWidget, SRenderStats)
-				.Visibility(EVisibility::HitTestInvisible);
 		}
 
 		FGraphAppearanceInfo FEditor::GetGraphAppearance() const
@@ -4443,7 +4574,8 @@ namespace Metasound
 					}
 				}
 
-				HighestMessageSeverity = Graph->GetHighestMessageSeverity();
+				constexpr bool bSetAuditionFocus = false;
+				SyncAuditionState(bSetAuditionFocus);
 
 				// Modify data has been observed both from synchronization & by
 				// updating views by this point, so full reset is completed here.
@@ -4726,7 +4858,9 @@ namespace Metasound
 
 		TUniquePtr<FGraphConnectionManager> FEditor::RebuildConnectionManager(UAudioComponent* PreviewComp) const
 		{
-			if (!PreviewComp)
+			using namespace Engine;
+
+			if (!PreviewComp || !Builder.IsValid())
 			{
 				return MakeUnique<FGraphConnectionManager>();
 			}
@@ -4737,21 +4871,14 @@ namespace Metasound
 				return MakeUnique<FGraphConnectionManager>();
 			}
 
-			const UMetasoundEditorSettings* EdSettings = GetDefault<UMetasoundEditorSettings>();
-			check(EdSettings);
-			const UMetaSoundSettings* Settings = GetDefault<UMetaSoundSettings>();
-			check(Settings);
-
-			const FMetaSoundPageSettings* TargetPageSettings = Settings->FindPageSettings(EdSettings->AuditionTargetPage);
-			const FMetaSoundPageSettings* BuilderPageSettings = Settings->FindPageSettings(Builder->GetConstBuilder().GetBuildPageID());
-			if (TargetPageSettings != BuilderPageSettings)
+			const FGuid ResolvedGraphPageID = FDocumentBuilderRegistry::GetChecked().ResolveTargetPageID(Source->GetConstDocumentChecked().RootGraph);
+			if (ResolvedGraphPageID != Builder->GetConstBuilder().GetBuildPageID())
 			{
 				return MakeUnique<FGraphConnectionManager>();
 			}
 
 			FAudioDevice* AudioDevice = PreviewComp->GetAudioDevice();
 			check(AudioDevice);
-			const FName& AudioBufferTypeName = GetMetasoundDataTypeName<FAudioBuffer>();
 			const FSampleRate DeviceSampleRate = static_cast<FSampleRate>(AudioDevice->GetSampleRate());
 			const uint32 PlayOrder = PreviewComp->GetLastPlayOrder();
 			const uint64 TransmitterID = Audio::GetTransmitterID(PreviewComp->GetAudioComponentID(), 0, PlayOrder);
@@ -4768,15 +4895,28 @@ namespace Metasound
 
 			if (FocusPageWidget.IsValid())
 			{
+				const FSlateColor* Color = nullptr;
 				// For now, this is one and the same, but may change in the future to accommodate conflating
 				// the static nature of the active build graph ID with what the user is visualizing.
 				const FMetaSoundPageSettings* PageSettings = nullptr;
 				if (Builder.IsValid())
 				{
-					const FGuid& PageID = Builder->GetConstBuilder().GetBuildPageID();
+					const FMetaSoundFrontendDocumentBuilder& DocBuilder = Builder->GetConstBuilder();
+					const FGuid& PageID = DocBuilder.GetBuildPageID();
+					if (bIsPlaying)
+					{
+						const FMetasoundFrontendGraphClass& GraphClass = DocBuilder.GetConstDocumentChecked().RootGraph;
+						const FGuid ResolvePageID = FDocumentBuilderRegistry::GetChecked().ResolveTargetPageID(GraphClass);
+						if (ResolvePageID == PageID)
+						{
+							Color = &FStyleColors::AccentGreen;
+						}
+					}
 					PageSettings = Settings->FindPageSettings(PageID);
 				}
-				FocusPageWidget->Update(PageSettings);
+
+				const FText GraphPageHeader = LOCTEXT("GraphPageTargetHeader", "Graph Page");
+				FocusPageWidget->Update(PageSettings, GraphPageHeader, Color);
 			}
 
 			if (RenderPageWidget.IsValid())
@@ -4786,10 +4926,12 @@ namespace Metasound
 				{
 					if (const UMetasoundEditorSettings* EditorSettings = GetDefault<UMetasoundEditorSettings>())
 					{
-						PageSettings = Settings->FindPageSettings(EditorSettings->AuditionTargetPage);
+						PageSettings = Settings->FindPageSettings(EditorSettings->AuditionPage);
 					}
 				}
-				RenderPageWidget->Update(PageSettings);
+
+				const FText RenderPageHeader = LOCTEXT("AuditionPageHeader", "Audition Page");
+				RenderPageWidget->Update(PageSettings, RenderPageHeader, &FStyleColors::AccentGreen);
 			}
 		}
 
@@ -4837,6 +4979,27 @@ namespace Metasound
 				ParentPtr->Stop();
 				ParentPtr->UpdatePageInfo(false);
 				ParentPtr->bRefreshGraph = true;
+				ParentPtr->RefreshExecVisibility(Args.PageID);
+			}
+		}
+
+		void FEditor::RefreshExecVisibility(const FGuid& InPageID) const
+		{
+			if (FocusPageWidget.IsValid())
+			{
+				TAttribute<EVisibility> ExecVisibility = TAttribute<EVisibility>::CreateLambda([this, InPageID]()
+				{
+					using namespace Engine;
+					if (Builder.IsValid())
+					{
+						const FMetaSoundFrontendDocumentBuilder& DocBuilder = Builder->GetConstBuilder();
+						const bool bIsPreviewing = IsPreviewingPageGraph(DocBuilder, InPageID);
+						return bIsPreviewing ? EVisibility::Visible : EVisibility::Collapsed;
+					}
+
+					return EVisibility::Collapsed;
+				});
+				FocusPageWidget->SetExecVisibility(MoveTemp(ExecVisibility));
 			}
 		}
 
diff --git a/Engine/Plugins/Runtime/Metasound/Source/MetasoundEditor/Private/MetasoundEditor.h b/Engine/Plugins/Runtime/Metasound/Source/MetasoundEditor/Private/MetasoundEditor.h
--- a/Engine/Plugins/Runtime/Metasound/Source/MetasoundEditor/Private/MetasoundEditor.h
+++ b/Engine/Plugins/Runtime/Metasound/Source/MetasoundEditor/Private/MetasoundEditor.h
@@ -53,6 +53,8 @@ class FSlateRect;
 class IDetailsView;
 class IToolkitHost;
 
+struct FMetaSoundFrontendDocumentBuilder;
+
 namespace Metasound::Editor
 {
 	class SFindInMetasound;
@@ -102,7 +104,10 @@ class UMetasoundPagesView : public UMetasoundEditorViewBase
 
 namespace Metasound::Editor
 {
-	bool PageEditorEnabled();
+	bool IsPreviewingMetaSound(const UObject& InMetaSound);
+	bool IsPreviewingPageInputDefault(const FMetaSoundFrontendDocumentBuilder& Builder, const FMetasoundFrontendClassInput& InClassInput, const FGuid& InPageID);
+	bool IsPreviewingPageGraph(const FMetaSoundFrontendDocumentBuilder& Builder, const FGuid& InPageID);
+	bool PageEditorEnabled(const FMetaSoundFrontendDocumentBuilder& Builder, bool bHasProjectPageValues, bool bPresetCanEditPageValues = false);
 
 	// Forward Declarations
 	class FMetasoundGraphMemberSchemaAction;
@@ -409,6 +414,8 @@ namespace Metasound::Editor
 	private:
 		int32 PromotableSelectedNodes();
 
+		void RefreshExecVisibility(const FGuid& InPageID) const;
+
 		/** Forces refresh of pages view. */
 		void RefreshPagesView();
 
@@ -421,8 +428,7 @@ namespace Metasound::Editor
 
 		void ExportNodesToText(FString& OutText) const;
 
-		// Sets the globally targeted audition page for execution to this editor's currently focused page
-		void SyncFocusedPage() const;
+		void SyncAuditionState(bool bSetAuditionFocus = true);
 
 		/** FNotifyHook interface */
 		virtual void NotifyPostChange(const FPropertyChangedEvent& PropertyChangedEvent, FProperty* PropertyThatChanged) override;
@@ -582,8 +588,12 @@ namespace Metasound::Editor
 		/** Whether or not registry is currently being primed due to assets still loading */
 		bool bPrimingRegistry = false;
 
+		/** Highest message severity set on last validation pass of graph. */
 		int32 HighestMessageSeverity = EMessageSeverity::Info;
 
+		/** If set, used to inform user of validation results on hover of play icon. */
+		FText GraphStatusDescriptionOverride;
+
 		TSharedPtr<SNotificationItem> NotificationPtr;
 
 		bool bMemberRenameRequested = false;
diff --git a/Engine/Plugins/Runtime/Metasound/Source/MetasoundEditor/Private/MetasoundEditorGraph.cpp b/Engine/Plugins/Runtime/Metasound/Source/MetasoundEditor/Private/MetasoundEditorGraph.cpp
--- a/Engine/Plugins/Runtime/Metasound/Source/MetasoundEditor/Private/MetasoundEditorGraph.cpp
+++ b/Engine/Plugins/Runtime/Metasound/Source/MetasoundEditor/Private/MetasoundEditorGraph.cpp
@@ -15,12 +15,14 @@
 #include "MetasoundEditorGraphNode.h"
 #include "MetasoundEditorGraphValidation.h"
 #include "MetasoundEditorModule.h"
+#include "MetasoundEditorSettings.h"
 #include "MetasoundEditorSubsystem.h"
 #include "MetasoundFrontendDataTypeRegistry.h"
 #include "MetasoundFrontendDocumentBuilder.h"
 #include "MetasoundFrontendNodeTemplateRegistry.h"
 #include "MetasoundFrontendSearchEngine.h"
 #include "MetasoundLog.h"
+#include "MetasoundSettings.h"
 #include "MetasoundUObjectRegistry.h"
 #include "MetasoundVariableNodes.h"
 #include "MetasoundVertex.h"
@@ -731,10 +733,19 @@ bool UMetasoundEditorGraphMemberDefaultLiteral::TryGetPreviewPageID(FGuid& OutPr
 	if (const FMetasoundFrontendClassInput* Input = Builder.FindGraphInput(Member->GetMemberName()))
 	{
 		const FGuid PageID = FDocumentBuilderRegistry::GetChecked().ResolveTargetPageID(*Input);
-		if (PageID == Builder.GetBuildPageID())
+
+		const UMetasoundEditorSettings* EdSettings = ::GetDefault<UMetasoundEditorSettings>();
+		const UMetaSoundSettings* Settings = ::GetDefault<UMetaSoundSettings>();
+		if (Settings && EdSettings)
 		{
-			OutPreviewPageID = PageID;
-			return true;
+			if (const FMetaSoundPageSettings* PageSettings = Settings->FindPageSettings(EdSettings->AuditionPage))
+			{
+				if (PageID == PageSettings->UniqueId)
+				{
+					OutPreviewPageID = PageID;
+					return true;
+				}
+			}
 		}
 	}
 
diff --git a/Engine/Plugins/Runtime/Metasound/Source/MetasoundEditor/Private/MetasoundEditorModule.cpp b/Engine/Plugins/Runtime/Metasound/Source/MetasoundEditor/Private/MetasoundEditorModule.cpp
--- a/Engine/Plugins/Runtime/Metasound/Source/MetasoundEditor/Private/MetasoundEditorModule.cpp
+++ b/Engine/Plugins/Runtime/Metasound/Source/MetasoundEditor/Private/MetasoundEditorModule.cpp
@@ -163,8 +163,9 @@ namespace Metasound
 
 					// Misc
 					Set("MetasoundEditor.Audition", new IMAGE_BRUSH_SVG(TEXT("Icons/metasound_page"), Icon16));
-					Set("MetasoundEditor.Speaker", new FSlateImageBrush(RootToContentDir(TEXT("/Icons/speaker_144x.png")), FVector2D(144.0f, 144.0f)));
 					Set("MetasoundEditor.Metasound.Icon", new IMAGE_BRUSH_SVG(TEXT("Icons/metasound_icon"), Icon16));
+					Set("MetasoundEditor.Page.Executing", new IMAGE_BRUSH_SVG(TEXT("Icons/metasound_page_exec"), Icon16));
+					Set("MetasoundEditor.Speaker", new FSlateImageBrush(RootToContentDir(TEXT("/Icons/speaker_144x.png")), FVector2D(144.0f, 144.0f)));
 
 					// Class Icons
 					auto SetClassIcon = [this, InIcon16 = Icon16, InIcon64 = Icon64](const FString& ClassName)
diff --git a/Engine/Plugins/Runtime/Metasound/Source/MetasoundEditor/Private/MetasoundEditorSettings.cpp b/Engine/Plugins/Runtime/Metasound/Source/MetasoundEditor/Private/MetasoundEditorSettings.cpp
--- a/Engine/Plugins/Runtime/Metasound/Source/MetasoundEditor/Private/MetasoundEditorSettings.cpp
+++ b/Engine/Plugins/Runtime/Metasound/Source/MetasoundEditor/Private/MetasoundEditorSettings.cpp
@@ -19,6 +19,9 @@
 
 #define LOCTEXT_NAMESPACE "MetasoundEditorSettings"
 
+const FName UMetasoundEditorSettings::DefaultAuditionPlatform = "Default";
+const FName UMetasoundEditorSettings::EditorAuditionPlatform = "Editor";
+
 UMetasoundEditorSettings::UMetasoundEditorSettings(const FObjectInitializer& ObjectInitializer)
 	: Super(ObjectInitializer)
 {
@@ -117,7 +120,7 @@ Metasound::Engine::FPageResolutionEditorResults UMetasoundEditorSettings::Resolv
 
 	FPageResolutionEditorResults PreviewInfo { .PlatformName = AuditionPlatform };
 
-	if (AuditionTargetPage == Frontend::DefaultPageName)
+	if (AuditionPage == Frontend::DefaultPageName)
 	{
 		check(InPageIDs.Contains(Frontend::DefaultPageID));
 		PreviewInfo.PageID = Frontend::DefaultPageID;
@@ -126,7 +129,7 @@ Metasound::Engine::FPageResolutionEditorResults UMetasoundEditorSettings::Resolv
 
 	if (const UMetaSoundSettings* Settings = GetDefault<UMetaSoundSettings>())
 	{
-		if (const FMetaSoundPageSettings* TargetPageSettings = Settings->FindPageSettings(AuditionTargetPage))
+		if (const FMetaSoundPageSettings* TargetPageSettings = Settings->FindPageSettings(AuditionPage))
 		{
 			auto PageIsCooked = [&InPageIDs, &PreviewInfo](const FMetaSoundPageSettings& PageSettings)
 			{
@@ -149,7 +152,7 @@ Metasound::Engine::FPageResolutionEditorResults UMetasoundEditorSettings::Resolv
 					bFoundMatch |= PageSettings.UniqueId == TargetPageID;
 					if (bFoundMatch)
 					{
-						if (PageIsCooked(PageSettings))
+						if (AuditionPlatform == EditorAuditionPlatform || PageIsCooked(PageSettings))
 						{
 							bPageSelected = true;
 							PreviewInfo.PageID = PageSettings.UniqueId;
@@ -166,10 +169,45 @@ Metasound::Engine::FPageResolutionEditorResults UMetasoundEditorSettings::Resolv
 
 TArray<FName> UMetasoundEditorSettings::GetAuditionPlatformNames()
 {
+	TArray<FName> PlatformNames { EditorAuditionPlatform, DefaultAuditionPlatform };
 	if (const UMetaSoundSettings* Settings = GetDefault<UMetaSoundSettings>())
 	{
-		return Settings->GetAllPlatformNamesImplementingTargets();
+		PlatformNames.Append(Settings->GetAllPlatformNamesImplementingTargets());
 	}
-	return { FPlatformProperties::IniPlatformName() };
+	return PlatformNames;
+}
+
+TArray<FName> UMetasoundEditorSettings::GetAuditionPageNames()
+{
+	const UMetaSoundSettings* Settings = GetDefault<UMetaSoundSettings>();
+	const UMetasoundEditorSettings* EdSettings = GetDefault<UMetasoundEditorSettings>();
+
+	if (Settings && EdSettings)
+	{
+		TSet<FName> AuditionPageNames;
+		if (EdSettings->AuditionPlatform == EditorAuditionPlatform)
+		{
+			Algo::Transform(Settings->GetProjectPageSettings(), AuditionPageNames, [](const FMetaSoundPageSettings& PageSettings) { return PageSettings.Name; });
+			AuditionPageNames.Add(Settings->GetDefaultPageSettings().Name);
+		}
+		else
+		{
+			const TArray<FGuid> AuditionPageIDs = Settings->GetCookedTargetPageIDs(EdSettings->AuditionPlatform);
+			Algo::Transform(AuditionPageIDs, AuditionPageNames, [&Settings](const FGuid& PageID)
+			{
+				const FMetaSoundPageSettings* PageSettings = Settings->FindPageSettings(PageID);
+				if (ensure(PageSettings))
+				{
+					return PageSettings->Name;
+				}
+
+				return FName { };
+			});
+		}
+
+		return AuditionPageNames.Array();
+	}
+
+	return { };
 }
 #undef LOCTEXT_NAMESPACE // "MetaSoundEditor"
diff --git a/Engine/Plugins/Runtime/Metasound/Source/MetasoundEditor/Private/MetasoundEditorSettings.h b/Engine/Plugins/Runtime/Metasound/Source/MetasoundEditor/Private/MetasoundEditorSettings.h
--- a/Engine/Plugins/Runtime/Metasound/Source/MetasoundEditor/Private/MetasoundEditorSettings.h
+++ b/Engine/Plugins/Runtime/Metasound/Source/MetasoundEditor/Private/MetasoundEditorSettings.h
@@ -57,10 +57,10 @@ enum class EMetasoundActiveDetailView : uint8
 UENUM()
 enum class EAuditionPageMode : uint8
 {
-	// Audition Target Page set automatically set to page focused in asset editor
+	// Sets Audition Page automatically to graph page focused in asset editor
 	Focused,
 
-	// Audition Target Page is specified by user (does not automatically change when new page is visualized)
+	// Audition Page is specified by user (does not automatically change when graph page is focused)
 	User
 };
 
@@ -139,6 +139,12 @@ class METASOUNDEDITOR_API UMetasoundEditorSettings : public UObject
 	GENERATED_UCLASS_BODY()
 
 public:
+	// Represents auditioning any platform using the default target/cook settings
+	static const FName DefaultAuditionPlatform;
+
+	// Represents auditioning as the editor, ignoring any explicit target/cook settings
+	static const FName EditorAuditionPlatform;
+
 	/** Whether to pin the MetaSound Patch asset type when creating new assets.
 	  * Requires editor restart for change to take effect.*/
 	UPROPERTY(EditAnywhere, config, DisplayName = "Pin MetaSound Patch in Asset Menu", Category = AssetMenu)
@@ -149,40 +155,41 @@ public:
 	UPROPERTY(EditAnywhere, config, DisplayName = "Pin MetaSound Source in Asset Menu", Category = AssetMenu)
 	bool bPinMetaSoundSourceInAssetMenu = true;
 
-	/** If true, uses editor page/platform audition settings in PIE. If false, uses project's defined values (see 'MetaSound'
-	  * Developer Settings for TargetPage, which can be manipulated via code/Blueprint.)
+	/** If true, uses editor page/platform audition settings in PIE. If false, uses project's defined values
+	  * (see project 'MetaSound' setting 'TargetPage', which can be manipulated via code/Blueprint.)
 	  */
-	UPROPERTY(EditAnywhere, config, Category = Audition)
+	UPROPERTY(EditAnywhere, config, Category = "Audition (Experimental)")
 	bool bApplyAuditionSettingsInPIE = true;
 
 	/** Default author title to use when authoring a new
 	  * MetaSound.  If empty, uses machine name by default.
 	  */
-	UPROPERTY(EditAnywhere, config, Category=General)
+	UPROPERTY(EditAnywhere, config, Category = General)
 	FString DefaultAuthor;
 
 	/* Currently set page audition mode. Set by the MetaSound Asset Editor. */
-	UPROPERTY(EditAnywhere, config, Category = Audition, meta = (DisplayName = "Audition Mode"))
+	UPROPERTY(EditAnywhere, config, Category = "Audition (Experimental)", meta = (DisplayName = "Page Audition Mode"))
 	EAuditionPageMode AuditionPageMode = EAuditionPageMode::Focused;
 
-	/** Name of platform to mock when previewing playback. This will limit playback to only pages that are cooked for the given platform.
-	  * Set in the MetaSound Asset Editor.
+	/** Name of platform to mock when previewing playback. This will limit playback
+	  * to fallback only to paged data that are cooked for the given platform.
+	  * (see project 'MetaSound' Settings --> 'Page Settings' array for order)
+	  * If set to 'Editor', ignores cook settings and allows fallback to all page.
 	  */
-	UPROPERTY(EditAnywhere, config, Category = Audition, meta = (DisplayName = "Platform", GetOptions = "MetasoundEditor.MetasoundEditorSettings.GetAuditionPlatformNames"))
-	FName AuditionPlatform;
+	UPROPERTY(EditAnywhere, config, Category = "Audition (Experimental)", meta = (DisplayName = "Page Audition Platform", GetOptions = "MetasoundEditor.MetasoundEditorSettings.GetAuditionPlatformNames"))
+	FName AuditionPlatform = EditorAuditionPlatform;
 
-	/** Name of the page to target when previewing playback in editor. If target page is not implemented for the set audition platform,
-	  *  uses order of cooked pages (see 'project MetaSound Settings --> Page Settings' for order) falling back to lower index-ordered page
-	  * implemented in MetaSound asset. Set in the MetaSound Asset Editor.
+	/** Name of the page to audition in editor. If unimplemented on the auditioned MetaSound, uses order of cooked pages
+	  * (see project 'MetaSound' Settings --> 'Page Settings' array for order) falling back to lower index-ordered page implemented
+	  * in MetaSound asset.
 	  */
-	UPROPERTY(EditAnywhere, config, Category = Audition, meta =
+	UPROPERTY(EditAnywhere, config, Category = "Audition (Experimental)", meta =
 	(
-		DisplayName = "Audition Page",
 		EditCondition = "AuditionPageMode == EAuditionPageMode::User",
 		EditConditionHides = true,
-		GetOptions = "MetasoundEngine.MetaSoundSettings.GetPageNames")
+		GetOptions = "MetasoundEditor.MetasoundEditorSettings.GetAuditionPageNames")
 	)
-	FName AuditionTargetPage = Metasound::Frontend::DefaultPageName;
+	FName AuditionPage = Metasound::Frontend::DefaultPageName;
 
 	/** Maps Pin Category To Pin Color */
 	TMap<FName, FLinearColor> CustomPinTypeColors;
@@ -303,6 +310,9 @@ public:
 
 	Metasound::Engine::FPageResolutionEditorResults ResolveAuditionPage(const TArray<FGuid>& InPageIDs) const;
 
+	UFUNCTION()
+	static TArray<FName> GetAuditionPageNames();
+
 	UFUNCTION()
 	static TArray<FName> GetAuditionPlatformNames();
 };
diff --git a/Engine/Plugins/Runtime/Metasound/Source/MetasoundEditor/Private/MetasoundEditorSubsystem.cpp b/Engine/Plugins/Runtime/Metasound/Source/MetasoundEditor/Private/MetasoundEditorSubsystem.cpp
--- a/Engine/Plugins/Runtime/Metasound/Source/MetasoundEditor/Private/MetasoundEditorSubsystem.cpp
+++ b/Engine/Plugins/Runtime/Metasound/Source/MetasoundEditor/Private/MetasoundEditorSubsystem.cpp
@@ -240,6 +240,40 @@ void UMetaSoundEditorSubsystem::InitEdGraph(UObject& InMetaSound)
 	Metasound::Editor::FGraphBuilder::BindEditorGraph(IDocumentBuilderRegistry::GetChecked().FindOrBeginBuilding(&InMetaSound));
 }
 
+bool UMetaSoundEditorSubsystem::IsPageAuditionPlatformCookTarget(FName InPageName) const
+{
+	if (const UMetaSoundSettings* Settings = GetDefault<UMetaSoundSettings>())
+	{
+		if (const FMetaSoundPageSettings* PageSettings = Settings->FindPageSettings(InPageName))
+		{
+			return IsPageAuditionPlatformCookTarget(PageSettings->UniqueId);
+		}
+	}
+
+	return false;
+}
+
+bool UMetaSoundEditorSubsystem::IsPageAuditionPlatformCookTarget(const FGuid& InPageID) const
+{
+#if WITH_EDITORONLY_DATA
+	if (const UMetaSoundSettings* Settings = GetDefault<UMetaSoundSettings>())
+	{
+		if (const UMetasoundEditorSettings* EditorSettings = GetDefault<UMetasoundEditorSettings>())
+		{
+			bool bIsAuditionable = false;
+			auto PageIsTargetable = [&InPageID, &bIsAuditionable](const FGuid& PlatformTargetPageID)
+			{
+				bIsAuditionable |= PlatformTargetPageID == InPageID;
+			};
+			Settings->IterateCookedTargetPageIDs(EditorSettings->AuditionPlatform, PageIsTargetable);
+			return bIsAuditionable;
+		}
+	}
+#endif // WITH_EDITORONLY_DATA
+
+	return false;
+}
+
 void UMetaSoundEditorSubsystem::RegisterGraphWithFrontend(UObject& InMetaSound, bool bInForceViewSynchronization) const
 {
 	Metasound::Editor::FGraphBuilder::RegisterGraphWithFrontend(InMetaSound, bInForceViewSynchronization);
@@ -295,44 +329,43 @@ bool UMetaSoundEditorSubsystem::SetFocusedPageInternal(FName PageName, const FGu
 {
 	using namespace Metasound::Frontend;
 
-	if (!PageName.IsNone())
+	bool bAuditionPageSet = false;
+	// Must set audition target page before setting build page ID as listeners
+	// to build page ID changes need to reliably be able to adjust to newly assigned
+	// audition target page.
+	UMetasoundEditorSettings* EditorSettings = GetMutableDefault<UMetasoundEditorSettings>();
+	check(EditorSettings);
+	if (EditorSettings->AuditionPageMode == EAuditionPageMode::Focused)
 	{
-		if (UMetasoundEditorSettings* EditorSettings = GetMutableDefault<UMetasoundEditorSettings>())
+		if (EditorSettings->AuditionPage != PageName)
 		{
-			if (EditorSettings->AuditionPageMode == EAuditionPageMode::Focused)
-			{
-				// Must set audition target page before setting build page ID as listeners
-				// to build page ID changes need to reliably be able to adjust to newly assigned
-				// audition target page.
-				EditorSettings->AuditionTargetPage = PageName;
-			}
+			EditorSettings->AuditionPage = PageName;
+			bAuditionPageSet = true;
 		}
+	}
 
-		const FMetaSoundFrontendDocumentBuilder& DocBuilder = Builder.GetConstBuilder();
-		if (DocBuilder.GetBuildPageID() != InPageID)
+	const FMetaSoundFrontendDocumentBuilder& DocBuilder = Builder.GetConstBuilder();
+	if (DocBuilder.GetBuildPageID() != InPageID)
+	{
+		const FScopedTransaction Transaction(FText::Format(LOCTEXT("SetFocusedPageTransactionFormat", "Set Focused Page '{0}'"), FText::FromName(PageName)), bPostTransaction);
+		Builder.Modify();
+		UObject& MetaSound = DocBuilder.CastDocumentObjectChecked<UObject>();
+		if (Builder.GetBuilder().SetBuildPageID(InPageID))
 		{
-			const FScopedTransaction Transaction(FText::Format(LOCTEXT("SetFocusedPageTransactionFormat", "Set Focused Page '{0}'"), FText::FromName(PageName)), bPostTransaction);
-			Builder.Modify();
-			UObject& MetaSound = DocBuilder.CastDocumentObjectChecked<UObject>();
-			if (Builder.GetBuilder().SetBuildPageID(InPageID))
-			{
-				// Reregister to ensure all future audible instances are using the new page implementation.
-				RegisterGraphWithFrontend(MetaSound);
-			}
+			// Reregister to ensure all future audible instances are using the new page implementation.
+			RegisterGraphWithFrontend(MetaSound);
+		}
 
-			if (GEditor && bOpenEditor)
+		if (GEditor && bOpenEditor)
+		{
+			if (UAssetEditorSubsystem* AssetEditorSubsystem = GEditor->GetEditorSubsystem<UAssetEditorSubsystem>())
 			{
-				if (UAssetEditorSubsystem* AssetEditorSubsystem = GEditor->GetEditorSubsystem<UAssetEditorSubsystem>())
-				{
-					AssetEditorSubsystem->OpenEditorForAsset(&MetaSound);
-				}
+				AssetEditorSubsystem->OpenEditorForAsset(&MetaSound);
 			}
 		}
-
-		return true;
 	}
 
-	return false;
+	return bAuditionPageSet;
 }
 
 bool UMetaSoundEditorSubsystem::UnregisterToolbarExtender(TSharedRef<FExtender> InExtender)
diff --git a/Engine/Plugins/Runtime/Metasound/Source/MetasoundEditor/Private/SMetasoundGraphNode.cpp b/Engine/Plugins/Runtime/Metasound/Source/MetasoundEditor/Private/SMetasoundGraphNode.cpp
--- a/Engine/Plugins/Runtime/Metasound/Source/MetasoundEditor/Private/SMetasoundGraphNode.cpp
+++ b/Engine/Plugins/Runtime/Metasound/Source/MetasoundEditor/Private/SMetasoundGraphNode.cpp
@@ -28,6 +28,7 @@
 #include "MetasoundEditorGraphSchema.h"
 #include "MetasoundEditorModule.h"
 #include "MetasoundEditorSettings.h"
+#include "MetasoundEditorSubsystem.h"
 #include "MetasoundFrontendNodeTemplateRegistry.h"
 #include "MetasoundFrontendRegistries.h"
 #include "MetasoundTrace.h"
@@ -766,29 +767,30 @@ namespace Metasound
 		{
 			using namespace Engine;
 
+			TSharedPtr<SWidget> OuterContentBox;
 			TWeakObjectPtr<UMetasoundEditorGraphInput> GraphMember;
 			{
 				GraphMember = Cast<UMetasoundEditorGraphInput>(GetMetaSoundMember());
 				if (!GraphMember.IsValid())
 				{
-					return { };
+					return OuterContentBox;
 				}
 			}
 
-			bool bUseAudioMaterialWidgets = false;
-			if (const UMetasoundEditorSettings* EditorSettings = GetDefault<UMetasoundEditorSettings>())
+			const UMetasoundEditorSettings* EditorSettings = GetDefault<UMetasoundEditorSettings>();
+			const UMetaSoundSettings* MetaSoundSettings = GetDefault<UMetaSoundSettings>();
+			if (!EditorSettings || !MetaSoundSettings)
 			{
-				bUseAudioMaterialWidgets = EditorSettings->bUseAudioMaterialWidgets;
+				return OuterContentBox;
 			}
 
 			const UMetasoundEditorGraph* OwningGraph = GraphMember->GetOwningGraph();
 			if (!OwningGraph || !OwningGraph->IsEditable() || GraphMember->GetVertexAccessType() == EMetasoundFrontendVertexAccessType::Unset)
 			{
-				return { };
+				return OuterContentBox;
 			}
 
-			TSharedPtr<SWidget> OuterContentBox;
-
+			const bool bUseAudioMaterialWidgets = EditorSettings->bUseAudioMaterialWidgets;
 			TWeakObjectPtr<UMetasoundEditorGraphMemberDefaultFloat> DefaultFloat = Cast<UMetasoundEditorGraphMemberDefaultFloat>(GraphMember->GetLiteral());
 			if (DefaultFloat.IsValid() && DefaultFloat->WidgetType != EMetasoundMemberDefaultWidget::None)
 			{
@@ -856,9 +858,6 @@ namespace Metasound
 				{
 					if (bUseAudioMaterialWidgets)
 					{
-						const UMetasoundEditorSettings* EditorSettings = GetDefault<UMetasoundEditorSettings>();
-						check(EditorSettings)
-
 						SAssignNew(FloatInputWidget, SAudioMaterialLabeledSlider)
 							.Owner(GraphMember->GetOwningGraph())
 							.Style(EditorSettings->GetSliderStyle())
@@ -992,9 +991,6 @@ namespace Metasound
 
 					if (bUseAudioMaterialWidgets)
 					{
-						const UMetasoundEditorSettings* EditorSettings = GetDefault<UMetasoundEditorSettings>();
-						check(EditorSettings)
-
 						SAssignNew(FloatInputWidget, SAudioMaterialLabeledKnob)
 							.Owner(GraphMember->GetOwningGraph())
 							.Style(EditorSettings->GetKnobStyle())
@@ -1155,9 +1151,6 @@ namespace Metasound
 								}
 							};
 
-							const UMetasoundEditorSettings* EditorSettings = GetDefault<UMetasoundEditorSettings>();
-							check(EditorSettings);
-
 							SAssignNew(MaterialButtonWidget, SAudioMaterialButton)
 								.AudioMaterialButtonStyle(EditorSettings->GetButtonStyle())
 								.OnBooleanValueChanged_Lambda(OnBoolValueChangedLambda)
diff --git a/Engine/Plugins/Runtime/Metasound/Source/MetasoundEditor/Private/SMetasoundStats.cpp b/Engine/Plugins/Runtime/Metasound/Source/MetasoundEditor/Private/SMetasoundStats.cpp
--- a/Engine/Plugins/Runtime/Metasound/Source/MetasoundEditor/Private/SMetasoundStats.cpp
+++ b/Engine/Plugins/Runtime/Metasound/Source/MetasoundEditor/Private/SMetasoundStats.cpp
@@ -6,46 +6,77 @@
 #include "Editor.h"
 #include "Internationalization/Text.h"
 #include "MetasoundDocumentBuilderRegistry.h"
+#include "MetasoundEditorModule.h"
 #include "MetasoundEditorSettings.h"
 #include "MetasoundFrontendDocumentBuilder.h"
 #include "MetasoundSettings.h"
 #include "MetasoundSource.h"
 #include "MetasoundGenerator.h"
+#include "Styling/SlateColor.h"
 #include "Templates/SharedPointer.h"
 #include "Widgets/DeclarativeSyntaxSupport.h"
 #include "Widgets/SBoxPanel.h"
 #include "Widgets/Text/STextBlock.h"
 
 
+#define LOCTEXT_NAMESPACE "MetaSoundEditor"
+
 namespace Metasound::Editor
 {
+	namespace StatsPrivate
+	{
+		const FLinearColor BaseTextColor(1, 1, 1, 0.30f);
+	}
+
 	void SPageStats::Construct(const FArguments& InArgs)
 	{
 		SVerticalBox::Construct(SVerticalBox::FArguments());
-		AddSlot()
-			.HAlign(HAlign_Left)
+
+		TSharedRef<SHorizontalBox> PageWidgetBox = SNew(SHorizontalBox);
+		PageWidgetBox->AddSlot()
+			.Padding(2.0f)
+			.HAlign(HAlign_Center)
 			[
 				SAssignNew(PageTextWidget, STextBlock)
 				.Visibility(EVisibility::HitTestInvisible)
 				.TextStyle(FAppStyle::Get(), "Graph.ZoomText")
-				.ColorAndOpacity(FLinearColor(1, 1, 1, 0.30f))
+				.ColorAndOpacity(StatsPrivate::BaseTextColor)
 			];
+
+		ExecImageWidget = SNew(SImage)
+			.Image(Style::CreateSlateIcon("MetasoundEditor.Page.Executing").GetIcon())
+			.DesiredSizeOverride(FVector2D(24.f, 24.f))
+			.ColorAndOpacity(FStyleColors::AccentGreen)
+			.Visibility(EVisibility::Collapsed);
+
+		PageWidgetBox->AddSlot()
+			.Padding(2.0f)
+			.HAlign(HAlign_Center)
+			.VAlign(VAlign_Center)
+			.AutoWidth()[ ExecImageWidget.ToSharedRef()];
+
+		AddSlot().HAlign(HAlign_Left) [ PageWidgetBox ];
+	}
+
+	void SPageStats::SetExecVisibility(TAttribute<EVisibility> InVisibility)
+	{
+		ExecImageWidget->SetVisibility(MoveTemp(InVisibility));
 	}
 
-	void SPageStats::Update(const FMetaSoundPageSettings* PageSettings)
+	void SPageStats::Update(const FMetaSoundPageSettings* PageSettings, const FText& Header, const FSlateColor* ColorOverride)
 	{
 		using namespace Engine;
 
-		FString PageInfo;
-		EVisibility Visibility = EVisibility::Hidden;
+		FText PageInfo;
+		EVisibility Visibility = EVisibility::Collapsed;
 		if (PageTextWidget.IsValid() && PageSettings)
 		{
-			PageInfo = FString::Printf(TEXT("Page: %s"), *PageSettings->Name.ToString());
+			PageInfo = FText::Format(LOCTEXT("PageStatsFormat", "{0}: {1}"), Header, FText::FromString(PageSettings->Name.ToString()));
 			Visibility = EVisibility::Visible;
 		}
 
-		PageTextWidget->SetText(FText::FromString(MoveTemp(PageInfo)));
-		PageTextWidget->SetVisibility(Visibility);
+		PageTextWidget->SetText(PageInfo);
+		PageTextWidget->SetColorAndOpacity(ColorOverride ? *ColorOverride : StatsPrivate::BaseTextColor);
 	}
 
 	void SRenderStats::Construct(const FArguments& InArgs)
@@ -116,3 +147,4 @@ namespace Metasound::Editor
 		}
 	}
 } // namespace Metasound::Editor
+#undef LOCTEXT_NAMESPACE
diff --git a/Engine/Plugins/Runtime/Metasound/Source/MetasoundEditor/Private/SMetasoundStats.h b/Engine/Plugins/Runtime/Metasound/Source/MetasoundEditor/Private/SMetasoundStats.h
--- a/Engine/Plugins/Runtime/Metasound/Source/MetasoundEditor/Private/SMetasoundStats.h
+++ b/Engine/Plugins/Runtime/Metasound/Source/MetasoundEditor/Private/SMetasoundStats.h
@@ -3,6 +3,7 @@
 
 #include "Templates/SharedPointer.h"
 #include "Widgets/DeclarativeSyntaxSupport.h"
+#include "Widgets/Images/SImage.h"
 #include "Widgets/SBoxPanel.h"
 
 
@@ -13,6 +14,7 @@ class UMetaSoundBuilderBase;
 class UMetaSoundSource;
 
 struct FMetaSoundPageSettings;
+struct FSlateColor;
 
 
 namespace Metasound::Editor
@@ -25,9 +27,11 @@ namespace Metasound::Editor
 		SLATE_END_ARGS()
 
 		void Construct(const FArguments& InArgs);
-		void Update(const FMetaSoundPageSettings* PageSettings);
+		void SetExecVisibility(TAttribute<EVisibility> InVisibility);
+		void Update(const FMetaSoundPageSettings* PageSettings, const FText& Header, const FSlateColor* TextColor);
 
 	private:
+		TSharedPtr<SImage> ExecImageWidget;
 		TSharedPtr<STextBlock> PageTextWidget;
 	};
 
diff --git a/Engine/Plugins/Runtime/Metasound/Source/MetasoundEditor/Public/MetasoundEditorSubsystem.h b/Engine/Plugins/Runtime/Metasound/Source/MetasoundEditor/Public/MetasoundEditorSubsystem.h
--- a/Engine/Plugins/Runtime/Metasound/Source/MetasoundEditor/Public/MetasoundEditorSubsystem.h
+++ b/Engine/Plugins/Runtime/Metasound/Source/MetasoundEditor/Public/MetasoundEditorSubsystem.h
@@ -69,6 +69,11 @@ public:
 	UE_DEPRECATED(5.5, "EdGraph is now transiently generated and privately managed for asset editor use only.")
 	void InitEdGraph(UObject& InMetaSound);
 
+	// Returns whether or not a page with the given name both exists and is set as
+	// a valid, cooked target for the currently set audition platform in editor.
+	bool IsPageAuditionPlatformCookTarget(FName InPageName) const;
+	bool IsPageAuditionPlatformCookTarget(const FGuid& InPageID) const;
+
 	// Wraps RegisterGraphWithFrontend logic in Frontend with any additional logic required to refresh editor & respective editor object state.
 	// @param InMetaSound - MetaSound to register
 	// @param bInForceSynchronize - Forces the synchronize flag for all open graphs being registered by this call (all referenced graphs and
@@ -79,16 +84,17 @@ public:
 	void RegisterToolbarExtender(TSharedRef<FExtender> InExtender);
 
 	// If the given page name is implemented on the provided builder, sets the focused page of
-	// the provided builder to the given page name if and sets the audition target page to
+	// the provided builder to the associated page and sets the audition page to
 	// the provided name. If the given builder has an asset editor open, optionally opens or brings
-	// that editor's associated page into user focus.
+	// that editor's associated PageID into user focus.
 	UFUNCTION(BlueprintCallable, Category = "Audio|MetaSound|Builder|Editor", meta = (ExpandEnumAsExecs = "OutResult"))
 	void SetFocusedPage(UMetaSoundBuilderBase* Builder, FName PageName, bool bOpenEditor, EMetaSoundBuilderResult& OutResult) const;
 
 	// If the given PageID is implemented on the provided builder, sets the focused page of
-	// the provided builder to the given PageID if and sets the audition target page to
+	// the provided builder to the associated page and sets the audition target page to
 	// the provided ID. If the given builder has an asset editor open, optionally opens or brings
-	// that editor's associated PageID into user focus.
+	// that editor's associated PageID into user focus. Returns whether or not the audition page
+	// was set to the provided focus page.
 	bool SetFocusedPage(UMetaSoundBuilderBase& Builder, const FGuid& InPageID, bool bOpenEditor, bool bPostTransaction = true) const;
 
 	// Unregisters toolbar extender that is displayed in the MetaSound Asset Editor.
diff --git a/Engine/Plugins/Runtime/Metasound/Source/MetasoundEngine/Private/Metasound.cpp b/Engine/Plugins/Runtime/Metasound/Source/MetasoundEngine/Private/Metasound.cpp
--- a/Engine/Plugins/Runtime/Metasound/Source/MetasoundEngine/Private/Metasound.cpp
+++ b/Engine/Plugins/Runtime/Metasound/Source/MetasoundEngine/Private/Metasound.cpp
@@ -32,6 +32,12 @@
 
 #include UE_INLINE_GENERATED_CPP_BY_NAME(Metasound)
 
+
+#if WITH_EDITOR
+#include "Misc/DataValidation.h"
+#endif // WITH_EDITOR
+
+
 #if WITH_EDITORONLY_DATA
 #include "EdGraph/EdGraph.h"
 #endif // WITH_EDITORONLY_DATA
@@ -121,11 +127,11 @@ void UMetaSoundPatch::PostEditUndo()
 
 EDataValidationResult UMetaSoundPatch::IsDataValid(FDataValidationContext& Context) const
 {
-	const EDataValidationResult Result = Metasound::Engine::FAssetHelper::IsDataValid(*this, Context);
+	const EDataValidationResult Result = Metasound::Engine::FAssetHelper::IsDataValid(*this, RootMetaSoundDocument, Context);
 	return CombineDataValidationResults(Result, Super::IsDataValid(Context));
 }
 
-#endif // WITHEDITOR
+#endif // WITH_EDITOR
 
 void UMetaSoundPatch::BeginDestroy()
 {
diff --git a/Engine/Plugins/Runtime/Metasound/Source/MetasoundEngine/Private/MetasoundEngineAsset.h b/Engine/Plugins/Runtime/Metasound/Source/MetasoundEngine/Private/MetasoundEngineAsset.h
--- a/Engine/Plugins/Runtime/Metasound/Source/MetasoundEngine/Private/MetasoundEngineAsset.h
+++ b/Engine/Plugins/Runtime/Metasound/Source/MetasoundEngine/Private/MetasoundEngineAsset.h
@@ -9,6 +9,7 @@
 #include "MetasoundFrontendDocumentIdGenerator.h"
 #include "MetasoundFrontendRegistryKey.h"
 #include "MetasoundGlobals.h"
+#include "MetasoundSettings.h"
 #include "MetasoundUObjectRegistry.h"
 #include "Misc/App.h"
 #include "Modules/ModuleManager.h"
@@ -105,8 +106,7 @@ namespace Metasound::Engine
 			}
 		}
 
-		template <typename TMetaSoundObject>
-		static EDataValidationResult IsClassNameUnique(const TMetaSoundObject& InMetaSound, FDataValidationContext& InOutContext)
+		static EDataValidationResult IsClassNameUnique(const FMetasoundFrontendDocument& Document, FDataValidationContext& InOutContext)
 		{
 			using namespace Metasound::Frontend;
 			using namespace Metasound::Engine;
@@ -136,7 +136,7 @@ namespace Metasound::Engine
 			AssetManager.SetLogActiveAssetsOnShutdown(false);
 
 			// Add error for multiple assets with the same class name
-			const FAssetKey Key(InMetaSound.GetConstDocument().RootGraph.Metadata);
+			const FAssetKey Key(Document.RootGraph.Metadata);
 			const TArray<FTopLevelAssetPath>* AssetPaths = AssetManager.FindAssetPaths(Key);
 			if (AssetPaths && AssetPaths->Num() > 1)
 			{
@@ -153,13 +153,51 @@ namespace Metasound::Engine
 			return Result;
 		}
 
-		template <typename TMetaSoundObject>
-		static EDataValidationResult IsDataValid(const TMetaSoundObject& InMetaSound, FDataValidationContext& InOutContext)
+		static EDataValidationResult IsDataValid(const UObject& MetaSound, const FMetasoundFrontendDocument& Document, FDataValidationContext& InOutContext)
 		{
+			using namespace Metasound::Engine;
+
 			EDataValidationResult Result = EDataValidationResult::Valid;
 			if (MetasoundEngineModulePrivate::EnableMetaSoundEditorAssetValidation)
 			{
-				Result = IsClassNameUnique(InMetaSound, InOutContext);
+				Result = IsClassNameUnique(Document, InOutContext);
+			}
+
+			const UMetaSoundSettings* Settings = GetDefault<UMetaSoundSettings>();
+			check(Settings);
+
+			TSet<FGuid> ValidPageIDs;
+			auto ErrorIfMissing = [&](const FGuid& PageID, const FText& DataDescriptor)
+			{
+				if (!ValidPageIDs.Contains(PageID))
+				{
+					if (const FMetaSoundPageSettings* PageSettings = Settings->FindPageSettings(PageID))
+					{
+						ValidPageIDs.Add(PageSettings->UniqueId);
+					}
+					else
+					{
+						Result = EDataValidationResult::Invalid;
+						InOutContext.AddMessage(FAssetData(&MetaSound), EMessageSeverity::Error, FText::Format(
+							LOCTEXT("InvalidPageDataFormat", "MetaSound contains invalid {0} with page ID '{1}': page not found in Project 'MetaSound' Settings. Remove page data or migrate to existing page identifier."),
+							DataDescriptor,
+							FText::FromString(PageID.ToString())));
+					}
+				}
+			};
+
+			const TArray<FMetasoundFrontendGraph>& Graphs = Document.RootGraph.GetConstGraphPages();
+			for (const FMetasoundFrontendGraph& Graph : Graphs)
+			{
+				ErrorIfMissing(Graph.PageID, LOCTEXT("GraphPageDescriptor", "graph"));
+			}
+
+			for (const FMetasoundFrontendClassInput& ClassInput : Document.RootGraph.Interface.Inputs)
+			{
+				ClassInput.IterateDefaults([&](const FGuid& PageID, const FMetasoundFrontendLiteral&)
+				{
+					ErrorIfMissing(PageID, FText::Format(LOCTEXT("InputPageDefaultDescriptorFormat", "input '{0}' default value"), FText::FromName(ClassInput.Name)));
+				});
 			}
 			return Result;
 		}
diff --git a/Engine/Plugins/Runtime/Metasound/Source/MetasoundEngine/Private/MetasoundSettings.cpp b/Engine/Plugins/Runtime/Metasound/Source/MetasoundEngine/Private/MetasoundSettings.cpp
--- a/Engine/Plugins/Runtime/Metasound/Source/MetasoundEngine/Private/MetasoundSettings.cpp
+++ b/Engine/Plugins/Runtime/Metasound/Source/MetasoundEngine/Private/MetasoundSettings.cpp
@@ -160,34 +160,54 @@ void UMetaSoundSettings::ConformPageSettings(bool bNotifyDefaultRenamed)
 {
 	using namespace Metasound;
 
-	bool bTargetFound = false;
-	bool bDefaultRenamed = false;
-	for (int32 Index = PageSettings.Num() - 1; Index >= 0; --Index)
+	DefaultPageSettings.UniqueId = Metasound::Frontend::DefaultPageID;
+	DefaultPageSettings.Name = Metasound::Frontend::DefaultPageName;
+	DefaultPageSettings.bIsDefaultPage = true;
+	DefaultPageSettings.ExcludeFromCook = false;
+
+	bool bInvalidDefaultRenamed = false;
+	TMap<FName, bool> PlatformHasTarget;
+	auto GatherPlatformTargets = [&PlatformHasTarget](const FMetaSoundPageSettings& Page)
+	{
+		PlatformHasTarget.FindOrAdd({ }) |= Page.CanTarget.Default;
+		for (const TPair<FName, bool>& Pair : Page.CanTarget.PerPlatform)
+		{
+			PlatformHasTarget.FindOrAdd(Pair.Key) |= Pair.Value;
+		}
+	};
+
+	GatherPlatformTargets(DefaultPageSettings);
+	for (FMetaSoundPageSettings& Page : PageSettings)
 	{
-		FMetaSoundPageSettings& Page = PageSettings[Index];
 		const bool bIsDefaultName = Page.Name == Frontend::DefaultPageName;
 		if (bIsDefaultName)
 		{
 			const TSet<FName> PageNames(GetPageNames());
 			Page.Name = SettingsPrivate::GenerateUniqueName(PageNames, *Page.Name.ToString());
-			bDefaultRenamed = true;
+			bInvalidDefaultRenamed = true;
 		}
-		bTargetFound |= TargetPageName == Page.Name;
-	}
 
-	DefaultPageSettings.CanTarget = PageSettings.IsEmpty();
-	if (DefaultPageSettings.CanTarget.GetValue())
-	{
-		TargetPageName = DefaultPageSettings.Name;
-	}
-	else
-	{
-		if (!bTargetFound)
-		{
-			TargetPageName = PageSettings.Last().Name;
-		}
+		GatherPlatformTargets(Page);
+
+		Page.bIsDefaultPage = false;
 	}
 
+	// Forces each platform to target at least one page setting.
+ 	for (const TPair<FName, bool>& Pair : PlatformHasTarget)
+ 	{
+ 		if (!Pair.Value)
+ 		{
+ 			if (Pair.Key.IsNone())
+ 			{
+ 				DefaultPageSettings.CanTarget.Default = true;
+ 			}
+ 			else
+ 			{
+ 				DefaultPageSettings.CanTarget.PerPlatform.FindOrAdd(Pair.Key) = true;
+ 			}
+ 		}
+ 	}
+
 #if WITH_EDITORONLY_DATA
 	{
 		FScopeLock Lock(&CookPlatformTargetCritSec);
@@ -198,7 +218,7 @@ void UMetaSoundSettings::ConformPageSettings(bool bNotifyDefaultRenamed)
 
 	TargetPageNameOverride.Reset();
 
-	if (bNotifyDefaultRenamed && bDefaultRenamed)
+	if (bNotifyDefaultRenamed && bInvalidDefaultRenamed)
 	{
 		OnDefaultRenamed.Broadcast();
 	}
@@ -241,7 +261,7 @@ const FMetaSoundPageSettings& UMetaSoundSettings::GetDefaultPageSettings() const
 #if WITH_EDITOR
 TArray<FName> UMetaSoundSettings::GetAllPlatformNamesImplementingTargets() const
 {
-	TSet<FName> PlatformNames { FPlatformProperties::IniPlatformName() };
+	TSet<FName> PlatformNames;
 	IteratePageSettings([&](const FMetaSoundPageSettings& PageSetting)
 	{
 		TArray<FName> PagePlatforms = PageSetting.GetTargetPlatforms();
@@ -374,6 +394,9 @@ void UMetaSoundSettings::PostEditChangeProperty(FPropertyChangedEvent& PropertyC
 void UMetaSoundSettings::PostInitProperties()
 {
 	Super::PostInitProperties();
+
+	constexpr bool bNotifyDefaultRenamed = false;
+	ConformPageSettings(bNotifyDefaultRenamed);
 }
 #endif // WITH_EDITOR
 
@@ -465,4 +488,41 @@ void UMetaSoundSettings::IteratePageSettings(TFunctionRef<void(const FMetaSoundP
 		}
 	}
 }
+
+void UMetaSoundSettings::PostLoad()
+{
+#if !WITH_EDITORONLY_DATA
+	if (const FMetaSoundPageSettings* Page = FindPageSettings(TargetPageName))
+	{
+		if (Page->CanTarget.GetValue())
+		{
+			return;
+		}
+	}
+
+	UE_LOG(LogMetaSound, Warning, TEXT("TargetPageName on load in MetaSound settings did not correspond to a valid targetable page."));
+	if (DefaultPageSettings.CanTarget.GetValue())
+	{
+		UE_LOG(LogMetaSound, Warning, TEXT("Setting target to 'Default' page settings."));
+		TargetPageNameOverride = DefaultPageSettings.Name;
+		return;
+	}
+	else
+	{
+		for (const FMetaSoundPageSettings& Page : PageSettings)
+		{
+			if (Page.CanTarget.GetValue())
+			{
+				UE_LOG(LogMetaSound, Warning, TEXT("Setting target to '%s' page settings."), *Page.Name.ToString());
+				TargetPageNameOverride = Page.Name;
+				return;
+			}
+		}
+
+		UE_LOG(LogMetaSound, Warning, TEXT("Setting target to 'Default' page settings & forcing default as targetable."));
+		DefaultPageSettings.CanTarget.Default = true;
+		TargetPageNameOverride = DefaultPageSettings.Name;
+	}
+#endif // !WITH_EDITORONLY_DATA
+}
 #undef LOCTEXT_NAMESPACE // MetaSound
diff --git a/Engine/Plugins/Runtime/Metasound/Source/MetasoundEngine/Private/MetasoundSource.cpp b/Engine/Plugins/Runtime/Metasound/Source/MetasoundEngine/Private/MetasoundSource.cpp
--- a/Engine/Plugins/Runtime/Metasound/Source/MetasoundEngine/Private/MetasoundSource.cpp
+++ b/Engine/Plugins/Runtime/Metasound/Source/MetasoundEngine/Private/MetasoundSource.cpp
@@ -48,6 +48,12 @@
 
 #include UE_INLINE_GENERATED_CPP_BY_NAME(MetasoundSource)
 
+
+#if WITH_EDITOR
+#include "Misc/DataValidation.h"
+#endif // WITH_EDITOR
+
+
 #if WITH_EDITORONLY_DATA
 #include "EdGraph/EdGraph.h"
 #endif // WITH_EDITORONLY_DATA
@@ -496,7 +502,7 @@ bool UMetaSoundSource::CanEditChange(const FProperty* InProperty) const
 
 EDataValidationResult UMetaSoundSource::IsDataValid(FDataValidationContext& Context) const
 {
-	const EDataValidationResult Result = Metasound::Engine::FAssetHelper::IsDataValid(*this, Context);
+	const EDataValidationResult Result = Metasound::Engine::FAssetHelper::IsDataValid(*this, RootMetasoundDocument, Context);
 	return CombineDataValidationResults(Result, Super::IsDataValid(Context));
 }
 
diff --git a/Engine/Plugins/Runtime/Metasound/Source/MetasoundEngine/Public/Metasound.h b/Engine/Plugins/Runtime/Metasound/Source/MetasoundEngine/Public/Metasound.h
--- a/Engine/Plugins/Runtime/Metasound/Source/MetasoundEngine/Public/Metasound.h
+++ b/Engine/Plugins/Runtime/Metasound/Source/MetasoundEngine/Public/Metasound.h
@@ -20,6 +20,10 @@
 
 
 // Forward Declarations
+#if WITH_EDITOR
+class FDataValidationContext;
+#endif // WITH_EDITOR
+
 namespace Metasound::Engine
 {
 	struct FAssetHelper;
diff --git a/Engine/Plugins/Runtime/Metasound/Source/MetasoundEngine/Public/MetasoundSettings.h b/Engine/Plugins/Runtime/Metasound/Source/MetasoundEngine/Public/MetasoundSettings.h
--- a/Engine/Plugins/Runtime/Metasound/Source/MetasoundEngine/Public/MetasoundSettings.h
+++ b/Engine/Plugins/Runtime/Metasound/Source/MetasoundEngine/Public/MetasoundSettings.h
@@ -66,7 +66,7 @@ struct METASOUNDENGINE_API FMetaSoundPageSettings
 	FGuid UniqueId = Metasound::Frontend::DefaultPageID;
 
 	/** Name of this page's setting to be displayed in editors and used for identification from Blueprint/native API. */
-	UPROPERTY(EditAnywhere, Category = "Pages")
+	UPROPERTY(EditAnywhere, Category = "Pages", meta = (EditCondition = "!bIsDefaultPage"))
 	FName Name = Metasound::Frontend::DefaultPageName;
 
 private:
@@ -77,11 +77,16 @@ private:
 	FPerPlatformBool CanTarget = true;
 
 #if WITH_EDITORONLY_DATA
+	// Just used to inform edit condition to enable/disable exclude from cook. Maintained by ConformPageSettings on object load/mutation.
+	// EditCondition meta mark-up is hack to avoid boolean being default added to name field
+	UPROPERTY(EditAnywhere, Category = "Pages", meta = (EditCondition = false, EditConditionHides))
+	bool bIsDefaultPage = true;
+
 	// When true, exclude page data when cooking from the assigned platform(s)/platform group(s).
 	// If false, page data may or may not be included in cook depending on whether or not the given
 	// page data is required in order to ensure a value is always resolved for the cook platform target(s).
 	// (Ignored if CanTarget true for corresponding platform/group).
-	UPROPERTY(EditAnywhere, Category = "Pages")
+	UPROPERTY(EditAnywhere, Category = "Pages", meta = (EditCondition = "!bIsDefaultPage"))
 	FPerPlatformBool ExcludeFromCook = false;
 #endif //WITH_EDITORONLY_DATA
 
@@ -166,22 +171,22 @@ private:
 
 	/** Page Name to target when attempting to execute MetaSound. If target page is not implemented (or cooked in a runtime build)
 	  * for the active platform, uses order of cooked pages (see 'Page Settings' for order) falling back to lower index-ordered page
-	  * implemented in MetaSound asset. If no fallback is found, uses default implementation. If any pages are specified, target
-	  * cannot be set directly to default (allows for removal of default implementation when cooking for performance reasons). */
-	UPROPERTY(EditAnywhere, config, Category = Pages, meta = (GetOptions = "MetasoundEngine.MetaSoundSettings.GetPageNames"))
+	  * implemented in MetaSound asset. If no fallback is found, uses default implementation.
+	  */
+	UPROPERTY(EditAnywhere, config, Category = "Pages (Experimental)", meta = (GetOptions = "MetasoundEngine.MetaSoundSettings.GetPageNames"))
 	FName TargetPageName = Metasound::Frontend::DefaultPageName;
 
+	/** Default page settings to be used in editor and if no other page settings are targeted or defined. */
+	UPROPERTY(EditAnywhere, config, Category = "Pages (Experimental)")
+	FMetaSoundPageSettings DefaultPageSettings;
+
 	/** Array of possible page settings that can be added to a MetaSound object. Order
 	  * defines default fallback logic whereby a higher index-ordered page
 	  * implemented in a MetaSound asset is higher priority (see 'Target Page').
 	  */
-	UPROPERTY(EditAnywhere, config, Category = Pages)
+	UPROPERTY(EditAnywhere, config, Category = "Pages (Experimental)")
 	TArray<FMetaSoundPageSettings> PageSettings;
 
-	/** Default page settings to be used in editor and if no other page settings are targeted or defined. */
-	UPROPERTY()
-	FMetaSoundPageSettings DefaultPageSettings;
-
 	/** Array of possible quality settings for Metasounds to chose from */
 	UPROPERTY(EditAnywhere, config, Category = Quality)
 	TArray<FMetaSoundQualitySettings> QualitySettings;
@@ -245,6 +250,8 @@ public:
 	// optionally set to reverse, iterates in reverse.
 	void IteratePageSettings(TFunctionRef<void(const FMetaSoundPageSettings&)> Iter, bool bReverse = false) const;
 
+	virtual void PostLoad() override;
+
 	// Sets the target page to the given name. Returns true if associated page settings were found
 	// and target set, false if not found and not set.
 	bool SetTargetPage(FName PageName);
diff --git a/Engine/Plugins/Runtime/Metasound/Source/MetasoundEngine/Public/MetasoundSource.h b/Engine/Plugins/Runtime/Metasound/Source/MetasoundEngine/Public/MetasoundSource.h
--- a/Engine/Plugins/Runtime/Metasound/Source/MetasoundEngine/Public/MetasoundSource.h
+++ b/Engine/Plugins/Runtime/Metasound/Source/MetasoundEngine/Public/MetasoundSource.h
@@ -23,10 +23,13 @@
 
 
 // Forward Declarations
+#if WITH_EDITOR
+class FDataValidationContext;
+#endif // WITH_EDITOR
 class UMetaSoundSettings;
-struct FMetaSoundQualitySettings;
 
 struct FMetaSoundFrontendDocumentBuilder;
+struct FMetaSoundQualitySettings;
 
 namespace Audio
 {
diff --git a/Engine/Plugins/Runtime/Metasound/Source/MetasoundFrontend/Private/MetasoundAssetBase.cpp b/Engine/Plugins/Runtime/Metasound/Source/MetasoundFrontend/Private/MetasoundAssetBase.cpp
--- a/Engine/Plugins/Runtime/Metasound/Source/MetasoundFrontend/Private/MetasoundAssetBase.cpp
+++ b/Engine/Plugins/Runtime/Metasound/Source/MetasoundFrontend/Private/MetasoundAssetBase.cpp
@@ -1048,15 +1048,16 @@ bool FMetasoundAssetBase::VersionDependencies(FMetaSoundFrontendDocumentBuilder&
 	const FGuid InitBuildPageID = Builder.GetBuildPageID();
 	Metasound::Frontend::FDocumentHandle DocHandle = GetDocumentHandle();
 	const FMetasoundFrontendGraphClass& RootGraph = Builder.GetConstDocumentChecked().RootGraph;
+	constexpr bool bBroadcastPageIDDelegate = false;
 	RootGraph.IterateGraphPages([&](const FMetasoundFrontendGraph& Graph)
 	{
 		// Set the build page ID to this graph as a hack to apply dependency versioning logic using
 		// the controller/handle API until auto-update is renamed & moved to use document builder API.
-		Builder.SetBuildPageID(Graph.PageID);
+		Builder.SetBuildPageID(Graph.PageID, bBroadcastPageIDDelegate);
 		FString OwningAssetName = GetOwningAssetName();
 		bDocumentModified |= FAutoUpdateRootGraph(MoveTemp(OwningAssetName), bInLogWarningsOnDroppedConnection).Transform(DocHandle);
 	});
-	Builder.SetBuildPageID(InitBuildPageID);
+	Builder.SetBuildPageID(InitBuildPageID, bBroadcastPageIDDelegate);
 #endif // WITH_EDITORONLY_DATA
 
 	return bDocumentModified;
diff --git a/Engine/Plugins/Runtime/Metasound/Source/MetasoundFrontend/Private/MetasoundFrontendDocumentBuilder.cpp b/Engine/Plugins/Runtime/Metasound/Source/MetasoundFrontend/Private/MetasoundFrontendDocumentBuilder.cpp
--- a/Engine/Plugins/Runtime/Metasound/Source/MetasoundFrontend/Private/MetasoundFrontendDocumentBuilder.cpp
+++ b/Engine/Plugins/Runtime/Metasound/Source/MetasoundFrontend/Private/MetasoundFrontendDocumentBuilder.cpp
@@ -3560,7 +3560,7 @@ void FMetaSoundFrontendDocumentBuilder::SetAuthor(const FString& InAuthor)
 #endif // WITH_EDITOR
 
 #if WITH_EDITORONLY_DATA
-bool FMetaSoundFrontendDocumentBuilder::SetBuildPageID(const FGuid& InBuildPageID)
+bool FMetaSoundFrontendDocumentBuilder::SetBuildPageID(const FGuid& InBuildPageID, bool bBroadcastDelegate)
 {
 	using namespace Metasound::Frontend;
 
@@ -3573,7 +3573,10 @@ bool FMetaSoundFrontendDocumentBuilder::SetBuildPageID(const FGuid& InBuildPageI
 
 			constexpr bool bPrimeCache = false;
 			DocumentCache->SetBuildPageID(BuildPageID);
-			DocumentDelegates->PageDelegates.OnPageSet.Broadcast({ BuildPageID });
+			if (bBroadcastDelegate)
+			{
+				DocumentDelegates->PageDelegates.OnPageSet.Broadcast({ BuildPageID });
+			}
 		}
 		return true;
 	}
diff --git a/Engine/Plugins/Runtime/Metasound/Source/MetasoundFrontend/Public/MetasoundFrontendDocumentBuilder.h b/Engine/Plugins/Runtime/Metasound/Source/MetasoundFrontend/Public/MetasoundFrontendDocumentBuilder.h
--- a/Engine/Plugins/Runtime/Metasound/Source/MetasoundFrontend/Public/MetasoundFrontendDocumentBuilder.h
+++ b/Engine/Plugins/Runtime/Metasound/Source/MetasoundFrontend/Public/MetasoundFrontendDocumentBuilder.h
@@ -408,8 +408,10 @@ public:
 	// Returns true if the builder is already targeting the given ID or if it successfully
 	// found a page implementation with the given ID and was able to switch to it, false if not.
 	// Swapping the targeted build graph ID clears the local cache, so swapping frequently can
-	// induce cash thrashing.
-	bool SetBuildPageID(const FGuid& InBuildPageID);
+	// induce cash thrashing. BroadcastDelegate should always be true unless dealing with the controller
+	// API (exposed as a mechanism for mutating via controllers while deprecating.  Option will be removed
+	// in a future build).
+	bool SetBuildPageID(const FGuid& InBuildPageID, bool bBroadcastDelegate = true);
 
 	// Sets the given input`s IsAdvancedDisplay state. AdvancedDisplay pins are hidden in the node by default.
 	// returns true if state was changed.
