Fix for "Editor" audition platform setting not validating pageID is implemented when resolving page ID

#rnx

diff --git a/Engine/Plugins/Runtime/Metasound/Source/MetasoundEditor/Private/MetasoundEditorSettings.cpp b/Engine/Plugins/Runtime/Metasound/Source/MetasoundEditor/Private/MetasoundEditorSettings.cpp
--- a/Engine/Plugins/Runtime/Metasound/Source/MetasoundEditor/Private/MetasoundEditorSettings.cpp
+++ b/Engine/Plugins/Runtime/Metasound/Source/MetasoundEditor/Private/MetasoundEditorSettings.cpp
@@ -131,16 +131,6 @@ Metasound::Engine::FPageResolutionEditorResults UMetasoundEditorSettings::Resolv
 	{
 		if (const FMetaSoundPageSettings* TargetPageSettings = Settings->FindPageSettings(AuditionPage))
 		{
-			auto PageIsCooked = [&InPageIDs, &PreviewInfo](const FMetaSoundPageSettings& PageSettings)
-			{
-				if (InPageIDs.Contains(PageSettings.UniqueId))
-				{
-					return !PageSettings.GetExcludeFromCook(PreviewInfo.PlatformName);
-				}
-
-				return false;
-			};
-
 			const FGuid& TargetPageID = TargetPageSettings->UniqueId;
 			constexpr bool bReverse = true;
 			bool bFoundMatch = false;
@@ -152,10 +142,13 @@ Metasound::Engine::FPageResolutionEditorResults UMetasoundEditorSettings::Resolv
 					bFoundMatch |= PageSettings.UniqueId == TargetPageID;
 					if (bFoundMatch)
 					{
-						if (AuditionPlatform == EditorAuditionPlatform || PageIsCooked(PageSettings))
+						if (InPageIDs.Contains(PageSettings.UniqueId))
 						{
-							bPageSelected = true;
-							PreviewInfo.PageID = PageSettings.UniqueId;
+							if (AuditionPlatform == EditorAuditionPlatform || !PageSettings.GetExcludeFromCook(PreviewInfo.PlatformName))
+							{
+								bPageSelected = true;
+								PreviewInfo.PageID = PageSettings.UniqueId;
+							}
 						}
 					}
 				}
