- Clean-up more spots using "Audition Target" terminology now replaced with "Audition" in MetaSound asset editor
- Fix crash due to raw FEditor pointer passing to menu options

#rnx

diff --git a/Engine/Plugins/Runtime/Metasound/Source/MetasoundEditor/Private/MetasoundEditor.cpp b/Engine/Plugins/Runtime/Metasound/Source/MetasoundEditor/Private/MetasoundEditor.cpp
--- a/Engine/Plugins/Runtime/Metasound/Source/MetasoundEditor/Private/MetasoundEditor.cpp
+++ b/Engine/Plugins/Runtime/Metasound/Source/MetasoundEditor/Private/MetasoundEditor.cpp
@@ -1997,15 +1997,20 @@ namespace Metasound
 			}
 			MenuBuilder.EndSection();
 
-			MenuBuilder.BeginSection("SetAuditionPageSectionHeader", LOCTEXT("SetAuditionPageDescription", "Audition Target Page"));
+			MenuBuilder.BeginSection("SetAuditionPageSectionHeader", LOCTEXT("SetAuditionPageDescription", "Audition Page"));
 			{
-				const FText FocusPageTooltip = LOCTEXT("EnableFocusTargetPageSwapTooltip", "Attempts to set audition target page to focused graph page. "
-					"If page focus is modified and newly selected graph is not targetable for a selected audition platform (see 'Audition' or `MetaSound Editor Settings`), "
-					"will fail to pass validation for preview audition.");
+				TSharedRef<FEditor> ThisShared = StaticCastSharedRef<FEditor>(AsShared());
+				const FText FocusPageTooltip = LOCTEXT("EnableAuditionFocusPageTooltip", "Sets audition page to focused graph page. "
+					"If page focus is modified and newly selected graph is not targetable, Editor will warn behavior is not reflected at runtime (see Project 'MetaSound Editor' Settings), ");
 				MenuBuilder.AddWidget(
 					SNew(SCheckBox)
-					.OnCheckStateChanged_Lambda([this](ECheckBoxState State)
+					.OnCheckStateChanged_Lambda([EditorPtr = TWeakPtr<FEditor>(ThisShared)](ECheckBoxState State)
 					{
+						TSharedPtr<FEditor> ThisEditor = EditorPtr.Pin();
+						if (!ThisEditor.IsValid())
+						{
+							return;
+						}
 						if (UMetasoundEditorSettings* EdSettings = GetMutableDefault<UMetasoundEditorSettings>())
 						{
 							switch (State)
@@ -2013,8 +2018,8 @@ namespace Metasound
 								case ECheckBoxState::Checked:
 								{
 									EdSettings->AuditionPageMode = EAuditionPageMode::Focused;
-									Stop();
-									SyncAuditionState();
+									ThisEditor->Stop();
+									ThisEditor->SyncAuditionState();
 									break;
 								}
 
@@ -2028,7 +2033,7 @@ namespace Metasound
 							}
 						}
 					})
-					.IsChecked_Lambda([this]()
+					.IsChecked_Lambda([]()
 					{
 						if (const UMetasoundEditorSettings* EdSettings = GetDefault<UMetasoundEditorSettings>())
 						{
@@ -2041,32 +2046,35 @@ namespace Metasound
 						return ECheckBoxState::Unchecked;
 					})
 					.ToolTipText(FocusPageTooltip),
-					LOCTEXT("EnableFocusTargetPageGraphSync", "Sync Audition Target With Focused Graph"),
+					LOCTEXT("EnableAuditionAndFocusGraphPageSync", "Sync Audition With Focused Graph Page"),
 					true,
 					true,
 					FocusPageTooltip
 				);
 
-				auto TryAddPageEntry = [this, &MenuBuilder](const FMetaSoundPageSettings& PageSettings)
+				auto TryAddPageEntry = [EditorPtr = TWeakPtr<FEditor>(ThisShared), &MenuBuilder](const FMetaSoundPageSettings& PageSettings)
 				{
 					const FName AuditionPage = PageSettings.Name;
 					const FText PageText = FText::FromName(PageSettings.Name);
 					FUIAction SetTargetPageAction;
 
-					SetTargetPageAction.ExecuteAction = FExecuteAction::CreateLambda([this, AuditionPage]()
+					SetTargetPageAction.ExecuteAction = FExecuteAction::CreateLambda([EditorPtr, AuditionPage]()
 					{
 						if (UMetasoundEditorSettings* EditorSettings = GetMutableDefault<UMetasoundEditorSettings>())
 						{
 							if (EditorSettings->AuditionPage != AuditionPage)
 							{
 								EditorSettings->AuditionPage = AuditionPage;
-								Stop();
-								SyncAuditionState();
+								if (TSharedPtr<FEditor> ThisEditor = EditorPtr.Pin())
+								{
+									ThisEditor->Stop();
+									ThisEditor->SyncAuditionState();
+								}
 							}
 						}
 					});
 
-					SetTargetPageAction.GetActionCheckState = FGetActionCheckState::CreateLambda([this, AuditionPage]()
+					SetTargetPageAction.GetActionCheckState = FGetActionCheckState::CreateLambda([AuditionPage]()
 					{
 						if (const UMetasoundEditorSettings* EditorSettings = GetDefault<UMetasoundEditorSettings>())
 						{
@@ -2078,7 +2086,7 @@ namespace Metasound
 
 						return ECheckBoxState::Unchecked;
 					});
-					SetTargetPageAction.CanExecuteAction = FCanExecuteAction::CreateLambda([this, AuditionPage]()
+					SetTargetPageAction.CanExecuteAction = FCanExecuteAction::CreateLambda([AuditionPage]()
 					{
 						if (const UMetasoundEditorSettings* EdSettings = GetDefault<UMetasoundEditorSettings>())
 						{
@@ -2092,7 +2100,7 @@ namespace Metasound
 						return false;
 					});
 
-					TAttribute<FText> TooltipAttribute = TAttribute<FText>::CreateLambda([this, AuditionPage, PageText]()
+					TAttribute<FText> TooltipAttribute = TAttribute<FText>::CreateLambda([AuditionPage, PageText]()
 					{
 						if (const UMetasoundEditorSettings* EditorSettings = GetDefault<UMetasoundEditorSettings>())
 						{
