Tooltip clean-up to clarify difference between Editor & Default Audition Platform

#rnx

diff --git a/Engine/Plugins/Runtime/Metasound/Source/MetasoundEditor/Private/MetasoundEditor.cpp b/Engine/Plugins/Runtime/Metasound/Source/MetasoundEditor/Private/MetasoundEditor.cpp
--- a/Engine/Plugins/Runtime/Metasound/Source/MetasoundEditor/Private/MetasoundEditor.cpp
+++ b/Engine/Plugins/Runtime/Metasound/Source/MetasoundEditor/Private/MetasoundEditor.cpp
@@ -1948,7 +1948,7 @@ namespace Metasound
 
 			MenuBuilder.BeginSection("SetAuditionPlatformSectionHeader", LOCTEXT("AuditionPlatformSectionName", "Audition Platform"));
 			{
-				auto CreatePlatformEntry = [this, &MenuBuilder](FName PlatformName, const FText& PlatformText)
+				auto CreatePlatformEntry = [this, &MenuBuilder](FName PlatformName, const FText& PlatformText, const FText& Tooltip)
 				{
 					FUIAction SetPlatformAction;
 					SetPlatformAction.ExecuteAction = FExecuteAction::CreateLambda([this, PlatformName]()
@@ -1973,12 +1973,7 @@ namespace Metasound
 						return ECheckBoxState::Unchecked;
 					});
 
-					MenuBuilder.AddMenuEntry(PlatformText,
-						FText::Format(LOCTEXT("SetAuditionPlatformToolTip", "Sets the page audition platform to '{0}'."), PlatformText),
-						FSlateIcon(),
-						SetPlatformAction,
-						{ },
-						EUserInterfaceActionType::RadioButton);
+					MenuBuilder.AddMenuEntry(PlatformText, Tooltip, FSlateIcon(), SetPlatformAction, { }, EUserInterfaceActionType::RadioButton);
 				};
 
 				TArray<FName> AuditionPlatforms = UMetasoundEditorSettings::GetAuditionPlatformNames();
@@ -1992,7 +1987,21 @@ namespace Metasound
 				for (const FName& PlatformName : AuditionPlatforms)
 				{
 					const FText PlatformText = FText::FromName(PlatformName);
-					CreatePlatformEntry(PlatformName, PlatformText);
+					FText Tooltip;
+					if (PlatformName == UMetasoundEditorSettings::DefaultAuditionPlatform)
+					{
+						Tooltip = LOCTEXT("SetDefaultPlatformToolTip", "Sets the page audition platform to 'Default', which follows target/cook settings for unspecified platforms.");
+					}
+					else if (PlatformName == UMetasoundEditorSettings::EditorAuditionPlatform)
+					{
+						Tooltip = LOCTEXT("SetEditorPlatformToolTip", "Sets the page audition platform to 'Editor', which ignores any explicit target/cook settings.");
+					}
+					else
+					{
+						Tooltip = FText::Format(LOCTEXT("SetAuditionPlatformToolTip", "Sets the page audition platform to '{0}'."), PlatformText);
+					}
+
+					CreatePlatformEntry(PlatformName, PlatformText, Tooltip);
 				}
 			}
 			MenuBuilder.EndSection();
@@ -2728,7 +2737,7 @@ namespace Metasound
 						if (!UMetaSoundEditorSubsystem::GetChecked().IsPageAuditionPlatformCookTarget(EdSettings->AuditionPage))
 						{
 							GraphStatusDescriptionOverride = LOCTEXT("InvalidAuditionPageWarning",
-								"Selected Audition Page in MetaSound Editor Settings is not a target page for the selectd 'Audition Platform'. "
+								"Selected Audition Page in MetaSound Editor Settings is not a target page for the selected 'Audition Platform'. "
 								"Execution may result in behavior that does not exhibit runtime behavior.");
 							if (HighestMessageSeverity > EMessageSeverity::Warning)
 							{
