Minor: Give return value a display name for FindOrBeginBuilding
#rnx

diff --git a/Engine/Plugins/Runtime/Metasound/Source/MetasoundEditor/Public/MetasoundEditorSubsystem.h b/Engine/Plugins/Runtime/Metasound/Source/MetasoundEditor/Public/MetasoundEditorSubsystem.h
--- a/Engine/Plugins/Runtime/Metasound/Source/MetasoundEditor/Public/MetasoundEditorSubsystem.h
+++ b/Engine/Plugins/Runtime/Metasound/Source/MetasoundEditor/Public/MetasoundEditorSubsystem.h
@@ -53,7 +53,7 @@ public:
 	// Returns a builder for the given MetaSound asset. Returns null if provided a transient MetaSound. For finding builders for transient
 	// MetaSounds, use the UMetaSoundBuilderSubsystem's API (FindPatchBuilder, FindSourceBuilder, FindBuilderByName etc.)
 	UFUNCTION(BlueprintCallable, Category = "Audio|MetaSound|Builder|Editor", meta = (DisplayName = "Find Or Begin Building MetaSound Asset", ExpandEnumAsExecs = "OutResult"))
-	UMetaSoundBuilderBase* FindOrBeginBuilding(TScriptInterface<IMetaSoundDocumentInterface> MetaSound, EMetaSoundBuilderResult& OutResult) const;
+	UPARAM(DisplayName = "Builder") UMetaSoundBuilderBase* FindOrBeginBuilding(TScriptInterface<IMetaSoundDocumentInterface> MetaSound, EMetaSoundBuilderResult& OutResult) const;
 
 	// Sets the visual location to InLocation of a given node InNode of a given builder's document.
 	UFUNCTION(BlueprintCallable, Category = "Audio|MetaSound|Builder|Editor", meta = (ExpandEnumAsExecs = "OutResult"))
