Gather text from source, remove LOC_DEFINE_REGION regions which were used to exclude LOCTEXT's from gather.
Resolve duplicate keys(identifiers)

diff --git a/Engine/Plugins/Runtime/Harmonix/Source/HarmonixMetasound/Public/HarmonixMetasound/Common.h b/Engine/Plugins/Runtime/Harmonix/Source/HarmonixMetasound/Public/HarmonixMetasound/Common.h
--- a/Engine/Plugins/Runtime/Harmonix/Source/HarmonixMetasound/Public/HarmonixMetasound/Common.h
+++ b/Engine/Plugins/Runtime/Harmonix/Source/HarmonixMetasound/Public/HarmonixMetasound/Common.h
@@ -30,7 +30,6 @@ namespace HarmonixMetasound
 	HARMONIXMETASOUND_API extern const FText& NAME##DisplayName;
 
 #if WITH_EDITOR
-#define LOC_DEFINE_REGION
 #define DEFINE_INPUT_METASOUND_PARAM(NAME, NAME_TEXT, TOOLTIP_TEXT)						\
 	const TCHAR* NAME##Name = TEXT(NAME_TEXT);											\
 	const FText NAME##Tooltip = LOCTEXT(#NAME "InputTooltip", TOOLTIP_TEXT);			\
@@ -39,7 +38,6 @@ namespace HarmonixMetasound
 	const TCHAR* NAME##Name = TEXT(NAME_TEXT);											\
 	const FText NAME##Tooltip = LOCTEXT(#NAME "OutputTooltip", TOOLTIP_TEXT);			\
     const FText NAME##DisplayName = LOCTEXT(#NAME "OutputDisplayName", NAME_TEXT);
-#undef LOC_DEFINE_REGION
 #else 
 #define DEFINE_INPUT_METASOUND_PARAM(NAME, NAME_TEXT, TOOLTIP_TEXT)		\
 	const TCHAR* NAME##Name = TEXT(NAME_TEXT);							\
diff --git a/Engine/Plugins/Runtime/Metasound/Source/MetasoundFrontend/Public/MetasoundOutputNode.h b/Engine/Plugins/Runtime/Metasound/Source/MetasoundFrontend/Public/MetasoundOutputNode.h
--- a/Engine/Plugins/Runtime/Metasound/Source/MetasoundFrontend/Public/MetasoundOutputNode.h
+++ b/Engine/Plugins/Runtime/Metasound/Source/MetasoundFrontend/Public/MetasoundOutputNode.h
@@ -28,8 +28,8 @@ namespace Metasound
 	public:
 		FMissingOutputNodeInputReferenceError(const INode& InNode, const FText& InDataType)
 			: FBuildErrorBase(
-				"MetasoundMissingOutputDataReferenceError",
-				METASOUND_LOCTEXT_FORMAT("MissingOutputNodeInputReferenceError", "Missing required output node input reference for type {0}.", InDataType))
+				"MetasoundMissingOutputDataReferenceForTypeError",
+				METASOUND_LOCTEXT_FORMAT("MissingOutputNodeInputReferenceForTypeError", "Missing required output node input reference for type {0}.", InDataType))
 		{
 			AddNode(InNode);
 		}
diff --git a/Engine/Plugins/Runtime/Metasound/Source/MetasoundGraphCore/Public/MetasoundParamHelper.h b/Engine/Plugins/Runtime/Metasound/Source/MetasoundGraphCore/Public/MetasoundParamHelper.h
--- a/Engine/Plugins/Runtime/Metasound/Source/MetasoundGraphCore/Public/MetasoundParamHelper.h
+++ b/Engine/Plugins/Runtime/Metasound/Source/MetasoundGraphCore/Public/MetasoundParamHelper.h
@@ -58,12 +58,10 @@
  * Use this macro to define your parameter if you used DECLARE_METASOUND_PARAM, otherwise just use METASOUND_PARAM in a source file
  */
 #if WITH_EDITOR
-#define LOC_DEFINE_REGION
 #define DEFINE_METASOUND_PARAM(NAME, NAME_TEXT, TOOLTIP_TEXT) \
 	const FLazyName NAME##Name = TEXT(NAME_TEXT); \
 	static const FText NAME##Tooltip = LOCTEXT(#NAME "Tooltip", TOOLTIP_TEXT); \
 	static const FText NAME##DisplayName = LOCTEXT(#NAME "DisplayName", NAME_TEXT);
-#undef LOC_DEFINE_REGION
 #else 
 #define DEFINE_METASOUND_PARAM(NAME, NAME_TEXT, TOOLTIP_TEXT) \
 	const FLazyName NAME##Name = TEXT(NAME_TEXT); \
@@ -72,12 +70,10 @@
 #endif // WITH_EDITOR
 
 #if WITH_EDITOR
-#define LOC_DEFINE_REGION
 #define METASOUND_PARAM(NAME, NAME_TEXT, TOOLTIP_TEXT) \
 	static const FLazyName NAME##Name = TEXT(NAME_TEXT); \
 	static const FText NAME##Tooltip = LOCTEXT(#NAME "Tooltip", TOOLTIP_TEXT); \
 	static const FText NAME##DisplayName = LOCTEXT(#NAME "DisplayName", NAME_TEXT); 
-#undef LOC_DEFINE_REGION
 #else 
 #define METASOUND_PARAM(NAME, NAME_TEXT, TOOLTIP_TEXT) \
 	static const FLazyName NAME##Name = TEXT(NAME_TEXT); \
