- Enable MetaSound audition resolution of page IDs to specify audition page ID
  - Allowings in-graph node modifiers (eg. float knob/slider) to display the proper fallback value for the given graph page's ID if a default page value is implemented in a higher page in the fallback array.
- Add tooltip describing fallback value
- Misc tooltip clean-up
#rnx

diff --git a/Engine/Plugins/Runtime/Metasound/Source/MetasoundEditor/Private/MetasoundEditor.cpp b/Engine/Plugins/Runtime/Metasound/Source/MetasoundEditor/Private/MetasoundEditor.cpp
--- a/Engine/Plugins/Runtime/Metasound/Source/MetasoundEditor/Private/MetasoundEditor.cpp
+++ b/Engine/Plugins/Runtime/Metasound/Source/MetasoundEditor/Private/MetasoundEditor.cpp
@@ -2009,8 +2009,11 @@ namespace Metasound
 			MenuBuilder.BeginSection("SetAuditionPageSectionHeader", LOCTEXT("SetAuditionPageDescription", "Audition Page"));
 			{
 				TSharedRef<FEditor> ThisShared = StaticCastSharedRef<FEditor>(AsShared());
-				const FText FocusPageTooltip = LOCTEXT("EnableAuditionFocusPageTooltip", "Sets audition page to focused graph page. "
-					"If page focus is modified and newly selected graph is not targetable, Editor will warn behavior is not reflected at runtime (see Project 'MetaSound Editor' Settings), ");
+				const FText FocusPageTooltip = LOCTEXT("EnableAuditionFocusPageTooltip",
+					"Synchronizes audition page to currently focused graph page.\r\n"
+					"If focused graph page is non-targetable for the selected audition\r\n"
+					"platform, will issue warning behavior is not reflected at runtime\r\n"
+					"(see 'MetaSound Editor' user settings).");
 				MenuBuilder.AddWidget(
 					SNew(SCheckBox)
 					.OnCheckStateChanged_Lambda([EditorPtr = TWeakPtr<FEditor>(ThisShared)](ECheckBoxState State)
@@ -2055,7 +2058,7 @@ namespace Metasound
 						return ECheckBoxState::Unchecked;
 					})
 					.ToolTipText(FocusPageTooltip),
-					LOCTEXT("EnableAuditionAndFocusGraphPageSync", "Sync Audition With Focused Graph Page"),
+					LOCTEXT("EnableAuditionAndFocusGraphPageSync", "Sync With Graph Page"),
 					true,
 					true,
 					FocusPageTooltip
diff --git a/Engine/Plugins/Runtime/Metasound/Source/MetasoundEditor/Private/MetasoundEditorGraphInputNode.cpp b/Engine/Plugins/Runtime/Metasound/Source/MetasoundEditor/Private/MetasoundEditorGraphInputNode.cpp
--- a/Engine/Plugins/Runtime/Metasound/Source/MetasoundEditor/Private/MetasoundEditorGraphInputNode.cpp
+++ b/Engine/Plugins/Runtime/Metasound/Source/MetasoundEditor/Private/MetasoundEditorGraphInputNode.cpp
@@ -19,6 +19,7 @@
 #include "MetasoundFrontendRegistries.h"
 #include "MetasoundFrontendSearchEngine.h"
 #include "MetasoundPrimitives.h"
+#include "MetasoundSettings.h"
 #include "Misc/Guid.h"
 #include "Sound/SoundWave.h"
 #include "UObject/ObjectMacros.h"
@@ -138,18 +139,19 @@ void UMetasoundEditorGraphInputNode::GetPinHoverText(const UEdGraphPin& Pin, FSt
 	{
 		if (ensure(Input))
 		{
-			FConstNodeHandle InputNode = Input->GetConstNodeHandle();
-			OutHoverText = InputNode->GetDescription().ToString();
+			OutHoverText = Input->GetDescription().ToString();
 			if (ShowNodeDebugData())
 			{
-				FConstOutputHandle OutputHandle = FGraphBuilder::FindReroutedConstOutputHandleFromPin(&Pin);
-				OutHoverText = FString::Format(TEXT("{0}\nVertex Name: {1}\nDataType: {2}\nID: {3}"),
+				if (const FMetasoundFrontendClassVertex* Vertex = Input->GetFrontendClassVertex())
 				{
-					OutHoverText,
-					OutputHandle->GetName().ToString(),
-					OutputHandle->GetDataType().ToString(),
-					OutputHandle->GetID().ToString(),
-				});
+					OutHoverText = FString::Format(TEXT("{0}\nVertex Name: {1}\nDataType: {2}\nID: {3}"),
+					{
+						OutHoverText,
+						Vertex->Name.ToString(),
+						Vertex->TypeName.ToString(),
+						Vertex->NodeID.ToString(),
+					});
+				}
 			}
 		}
 	}
@@ -198,13 +200,61 @@ void UMetasoundEditorGraphInputNode::Validate(Metasound::Editor::FGraphNodeValid
 FText UMetasoundEditorGraphInputNode::GetTooltipText() const
 {
 	//If Constructor input
-	if (Input && Input->GetVertexAccessType() == EMetasoundFrontendVertexAccessType::Value)
+	if (Input)
 	{
-		UMetasoundEditorGraph* Graph = CastChecked<UMetasoundEditorGraph>(GetGraph());
-		if (Graph->IsPreviewing())
+		if (Input->GetVertexAccessType() == EMetasoundFrontendVertexAccessType::Value)
+		{
+			UMetasoundEditorGraph* Graph = CastChecked<UMetasoundEditorGraph>(GetGraph());
+			if (Graph->IsPreviewing())
+			{
+				FText ToolTip = LOCTEXT("Metasound_ConstructorInputNodeDescription", "Editing constructor values is disabled while previewing.");
+				return ToolTip;
+			}
+		}
+
+		if (UMetasoundEditorGraphMemberDefaultLiteral* Literal = Input->GetLiteral())
+		{
+			bool bPageDefaultImplemented = false;
+			const FMetaSoundFrontendDocumentBuilder& Builder = Input->GetFrontendBuilderChecked();
+			const FGuid& BuildPageID = Builder.GetBuildPageID();
+			Literal->IterateDefaults([&bPageDefaultImplemented, &BuildPageID](const FGuid& InDefaultPageID, FMetasoundFrontendLiteral)
+			{
+				bPageDefaultImplemented |= InDefaultPageID == BuildPageID;
+			});
+
+			if (!bPageDefaultImplemented)
+			{
+				const UMetaSoundSettings* MetaSoundSettings = GetDefault<UMetaSoundSettings>();
+				check(MetaSoundSettings);
+
+				const UMetasoundEditorSettings* EditorSettings = GetDefault<UMetasoundEditorSettings>();
+				check(EditorSettings);
+
+				if (const FMetaSoundPageSettings* PageSettings = MetaSoundSettings->FindPageSettings(BuildPageID))
+				{
+					if (const FMetasoundFrontendClassInput* ClassInput = Builder.FindGraphInput(Input->GetMemberName()))
+					{
+						const FGuid FallbackPageID = EditorSettings->ResolveAuditionPage(*ClassInput, BuildPageID);
+						if (const FMetaSoundPageSettings* FallbackSettings = MetaSoundSettings->FindPageSettings(FallbackPageID))
+						{
+							return FText::Format(
+								LOCTEXT("DefaultPageValueDisabledNotImplemented",
+									"No '{0}' page default value implemented.\r\n"
+									"Showing platform/platform group '{1}' fallback '{2}'.\r\n"
+									"(See 'Audition' menu or 'MetaSound' Editor Preferences to change 'Audition Platform')."),
+								FText::FromName(PageSettings->Name),
+								FText::FromName(EditorSettings->AuditionPlatform),
+								FText::FromName(FallbackSettings->Name));
+						}
+					}
+				}
+			}
+		}
+
+		const FText InputDescription = Input->GetDescription();
+		if (!InputDescription.IsEmpty())
 		{
-			FText ToolTip = LOCTEXT("Metasound_ConstructorInputNodeDescription", "Editing constructor values is disabled while previewing.");
-			return ToolTip;
+			return InputDescription;
 		}
 	}
 
diff --git a/Engine/Plugins/Runtime/Metasound/Source/MetasoundEditor/Private/MetasoundEditorSettings.cpp b/Engine/Plugins/Runtime/Metasound/Source/MetasoundEditor/Private/MetasoundEditorSettings.cpp
--- a/Engine/Plugins/Runtime/Metasound/Source/MetasoundEditor/Private/MetasoundEditorSettings.cpp
+++ b/Engine/Plugins/Runtime/Metasound/Source/MetasoundEditor/Private/MetasoundEditorSettings.cpp
@@ -129,35 +129,54 @@ Metasound::Engine::FPageResolutionEditorResults UMetasoundEditorSettings::Resolv
 
 	if (const UMetaSoundSettings* Settings = GetDefault<UMetaSoundSettings>())
 	{
-		if (const FMetaSoundPageSettings* TargetPageSettings = Settings->FindPageSettings(AuditionPage))
+		if (const FMetaSoundPageSettings* AuditionPageSettings = Settings->FindPageSettings(AuditionPage))
 		{
-			const FGuid& TargetPageID = TargetPageSettings->UniqueId;
-			constexpr bool bReverse = true;
-			bool bFoundMatch = false;
-			bool bPageSelected = false;
-			Settings->IteratePageSettings([&](const FMetaSoundPageSettings& PageSettings)
+			const FGuid& AuditionPageID = AuditionPageSettings->UniqueId;
+			PreviewInfo.PageID = ResolveAuditionPage(InPageIDs, AuditionPageID);
+		}
+	}
+	return PreviewInfo;
+
+}
+
+FGuid UMetasoundEditorSettings::ResolveAuditionPage(const FMetasoundFrontendClassInput& InClassInput, const FGuid& InAuditionPageID) const
+{
+	TArray<FGuid> PageIDs;
+	InClassInput.IterateDefaults([&PageIDs](const FGuid& PageID, const FMetasoundFrontendLiteral&)
+	{
+		PageIDs.Add(PageID);
+	});
+	return ResolveAuditionPage(PageIDs, InAuditionPageID);
+}
+
+FGuid UMetasoundEditorSettings::ResolveAuditionPage(const TArray<FGuid>& InPageIDs, const FGuid& InAuditionPageID) const
+{
+	FGuid ResolvedPageID = Metasound::Frontend::DefaultPageID;
+	if (const UMetaSoundSettings* Settings = GetDefault<UMetaSoundSettings>())
+	{
+		constexpr bool bReverse = true;
+		bool bFoundMatch = false;
+		bool bPageSelected = false;
+		Settings->IteratePageSettings([&](const FMetaSoundPageSettings& PageSettings)
 			{
 				if (!bPageSelected)
 				{
-					bFoundMatch |= PageSettings.UniqueId == TargetPageID;
+					bFoundMatch |= PageSettings.UniqueId == InAuditionPageID;
 					if (bFoundMatch)
 					{
 						if (InPageIDs.Contains(PageSettings.UniqueId))
 						{
-							if (AuditionPlatform == EditorAuditionPlatform || !PageSettings.GetExcludeFromCook(PreviewInfo.PlatformName))
+							if (AuditionPlatform == EditorAuditionPlatform || !PageSettings.GetExcludeFromCook(AuditionPlatform))
 							{
 								bPageSelected = true;
-								PreviewInfo.PageID = PageSettings.UniqueId;
+								ResolvedPageID = PageSettings.UniqueId;
 							}
 						}
 					}
 				}
 			}, bReverse);
-		}
 	}
-
-	return PreviewInfo;
-
+	return ResolvedPageID;
 }
 
 TArray<FName> UMetasoundEditorSettings::GetAuditionPlatformNames()
diff --git a/Engine/Plugins/Runtime/Metasound/Source/MetasoundEditor/Private/MetasoundEditorSettings.h b/Engine/Plugins/Runtime/Metasound/Source/MetasoundEditor/Private/MetasoundEditorSettings.h
--- a/Engine/Plugins/Runtime/Metasound/Source/MetasoundEditor/Private/MetasoundEditorSettings.h
+++ b/Engine/Plugins/Runtime/Metasound/Source/MetasoundEditor/Private/MetasoundEditorSettings.h
@@ -309,6 +309,12 @@ public:
 	const FAudioMaterialMeterStyle* GetMeterStyle() const;
 
 	Metasound::Engine::FPageResolutionEditorResults ResolveAuditionPage(const TArray<FGuid>& InPageIDs) const;
+	
+	/** Given the provided AuditionPageID, returns the resolved PageID from the provided array of values based on fallback logic. */
+	FGuid ResolveAuditionPage(const TArray<FGuid>& InPageIDs, const FGuid& InAuditionPageID) const;
+
+	/** Given the provided class input and AuditionPageID, returns the resolved PageID from the provided array of values based on fallback logic. */
+	FGuid ResolveAuditionPage(const FMetasoundFrontendClassInput& InClassInput, const FGuid& InAuditionPageID) const;
 
 	UFUNCTION()
 	static TArray<FName> GetAuditionPageNames();
diff --git a/Engine/Plugins/Runtime/Metasound/Source/MetasoundEditor/Private/SMetasoundGraphNode.cpp b/Engine/Plugins/Runtime/Metasound/Source/MetasoundEditor/Private/SMetasoundGraphNode.cpp
--- a/Engine/Plugins/Runtime/Metasound/Source/MetasoundEditor/Private/SMetasoundGraphNode.cpp
+++ b/Engine/Plugins/Runtime/Metasound/Source/MetasoundEditor/Private/SMetasoundGraphNode.cpp
@@ -31,6 +31,7 @@
 #include "MetasoundEditorSubsystem.h"
 #include "MetasoundFrontendNodeTemplateRegistry.h"
 #include "MetasoundFrontendRegistries.h"
+#include "MetasoundSettings.h"
 #include "MetasoundTrace.h"
 #include "MetasoundTrigger.h"
 #include "NodeFactory.h"
@@ -758,6 +759,15 @@ namespace Metasound
 			});
 		}
 
+		FText SMetaSoundGraphNode::GetInputWidgetTooltip() const
+		{
+			if (const UMetasoundEditorGraphMemberNode* Node = GetMetaSoundMemberNode())
+			{
+				return Node->GetTooltipText();
+			}
+			return FText();
+		}
+
 		UMetasoundEditorGraphMemberNode* SMetaSoundGraphNode::GetMetaSoundMemberNode() const
 		{
 			return Cast<UMetasoundEditorGraphMemberNode>(&GetMetaSoundNode());
@@ -1051,14 +1061,14 @@ namespace Metasound
 				FGuid ResolvedPageID = Frontend::DefaultPageID;
 				if (ensure(ClassInput))
 				{
-					ResolvedPageID = Engine::FDocumentBuilderRegistry::GetChecked().ResolveTargetPageID(*ClassInput);
+					ResolvedPageID = EditorSettings->ResolveAuditionPage(*ClassInput, InBuilder.GetBuildPageID());
 				}
 
 				FloatInputWidget->SetOutputRange(DefaultFloat->GetRange());
 				FloatInputWidget->SetUnitsTextReadOnly(true);
 				FloatInputWidget->SetSliderValue(FloatInputWidget->GetSliderValue(DefaultFloat->GetDefaultAs<float>(ResolvedPageID)));
 				FloatInputWidget->SetEnabled(GetInputWidgetEnabled());
-
+				FloatInputWidget->SetToolTipText(GetInputWidgetTooltip());
 				// Setup & clear delegate if necessary (ex. if was just saved)
 				if (InputSliderOnValueChangedDelegateHandle.IsValid())
 				{
diff --git a/Engine/Plugins/Runtime/Metasound/Source/MetasoundEditor/Private/SMetasoundGraphNode.h b/Engine/Plugins/Runtime/Metasound/Source/MetasoundEditor/Private/SMetasoundGraphNode.h
--- a/Engine/Plugins/Runtime/Metasound/Source/MetasoundEditor/Private/SMetasoundGraphNode.h
+++ b/Engine/Plugins/Runtime/Metasound/Source/MetasoundEditor/Private/SMetasoundGraphNode.h
@@ -74,6 +74,9 @@ namespace Metasound::Editor
 		// Returns attribute that returns whether or not the input widget should be enabled.
 		TAttribute<bool> GetInputWidgetEnabled() const;
 
+		// Returns the input widget tooltip text.
+		FText GetInputWidgetTooltip() const;
+
 		// If this node represents a graph member node, returns corresponding member.
 		UMetasoundEditorGraphMember* GetMetaSoundMember();
 
