Fix issues with MetaSound member metadata (ex. float input range) (changes not being serialized + values not inherited on preset creation)

diff --git a/Engine/Plugins/Runtime/Metasound/Source/MetasoundEditor/Private/MetasoundEditorSubsystem.cpp b/Engine/Plugins/Runtime/Metasound/Source/MetasoundEditor/Private/MetasoundEditorSubsystem.cpp
--- a/Engine/Plugins/Runtime/Metasound/Source/MetasoundEditor/Private/MetasoundEditorSubsystem.cpp
+++ b/Engine/Plugins/Runtime/Metasound/Source/MetasoundEditor/Private/MetasoundEditorSubsystem.cpp
@@ -9,6 +9,8 @@
 #include "MetasoundEditorGraph.h"
 #include "MetasoundEditorGraphBuilder.h"
 #include "MetasoundEditorGraphSchema.h"
+#include "MetasoundEditorGraphMemberDefaults.h"
+#include "MetasoundEditorModule.h"
 #include "MetasoundEditorSettings.h"
 #include "MetasoundFactory.h"
 #include "MetasoundSettings.h"
@@ -116,17 +118,52 @@ TScriptInterface<IMetaSoundDocumentInterface> UMetaSoundEditorSubsystem::BuildTo
 	return nullptr;
 }
 
+UMetasoundEditorGraphMemberDefaultLiteral* UMetaSoundEditorSubsystem::CreateMemberMetadata(
+	FMetaSoundFrontendDocumentBuilder& Builder,
+	FName InMemberName,
+	TSubclassOf<UMetasoundEditorGraphMemberDefaultLiteral> LiteralClass) const
+{
+	using namespace Metasound::Engine;
+
+	// If preset and input inherits from default, copy member metadata from referenced graph
+	if (Builder.IsPreset())
+	{
+		const FMetasoundFrontendGraphClass& RootGraph = Builder.GetConstDocumentChecked().RootGraph;
+		const TSet<FName>& InputsInheritingDefault = RootGraph.PresetOptions.InputsInheritingDefault;
+		if (InputsInheritingDefault.Contains(InMemberName))
+		{
+			// Get referenced asset to inherit metadata from
+			FMetasoundAssetBase* ReferencedPresetAsset = Builder.GetReferencedPresetAsset();
+			check(ReferencedPresetAsset);
+
+			FMetaSoundFrontendDocumentBuilder& ReferencedBuilder = FDocumentBuilderRegistry::GetChecked().FindOrBeginBuilding(ReferencedPresetAsset->GetOwningAsset());
+			const FMetasoundFrontendClassInput* ClassInput = ReferencedBuilder.FindGraphInput(InMemberName);
+			if (ClassInput)
+			{
+				if (UMetaSoundFrontendMemberMetadata* ReferencedMemberMetadata = ReferencedBuilder.FindMemberMetadata(ClassInput->NodeID))
+				{
+					return NewObject<UMetasoundEditorGraphMemberDefaultLiteral>(&Builder.CastDocumentObjectChecked<UObject>(), LiteralClass, FName(), RF_Transactional, ReferencedMemberMetadata);
+				}
+			}
+		}
+	}
+	// Otherwise, create brand new member metadata
+	return NewObject<UMetasoundEditorGraphMemberDefaultLiteral>(&Builder.CastDocumentObjectChecked<UObject>(), LiteralClass, FName(), RF_Transactional, nullptr);
+}
+
 bool UMetaSoundEditorSubsystem::BindMemberMetadata(
 	FMetaSoundFrontendDocumentBuilder& Builder,
 	UMetasoundEditorGraphMember& InMember,
 	TSubclassOf<UMetasoundEditorGraphMemberDefaultLiteral> LiteralClass,
 	UMetasoundEditorGraphMemberDefaultLiteral* TemplateObject)
 {
+	UMetasoundEditorGraphMemberDefaultLiteral* NewLiteral = nullptr;
 	const FGuid& MemberID = InMember.GetMemberID();
 
 	if (TemplateObject)
 	{
 		Builder.ClearMemberMetadata(MemberID);
+		NewLiteral = NewObject<UMetasoundEditorGraphMemberDefaultLiteral>(&Builder.CastDocumentObjectChecked<UObject>(), LiteralClass, FName(), RF_Transactional, TemplateObject);
 	}
 	else
 	{
@@ -137,7 +174,9 @@ bool UMetaSoundEditorSubsystem::BindMemberMetadata(
 		}
 	}
 
-	if (UMetasoundEditorGraphMemberDefaultLiteral* NewLiteral = NewObject<UMetasoundEditorGraphMemberDefaultLiteral>(&Builder.CastDocumentObjectChecked<UObject>(), LiteralClass, FName(), RF_Transactional, TemplateObject))
+	NewLiteral = CreateMemberMetadata(Builder, InMember.GetMemberName(), LiteralClass);
+
+	if (NewLiteral)
 	{
 		NewLiteral->MemberID = MemberID;
 
@@ -164,6 +203,73 @@ UMetaSoundBuilderBase* UMetaSoundEditorSubsystem::FindOrBeginBuilding(TScriptInt
 	return nullptr;
 }
 
+UMetaSoundFrontendMemberMetadata* UMetaSoundEditorSubsystem::FindOrCreateGraphInputMetadata(UMetaSoundBuilderBase* InBuilder, FName InputName, EMetaSoundBuilderResult& OutResult)
+{
+	using namespace Metasound::Editor;
+	using namespace Metasound::Frontend; 
+	
+	if (InBuilder)
+	{
+		FMetaSoundNodeHandle GraphInputNodeHandle = InBuilder->FindGraphInputNode(InputName, OutResult);
+		if (OutResult == EMetaSoundBuilderResult::Succeeded)
+		{
+			FMetaSoundFrontendDocumentBuilder& DocBuilder = InBuilder->GetBuilder();
+			// Look for existing metadata
+			if (UMetaSoundFrontendMemberMetadata* MemberMetadata = DocBuilder.FindMemberMetadata(GraphInputNodeHandle.NodeID))
+			{
+				OutResult = EMetaSoundBuilderResult::Succeeded;
+				return MemberMetadata;
+			}
+			// Create new metadata
+			else
+			{
+				// Get literal class
+				FDataTypeRegistryInfo DataTypeInfo;
+				IMetasoundEditorModule& EditorModule = FModuleManager::GetModuleChecked<IMetasoundEditorModule>("MetaSoundEditor");
+				const FName TypeName = DocBuilder.FindGraphInput(InputName)->TypeName;
+				IDataTypeRegistry::Get().GetDataTypeInfo(TypeName, DataTypeInfo);
+				const EMetasoundFrontendLiteralType LiteralType = static_cast<EMetasoundFrontendLiteralType>(DataTypeInfo.PreferredLiteralType);
+
+				TSubclassOf<UMetasoundEditorGraphMemberDefaultLiteral> LiteralClass = EditorModule.FindDefaultLiteralClass(LiteralType);
+				if (!LiteralClass)
+				{
+					LiteralClass = UMetasoundEditorGraphMemberDefaultLiteral::StaticClass();
+				}
+
+
+				if (InBuilder->IsPreset())
+				{
+					// Ensure this MetaSound's dependencies are registered to lookup inherited metadata
+					// Needed in the case where this function is called from a BP,
+					// metadata wasn't created previously, and the MetaSound editor was never opened this session
+					UObject* MetaSound = InBuilder->GetConstBuilder().GetMetasoundAsset().GetOwningAsset();
+					check(MetaSound);
+					RegisterGraphWithFrontend(*MetaSound);
+				}
+
+				// Create new literal and setup
+				UMetasoundEditorGraphMemberDefaultLiteral* NewLiteral = CreateMemberMetadata(DocBuilder, InputName, LiteralClass);
+				if (NewLiteral)
+				{
+					NewLiteral->MemberID = GraphInputNodeHandle.NodeID;
+					NewLiteral->Initialize();
+					DocBuilder.SetMemberMetadata(*NewLiteral);
+					
+					OutResult = EMetaSoundBuilderResult::Succeeded;
+					return NewLiteral;
+				}
+			}
+		}
+		else
+		{
+			UE_LOG(LogMetaSound, Display, TEXT("Failed to find graph input node for input '%s' with builder '%s'."), *InputName.ToString(), *InBuilder->GetName());
+		}
+	}
+
+	OutResult = EMetaSoundBuilderResult::Failed;
+	return nullptr;
+}
+
 UMetaSoundEditorSubsystem& UMetaSoundEditorSubsystem::GetChecked()
 {
 	checkf(GEditor, TEXT("Cannot access UMetaSoundEditorSubsystem without editor loaded"));
diff --git a/Engine/Plugins/Runtime/Metasound/Source/MetasoundEditor/Public/MetasoundEditorSubsystem.h b/Engine/Plugins/Runtime/Metasound/Source/MetasoundEditor/Public/MetasoundEditorSubsystem.h
--- a/Engine/Plugins/Runtime/Metasound/Source/MetasoundEditor/Public/MetasoundEditorSubsystem.h
+++ b/Engine/Plugins/Runtime/Metasound/Source/MetasoundEditor/Public/MetasoundEditorSubsystem.h
@@ -4,6 +4,7 @@
 #include "AssetRegistry/AssetData.h"
 #include "EditorSubsystem.h"
 #include "Framework/MultiBox/MultiBoxExtender.h"
+#include "MetasoundBuilderBase.h"
 #include "MetasoundDocumentInterface.h"
 #include "UObject/ScriptInterface.h"
 
@@ -50,10 +51,20 @@ public:
 		EMetaSoundBuilderResult& OutResult,
 		UPARAM(DisplayName = "Template SoundWave") const USoundWave* TemplateSoundWave = nullptr);
 
+	// Creates new member metadata for a member of a given builder, copying data from the referenced asset in the case of preset inherited inputs
+	UMetasoundEditorGraphMemberDefaultLiteral* CreateMemberMetadata(
+		FMetaSoundFrontendDocumentBuilder& Builder,
+		FName InMemberName,
+		TSubclassOf<UMetasoundEditorGraphMemberDefaultLiteral> LiteralClass) const;
+
 	// Returns a builder for the given MetaSound asset. Returns null if provided a transient MetaSound. For finding builders for transient
 	// MetaSounds, use the UMetaSoundBuilderSubsystem's API (FindPatchBuilder, FindSourceBuilder, FindBuilderByName etc.)
 	UFUNCTION(BlueprintCallable, Category = "Audio|MetaSound|Builder|Editor", meta = (DisplayName = "Find Or Begin Building MetaSound Asset", ExpandEnumAsExecs = "OutResult"))
 	UPARAM(DisplayName = "Builder") UMetaSoundBuilderBase* FindOrBeginBuilding(TScriptInterface<IMetaSoundDocumentInterface> MetaSound, EMetaSoundBuilderResult& OutResult) const;
+	
+	// Find graph input metadata (which includes editor only range information for floats) for a given input. If the metadata does not exist, create it. 
+	UFUNCTION(BlueprintCallable, Category = "Audio|MetaSound|Builder", meta = (ExpandEnumAsExecs = "OutResult"))
+	UPARAM(DisplayName = "Input Metadata") UMetaSoundFrontendMemberMetadata* FindOrCreateGraphInputMetadata(UPARAM(DisplayName = "Builder") UMetaSoundBuilderBase* InBuilder, FName InputName, EMetaSoundBuilderResult& OutResult);
 
 	// Sets the visual location to InLocation of a given node InNode of a given builder's document.
 	UFUNCTION(BlueprintCallable, Category = "Audio|MetaSound|Builder|Editor", meta = (ExpandEnumAsExecs = "OutResult"))
diff --git a/Engine/Plugins/Runtime/Metasound/Source/MetasoundEngine/Public/MetasoundBuilderBase.cpp b/Engine/Plugins/Runtime/Metasound/Source/MetasoundEngine/Public/MetasoundBuilderBase.cpp
--- a/Engine/Plugins/Runtime/Metasound/Source/MetasoundEngine/Public/MetasoundBuilderBase.cpp
+++ b/Engine/Plugins/Runtime/Metasound/Source/MetasoundEngine/Public/MetasoundBuilderBase.cpp
@@ -764,26 +764,10 @@ UObject* UMetaSoundBuilderBase::GetReferencedPresetAsset() const
 		return nullptr;
 	}
 
-	// Find the single external node which is the referenced preset asset, 
-	// and find the asset with its registry key 
-	auto FindExternalNode = [this](const FMetasoundFrontendNode& Node)
+	if (FMetasoundAssetBase* Asset = Builder.GetReferencedPresetAsset())
 	{
-		const FMetasoundFrontendClass* Class = Builder.FindDependency(Node.ClassID);
-		check(Class);
-		return Class->Metadata.GetType() == EMetasoundFrontendClassType::External;
-	};
-	const FMetasoundFrontendNode* Node = Builder.FindConstBuildGraphChecked().Nodes.FindByPredicate(FindExternalNode);
-	if (Node != nullptr)
-	{
-		const FMetasoundFrontendClass* NodeClass = Builder.FindDependency(Node->ClassID);
-		check(NodeClass);
-		const FNodeRegistryKey NodeClassRegistryKey = FNodeRegistryKey(NodeClass->Metadata);
-		if (FMetasoundAssetBase* Asset = IMetaSoundAssetManager::GetChecked().FindAsset(NodeClassRegistryKey))
-		{
-			return Asset->GetOwningAsset();
-		}
+		return Asset->GetOwningAsset();
 	}
-
 	return nullptr;
 }
 
diff --git a/Engine/Plugins/Runtime/Metasound/Source/MetasoundFrontend/Private/MetasoundFrontendDocumentBuilder.cpp b/Engine/Plugins/Runtime/Metasound/Source/MetasoundFrontend/Private/MetasoundFrontendDocumentBuilder.cpp
--- a/Engine/Plugins/Runtime/Metasound/Source/MetasoundFrontend/Private/MetasoundFrontendDocumentBuilder.cpp
+++ b/Engine/Plugins/Runtime/Metasound/Source/MetasoundFrontend/Private/MetasoundFrontendDocumentBuilder.cpp
@@ -2216,6 +2216,33 @@ FMetasoundAssetBase& FMetaSoundFrontendDocumentBuilder::GetMetasoundAsset() cons
 	return *Asset;
 }
 
+FMetasoundAssetBase* FMetaSoundFrontendDocumentBuilder::GetReferencedPresetAsset() const
+{
+	using namespace Metasound::Frontend;
+	if (!IsPreset())
+	{
+		return nullptr;
+	}
+
+	// Find the single external node which is the referenced preset asset, 
+	// and find the asset with its registry key 
+	auto FindExternalNode = [this](const FMetasoundFrontendNode& Node)
+	{
+		const FMetasoundFrontendClass* Class = FindDependency(Node.ClassID);
+		check(Class);
+		return Class->Metadata.GetType() == EMetasoundFrontendClassType::External;
+	};
+	const FMetasoundFrontendNode* Node = FindConstBuildGraphChecked().Nodes.FindByPredicate(FindExternalNode);
+	if (Node != nullptr)
+	{
+		const FMetasoundFrontendClass* NodeClass = FindDependency(Node->ClassID);
+		check(NodeClass);
+		const FNodeRegistryKey NodeClassRegistryKey = FNodeRegistryKey(NodeClass->Metadata);
+		return IMetaSoundAssetManager::GetChecked().FindAsset(NodeClassRegistryKey);
+	}
+	return nullptr;
+}
+
 const FGuid& FMetaSoundFrontendDocumentBuilder::GetBuildPageID() const
 {
 	return BuildPageID;
diff --git a/Engine/Plugins/Runtime/Metasound/Source/MetasoundFrontend/Private/MetasoundFrontendTransform.cpp b/Engine/Plugins/Runtime/Metasound/Source/MetasoundFrontend/Private/MetasoundFrontendTransform.cpp
--- a/Engine/Plugins/Runtime/Metasound/Source/MetasoundFrontend/Private/MetasoundFrontendTransform.cpp
+++ b/Engine/Plugins/Runtime/Metasound/Source/MetasoundFrontend/Private/MetasoundFrontendTransform.cpp
@@ -675,6 +675,11 @@ namespace Metasound
 			TArray<FMetasoundFrontendClassInput> ClassInputs = GenerateRequiredClassInputs(InDocument, PresetGraphHandle, InputsInheritingDefault);
 			TArray<FMetasoundFrontendClassOutput> ClassOutputs = GenerateRequiredClassOutputs(InDocument, PresetGraphHandle);
 
+#if WITH_EDITORONLY_DATA
+			// Cache off member metadata so it be can be readded if necessary after the graph is cleared 
+			const FMemberIDToMetadataMap CachedMemberMetadata = InDocument->GetMetadata()->MemberMetadata;
+#endif // WITH_EDITORONLY_DATA
+
 			FGuid PresetNodeID;
 			PresetGraphHandle->IterateConstNodes([InPresetNodeID = &PresetNodeID](FConstNodeHandle PresetNodeHandle)
 			{
@@ -723,6 +728,11 @@ namespace Metasound
 			AddAndConnectInputs(ClassInputs, PresetGraphHandle, ReferencedNodeHandle);
 			AddAndConnectOutputs(ClassOutputs, PresetGraphHandle, ReferencedNodeHandle);
 
+#if WITH_EDITORONLY_DATA
+			FMemberIDToMetadataMap& MemberMetadata = InDocument->GetMetadata()->MemberMetadata;
+			AddMemberMetadata(CachedMemberMetadata, PresetGraphHandle, MemberMetadata);
+#endif // WITH_EDITORONLY_DATA
+
 			return true;
 		}
 	
@@ -738,6 +748,24 @@ namespace Metasound
 			return Transform(IDocumentController::CreateDocumentHandle(InDocument));
 		}
 
+#if WITH_EDITORONLY_DATA
+		void FRebuildPresetRootGraph::AddMemberMetadata(const FMemberIDToMetadataMap& InCachedMemberMetadata, FGraphHandle& InPresetGraphHandle, FMemberIDToMetadataMap& InOutMemberMetadata) const
+		{
+			// Add member metadata if a member with the corresponding node ID exists in the preset graph
+			if (!InCachedMemberMetadata.IsEmpty())
+			{
+				for (const TPair<FGuid, TObjectPtr<UMetaSoundFrontendMemberMetadata>>& MemberMetadataPair : InCachedMemberMetadata)
+				{
+					FConstNodeHandle FoundNodeHandle = InPresetGraphHandle->GetNodeWithID(MemberMetadataPair.Key);
+					if (FoundNodeHandle->IsValid())
+					{
+						InOutMemberMetadata.Add(MemberMetadataPair.Key, MemberMetadataPair.Value);
+					}
+				}
+			}
+		}
+#endif // WITH_EDITORONLY_DATA
+
 		void FRebuildPresetRootGraph::AddAndConnectInputs(const TArray<FMetasoundFrontendClassInput>& InClassInputs, FGraphHandle& InPresetGraphHandle, FNodeHandle& InReferencedNode) const
 		{
 			// Add inputs and space appropriately
diff --git a/Engine/Plugins/Runtime/Metasound/Source/MetasoundFrontend/Public/MetasoundFrontendDocumentBuilder.h b/Engine/Plugins/Runtime/Metasound/Source/MetasoundFrontend/Public/MetasoundFrontendDocumentBuilder.h
--- a/Engine/Plugins/Runtime/Metasound/Source/MetasoundFrontend/Public/MetasoundFrontendDocumentBuilder.h
+++ b/Engine/Plugins/Runtime/Metasound/Source/MetasoundFrontend/Public/MetasoundFrontendDocumentBuilder.h
@@ -293,6 +293,9 @@ public:
 	UE_DEPRECATED(5.5, "Use GetConstDocumentInterfaceChecked instead")
 	const IMetaSoundDocumentInterface& GetDocumentInterface() const;
 	FMetasoundAssetBase& GetMetasoundAsset() const;
+	
+	// Get the asset referenced by this builder's preset asset, nullptr if builder is not a preset.
+	FMetasoundAssetBase* GetReferencedPresetAsset() const;
 
 	int32 GetTransactionCount() const;
 
diff --git a/Engine/Plugins/Runtime/Metasound/Source/MetasoundFrontend/Public/MetasoundFrontendTransform.h b/Engine/Plugins/Runtime/Metasound/Source/MetasoundFrontend/Public/MetasoundFrontendTransform.h
--- a/Engine/Plugins/Runtime/Metasound/Source/MetasoundFrontend/Public/MetasoundFrontendTransform.h
+++ b/Engine/Plugins/Runtime/Metasound/Source/MetasoundFrontend/Public/MetasoundFrontendTransform.h
@@ -177,6 +177,11 @@ namespace Metasound::Frontend
 		// Add outputs to parent graph and connect to wrapped graph node.
 		void AddAndConnectOutputs(const TArray<FMetasoundFrontendClassOutput>& InClassOutputs, FGraphHandle& InParentGraphHandle, FNodeHandle& InReferencedNode) const;
 
+#if WITH_EDITORONLY_DATA
+		using FMemberIDToMetadataMap = TMap<FGuid, TObjectPtr<UMetaSoundFrontendMemberMetadata>>;
+		void AddMemberMetadata(const FMemberIDToMetadataMap& InCachedMemberMetadata, FGraphHandle& InPresetGraphHandle, FMemberIDToMetadataMap& InOutMemberMetadata) const;
+#endif // WITH_EDITORONLY_DATA
+
 		FConstDocumentHandle ReferencedDocument = IDocumentController::GetInvalidHandle();
 	};
 
