Add nullptr check in MetaSound editor subsystem

diff --git a/Engine/Plugins/Runtime/Metasound/Source/MetasoundEditor/Private/MetasoundEditorSubsystem.cpp b/Engine/Plugins/Runtime/Metasound/Source/MetasoundEditor/Private/MetasoundEditorSubsystem.cpp
--- a/Engine/Plugins/Runtime/Metasound/Source/MetasoundEditor/Private/MetasoundEditorSubsystem.cpp
+++ b/Engine/Plugins/Runtime/Metasound/Source/MetasoundEditor/Private/MetasoundEditorSubsystem.cpp
@@ -134,15 +134,16 @@ UMetasoundEditorGraphMemberDefaultLiteral* UMetaSoundEditorSubsystem::CreateMemb
 		{
 			// Get referenced asset to inherit metadata from
 			FMetasoundAssetBase* ReferencedPresetAsset = Builder.GetReferencedPresetAsset();
-			check(ReferencedPresetAsset);
-
-			FMetaSoundFrontendDocumentBuilder& ReferencedBuilder = FDocumentBuilderRegistry::GetChecked().FindOrBeginBuilding(ReferencedPresetAsset->GetOwningAsset());
-			const FMetasoundFrontendClassInput* ClassInput = ReferencedBuilder.FindGraphInput(InMemberName);
-			if (ClassInput)
+			if (ReferencedPresetAsset)
 			{
-				if (UMetaSoundFrontendMemberMetadata* ReferencedMemberMetadata = ReferencedBuilder.FindMemberMetadata(ClassInput->NodeID))
+				FMetaSoundFrontendDocumentBuilder& ReferencedBuilder = FDocumentBuilderRegistry::GetChecked().FindOrBeginBuilding(ReferencedPresetAsset->GetOwningAsset());
+				const FMetasoundFrontendClassInput* ClassInput = ReferencedBuilder.FindGraphInput(InMemberName);
+				if (ClassInput)
 				{
-					return NewObject<UMetasoundEditorGraphMemberDefaultLiteral>(&Builder.CastDocumentObjectChecked<UObject>(), LiteralClass, FName(), RF_Transactional, ReferencedMemberMetadata);
+					if (UMetaSoundFrontendMemberMetadata* ReferencedMemberMetadata = ReferencedBuilder.FindMemberMetadata(ClassInput->NodeID))
+					{
+						return NewObject<UMetasoundEditorGraphMemberDefaultLiteral>(&Builder.CastDocumentObjectChecked<UObject>(), LiteralClass, FName(), RF_Transactional, ReferencedMemberMetadata);
+					}
 				}
 			}
 		}
