Update Builder API to support replacing MetaSound class dependencies via Editor BP script

#rnx
#tests FixUpDuplicateMetaSoundReferences script

diff --git a/Engine/Plugins/Runtime/Metasound/Source/MetasoundEditor/Private/MetasoundEditorModule.cpp b/Engine/Plugins/Runtime/Metasound/Source/MetasoundEditor/Private/MetasoundEditorModule.cpp
--- a/Engine/Plugins/Runtime/Metasound/Source/MetasoundEditor/Private/MetasoundEditorModule.cpp
+++ b/Engine/Plugins/Runtime/Metasound/Source/MetasoundEditor/Private/MetasoundEditorModule.cpp
@@ -764,7 +764,6 @@ namespace Metasound
 				if (MigrationContext.GetCurrentStep() == UE::AssetTools::FPackageMigrationContext::EPackageMigrationStep::InstancedPackagesLoaded)
 				{
 					// Gather the new MetaSound assets
-					TMap<FMetasoundFrontendClassName, FMetasoundFrontendClassName> OldToNewReferencedClassNames;
 					TArray<FMetaSoundFrontendDocumentBuilder> NewMetaSoundAssetBuilders;
 					for (const UE::AssetTools::FPackageMigrationContext::FMigrationPackageData& MigrationPackageData : MigrationContext.GetMigrationPackagesData())
 					{
@@ -782,11 +781,12 @@ namespace Metasound
 
 					// Assign new class names and cache mapping with old one
 					IMetaSoundAssetManager& AssetManager = IMetaSoundAssetManager::GetChecked();
+					TMap<FNodeRegistryKey, FNodeRegistryKey> OldToNewReferenceKeys;
 					for (FMetaSoundFrontendDocumentBuilder& MetaSoundBuilder : NewMetaSoundAssetBuilders)
 					{
-						const FMetasoundFrontendClassName OldReferencedClassName = MetaSoundBuilder.GetConstDocumentChecked().RootGraph.Metadata.GetClassName();
-						const FMetasoundFrontendClassName NewReferencedClassName = MetaSoundBuilder.GenerateNewClassName();
-						OldToNewReferencedClassNames.FindOrAdd(OldReferencedClassName) = NewReferencedClassName;
+						FNodeRegistryKey OldRegistryKey(MetaSoundBuilder.GetConstDocumentChecked().RootGraph.Metadata);
+						FNodeRegistryKey NewRegistryKey(EMetasoundFrontendClassType::External, MetaSoundBuilder.GenerateNewClassName(), OldRegistryKey.Version);
+						OldToNewReferenceKeys.FindOrAdd(MoveTemp(OldRegistryKey)) = MoveTemp(NewRegistryKey);
 
 						UObject& MetaSoundObject = MetaSoundBuilder.CastDocumentObjectChecked<UObject>();
 						AssetManager.AddOrUpdateAsset(MetaSoundObject);
@@ -795,7 +795,7 @@ namespace Metasound
 					// Fix up dependencies
 					for (FMetaSoundFrontendDocumentBuilder& MetaSoundBuilder : NewMetaSoundAssetBuilders)
 					{
-						MetaSoundBuilder.UpdateDependencyClassNames(OldToNewReferencedClassNames);
+						MetaSoundBuilder.UpdateDependencyRegistryData(OldToNewReferenceKeys);
 					}
 				}
 			}
diff --git a/Engine/Plugins/Runtime/Metasound/Source/MetasoundEngine/Private/MetasoundAssetSubsystem.cpp b/Engine/Plugins/Runtime/Metasound/Source/MetasoundEngine/Private/MetasoundAssetSubsystem.cpp
--- a/Engine/Plugins/Runtime/Metasound/Source/MetasoundEngine/Private/MetasoundAssetSubsystem.cpp
+++ b/Engine/Plugins/Runtime/Metasound/Source/MetasoundEngine/Private/MetasoundAssetSubsystem.cpp
@@ -9,6 +9,7 @@
 #include "Metasound.h"
 #include "MetasoundAssetBase.h"
 #include "MetasoundBuilderSubsystem.h"
+#include "MetasoundDocumentBuilderRegistry.h"
 #include "MetasoundEngineAsset.h"
 #include "MetasoundFrontendDocument.h"
 #include "MetasoundFrontendDocumentBuilder.h"
@@ -152,9 +153,13 @@ namespace Metasound::Engine
 
 		void RebuildDenyListCache(const UAssetManager& InAssetManager);
 		void RegisterAssetClassesInDirectories(const TArray<FMetaSoundAssetDirectory>& Directories);
+
+#if WITH_EDITOR
+		bool ReplaceReferencesInDirectory(const TArray<FMetaSoundAssetDirectory>& InDirectories, const Metasound::Frontend::FNodeRegistryKey& OldClassKey, const Metasound::Frontend::FNodeRegistryKey& NewClassKey) const;
+#endif // WITH_EDITOR
 		void RequestAsyncLoadReferencedAssets(FMetasoundAssetBase& InAssetBase);
 		void OnAssetScanComplete();
-		void SearchAndIterateDirectoryAssets(const TArray<FDirectoryPath>& InDirectories, TFunctionRef<void(const FAssetData&)> InFunction);
+		void SearchAndIterateDirectoryAssets(const TArray<FDirectoryPath>& InDirectories, TFunctionRef<void(const FAssetData&)> InFunction) const;
 		FMetasoundAssetBase* TryLoadAsset(const FSoftObjectPath& InObjectPath) const;
 		void UnregisterAssetClassesInDirectories(const TArray<FMetaSoundAssetDirectory>& Directories);
 
@@ -174,6 +179,7 @@ namespace Metasound::Engine
 		virtual const FMetasoundAssetBase* GetAsAsset(const UObject& InObject) const override;
 #if WITH_EDITOR
 		virtual TSet<FAssetInfo> GetReferencedAssetClasses(const FMetasoundAssetBase& InAssetBase) const override;
+		virtual bool ReassignClassName(TScriptInterface<IMetaSoundDocumentInterface> DocInterface) override;
 #endif // WITH_EDITOR
 		virtual void IterateAssets(TFunctionRef<void(const FAssetKey, const TArray<FTopLevelAssetPath>&)> Iter) const override;
 		virtual void ReloadMetaSoundAssets() const override;
@@ -588,6 +594,48 @@ namespace Metasound::Engine
 		}
 		return MoveTemp(OutAssetInfos);
 	}
+
+	bool FMetaSoundAssetManager::ReassignClassName(TScriptInterface<IMetaSoundDocumentInterface> DocInterface)
+	{
+#if WITH_EDITORONLY_DATA
+		UObject* MetaSoundObject = DocInterface.GetObject();
+		if (!MetaSoundObject)
+		{
+			return false;
+		}
+
+		FMetasoundAssetBase* AssetBase = GetAsAsset(*MetaSoundObject);
+		if (!AssetBase)
+		{
+			return false;
+		}
+
+		FMetaSoundFrontendDocumentBuilder& Builder = FDocumentBuilderRegistry::GetChecked().FindOrBeginBuilding(DocInterface);
+
+		const FMetasoundFrontendClassMetadata& ClassMetadata = Builder.GetConstDocumentChecked().RootGraph.Metadata;
+		const FTopLevelAssetPath Path(MetaSoundObject);
+
+		AssetBase->UnregisterGraphWithFrontend();
+
+		{
+			const FAssetKey OldAssetKey(ClassMetadata.GetClassName(), ClassMetadata.GetVersion());
+			AssetSubsystemPrivate::RemovePath(PathMap, OldAssetKey, Path);
+		}
+
+		Builder.GenerateNewClassName();
+
+		{
+			const FAssetKey NewAssetKey(ClassMetadata.GetClassName(), ClassMetadata.GetVersion());
+			AssetSubsystemPrivate::AddPath(PathMap, NewAssetKey, Path);
+		}
+
+		AssetBase->UpdateAndRegisterForExecution();
+		return true;
+
+#else // !WITH_EDITORONLY_DATA
+		return false;
+#endif // !WITH_EDITORONLY_DATA
+	}
 #endif // WITH_EDITOR
 
 	void FMetaSoundAssetManager::IterateAssets(TFunctionRef<void(const FAssetKey, const TArray<FTopLevelAssetPath>&)> Iter) const
@@ -726,6 +774,56 @@ namespace Metasound::Engine
 		}
 	}
 
+#if WITH_EDITOR
+	bool FMetaSoundAssetManager::ReplaceReferencesInDirectory(const TArray<FMetaSoundAssetDirectory>& InDirectories, const Metasound::Frontend::FNodeRegistryKey& OldClassKey, const Metasound::Frontend::FNodeRegistryKey& NewClassKey) const
+	{
+		using namespace Frontend;
+
+		bool bReferencesReplaced = false;
+
+#if WITH_EDITORONLY_DATA
+		if (!NewClassKey.IsValid())
+		{
+			return bReferencesReplaced;
+		}
+
+		FMetasoundFrontendClass NewClass;
+		const bool bNewClassExists = ISearchEngine::Get().FindClassWithHighestVersion(NewClassKey.ClassName, NewClass);
+		if (bNewClassExists)
+		{
+			TArray<FDirectoryPath> Directories;
+			Algo::Transform(InDirectories, Directories, [](const FMetaSoundAssetDirectory& AssetDir) { return AssetDir.Directory; });
+
+			TMap<FNodeRegistryKey, FNodeRegistryKey> OldToNewReferenceKeys = { { OldClassKey, NewClassKey } };
+			SearchAndIterateDirectoryAssets(Directories, [this, &bReferencesReplaced, &OldToNewReferenceKeys](const FAssetData& AssetData)
+			{
+				if (UObject* MetaSoundObject = AssetData.GetAsset())
+				{
+					MetaSoundObject->Modify();
+					FMetaSoundFrontendDocumentBuilder& Builder = FDocumentBuilderRegistry::GetChecked().FindOrBeginBuilding(MetaSoundObject);
+					const bool bDependencyUpdated = Builder.UpdateDependencyRegistryData(OldToNewReferenceKeys);
+					if (bDependencyUpdated)
+					{
+						bReferencesReplaced = true;
+						Builder.RemoveUnusedDependencies();
+						if (FMetasoundAssetBase* AssetBase = GetAsAsset(*MetaSoundObject); ensure(AssetBase))
+						{
+							AssetBase->RebuildReferencedAssetClasses();
+						}
+					}
+				}
+			});
+		}
+		else
+		{
+			UE_LOG(LogMetaSound, Display, TEXT("Cannot replace references in MetaSound assets found in given directory/directories: NewClass '%s' does not exist"), *NewClassKey.ToString());
+		}
+#endif // WITH_EDITORONLY_DATA
+
+		return bReferencesReplaced;
+	}
+#endif // WITH_EDITOR
+
 	void FMetaSoundAssetManager::RequestAsyncLoadReferencedAssets(FMetasoundAssetBase& InAssetBase)
 	{
 		const TSet<FSoftObjectPath>& AsyncReferences = InAssetBase.GetAsyncReferencedAssetClassPaths();
@@ -1117,12 +1215,39 @@ UMetaSoundAssetSubsystem& UMetaSoundAssetSubsystem::GetChecked()
 	return *Subsystem;
 }
 
+#if WITH_EDITOR
+bool UMetaSoundAssetSubsystem::ReassignClassName(TScriptInterface<IMetaSoundDocumentInterface> DocInterface)
+{
+	using namespace Metasound::Engine;
+	return FMetaSoundAssetManager::GetChecked().ReassignClassName(DocInterface);
+}
+#endif // WITH_EDITOR
+
 void UMetaSoundAssetSubsystem::RegisterAssetClassesInDirectories(const TArray<FMetaSoundAssetDirectory>& InDirectories)
 {
 	using namespace Metasound::Engine;
 	FMetaSoundAssetManager::GetChecked().RegisterAssetClassesInDirectories(InDirectories);
 }
 
+#if WITH_EDITOR
+bool UMetaSoundAssetSubsystem::ReplaceReferencesInDirectory(
+	const TArray<FMetaSoundAssetDirectory>& InDirectories,
+	const FMetasoundFrontendClassName& OldClassName,
+	const FMetasoundFrontendClassName& NewClassName,
+	const FMetasoundFrontendVersionNumber OldVersion,
+	const FMetasoundFrontendVersionNumber NewVersion)
+{
+	using namespace Metasound::Engine;
+	using namespace Metasound::Frontend;
+
+	return FMetaSoundAssetManager::GetChecked().ReplaceReferencesInDirectory(
+		InDirectories,
+		FNodeRegistryKey(EMetasoundFrontendClassType::External, OldClassName, OldVersion),
+		FNodeRegistryKey(EMetasoundFrontendClassType::External, NewClassName, NewVersion)
+	);
+}
+#endif // WITH_EDITOR
+
 void UMetaSoundAssetSubsystem::UnregisterAssetClassesInDirectories(const TArray<FMetaSoundAssetDirectory>& InDirectories)
 {
 	using namespace Metasound::Engine;
diff --git a/Engine/Plugins/Runtime/Metasound/Source/MetasoundEngine/Public/MetasoundAssetSubsystem.h b/Engine/Plugins/Runtime/Metasound/Source/MetasoundEngine/Public/MetasoundAssetSubsystem.h
--- a/Engine/Plugins/Runtime/Metasound/Source/MetasoundEngine/Public/MetasoundAssetSubsystem.h
+++ b/Engine/Plugins/Runtime/Metasound/Source/MetasoundEngine/Public/MetasoundAssetSubsystem.h
@@ -6,6 +6,7 @@
 #include "MetasoundAssetBase.h"
 #include "MetasoundAssetManager.h"
 #include "MetasoundBuilderBase.h"
+#include "MetasoundDocumentInterface.h"
 #include "Subsystems/EngineSubsystem.h"
 #include "UObject/Object.h"
 
@@ -125,6 +126,22 @@ public:
 	UE_DEPRECATED(5.5, "Moved to internal implementation, use IMetaSoundAssetManager::GetChecked() and analogous call")
 	virtual void WaitUntilAsyncLoadReferencedAssetsComplete(FMetasoundAssetBase& InAssetBase) { }
 
+#if WITH_EDITOR
+	UFUNCTION(BlueprintCallable, Category = "MetaSounds|Utilities")
+	UPARAM(DisplayName = "Reassigned") bool ReassignClassName(TScriptInterface<IMetaSoundDocumentInterface> DocInterface);
+
+	// Replaces dependencies in a MetaSound with the given class name and version with another MetaSound with the given
+	// class name and version.  Can be asset or code-defined.  It is up to the caller to validate the two classes have
+	// matching interfaces (Swapping with classes of unmatched interfaces can leave MetaSound in non-executable state).
+	UFUNCTION(BlueprintCallable, Category = "MetaSounds|Utilities", meta = (AdvancedDisplay = "3"))
+	UPARAM(DisplayName = "References Replaced") bool ReplaceReferencesInDirectory(
+		const TArray<FMetaSoundAssetDirectory>& InDirectories,
+		const FMetasoundFrontendClassName& OldClassName,
+		const FMetasoundFrontendClassName& NewClassName,
+		const FMetasoundFrontendVersionNumber OldVersion = FMetasoundFrontendVersionNumber(),
+		const FMetasoundFrontendVersionNumber NewVersion = FMetasoundFrontendVersionNumber());
+#endif // WITH_EDITOR
+
 	UFUNCTION(BlueprintCallable, Category = "MetaSounds|Registration")
 	void RegisterAssetClassesInDirectories(const TArray<FMetaSoundAssetDirectory>& Directories);
 
diff --git a/Engine/Plugins/Runtime/Metasound/Source/MetasoundEngine/Public/MetasoundBuilderBase.h b/Engine/Plugins/Runtime/Metasound/Source/MetasoundEngine/Public/MetasoundBuilderBase.h
--- a/Engine/Plugins/Runtime/Metasound/Source/MetasoundEngine/Public/MetasoundBuilderBase.h
+++ b/Engine/Plugins/Runtime/Metasound/Source/MetasoundEngine/Public/MetasoundBuilderBase.h
@@ -115,7 +115,7 @@ enum class EMetaSoundBuilderResult : uint8
 };
 
 /** Base implementation of MetaSound builder */
-UCLASS(Abstract)
+UCLASS(Abstract, BlueprintType, Transient)
 class METASOUNDENGINE_API UMetaSoundBuilderBase : public UObject
 {
 	GENERATED_BODY()
@@ -279,8 +279,9 @@ public:
 	UFUNCTION(BlueprintCallable, Category = "Audio|MetaSound|Builder", meta = (ExpandEnumAsExecs = "OutResult"))
 	UPARAM(DisplayName = "Node ClassVersion") FMetasoundFrontendVersion FindNodeClassVersion(const FMetaSoundNodeHandle& NodeHandle, EMetaSoundBuilderResult& OutResult);
 
-	// Returns the document's root graph class name
-	FMetasoundFrontendClassName GetRootGraphClassName() const;
+	// Returns the MetaSound asset's graph class name (used by the MetaSound Node Class Registry)
+	UFUNCTION(BlueprintCallable, Category = "Audio|MetaSound|Builder|Advanced", meta = (DisplayName = "Get MetaSound Class Name"))
+	UPARAM(DisplayName = "Class Name") FMetasoundFrontendClassName GetRootGraphClassName() const;
 
 	// Returns node input's data if valid (including things like name and datatype).
 	UFUNCTION(BlueprintCallable, Category = "Audio|MetaSound|Builder", meta = (ExpandEnumAsExecs = "OutResult"))
@@ -288,15 +289,15 @@ public:
 
 	// Returns node input's literal value if set on graph, otherwise fails and returns default literal.
 	UFUNCTION(BlueprintCallable, Category = "Audio|MetaSound|Builder", meta = (ExpandEnumAsExecs = "OutResult"))
-	FMetasoundFrontendLiteral GetNodeInputDefault(const FMetaSoundBuilderNodeInputHandle& InputHandle, EMetaSoundBuilderResult& OutResult);
+	UPARAM(DisplayName = "Default") FMetasoundFrontendLiteral GetNodeInputDefault(const FMetaSoundBuilderNodeInputHandle& InputHandle, EMetaSoundBuilderResult& OutResult);
 
 	// Returns node input's class literal value if set, otherwise fails and returns default literal.
 	UFUNCTION(BlueprintCallable, Category = "Audio|MetaSound|Builder", meta = (ExpandEnumAsExecs = "OutResult"))
-	FMetasoundFrontendLiteral GetNodeInputClassDefault(const FMetaSoundBuilderNodeInputHandle& InputHandle, EMetaSoundBuilderResult& OutResult);
+	UPARAM(DisplayName = "Default") FMetasoundFrontendLiteral GetNodeInputClassDefault(const FMetaSoundBuilderNodeInputHandle& InputHandle, EMetaSoundBuilderResult& OutResult);
 
 	// Returns whether the given node input is a constructor pin
 	UFUNCTION(BlueprintCallable, Category = "Audio|MetaSound|Builder")
-	bool GetNodeInputIsConstructorPin(const FMetaSoundBuilderNodeInputHandle& InputHandle) const;
+	UPARAM(DisplayName = "Is Constructor Pin") bool GetNodeInputIsConstructorPin(const FMetaSoundBuilderNodeInputHandle& InputHandle) const;
 
 	// Returns node output's data if valid (including things like name and datatype).
 	UFUNCTION(BlueprintCallable, Category = "Audio|MetaSound|Builder", meta = (ExpandEnumAsExecs = "OutResult"))
@@ -304,15 +305,15 @@ public:
 	
 	// Returns whether the given node output is a constructor pin
 	UFUNCTION(BlueprintCallable, Category = "Audio|MetaSound|Builder")
-	bool GetNodeOutputIsConstructorPin(const FMetaSoundBuilderNodeOutputHandle& OutputHandle) const;
+	UPARAM(DisplayName = "Is Constructor Pin") bool GetNodeOutputIsConstructorPin(const FMetaSoundBuilderNodeOutputHandle& OutputHandle) const;
 
 	// Return the asset referenced by this preset builder. Returns nullptr if the builder is not a preset.
 	UFUNCTION(BlueprintCallable, Category = "Audio|MetaSound|Builder")
-	UObject* GetReferencedPresetAsset() const;
+	UPARAM(DisplayName = "Referenced Preset") UObject* GetReferencedPresetAsset() const;
 
 	// Returns if a given interface is declared.
 	UFUNCTION(BlueprintCallable, Category = "Audio|MetaSound|Builder")
-	UPARAM(DisplayName = "IsDeclared") bool InterfaceIsDeclared(FName InterfaceName) const;
+	UPARAM(DisplayName = "Is Declared") bool InterfaceIsDeclared(FName InterfaceName) const;
 
 	// Returns if a given node output and node input are connected.
 	UFUNCTION(BlueprintCallable, Category = "Audio|MetaSound|Builder")
@@ -328,7 +329,7 @@ public:
 
 	// Returns whether this is a preset.
 	UFUNCTION(BlueprintCallable, Category = "Audio|MetaSound|Builder")
-	bool IsPreset() const;
+	UPARAM(DisplayName = "Is Preset") bool IsPreset() const;
 
 	// Converts this preset to a fully accessible MetaSound; sets result to succeeded if it was converted successfully and failed if it was not.
 	UFUNCTION(BlueprintCallable, Category = "Audio|MetaSound|Builder", meta = (ExpandEnumAsExecs = "OutResult"))
diff --git a/Engine/Plugins/Runtime/Metasound/Source/MetasoundFrontend/Private/MetasoundFrontendDocumentBuilder.cpp b/Engine/Plugins/Runtime/Metasound/Source/MetasoundFrontend/Private/MetasoundFrontendDocumentBuilder.cpp
--- a/Engine/Plugins/Runtime/Metasound/Source/MetasoundFrontend/Private/MetasoundFrontendDocumentBuilder.cpp
+++ b/Engine/Plugins/Runtime/Metasound/Source/MetasoundFrontend/Private/MetasoundFrontendDocumentBuilder.cpp
@@ -2101,7 +2101,6 @@ FMetasoundFrontendClassName FMetaSoundFrontendDocumentBuilder::GenerateNewClassN
 {
 	using namespace Metasound::Frontend;
 	FMetasoundFrontendClassMetadata& Metadata = GetDocumentChecked().RootGraph.Metadata;
-	const FMetasoundFrontendClassName OldClassName = Metadata.GetClassName();
 	const FMetasoundFrontendClassName NewClassName(FName(), FName(*FGuid::NewGuid().ToString()), FName());
 	Metadata.SetClassName(NewClassName);
 	return NewClassName;
@@ -4441,27 +4440,59 @@ bool FMetaSoundFrontendDocumentBuilder::UnlinkVariableNode(const FGuid& InNodeID
 	return false;
 }
 
-bool FMetaSoundFrontendDocumentBuilder::UpdateDependencyClassNames(const TMap<FMetasoundFrontendClassName, FMetasoundFrontendClassName>& OldToNewReferencedClassNames)
+#if WITH_EDITOR
+bool FMetaSoundFrontendDocumentBuilder::UpdateDependencyRegistryData(const TMap<Metasound::Frontend::FNodeRegistryKey, Metasound::Frontend::FNodeRegistryKey>& OldToNewClassKeys)
 {
-	if (!DocumentDelegates.IsValid())
-	{
-		return false;
-	}
+	using namespace Metasound::Frontend;
 
-	for (FMetasoundFrontendClass& Dependency : GetDocumentChecked().Dependencies)
+	bool bUpdated = false;
+	if (DocumentDelegates.IsValid())
 	{
-		const FMetasoundFrontendClassName OldName = Dependency.Metadata.GetClassName();
-		if (const FMetasoundFrontendClassName* NewName = OldToNewReferencedClassNames.Find(OldName))
+		FMetasoundFrontendDocument& Document = GetDocumentChecked();
+		for (FMetasoundFrontendClass& Dependency : Document.Dependencies)
+		{
+			const FNodeRegistryKey OldKey(Dependency.Metadata);
+			if (const FNodeRegistryKey* NewKey = OldToNewClassKeys.Find(OldKey))
+			{
+				if (Dependency.Metadata.GetType() == EMetasoundFrontendClassType::External)
+				{
+					bUpdated = true;
+					const int32* DependencyIndex = DocumentCache->FindDependencyIndex(Dependency.ID);
+					check(DependencyIndex);
+					DocumentDelegates->OnRenamingDependencyClass.Broadcast(*DependencyIndex, NewKey->ClassName);
+					Dependency.Metadata.SetType(NewKey->Type);
+					Dependency.Metadata.SetClassName(NewKey->ClassName);
+					Dependency.Metadata.SetVersion(NewKey->Version);
+				}
+			}
+		}
+
+#if WITH_EDITORONLY_DATA
+		if (bUpdated)
 		{
-			const int32* DependencyIndex = DocumentCache->FindDependencyIndex(Dependency.ID);
-			check(DependencyIndex);
-			DocumentDelegates->OnRenamingDependencyClass.Broadcast(*DependencyIndex, *NewName);
-			Dependency.Metadata.SetClassName(*NewName);
+			Document.Metadata.ModifyContext.SetDocumentModified();
 		}
+#endif // WITH_EDITORONLY_DATA
 	}
-	
-	return true;
+
+	return bUpdated;
+}
+
+bool FMetaSoundFrontendDocumentBuilder::UpdateDependencyClassNames(const TMap<FMetasoundFrontendClassName, FMetasoundFrontendClassName>& OldToNewReferencedClassNames)
+{
+	using namespace Metasound::Frontend;
+
+	TMap<FNodeRegistryKey, FNodeRegistryKey> OldToNewKeys;
+	Algo::Transform(OldToNewReferencedClassNames, OldToNewKeys, [](const TPair<FMetasoundFrontendClassName, FMetasoundFrontendClassName>& ClassNamePair)
+	{
+		return TPair<FNodeRegistryKey, FNodeRegistryKey>(
+			FNodeRegistryKey(EMetasoundFrontendClassType::External, ClassNamePair.Key, FMetasoundFrontendVersionNumber()),
+			FNodeRegistryKey(EMetasoundFrontendClassType::External, ClassNamePair.Value, FMetasoundFrontendVersionNumber())
+		);
+	});
+	return UpdateDependencyRegistryData(OldToNewKeys);
 }
+#endif // WITH_EDITOR
 
 #if WITH_EDITORONLY_DATA
 bool FMetaSoundFrontendDocumentBuilder::VersionInterfaces()
diff --git a/Engine/Plugins/Runtime/Metasound/Source/MetasoundFrontend/Public/MetasoundAssetManager.h b/Engine/Plugins/Runtime/Metasound/Source/MetasoundFrontend/Public/MetasoundAssetManager.h
--- a/Engine/Plugins/Runtime/Metasound/Source/MetasoundFrontend/Public/MetasoundAssetManager.h
+++ b/Engine/Plugins/Runtime/Metasound/Source/MetasoundFrontend/Public/MetasoundAssetManager.h
@@ -207,7 +207,13 @@ namespace Metasound::Frontend
 		// Try to load referenced assets of the given asset or return them if they are already loaded (non-recursive).
 		// @return - True if all referenced assets successfully loaded, false if not.
 		virtual bool TryLoadReferencedAssets(const FMetasoundAssetBase& InAssetBase, TArray<FMetasoundAssetBase*>& OutReferencedAssets) const = 0;
-			
+
+#if WITH_EDITOR
+		// Assigns a new arbitrary class name to the given document, which can cause references to be invalidated.
+		// (See
+		virtual bool ReassignClassName(TScriptInterface<IMetaSoundDocumentInterface> DocInterface) = 0;
+#endif // WITH_EDITOR
+
 		// Requests an async load of all async referenced assets of the input asset.
 		virtual void RequestAsyncLoadReferencedAssets(FMetasoundAssetBase& InAssetBase) = 0;
 
diff --git a/Engine/Plugins/Runtime/Metasound/Source/MetasoundFrontend/Public/MetasoundFrontendDocumentBuilder.h b/Engine/Plugins/Runtime/Metasound/Source/MetasoundFrontend/Public/MetasoundFrontendDocumentBuilder.h
--- a/Engine/Plugins/Runtime/Metasound/Source/MetasoundFrontend/Public/MetasoundFrontendDocumentBuilder.h
+++ b/Engine/Plugins/Runtime/Metasound/Source/MetasoundFrontend/Public/MetasoundFrontendDocumentBuilder.h
@@ -281,11 +281,11 @@ public:
 		return *CastChecked<TObjectType>(Owner);
 	}
 
-	// Generates and returns new class name for the given builder's document. Should ONLY be called on new assets
-	// using transient frontend builder, as using a persistent builder registered with the DocumentBuilderRegistry
-	// may result in corrupt records keyed off of undefined class name being potentially generated over.  In addition,
-	// this can potentially leave existing node references in an abandoned state to this class causing MetaSound generator
-	// build errors.
+	// Generates and returns new class name for the given builder's document. Should be used with extreme caution
+	// (i.e. on new assets, when migrating assets, or upon generation of transient MetaSounds), as using a persistent
+	// builder registered with the DocumentBuilderRegistry may result in stale asset records keyed off of an undefined class
+	// name.  In addition, this can potentially leave existing node references in an abandoned state to this class causing
+	// asset validation errors.
 	FMetasoundFrontendClassName GenerateNewClassName();
 
 	Metasound::Frontend::FDocumentModifyDelegates& GetDocumentDelegates();
@@ -481,8 +481,14 @@ public:
 
 	bool SwapGraphInput(const FMetasoundFrontendClassVertex& InExistingInputVertex, const FMetasoundFrontendClassVertex& NewInputVertex);
 	bool SwapGraphOutput(const FMetasoundFrontendClassVertex& InExistingOutputVertex, const FMetasoundFrontendClassVertex& NewOutputVertex);
+
+#if WITH_EDITOR
+	UE_DEPRECATED(5.5, "Use 'UpdateDependencyRegistryData' instead and supply keys (comprised of name, version & node class type)")
 	bool UpdateDependencyClassNames(const TMap<FMetasoundFrontendClassName, FMetasoundFrontendClassName>& OldToNewReferencedClassNames);
 
+	bool UpdateDependencyRegistryData(const TMap<Metasound::Frontend::FNodeRegistryKey, Metasound::Frontend::FNodeRegistryKey>& OldToNewClassKeys);
+#endif // WITH_EDITOR
+
 #if WITH_EDITORONLY_DATA
 	// Transforms template nodes within the given builder's document, which can include swapping associated edges and/or
 	// replacing nodes with other, registry-defined concrete node class instances. Returns true if any template nodes were processed.
