Use thread safe local reference to object instead of attempting to access MetaSoundAssetManager to return referenced asset when retrieving referenced preset asset.

#rnx

diff --git a/Engine/Plugins/Runtime/Metasound/Source/MetasoundEditor/Private/MetasoundEditorGraph.cpp b/Engine/Plugins/Runtime/Metasound/Source/MetasoundEditor/Private/MetasoundEditorGraph.cpp
--- a/Engine/Plugins/Runtime/Metasound/Source/MetasoundEditor/Private/MetasoundEditorGraph.cpp
+++ b/Engine/Plugins/Runtime/Metasound/Source/MetasoundEditor/Private/MetasoundEditorGraph.cpp
@@ -2309,7 +2309,11 @@ void UMetasoundEditorGraph::MigrateEditorDocumentData(FMetaSoundFrontendDocument
 				if (UMetasoundEditorGraphMemberDefaultLiteral* DefaultLiteral = Input.GetLiteral())
 				{
 					TSubclassOf<UMetasoundEditorGraphMemberDefaultLiteral> SubClass = DefaultLiteral->GetClass();
+
+					// Migration can occur on async thread, and bind can create a new literal
+					FGCScopeGuard ScopeGuard;
 					EditorSubsystem->BindMemberMetadata(OutBuilder, Input, SubClass, DefaultLiteral);
+					DefaultLiteral->ClearInternalFlags(EInternalObjectFlags::Async);
 				}
 			}
 		}
diff --git a/Engine/Plugins/Runtime/Metasound/Source/MetasoundFrontend/Private/MetasoundFrontendDocumentBuilder.cpp b/Engine/Plugins/Runtime/Metasound/Source/MetasoundFrontend/Private/MetasoundFrontendDocumentBuilder.cpp
--- a/Engine/Plugins/Runtime/Metasound/Source/MetasoundFrontend/Private/MetasoundFrontendDocumentBuilder.cpp
+++ b/Engine/Plugins/Runtime/Metasound/Source/MetasoundFrontend/Private/MetasoundFrontendDocumentBuilder.cpp
@@ -2236,8 +2236,20 @@ FMetasoundAssetBase* FMetaSoundFrontendDocumentBuilder::GetReferencedPresetAsset
 	{
 		const FMetasoundFrontendClass* NodeClass = FindDependency(Node->ClassID);
 		check(NodeClass);
-		const FNodeRegistryKey NodeClassRegistryKey = FNodeRegistryKey(NodeClass->Metadata);
-		return IMetaSoundAssetManager::GetChecked().FindAsset(NodeClassRegistryKey);
+		const FAssetKey NodeAssetKey(NodeClass->Metadata);
+		const TArray<FMetasoundAssetBase*> ReferencedAssets = GetMetasoundAsset().GetReferencedAssets();
+		for (FMetasoundAssetBase* RefAsset : ReferencedAssets)
+		{
+			TScriptInterface<IMetaSoundDocumentInterface> RefDocInterface = RefAsset->GetOwningAsset();
+			if (RefDocInterface.GetObject() != nullptr)
+			{
+				const FAssetKey AssetKey(RefDocInterface->GetConstDocument().RootGraph.Metadata);
+				if (AssetKey == NodeAssetKey)
+				{
+					return RefAsset;
+				}
+			}
+		}
 	}
 	return nullptr;
 }
