Make MetaSoundAssetManager thread safe (to allow safely using it during async MetaSound asset load)

#rnx

diff --git a/Engine/Plugins/Runtime/Metasound/Source/MetasoundEditor/Private/MetasoundEditorGraphBuilder.cpp b/Engine/Plugins/Runtime/Metasound/Source/MetasoundEditor/Private/MetasoundEditorGraphBuilder.cpp
--- a/Engine/Plugins/Runtime/Metasound/Source/MetasoundEditor/Private/MetasoundEditorGraphBuilder.cpp
+++ b/Engine/Plugins/Runtime/Metasound/Source/MetasoundEditor/Private/MetasoundEditorGraphBuilder.cpp
@@ -287,9 +287,10 @@ namespace Metasound
 					bool bIsClassNative = FMetasoundFrontendRegistryContainer::Get()->IsNodeNative(RegistryKey);
 					if (!bIsClassNative)
 					{
-						if (const FTopLevelAssetPath* Path = IMetaSoundAssetManager::GetChecked().FindAssetPath(RegistryKey))
+						const FTopLevelAssetPath Path = IMetaSoundAssetManager::GetChecked().FindAssetPath(RegistryKey);
+						if (Path.IsValid())
 						{
-							DisplayName = FText::FromName(Path->GetAssetName());
+							DisplayName = FText::FromName(Path.GetAssetName());
 						}
 					}
 				}
diff --git a/Engine/Plugins/Runtime/Metasound/Source/MetasoundEngine/Private/MetasoundAssetSubsystem.cpp b/Engine/Plugins/Runtime/Metasound/Source/MetasoundEngine/Private/MetasoundAssetSubsystem.cpp
--- a/Engine/Plugins/Runtime/Metasound/Source/MetasoundEngine/Private/MetasoundAssetSubsystem.cpp
+++ b/Engine/Plugins/Runtime/Metasound/Source/MetasoundEngine/Private/MetasoundAssetSubsystem.cpp
@@ -6,6 +6,7 @@
 #include "AssetRegistry/AssetRegistryModule.h"
 #include "Async/Async.h"
 #include "Engine/AssetManager.h"
+#include "HAL/CriticalSection.h"
 #include "Metasound.h"
 #include "MetasoundAssetBase.h"
 #include "MetasoundBuilderSubsystem.h"
@@ -86,10 +87,11 @@ namespace Metasound::Engine
 			return bSuccess;
 		}
 
-		bool RemovePath(TMap<Frontend::FAssetKey, TArray<FTopLevelAssetPath>>& InMap, const Frontend::FAssetKey& AssetKey, const FTopLevelAssetPath& AssetPath)
+		bool RemovePath(FCriticalSection* MapCritSec, TMap<Frontend::FAssetKey, TArray<FTopLevelAssetPath>>& Map, const Frontend::FAssetKey& AssetKey, const FTopLevelAssetPath& AssetPath)
 		{
-			check(IsInGameThread());
-			if (TArray<FTopLevelAssetPath>* MapAssetPaths = InMap.Find(AssetKey))
+			check(MapCritSec);
+			FScopeLock Lock(MapCritSec);
+			if (TArray<FTopLevelAssetPath>* MapAssetPaths = Map.Find(AssetKey))
 			{
 				auto ComparePaths = [&AssetPath](const FTopLevelAssetPath& Path) 
 				{
@@ -106,7 +108,7 @@ namespace Metasound::Engine
 				{
 					if (MapAssetPaths->IsEmpty())
 					{
-						InMap.Remove(AssetKey);
+						Map.Remove(AssetKey);
 					}
 					return true;
 				}
@@ -115,10 +117,11 @@ namespace Metasound::Engine
 			return false;
 		}
 
-		void AddPath(TMap<Frontend::FAssetKey, TArray<FTopLevelAssetPath>>& InMap, const Frontend::FAssetKey& AssetKey, const FTopLevelAssetPath& AssetPath)
+		void AddPath(FCriticalSection* MapCritSec, TMap<Frontend::FAssetKey, TArray<FTopLevelAssetPath>>& Map, const Frontend::FAssetKey& AssetKey, const FTopLevelAssetPath& AssetPath)
 		{
-			check(IsInGameThread());
-			TArray<FTopLevelAssetPath>& Paths = InMap.FindOrAdd(AssetKey);
+			check(MapCritSec);
+			FScopeLock Lock(MapCritSec);
+			TArray<FTopLevelAssetPath>& Paths = Map.FindOrAdd(AssetKey);
 			Paths.AddUnique(AssetPath);
 #if !NO_LOGGING
 			if (Paths.Num() > 1)
@@ -173,8 +176,8 @@ namespace Metasound::Engine
 		virtual bool ContainsKey(const FAssetKey& InKey) const override;
 		virtual FMetasoundAssetBase* FindAsset(const FAssetKey& InKey) const override;
 		virtual TScriptInterface<IMetaSoundDocumentInterface> FindAssetAsDocumentInterface(const Frontend::FAssetKey& InKey) const override;
-		virtual const FTopLevelAssetPath* FindAssetPath(const FAssetKey& InKey) const override;
-		virtual const TArray<FTopLevelAssetPath>* FindAssetPaths(const FAssetKey& InKey) const override;
+		virtual FTopLevelAssetPath FindAssetPath(const FAssetKey& InKey) const override;
+		virtual TArray<FTopLevelAssetPath> FindAssetPaths(const FAssetKey& InKey) const override;
 		virtual FMetasoundAssetBase* GetAsAsset(UObject& InObject) const override;
 		virtual const FMetasoundAssetBase* GetAsAsset(const UObject& InObject) const override;
 #if WITH_EDITOR
@@ -212,26 +215,32 @@ namespace Metasound::Engine
 		std::atomic<bool> bIsInitialAssetScanComplete = false;
 		TMap<FAssetKey, TArray<FTopLevelAssetPath>> PathMap;
 
+		// Critical section primarily for allowing safe access of path map during async loading of MetaSound assets.
+		mutable FCriticalSection PathMapCriticalSection;
+
 		bool bLogActiveAssetsOnShutdown = true;
 	};
 
 	FMetaSoundAssetManager::~FMetaSoundAssetManager()
 	{
 #if !NO_LOGGING
-		if (!PathMap.IsEmpty())
+		if (bLogActiveAssetsOnShutdown)
 		{
-			if (bLogActiveAssetsOnShutdown)
+			TMap<FAssetKey, TArray<FTopLevelAssetPath>> PathsOnShutdown;
+			{
+				FScopeLock Lock(&PathMapCriticalSection);
+				PathsOnShutdown = MoveTemp(PathMap);
+				PathMap.Reset();
+			}
+
+			if (!PathsOnShutdown.IsEmpty())
 			{
-				TSet<FAssetKey> Keys;
-				if (int32 NumKeys = PathMap.GetKeys(Keys); NumKeys > 0)
+				UE_LOG(LogMetaSound, Display, TEXT("AssetManager is shutting down with the following %i assets active:"), PathsOnShutdown.Num());
+				for (const TPair<FAssetKey, TArray<FTopLevelAssetPath>>& Pair : PathsOnShutdown)
 				{
-					UE_LOG(LogMetaSound, Display, TEXT("AssetManager is shutting down with the following %i assets active:"), NumKeys);
-					for (const TPair<FAssetKey, TArray<FTopLevelAssetPath>>& Pair : PathMap)
+					for (const FTopLevelAssetPath& Path : Pair.Value)
 					{
-						for (const FTopLevelAssetPath& Path : Pair.Value)
-						{
-							UE_LOG(LogMetaSound, Display, TEXT("- %s"), *Path.ToString());
-						}
+						UE_LOG(LogMetaSound, Display, TEXT("- %s"), *Path.ToString());
 					}
 				}
 			}
@@ -330,7 +339,7 @@ namespace Metasound::Engine
 
 		if (AssetKey.IsValid())
 		{
-			AssetSubsystemPrivate::AddPath(PathMap, AssetKey, FTopLevelAssetPath(&InObject));
+			AssetSubsystemPrivate::AddPath(&PathMapCriticalSection, PathMap, AssetKey, FTopLevelAssetPath(&InObject));
 		}
 
 		return AssetKey;
@@ -374,7 +383,7 @@ namespace Metasound::Engine
 			const FAssetKey AssetKey = FAssetKey(ClassInfo.ClassName, ClassInfo.Version);
 			if (AssetKey.IsValid())
 			{
-				AssetSubsystemPrivate::AddPath(PathMap, AssetKey, ClassInfo.AssetPath);
+				AssetSubsystemPrivate::AddPath(&PathMapCriticalSection, PathMap, AssetKey, ClassInfo.AssetPath);
 			}
 
 			return AssetKey;
@@ -398,7 +407,7 @@ namespace Metasound::Engine
 
 	bool FMetaSoundAssetManager::ContainsKey(const Metasound::Frontend::FAssetKey& InKey) const
 	{
-		check(IsInGameThread());
+		FScopeLock Lock(&PathMapCriticalSection);
 		return PathMap.Contains(InKey);
 	}
 
@@ -424,9 +433,10 @@ namespace Metasound::Engine
 
 	FMetasoundAssetBase* FMetaSoundAssetManager::FindAsset(const Metasound::Frontend::FAssetKey& InKey) const
 	{
-		if (const FTopLevelAssetPath* AssetPath = FindAssetPath(InKey))
+		FTopLevelAssetPath AssetPath = FindAssetPath(InKey);
+		if (AssetPath.IsValid())
 		{
-			if (UObject* Object = FSoftObjectPath(*AssetPath, { }).ResolveObject())
+			if (UObject* Object = FSoftObjectPath(AssetPath, { }).ResolveObject())
 			{
 				return GetAsAsset(*Object);
 			}
@@ -437,9 +447,10 @@ namespace Metasound::Engine
 
 	TScriptInterface<IMetaSoundDocumentInterface> FMetaSoundAssetManager::FindAssetAsDocumentInterface(const Metasound::Frontend::FAssetKey& InKey) const
 	{
-		if (const FTopLevelAssetPath* AssetPath = FindAssetPath(InKey))
+		const FTopLevelAssetPath AssetPath = FindAssetPath(InKey);
+		if (AssetPath.IsValid())
 		{
-			if (UObject* Object = FSoftObjectPath(*AssetPath, { }).ResolveObject())
+			if (UObject* Object = FSoftObjectPath(AssetPath, { }).ResolveObject())
 			{
 				return TScriptInterface<IMetaSoundDocumentInterface>(Object);
 			}
@@ -448,26 +459,26 @@ namespace Metasound::Engine
 		return nullptr;
 	}
 
-	const FTopLevelAssetPath* FMetaSoundAssetManager::FindAssetPath(const Metasound::Frontend::FAssetKey& InKey) const
+	FTopLevelAssetPath FMetaSoundAssetManager::FindAssetPath(const Metasound::Frontend::FAssetKey& InKey) const
 	{
-		check(IsInGameThread());
+		FScopeLock Lock(&PathMapCriticalSection);
 		if (const TArray<FTopLevelAssetPath>* Paths = PathMap.Find(InKey))
 		{
 			if (!Paths->IsEmpty())
 			{
-				return &Paths->Last();
+				return Paths->Last();
 			}
 		}
 
 		return nullptr;
 	}
 
-	const TArray<FTopLevelAssetPath>* FMetaSoundAssetManager::FindAssetPaths(const Metasound::Frontend::FAssetKey& InKey) const
+	TArray<FTopLevelAssetPath> FMetaSoundAssetManager::FindAssetPaths(const Metasound::Frontend::FAssetKey& InKey) const
 	{
-		check(IsInGameThread());
+		FScopeLock Lock(&PathMapCriticalSection);
 		if (const TArray<FTopLevelAssetPath>* Paths = PathMap.Find(InKey))
 		{
-			return Paths;
+			return *Paths;
 		}
 
 		return { };
@@ -546,9 +557,10 @@ namespace Metasound::Engine
 			}
 
 			const FAssetKey AssetKey(Class.Metadata);
-			if (const FTopLevelAssetPath* ObjectPath = FindAssetPath(AssetKey))
+			FTopLevelAssetPath ObjectPath = FindAssetPath(AssetKey);
+			if (ObjectPath.IsValid())
 			{
-				FAssetInfo AssetInfo { FNodeRegistryKey(Class.Metadata), FSoftObjectPath(*ObjectPath) };
+				FAssetInfo AssetInfo { FNodeRegistryKey(Class.Metadata), FSoftObjectPath(ObjectPath) };
 				OutAssetInfos.Add(MoveTemp(AssetInfo));
 			}
 			else
@@ -619,14 +631,14 @@ namespace Metasound::Engine
 
 		{
 			const FAssetKey OldAssetKey(ClassMetadata.GetClassName(), ClassMetadata.GetVersion());
-			AssetSubsystemPrivate::RemovePath(PathMap, OldAssetKey, Path);
+			AssetSubsystemPrivate::RemovePath(&PathMapCriticalSection, PathMap, OldAssetKey, Path);
 		}
 
 		Builder.GenerateNewClassName();
 
 		{
 			const FAssetKey NewAssetKey(ClassMetadata.GetClassName(), ClassMetadata.GetVersion());
-			AssetSubsystemPrivate::AddPath(PathMap, NewAssetKey, Path);
+			AssetSubsystemPrivate::AddPath(&PathMapCriticalSection, PathMap, NewAssetKey, Path);
 		}
 
 		AssetBase->UpdateAndRegisterForExecution();
@@ -719,7 +731,7 @@ namespace Metasound::Engine
 		}
 
 		const FAssetKey AssetKey(Metadata.GetClassName(), Metadata.GetVersion());
-		AssetSubsystemPrivate::RemovePath(PathMap, AssetKey, AssetPath);
+		AssetSubsystemPrivate::RemovePath(&PathMapCriticalSection, PathMap, AssetKey, AssetPath);
 	}
 
 	void FMetaSoundAssetManager::RemoveAsset(const FAssetData& InAssetData)
@@ -737,7 +749,7 @@ namespace Metasound::Engine
 			}
 
 			const FAssetKey AssetKey(ClassInfo.ClassName, ClassInfo.Version);
-			AssetSubsystemPrivate::RemovePath(PathMap, AssetKey, AssetPath);
+			AssetSubsystemPrivate::RemovePath(&PathMapCriticalSection, PathMap, AssetKey, AssetPath);
 		}
 	}
 
@@ -762,13 +774,13 @@ namespace Metasound::Engine
 		{
 			const FAssetKey AssetKey(ClassInfo.ClassName, ClassInfo.Version);
 			const FTopLevelAssetPath OldPath(InOldObjectPath);
-			AssetSubsystemPrivate::RemovePath(PathMap, AssetKey, OldPath);
+			AssetSubsystemPrivate::RemovePath(&PathMapCriticalSection, PathMap, AssetKey, OldPath);
 
 			if (ClassInfo.AssetClassID.IsValid())
 			{
 				if (AssetKey.IsValid())
 				{
-					AssetSubsystemPrivate::AddPath(PathMap, AssetKey, ClassInfo.AssetPath);
+					AssetSubsystemPrivate::AddPath(&PathMapCriticalSection, PathMap, AssetKey, ClassInfo.AssetPath);
 				}
 			}
 		}
@@ -928,9 +940,10 @@ namespace Metasound::Engine
 
 	FMetasoundAssetBase* FMetaSoundAssetManager::TryLoadAssetFromKey(const Metasound::Frontend::FAssetKey& InAssetKey) const
 	{
-		if (const FTopLevelAssetPath* ObjectPath = FindAssetPath(InAssetKey))
+		FTopLevelAssetPath ObjectPath = FindAssetPath(InAssetKey);
+		if (ObjectPath.IsValid())
 		{
-			const FSoftObjectPath SoftPath(*ObjectPath);
+			const FSoftObjectPath SoftPath(ObjectPath);
 			return TryLoadAsset(SoftPath);
 		}
 
@@ -1004,7 +1017,7 @@ namespace Metasound::Engine
 						FMetasoundFrontendRegistryContainer::Get()->UnregisterNode(RegistryKey);
 						const FTopLevelAssetPath AssetPath(AssetData.PackageName, AssetData.AssetName);
 						const FAssetKey AssetKey(AssetClassInfo.ClassName, AssetClassInfo.Version);
-						AssetSubsystemPrivate::RemovePath(PathMap, AssetKey, AssetPath);
+						AssetSubsystemPrivate::RemovePath(&PathMapCriticalSection, PathMap, AssetKey, AssetPath);
 					}
 				}
 			}
@@ -1176,9 +1189,10 @@ const FSoftObjectPath* UMetaSoundAssetSubsystem::FindObjectPathFromKey(const Met
 	using namespace Metasound::Frontend;
 	static FSoftObjectPath TempPath;
 	TempPath.Reset();
-	if (const FTopLevelAssetPath* Path = IMetaSoundAssetManager::GetChecked().FindAssetPath(InRegistryKey))
+	const FTopLevelAssetPath Path = IMetaSoundAssetManager::GetChecked().FindAssetPath(InRegistryKey);
+	if (Path.IsValid())
 	{
-		TempPath = FSoftObjectPath(*Path);
+		TempPath = FSoftObjectPath(Path);
 	}
 	return &TempPath;
 }
diff --git a/Engine/Plugins/Runtime/Metasound/Source/MetasoundEngine/Private/MetasoundEngineAsset.h b/Engine/Plugins/Runtime/Metasound/Source/MetasoundEngine/Private/MetasoundEngineAsset.h
--- a/Engine/Plugins/Runtime/Metasound/Source/MetasoundEngine/Private/MetasoundEngineAsset.h
+++ b/Engine/Plugins/Runtime/Metasound/Source/MetasoundEngine/Private/MetasoundEngineAsset.h
@@ -137,13 +137,13 @@ namespace Metasound::Engine
 
 			// Add error for multiple assets with the same class name
 			const FAssetKey Key(Document.RootGraph.Metadata);
-			const TArray<FTopLevelAssetPath>* AssetPaths = AssetManager.FindAssetPaths(Key);
-			if (AssetPaths && AssetPaths->Num() > 1)
+			const TArray<FTopLevelAssetPath> AssetPaths = AssetManager.FindAssetPaths(Key);
+			if (AssetPaths.Num() > 1)
 			{
 				Result = EDataValidationResult::Invalid;
 
 				TArray<FText> PathStrings;
-				Algo::Transform(*AssetPaths, PathStrings, [](const FTopLevelAssetPath& Path) { return FText::FromString(Path.ToString()); });
+				Algo::Transform(AssetPaths, PathStrings, [](const FTopLevelAssetPath& Path) { return FText::FromString(Path.ToString()); });
 				InOutContext.AddError(FText::Format(LOCTEXT("UniqueClassNameValidation",
 					"Multiple assets use the same class name which may result in unintended behavior. This may happen when an asset is moved, then the move is reverted in revision control without removing the newly created asset. Please remove the offending asset or duplicate it to automatically generate a new class name." \
 					"\nConflicting Asset Paths:\n{0}"), FText::Join(FText::FromString(TEXT("\n")), PathStrings)));
diff --git a/Engine/Plugins/Runtime/Metasound/Source/MetasoundEngineTest/Private/EngineTestMetaSoundAutomatedNodeTest.cpp b/Engine/Plugins/Runtime/Metasound/Source/MetasoundEngineTest/Private/EngineTestMetaSoundAutomatedNodeTest.cpp
--- a/Engine/Plugins/Runtime/Metasound/Source/MetasoundEngineTest/Private/EngineTestMetaSoundAutomatedNodeTest.cpp
+++ b/Engine/Plugins/Runtime/Metasound/Source/MetasoundEngineTest/Private/EngineTestMetaSoundAutomatedNodeTest.cpp
@@ -69,9 +69,10 @@ namespace Metasound::EngineTest{
 
 	FString GetPrettyName(const Frontend::FNodeRegistryKey& InRegistryKey)
 	{
-		if (const FTopLevelAssetPath* ObjectPath = Frontend::IMetaSoundAssetManager::GetChecked().FindAssetPath(InRegistryKey))
+		const FTopLevelAssetPath ObjectPath = Frontend::IMetaSoundAssetManager::GetChecked().FindAssetPath(InRegistryKey);
+		if (ObjectPath.IsValid())
 		{
-			return ObjectPath->ToString();
+			return ObjectPath.ToString();
 		}
 
 		FMetasoundFrontendRegistryContainer* NodeRegistry = FMetasoundFrontendRegistryContainer::Get();
diff --git a/Engine/Plugins/Runtime/Metasound/Source/MetasoundFrontend/Private/MetasoundFrontendDocumentBuilder.cpp b/Engine/Plugins/Runtime/Metasound/Source/MetasoundFrontend/Private/MetasoundFrontendDocumentBuilder.cpp
--- a/Engine/Plugins/Runtime/Metasound/Source/MetasoundFrontend/Private/MetasoundFrontendDocumentBuilder.cpp
+++ b/Engine/Plugins/Runtime/Metasound/Source/MetasoundFrontend/Private/MetasoundFrontendDocumentBuilder.cpp
@@ -1184,9 +1184,10 @@ const FMetasoundFrontendNode* FMetaSoundFrontendDocumentBuilder::AddGraphNode(co
 		using namespace Metasound::Frontend;
 
 		// Cache the asset name on the node if it node is reference to asset-defined graph.
-		if (const FTopLevelAssetPath* Path = IMetaSoundAssetManager::GetChecked().FindAssetPath(FAssetKey(ClassKey.ClassName, ClassKey.Version)))
+		const FTopLevelAssetPath Path = IMetaSoundAssetManager::GetChecked().FindAssetPath(FAssetKey(ClassKey.ClassName, ClassKey.Version));
+		if (Path.IsValid())
 		{
-			InOutNode.Name = Path->GetAssetName();
+			InOutNode.Name = Path.GetAssetName();
 			return;
 		}
 
diff --git a/Engine/Plugins/Runtime/Metasound/Source/MetasoundFrontend/Private/MetasoundFrontendGraphController.cpp b/Engine/Plugins/Runtime/Metasound/Source/MetasoundFrontend/Private/MetasoundFrontendGraphController.cpp
--- a/Engine/Plugins/Runtime/Metasound/Source/MetasoundFrontend/Private/MetasoundFrontendGraphController.cpp
+++ b/Engine/Plugins/Runtime/Metasound/Source/MetasoundFrontend/Private/MetasoundFrontendGraphController.cpp
@@ -1681,9 +1681,10 @@ namespace Metasound
 						if (NodeClass->Metadata.GetType() == EMetasoundFrontendClassType::External)
 						{
 							const FNodeRegistryKey RegistryKey = FNodeRegistryKey(NodeClass->Metadata);
-							if (const FTopLevelAssetPath* Path = AssetManager->FindAssetPath(RegistryKey))
+							const FTopLevelAssetPath Path = AssetManager->FindAssetPath(RegistryKey);
+							if (Path.IsValid())
 							{
-								Node.Name = Path->GetAssetName();
+								Node.Name = Path.GetAssetName();
 							}
 						}
 					}
diff --git a/Engine/Plugins/Runtime/Metasound/Source/MetasoundFrontend/Public/MetasoundAssetManager.h b/Engine/Plugins/Runtime/Metasound/Source/MetasoundFrontend/Public/MetasoundAssetManager.h
--- a/Engine/Plugins/Runtime/Metasound/Source/MetasoundFrontend/Public/MetasoundAssetManager.h
+++ b/Engine/Plugins/Runtime/Metasound/Source/MetasoundFrontend/Public/MetasoundAssetManager.h
@@ -172,10 +172,10 @@ namespace Metasound::Frontend
 
 		// Returns path associated with the given key (null if key is not registered with the AssetManager or was not loaded from asset)
 		// If multiple assets are associated with the given key, the last one is returned. 
-		virtual const FTopLevelAssetPath* FindAssetPath(const Metasound::Frontend::FAssetKey& InAssetKey) const = 0;
+		virtual FTopLevelAssetPath FindAssetPath(const Metasound::Frontend::FAssetKey& InAssetKey) const = 0;
 		
 		// Returns all paths associated with the given key (null if key is not registered with the AssetManager or was not loaded from asset)
-		virtual const TArray<FTopLevelAssetPath>* FindAssetPaths(const Metasound::Frontend::FAssetKey& InAssetKey) const = 0;
+		virtual TArray<FTopLevelAssetPath> FindAssetPaths(const Metasound::Frontend::FAssetKey& InAssetKey) const = 0;
 
 		// Converts an object to an AssetBase if its a registered asset
 		virtual FMetasoundAssetBase* GetAsAsset(UObject& InObject) const = 0;
