Minor comment clean-up
#rnx

diff --git a/Engine/Plugins/Runtime/Metasound/Source/MetasoundFrontend/Public/MetasoundAssetManager.h b/Engine/Plugins/Runtime/Metasound/Source/MetasoundFrontend/Public/MetasoundAssetManager.h
--- a/Engine/Plugins/Runtime/Metasound/Source/MetasoundFrontend/Public/MetasoundAssetManager.h
+++ b/Engine/Plugins/Runtime/Metasound/Source/MetasoundFrontend/Public/MetasoundAssetManager.h
@@ -163,18 +163,18 @@ namespace Metasound::Frontend
 		// Whether or not the asset manager has loaded the given asset
 		virtual bool ContainsKey(const Metasound::Frontend::FAssetKey& InAssetKey) const = 0;
 
-		// Returns object (if loaded) associated with the given key (null if key is not registered with the AssetManager)
+		// Returns object (if loaded) associated with the given key (null if key not registered with the AssetManager)
 		// If multiple assets are associated with the given key, the last one is returned. 
 		virtual FMetasoundAssetBase* FindAsset(const Metasound::Frontend::FAssetKey& InAssetKey) const = 0;
 
-		// Returns object (if loaded) associated with the given key as a Document Interface (null if key is not registered with the AssetManager)
+		// Returns object (if loaded) associated with the given key as a Document Interface (null if key not registered with the AssetManager)
 		virtual TScriptInterface<IMetaSoundDocumentInterface> FindAssetAsDocumentInterface(const Frontend::FAssetKey& InKey) const = 0;
 
-		// Returns path associated with the given key (null if key is not registered with the AssetManager or was not loaded from asset)
-		// If multiple assets are associated with the given key, the last one is returned. 
+		// Returns path associated with the given key (returns invalid asset path if key not registered with the AssetManager or was not loaded from asset)
+		// If multiple assets are associated with the given key, the last one is returned.
 		virtual FTopLevelAssetPath FindAssetPath(const Metasound::Frontend::FAssetKey& InAssetKey) const = 0;
 		
-		// Returns all paths associated with the given key (null if key is not registered with the AssetManager or was not loaded from asset)
+		// Returns all paths associated with the given key (returns empty array if key not registered with the AssetManager or was not loaded from asset)
 		virtual TArray<FTopLevelAssetPath> FindAssetPaths(const Metasound::Frontend::FAssetKey& InAssetKey) const = 0;
 
 		// Converts an object to an AssetBase if its a registered asset
@@ -201,7 +201,7 @@ namespace Metasound::Frontend
 		// Attempts to load an FMetasoundAssetBase from the given path, or returns it if its already loaded
 		virtual FMetasoundAssetBase* TryLoadAsset(const FSoftObjectPath& InObjectPath) const = 0;
 
-		// Returns asset associated with the given key (null if key is not registered with the AssetManager or was not loaded from asset)
+		// Returns asset associated with the given key (null if key not registered with the AssetManager or was not loaded from asset)
 		virtual FMetasoundAssetBase* TryLoadAssetFromKey(const Metasound::Frontend::FAssetKey& InAssetKey) const = 0;
 
 		// Try to load referenced assets of the given asset or return them if they are already loaded (non-recursive).
