PR #12031: Fix Regression for Stereo Delay when using metasounds generated from ?

diff --git a/Engine/Plugins/Runtime/Metasound/Source/MetasoundStandardNodes/Private/MetasoundStereoDelayNode.cpp b/Engine/Plugins/Runtime/Metasound/Source/MetasoundStandardNodes/Private/MetasoundStereoDelayNode.cpp
--- a/Engine/Plugins/Runtime/Metasound/Source/MetasoundStandardNodes/Private/MetasoundStereoDelayNode.cpp
+++ b/Engine/Plugins/Runtime/Metasound/Source/MetasoundStandardNodes/Private/MetasoundStereoDelayNode.cpp
@@ -152,11 +152,14 @@ namespace Metasound
 		, Feedback(InFeedback)
 		, LeftAudioOutput(FAudioBufferWriteRef::CreateNew(InSettings))
 		, RightAudioOutput(FAudioBufferWriteRef::CreateNew(InSettings))
-		, PrevDelayTimeMsec(GetInputDelayTimeMsecClamped())
+		, PrevDelayTimeMsec(0.0f)
 		, PrevDelayRatio(GetInputDelayRatioClamped())
 		, MaxDelayTimeSeconds(FMath::Clamp(InMaxDelayTimeSeconds, StereoDelay::MinDelaySeconds, StereoDelay::MaxDelaySeconds))
 		, TriggerReset(InTriggerReset)
 	{
+		// Depends on MaxDelayTimeSeconds being set first.
+		PrevDelayTimeMsec = GetInputDelayTimeMsecClamped();
+		
 		LeftDelayBuffer.Init(InSettings.GetSampleRate(), MaxDelayTimeSeconds);
 		LeftDelayBuffer.SetDelayMsec(PrevDelayTimeMsec * (1.0f + PrevDelayRatio));
 		RightDelayBuffer.Init(InSettings.GetSampleRate(), MaxDelayTimeSeconds);
