Fix-up Array drag/drop/copy/paste manipulators refactored/bugged from earlier page changes

#rnx

diff --git a/Engine/Plugins/Runtime/Metasound/Source/MetasoundEditor/Private/MetasoundDefaultLiteralCustomization.cpp b/Engine/Plugins/Runtime/Metasound/Source/MetasoundEditor/Private/MetasoundDefaultLiteralCustomization.cpp
--- a/Engine/Plugins/Runtime/Metasound/Source/MetasoundEditor/Private/MetasoundDefaultLiteralCustomization.cpp
+++ b/Engine/Plugins/Runtime/Metasound/Source/MetasoundEditor/Private/MetasoundDefaultLiteralCustomization.cpp
@@ -334,10 +334,10 @@ namespace Metasound::Editor
 
 	void FMetasoundDefaultLiteralCustomizationBase::CustomizeDefaults(UMetasoundEditorGraphMemberDefaultLiteral& InLiteral, IDetailLayoutBuilder& InDetailLayout)
 	{
-		CustomizePageDefaultRows(InLiteral, InDetailLayout, nullptr);
+		CustomizePageDefaultRows(InLiteral, InDetailLayout);
 	}
 
-	void FMetasoundDefaultLiteralCustomizationBase::CustomizePageDefaultRows(UMetasoundEditorGraphMemberDefaultLiteral& InLiteral, IDetailLayoutBuilder& InDetailLayout, FOnDefaultPageRowAdded* OnDefaultPageRowAdded)
+	void FMetasoundDefaultLiteralCustomizationBase::CustomizePageDefaultRows(UMetasoundEditorGraphMemberDefaultLiteral& InLiteral, IDetailLayoutBuilder& InDetailLayout)
 	{
 		DefaultProperties.Reset();
 
@@ -373,12 +373,7 @@ namespace Metasound::Editor
 
 				for (uint32 Index = 0; Index < NumElements; ++Index)
 				{
-					TSharedPtr<IPropertyHandle> ElementProperty = DefaultValueArray->GetElement(Index);
-					if (!ElementProperty.IsValid())
-					{
-						continue;
-					}
-
+					TSharedRef<IPropertyHandle> ElementProperty = DefaultValueArray->GetElement(Index);
 					TSharedPtr<IPropertyHandle> ValueProperty = ElementProperty->GetChildHandle("Value");
 					if (!ValueProperty.IsValid())
 					{
@@ -394,7 +389,7 @@ namespace Metasound::Editor
 
 						(*ValueRow.CustomNameWidget())
 						[
-							BuildPageDefaultNameWidget(InLiteral, ElementProperty.ToSharedRef())
+							BuildPageDefaultNameWidget(InLiteral, ElementProperty)
 						];
 
 						ValueRow.ShowPropertyButtons(false);
@@ -404,12 +399,6 @@ namespace Metasound::Editor
 						];
 					}
 					ValueRow.IsEnabled(GetEnabled());
-
-					// TODO: Fix half-implemented array multi-add broken by page implementation
-// 					if (bIsPagedDefault && bShowPageModifiers && OnDefaultPageRowAdded)
-// 					{
-// 						(*OnDefaultPageRowAdded)(ValueRow, ElementProperty->AsShared());
-// 					}
 				}
 			}
 		}
diff --git a/Engine/Plugins/Runtime/Metasound/Source/MetasoundEditor/Private/MetasoundEditorModule.cpp b/Engine/Plugins/Runtime/Metasound/Source/MetasoundEditor/Private/MetasoundEditorModule.cpp
--- a/Engine/Plugins/Runtime/Metasound/Source/MetasoundEditor/Private/MetasoundEditorModule.cpp
+++ b/Engine/Plugins/Runtime/Metasound/Source/MetasoundEditor/Private/MetasoundEditorModule.cpp
@@ -618,15 +618,15 @@ namespace Metasound
 					FOnGetDetailCustomizationInstance::CreateLambda([]() { return MakeShared<FMetasoundVariableDetailCustomization>(); }));
 
 				PropertyModule.RegisterCustomPropertyTypeLayout(
-					"MetaSoundEditorGraphMemberDefaultBoolRef",
+					"MetasoundEditorGraphMemberDefaultBoolRef",
 					FOnGetPropertyTypeCustomizationInstance::CreateLambda([]() { return MakeShared<FMetasoundMemberDefaultBoolDetailCustomization>(); }));
 
 				PropertyModule.RegisterCustomPropertyTypeLayout(
-					"MetaSoundEditorGraphMemberDefaultIntRef",
+					"MetasoundEditorGraphMemberDefaultIntRef",
 					FOnGetPropertyTypeCustomizationInstance::CreateLambda([]() { return MakeShared<FMetasoundMemberDefaultIntDetailCustomization>(); }));
 
 				PropertyModule.RegisterCustomPropertyTypeLayout(
-					"MetaSoundEditorGraphMemberDefaultObjectRef",
+					"MetasoundEditorGraphMemberDefaultObjectRef",
 					FOnGetPropertyTypeCustomizationInstance::CreateLambda([]() { return MakeShared<FMetasoundMemberDefaultObjectDetailCustomization>(); }));
 
 				LiteralCustomizationFactories.Add(UMetasoundEditorGraphMemberDefaultLiteral::StaticClass(), MakeUnique<FMetasoundDefaultLiteralCustomizationFactory>());
diff --git a/Engine/Plugins/Runtime/Metasound/Source/MetasoundEditor/Private/MetasoundNodeDetailCustomization.cpp b/Engine/Plugins/Runtime/Metasound/Source/MetasoundEditor/Private/MetasoundNodeDetailCustomization.cpp
--- a/Engine/Plugins/Runtime/Metasound/Source/MetasoundEditor/Private/MetasoundNodeDetailCustomization.cpp
+++ b/Engine/Plugins/Runtime/Metasound/Source/MetasoundEditor/Private/MetasoundNodeDetailCustomization.cpp
@@ -77,6 +77,11 @@ namespace Metasound
 			{
 				using namespace Frontend;
 
+				if (!ElementPropertyHandle.IsValid())
+				{
+					return false;
+				}
+
 				OutDataTypeInfo = { };
 				TArray<UObject*>OuterObjects;
 				ElementPropertyHandle->GetOuterObjects(OuterObjects);
@@ -118,9 +123,9 @@ namespace Metasound
 			}
 
 			// Paste execute action for object member default values
-			FExecuteAction CreateDefaultValueObjectPasteExecuteAction(TSharedPtr<IPropertyHandle> PropertyHandle, bool bIsArray)
+			FExecuteAction CreateDefaultValueObjectPasteExecuteAction(TSharedPtr<IPropertyHandle> PropertyHandle)
 			{
-				return FExecuteAction::CreateLambda([PropertyHandle = PropertyHandle, bIsArray = bIsArray]()
+				return FExecuteAction::CreateLambda([PropertyHandle]()
 				{
 					const FScopedTransaction Transaction(LOCTEXT("PasteObjectArrayProperty", "Paste Property"));
 
@@ -157,8 +162,9 @@ namespace Metasound
 						ClipboardValue.ParseIntoArrayWS(Values, TEXT(","), true);
 					}
 
-					if (Values.Num() > 0)
+					if (!Values.IsEmpty())
 					{
+						const bool bIsArray = PropertyHandle->AsArray().IsValid();
 						TStringBuilder<512> Builder;
 						if (bIsArray)
 						{
@@ -212,11 +218,11 @@ namespace Metasound
 			}
 
 			// Create copy/paste actions for member default value for object and object array types
-			void CreateDefaultValueObjectCopyPasteActions(FDetailWidgetRow& InWidgetRow, TSharedPtr<IPropertyHandle> PropertyHandle, bool bIsArray)
+			void CreateDefaultValueObjectCopyPasteActions(FDetailWidgetRow& InWidgetRow, TSharedPtr<IPropertyHandle> PropertyHandle)
 			{
 				// Copy action
 				FUIAction CopyAction;
-				CopyAction.ExecuteAction = FExecuteAction::CreateLambda([PropertyHandle = PropertyHandle]()
+				CopyAction.ExecuteAction = FExecuteAction::CreateLambda([PropertyHandle]()
 				{
 					FString Value;
 					if (PropertyHandle->GetValueAsFormattedString(Value, PPF_Copy) == FPropertyAccess::Success)
@@ -258,7 +264,7 @@ namespace Metasound
 					}
 				});
 
-				PasteAction.ExecuteAction = CreateDefaultValueObjectPasteExecuteAction(PropertyHandle, bIsArray);
+				PasteAction.ExecuteAction = CreateDefaultValueObjectPasteExecuteAction(PropertyHandle);
 
 				InWidgetRow.CopyAction(CopyAction);
 				InWidgetRow.PasteAction(PasteAction);
@@ -438,96 +444,7 @@ namespace Metasound
 		{
 			check(DefaultCategoryBuilder);
 
-			FOnDefaultPageRowAdded OnDefaultPageRowAdded = [this, &InDetailLayout](IDetailPropertyRow& ValueRow, TSharedRef<IPropertyHandle> PageDefaultProperty)
-			{
-				TSharedPtr<IPropertyHandle> ValueProperty = PageDefaultProperty->GetChildHandle(GET_MEMBER_NAME_CHECKED(FMetasoundEditorMemberPageDefaultObjectArray, Value));
-				if (ValueProperty.IsValid())
-				{
-					TSharedPtr<IPropertyHandleArray> ArrayProperty = ValueProperty->AsArray();
-					if (ArrayProperty.IsValid())
-					{
-						uint32 NumElements = 0;
-						ArrayProperty->GetNumElements(NumElements);
-						for (uint32 Index = 0; Index < NumElements; ++Index)
-						{
-							TSharedPtr<IPropertyHandle> ElementProperty = ArrayProperty->GetElement(Index);
-							if (ElementProperty.IsValid())
-							{
-								TSharedPtr<IPropertyHandle> DefaultValueProperty = ElementProperty->GetChildHandle(GET_MEMBER_NAME_CHECKED(FMetasoundEditorGraphMemberDefaultObjectRef, Object));
-
-								IDetailPropertyRow& Row = DefaultCategoryBuilder->AddProperty(ElementProperty);
-								constexpr bool bShowChildren = false;
-								Row.ShowPropertyButtons(false)
-								.CustomWidget(bShowChildren);
-								(*Row.CustomNameWidget())
-								[
-									ElementProperty->CreatePropertyNameWidget()
-								];
-								(*Row.CustomValueWidget())
-								[
-									SNew(SAssetDropTarget)
-									.bSupportsMultiDrop(true)
-									.OnAreAssetsAcceptableForDropWithReason_Lambda([this, DefaultValueProperty](TArrayView<FAssetData> InAssets, FText& OutReason)
-									{
-										Frontend::FDataTypeRegistryInfo DataTypeInfo;
-										const bool bMemberFound = MemberCustomizationPrivate::GetDataTypeFromElementPropertyHandle(DefaultValueProperty, DataTypeInfo);
-										bool bCanDrop = bMemberFound;
-										if (UClass* ProxyGenClass = DataTypeInfo.ProxyGeneratorClass; bCanDrop && bMemberFound)
-										{
-											bCanDrop = true;
-											for (const FAssetData& AssetData : InAssets)
-											{
-												if (UClass* Class = AssetData.GetClass())
-												{
-PRAGMA_DISABLE_DEPRECATION_WARNINGS
-													const IMetasoundEditorModule& EditorModule = FModuleManager::GetModuleChecked<IMetasoundEditorModule>("MetaSoundEditor");
-													if (EditorModule.IsExplicitProxyClass(*DataTypeInfo.ProxyGeneratorClass))
-													{
-														bCanDrop &= Class == DataTypeInfo.ProxyGeneratorClass;
-														continue;
-													}
-PRAGMA_ENABLE_DEPRECATION_WARNINGS
-													if (DataTypeInfo.bIsExplicit)
-													{
-														bCanDrop &= Class == DataTypeInfo.ProxyGeneratorClass;
-													}
-													else
-													{
-														bCanDrop &= Class->IsChildOf(DataTypeInfo.ProxyGeneratorClass);
-													}
-												}
-											}
-										}
-
-										return bCanDrop;
-									})
-									.OnAssetsDropped_Lambda([this, ArrayProperty](const FDragDropEvent& DragDropEvent, TArrayView<FAssetData> InAssets)
-									{
-										if (ArrayProperty.IsValid())
-										{
-											FScopedTransaction Transaction(LOCTEXT("DragDropInputAssets", "Drop Asset(s) on MetaSound Input"));
-											for (const FAssetData& AssetData : InAssets)
-											{
-												uint32 AddIndex = INDEX_NONE;
-												ArrayProperty->GetNumElements(AddIndex);
-												ArrayProperty->AddItem();
-												TSharedPtr<IPropertyHandle> ElementHandle = ArrayProperty->GetElement(static_cast<int32>(AddIndex));
-												TSharedPtr<IPropertyHandle> ObjectHandle = ElementHandle->GetChildHandle(GET_MEMBER_NAME_CHECKED(FMetasoundEditorGraphMemberDefaultObjectRef, Object));
-												ObjectHandle->SetValue(AssetData.GetAsset());
-											}
-										}
-									})
-									[
-										DefaultValueProperty->CreatePropertyValueWidget()
-									]
-								];
-								Row.Visibility(GetDefaultVisibility());
-							}
-						}
-					}
-				}
-			};
-			FMetasoundDefaultLiteralCustomizationBase::CustomizePageDefaultRows(InLiteral, InDetailLayout, &OnDefaultPageRowAdded);
+			FMetasoundDefaultLiteralCustomizationBase::CustomizePageDefaultRows(InLiteral, InDetailLayout);
 		}
 
 		FText FMetasoundMemberDefaultBoolDetailCustomization::GetPropertyNameOverride() const
@@ -748,123 +665,80 @@ PRAGMA_ENABLE_DEPRECATION_WARNINGS
 
 		TSharedRef<SWidget> FMetasoundDefaultMemberElementDetailCustomizationBase::CreateValueWidget(TSharedPtr<IPropertyHandleArray> ParentPropertyHandleArray, TSharedPtr<IPropertyHandle> StructPropertyHandle) const
 		{
-			TSharedRef<SWidget> ValueWidget = CreateStructureWidget(StructPropertyHandle);
-			if (!ParentPropertyHandleArray.IsValid())
-			{
-				return ValueWidget;
-			}
-
-			TSharedPtr<IPropertyHandle> StructPropertyPtr = StructPropertyHandle;
-			FExecuteAction InsertAction = FExecuteAction::CreateLambda([ParentPropertyHandleArray, StructPropertyPtr]
-			{
-				const int32 ArrayIndex = StructPropertyPtr.IsValid() ? StructPropertyPtr->GetIndexInArray() : INDEX_NONE;
-				if (ParentPropertyHandleArray.IsValid() && ArrayIndex >= 0)
-				{
-					ParentPropertyHandleArray->Insert(ArrayIndex);
-				}
-			});
-
-			FExecuteAction DeleteAction = FExecuteAction::CreateLambda([ParentPropertyHandleArray, StructPropertyPtr]
-			{
-				const int32 ArrayIndex = StructPropertyPtr.IsValid() ? StructPropertyPtr->GetIndexInArray() : INDEX_NONE;
-				if (ParentPropertyHandleArray.IsValid() && ArrayIndex >= 0)
-				{
-					ParentPropertyHandleArray->DeleteItem(ArrayIndex);
-				}
-			});
-
-			FExecuteAction DuplicateAction = FExecuteAction::CreateLambda([ParentPropertyHandleArray, StructPropertyPtr]
-			{
-				const int32 ArrayIndex = StructPropertyPtr.IsValid() ? StructPropertyPtr->GetIndexInArray() : INDEX_NONE;
-				if (ParentPropertyHandleArray.IsValid() && ArrayIndex >= 0)
-				{
-					ParentPropertyHandleArray->DuplicateItem(ArrayIndex);
-				}
-			});
-
-			return SNew(SHorizontalBox)
-				+ SHorizontalBox::Slot()
-				.FillWidth(1.0f)
-				.Padding(1.0f, 0.0f, 0.0f, 0.0f)
-				.VAlign(VAlign_Center)
-				[
-					ValueWidget
-				]
-				+ SHorizontalBox::Slot()
-				.AutoWidth()
-				.Padding(-6.0f, 0.0f, 0.0f, 0.0f) // Negative padding intentional on the left to bring the dropdown closer to the other buttons
-				.VAlign(VAlign_Center)
-				[
-					PropertyCustomizationHelpers::MakeInsertDeleteDuplicateButton(InsertAction, DeleteAction, DuplicateAction)
-				];
+			return CreateStructureWidget(StructPropertyHandle);
 		}
 
 		void FMetasoundDefaultMemberElementDetailCustomizationBase::CustomizeChildren(TSharedRef<IPropertyHandle> StructPropertyHandle, IDetailChildrenBuilder& ChildBuilder, IPropertyTypeCustomizationUtils& StructCustomizationUtils)
 		{
 			TSharedPtr<IPropertyHandleArray> ParentPropertyHandleArray;
-			TSharedPtr<IPropertyHandle> ElementPropertyHandle = StructPropertyHandle;
-			if (ElementPropertyHandle.IsValid())
+			TSharedRef<IPropertyHandle> ElementPropertyHandle = StructPropertyHandle;
 			{
-				TSharedPtr<IPropertyHandle> ParentProperty = ElementPropertyHandle->GetParentHandle();
+				TSharedPtr<IPropertyHandle> ParentProperty = StructPropertyHandle->GetParentHandle();
 				if (ParentProperty.IsValid() && ParentProperty->GetProperty() != nullptr)
 				{
 					ParentPropertyHandleArray = ParentProperty->AsArray();
 					if (ParentPropertyHandleArray.IsValid())
 					{
-						ElementPropertyHandle = ParentProperty;
+						ElementPropertyHandle = ParentProperty.ToSharedRef();
 					}
 				}
 			}
 
 			const bool bMemberFound = MemberCustomizationPrivate::GetDataTypeFromElementPropertyHandle(ElementPropertyHandle, DataTypeInfo);
+			ensureAlways(bMemberFound);
+
+			IDetailPropertyRow& ValueRow = ChildBuilder.AddProperty(StructPropertyHandle);
+			TSharedPtr<SWidget> NameWidget;
+			TSharedPtr<SWidget> ValueWidget;
+			ValueRow.GetDefaultWidgets(NameWidget, ValueWidget);
 
-			TSharedRef<SWidget> ValueWidget = CreateValueWidget(ParentPropertyHandleArray, StructPropertyHandle);
-			FDetailWidgetRow& ValueRow = ChildBuilder.AddCustomRow(MemberCustomizationStyle::DefaultPropertyText);
+			constexpr bool bShowChildren = false;
+			ValueRow.CustomWidget(bShowChildren);
 			if (ParentPropertyHandleArray.IsValid())
 			{
-				ValueRow.NameContent()
+				(*ValueRow.CustomNameWidget())
 				[
 					StructPropertyHandle->CreatePropertyNameWidget()
 				];
 			}
 			else
 			{
-				ValueRow.NameContent()
+				(*ValueRow.CustomNameWidget())
 				[
 					CreateNameWidget(StructPropertyHandle)
 				];
 			}
 
-			TArray<UObject*> OuterObjects;
-			StructPropertyHandle->GetOuterObjects(OuterObjects);
-			TArray<TWeakObjectPtr<UMetasoundEditorGraphInput>> Inputs;
-			for (UObject* Object : OuterObjects)
 			{
-				if (UMetasoundEditorGraphInput* Input = Cast<UMetasoundEditorGraphInput>(Object))
+				TArray<UObject*> OuterObjects;
+				StructPropertyHandle->GetOuterObjects(OuterObjects);
+				TArray<TWeakObjectPtr<UMetasoundEditorGraphInput>> Inputs;
+				for (UObject* Object : OuterObjects)
 				{
-					Inputs.Add(Input);
+					if (UMetasoundEditorGraphInput* Input = Cast<UMetasoundEditorGraphInput>(Object))
+					{
+						Inputs.Add(Input);
+					}
 				}
-			}
 
-			FSimpleDelegate UpdateFrontendDefaultLiteral = FSimpleDelegate::CreateLambda([InInputs = Inputs]()
-			{
-				for (const TWeakObjectPtr<UMetasoundEditorGraphInput>& GraphInput : InInputs)
+				FSimpleDelegate UpdateFrontendDefaultLiteral = FSimpleDelegate::CreateLambda([InInputs = MoveTemp(Inputs)]()
 				{
-					if (GraphInput.IsValid())
+					for (const TWeakObjectPtr<UMetasoundEditorGraphInput>& GraphInput : InInputs)
 					{
-						constexpr bool bPostTransaction = true;
-						GraphInput->UpdateFrontendDefaultLiteral(bPostTransaction);
+						if (GraphInput.IsValid())
+						{
+							constexpr bool bPostTransaction = true;
+							GraphInput->UpdateFrontendDefaultLiteral(bPostTransaction);
+						}
 					}
-				}
-			});
-			StructPropertyHandle->SetOnChildPropertyValueChanged(UpdateFrontendDefaultLiteral);
+				});
+				StructPropertyHandle->SetOnChildPropertyValueChanged(UpdateFrontendDefaultLiteral);
+			}
 
-			ValueRow.ValueContent()
+			(*ValueRow.CustomValueWidget())
 			[
-				ValueWidget
+				CreateValueWidget(ParentPropertyHandleArray, StructPropertyHandle)
 			];
-
-			MemberCustomizationPrivate::CreateDefaultValueObjectCopyPasteActions(ValueRow, StructPropertyHandle, /*bIsArray=*/false);
 		}
 
 		void FMetasoundDefaultMemberElementDetailCustomizationBase::CustomizeHeader(TSharedRef<IPropertyHandle> StructPropertyHandle, FDetailWidgetRow& HeaderRow, IPropertyTypeCustomizationUtils& StructCustomizationUtils)
@@ -1817,7 +1691,7 @@ PRAGMA_ENABLE_DEPRECATION_WARNINGS
 
 		void FMetaSoundNodeExtensionHandler::ExtendWidgetRow(FDetailWidgetRow& InWidgetRow, const IDetailLayoutBuilder& InDetailBuilder, const UClass* InObjectClass, TSharedPtr<IPropertyHandle> PropertyHandle)
 		{
-			MemberCustomizationPrivate::CreateDefaultValueObjectCopyPasteActions(InWidgetRow, PropertyHandle, /*bIsArray=*/true);
+			MemberCustomizationPrivate::CreateDefaultValueObjectCopyPasteActions(InWidgetRow, PropertyHandle);
 		}
 	} // namespace Editor
 } // namespace Metasound
diff --git a/Engine/Plugins/Runtime/Metasound/Source/MetasoundEditor/Public/MetasoundDefaultLiteralCustomization.h b/Engine/Plugins/Runtime/Metasound/Source/MetasoundEditor/Public/MetasoundDefaultLiteralCustomization.h
--- a/Engine/Plugins/Runtime/Metasound/Source/MetasoundEditor/Public/MetasoundDefaultLiteralCustomization.h
+++ b/Engine/Plugins/Runtime/Metasound/Source/MetasoundEditor/Public/MetasoundDefaultLiteralCustomization.h
@@ -46,7 +46,7 @@ namespace Metasound::Editor
 		virtual void SetResetOverride(const TOptional<FResetToDefaultOverride>& InResetOverride);
 
 	protected:
-		void CustomizePageDefaultRows(UMetasoundEditorGraphMemberDefaultLiteral& InLiteral, IDetailLayoutBuilder& InDetailLayout, FOnDefaultPageRowAdded* OnPageRowAdded = nullptr);
+		void CustomizePageDefaultRows(UMetasoundEditorGraphMemberDefaultLiteral& InLiteral, IDetailLayoutBuilder& InDetailLayout);
 
 		TArray<TSharedPtr<IPropertyHandle>> DefaultProperties;
 
diff --git a/Engine/Plugins/Runtime/Metasound/Source/MetasoundEditor/Public/MetasoundEditorGraphMemberDefaults.h b/Engine/Plugins/Runtime/Metasound/Source/MetasoundEditor/Public/MetasoundEditorGraphMemberDefaults.h
--- a/Engine/Plugins/Runtime/Metasound/Source/MetasoundEditor/Public/MetasoundEditorGraphMemberDefaults.h
+++ b/Engine/Plugins/Runtime/Metasound/Source/MetasoundEditor/Public/MetasoundEditorGraphMemberDefaults.h
@@ -554,7 +554,7 @@ struct FMetasoundEditorMemberPageDefaultObjectRef : public FMetasoundEditorMembe
 	{
 	}
 
-	UPROPERTY(EditAnywhere, Category = DefaultValue, meta = (ShowOnlyInnerProperties))
+	UPROPERTY(EditAnywhere, Category = DefaultValue)
 	FMetasoundEditorGraphMemberDefaultObjectRef Value;
 };
 
