Fix for crash when swapping MetaSound interfaces on asset that implements multiple pages

#tests swapping output format on paged MetaSound asset
#rnx

diff --git a/Engine/Plugins/Runtime/Metasound/Source/MetasoundFrontend/Private/MetasoundFrontendDocumentBuilder.cpp b/Engine/Plugins/Runtime/Metasound/Source/MetasoundFrontend/Private/MetasoundFrontendDocumentBuilder.cpp
--- a/Engine/Plugins/Runtime/Metasound/Source/MetasoundFrontend/Private/MetasoundFrontendDocumentBuilder.cpp
+++ b/Engine/Plugins/Runtime/Metasound/Source/MetasoundFrontend/Private/MetasoundFrontendDocumentBuilder.cpp
@@ -737,7 +737,7 @@ void FMetaSoundFrontendDocumentBuilder::AddEdge(FMetasoundFrontendEdge&& InNewEd
 #if DO_CHECK
 	{
 		const IDocumentGraphEdgeCache& EdgeCache = DocumentCache->GetEdgeCache(PageID);
-		checkf(!EdgeCache.IsNodeInputConnected(InNewEdge.ToNodeID, InNewEdge.ToVertexID), TEXT("Failed to add edge in MetaSound Builder: Destination input a lready connected"));
+		checkf(!EdgeCache.IsNodeInputConnected(InNewEdge.ToNodeID, InNewEdge.ToVertexID), TEXT("Failed to add edge in MetaSound Builder: Destination input already connected"));
 
 		const EInvalidEdgeReason Reason = IsValidEdge(InNewEdge, &PageID);
 		checkf(Reason == Metasound::Frontend::EInvalidEdgeReason::None, TEXT("Attempted call to AddEdge in MetaSound Builder where edge is invalid: %s."), *LexToString(Reason));
@@ -2558,13 +2558,13 @@ void FMetaSoundFrontendDocumentBuilder::InitNodeLocations()
 	using namespace Metasound;
 	using namespace Metasound::Frontend;
 
-	FVector2D InputNodeLocation = FVector2D::ZeroVector;
-	FVector2D ExternalNodeLocation = InputNodeLocation + DisplayStyle::NodeLayout::DefaultOffsetX;
-	FVector2D OutputNodeLocation = ExternalNodeLocation + DisplayStyle::NodeLayout::DefaultOffsetX;
-
 	FMetasoundFrontendDocument& Document = GetDocumentChecked();
 	Document.RootGraph.IterateGraphPages([&](FMetasoundFrontendGraph& Graph)
 	{
+		FVector2D InputNodeLocation = FVector2D::ZeroVector;
+		FVector2D ExternalNodeLocation = InputNodeLocation + DisplayStyle::NodeLayout::DefaultOffsetX;
+		FVector2D OutputNodeLocation = ExternalNodeLocation + DisplayStyle::NodeLayout::DefaultOffsetX;
+
 		TArray<FMetasoundFrontendNode>& Nodes = Graph.Nodes;
 		for (FMetasoundFrontendNode& Node : Nodes)
 		{
@@ -4259,7 +4259,7 @@ bool FMetaSoundFrontendDocumentBuilder::SwapGraphInput(const FMetasoundFrontendC
 #endif // WITH_EDITOR
 
 	// 2. Gather data from existing member/node needed to swap
-	TArray<FMetasoundFrontendEdge> RemovedEdges;
+	TMultiMap<FGuid, FMetasoundFrontendEdge> RemovedEdgesPerPage;
 
 	const FMetasoundFrontendClassInput* ExistingInputClass = InterfaceCache.FindInput(InExistingInputVertex.Name);
 	checkf(ExistingInputClass, TEXT("'SwapGraphInput' failed to find original graph input"));
@@ -4279,7 +4279,10 @@ bool FMetaSoundFrontendDocumentBuilder::SwapGraphInput(const FMetasoundFrontendC
 
 		const FGuid VertexID = ExistingInputNode->Interface.Outputs.Last().VertexID;
 		TArray<const FMetasoundFrontendEdge*> Edges = DocumentCache->GetEdgeCache(Graph.PageID).FindEdges(NodeID, VertexID);
-		Algo::Transform(Edges, RemovedEdges, [](const FMetasoundFrontendEdge* Edge) { return *Edge; });
+		Algo::Transform(Edges, RemovedEdgesPerPage, [PageID = Graph.PageID](const FMetasoundFrontendEdge* Edge)
+		{
+			return TPair<FGuid, FMetasoundFrontendEdge>(PageID, *Edge);
+		});
 	});
 
 	// 3. Remove existing graph vertex
@@ -4314,12 +4317,14 @@ bool FMetaSoundFrontendDocumentBuilder::SwapGraphInput(const FMetasoundFrontendC
 #endif // WITH_EDITOR
 
 		// 5b. Add to new copy existing node edges
+		TArray<FMetasoundFrontendEdge> RemovedEdges;
+		RemovedEdgesPerPage.MultiFind(Graph.PageID, RemovedEdges);
 		for (const FMetasoundFrontendEdge& RemovedEdge : RemovedEdges)
 		{
 			FMetasoundFrontendEdge NewEdge = RemovedEdge;
 			NewEdge.FromNodeID = NewInputNode->GetID();
 			NewEdge.FromVertexID = NewInputNode->Interface.Outputs.Last().VertexID;
-			AddEdge(MoveTemp(NewEdge));
+			AddEdge(MoveTemp(NewEdge), &Graph.PageID);
 		}
 	});
 
@@ -4346,8 +4351,8 @@ bool FMetaSoundFrontendDocumentBuilder::SwapGraphOutput(const FMetasoundFrontend
 	TMap<FGuid, FPageNodeLocations> PageNodeLocations;
 #endif // WITH_EDITOR
 
-	// 2. Gather data from existing member/node needed to swap
-	TArray<FMetasoundFrontendEdge> RemovedEdges;
+	// 2. Gather data from existing page member/node needed to swap
+	TMultiMap<FGuid, FMetasoundFrontendEdge> RemovedEdgesPerPage;
 
 	const FMetasoundFrontendClassOutput* ExistingOutputClass = InterfaceCache.FindOutput(InExistingOutputVertex.Name);
 	checkf(ExistingOutputClass, TEXT("'SwapGraphOutput' failed to find original graph output"));
@@ -4367,7 +4372,10 @@ bool FMetaSoundFrontendDocumentBuilder::SwapGraphOutput(const FMetasoundFrontend
 
 		const FGuid VertexID = ExistingOutputNode->Interface.Inputs.Last().VertexID;
 		TArray<const FMetasoundFrontendEdge*> Edges = DocumentCache->GetEdgeCache(Graph.PageID).FindEdges(NodeID, VertexID);
-		Algo::Transform(Edges, RemovedEdges, [](const FMetasoundFrontendEdge* Edge) { return *Edge; });
+		Algo::Transform(Edges, RemovedEdgesPerPage, [PageID = Graph.PageID](const FMetasoundFrontendEdge* Edge)
+		{
+			return TPair<FGuid, FMetasoundFrontendEdge>(PageID, *Edge);
+		});
 	});
 
 	// 3. Remove existing graph vertex
@@ -4402,12 +4410,14 @@ bool FMetaSoundFrontendDocumentBuilder::SwapGraphOutput(const FMetasoundFrontend
 #endif // WITH_EDITOR
 
 		// 5b. Add to new copy existing node edges
+		TArray<FMetasoundFrontendEdge> RemovedEdges;
+		RemovedEdgesPerPage.MultiFind(Graph.PageID, RemovedEdges);
 		for (const FMetasoundFrontendEdge& RemovedEdge : RemovedEdges)
 		{
 			FMetasoundFrontendEdge NewEdge = RemovedEdge;
 			NewEdge.ToNodeID = NewOutputNode->GetID();
 			NewEdge.ToVertexID = NewOutputNode->Interface.Inputs.Last().VertexID;
-			AddEdge(MoveTemp(NewEdge));
+			AddEdge(MoveTemp(NewEdge), &Graph.PageID);
 		}
 	});
 
