Fix for calling CreateMemberMetadata when handing a template object to BindMemberMetadata

#rnx

diff --git a/Engine/Plugins/Runtime/Metasound/Source/MetasoundEditor/Private/MetasoundEditorSubsystem.cpp b/Engine/Plugins/Runtime/Metasound/Source/MetasoundEditor/Private/MetasoundEditorSubsystem.cpp
--- a/Engine/Plugins/Runtime/Metasound/Source/MetasoundEditor/Private/MetasoundEditorSubsystem.cpp
+++ b/Engine/Plugins/Runtime/Metasound/Source/MetasoundEditor/Private/MetasoundEditorSubsystem.cpp
@@ -173,9 +173,9 @@ bool UMetaSoundEditorSubsystem::BindMemberMetadata(
 			InMember.Literal = CastChecked<UMetasoundEditorGraphMemberDefaultLiteral>(Literal);
 			return false;
 		}
-	}
 
-	NewLiteral = CreateMemberMetadata(Builder, InMember.GetMemberName(), LiteralClass);
+		NewLiteral = CreateMemberMetadata(Builder, InMember.GetMemberName(), LiteralClass);
+	}
 
 	if (NewLiteral)
 	{
