Minor: Add ability to hold shift as modifier to see debug node data in MetaSound Asset Editor. Allows for faster development for quick look-up while mocking builder behavior without having to jockey cvar.

#rnx

diff --git a/Engine/Plugins/Runtime/Metasound/Source/MetasoundEditor/Private/MetasoundEditorGraphInputNode.cpp b/Engine/Plugins/Runtime/Metasound/Source/MetasoundEditor/Private/MetasoundEditorGraphInputNode.cpp
--- a/Engine/Plugins/Runtime/Metasound/Source/MetasoundEditor/Private/MetasoundEditorGraphInputNode.cpp
+++ b/Engine/Plugins/Runtime/Metasound/Source/MetasoundEditor/Private/MetasoundEditorGraphInputNode.cpp
@@ -144,7 +144,7 @@ void UMetasoundEditorGraphInputNode::GetPinHoverText(const UEdGraphPin& Pin, FSt
 			{
 				if (const FMetasoundFrontendClassVertex* Vertex = Input->GetFrontendClassVertex())
 				{
-					OutHoverText = FString::Format(TEXT("{0}\nVertex Name: {1}\nDataType: {2}\nID: {3}"),
+					OutHoverText = FString::Format(TEXT("Description: {0}\nVertex Name: {1}\nDataType: {2}\nID: {3}"),
 					{
 						OutHoverText,
 						Vertex->Name.ToString(),
diff --git a/Engine/Plugins/Runtime/Metasound/Source/MetasoundEditor/Private/MetasoundEditorGraphNode.cpp b/Engine/Plugins/Runtime/Metasound/Source/MetasoundEditor/Private/MetasoundEditorGraphNode.cpp
--- a/Engine/Plugins/Runtime/Metasound/Source/MetasoundEditor/Private/MetasoundEditorGraphNode.cpp
+++ b/Engine/Plugins/Runtime/Metasound/Source/MetasoundEditor/Private/MetasoundEditorGraphNode.cpp
@@ -5,7 +5,9 @@
 #include "EdGraph/EdGraphPin.h"
 #include "Editor/EditorEngine.h"
 #include "Engine/Font.h"
+#include "Framework/Application/SlateApplication.h"
 #include "Framework/Commands/GenericCommands.h"
+#include "GenericPlatform/GenericApplication.h"
 #include "GraphEditorActions.h"
 #include "HAL/IConsoleManager.h"
 #include "Logging/TokenizedMessage.h"
@@ -78,7 +80,7 @@ bool UMetasoundEditorGraphNode::ShowNodeDebugData()
 {
 	using namespace Metasound::Editor;
 
-	return GraphNodePrivate::ShowNodeDebugData != 0;
+	return GraphNodePrivate::ShowNodeDebugData != 0 || FSlateApplication::Get().GetModifierKeys().AreModifersDown(EModifierKey::Shift);
 }
 
 bool UMetasoundEditorGraphNode::RemoveFromDocument() const
@@ -481,9 +483,9 @@ void UMetasoundEditorGraphNode::GetPinHoverText(const UEdGraphPin& Pin, FString&
 			OutHoverText = InputHandle->GetTooltip().ToString();
 		}
 
-		if (GraphNodePrivate::ShowNodeDebugData)
+		if (ShowNodeDebugData())
 		{
-			OutHoverText = FString::Format(TEXT("{0}\nVertex Name: {1}\nDataType: {2}\nID: {3}"),
+			OutHoverText = FString::Format(TEXT("Description: {0}\nVertex Name: {1}\nDataType: {2}\nID: {3}"),
 			{
 				OutHoverText,
 				InputHandle->GetName().ToString(),
@@ -496,9 +498,9 @@ void UMetasoundEditorGraphNode::GetPinHoverText(const UEdGraphPin& Pin, FString&
 	{
 		FConstOutputHandle OutputHandle = FGraphBuilder::FindReroutedConstOutputHandleFromPin(&Pin);
 		OutHoverText = OutputHandle->GetTooltip().ToString();
-		if (GraphNodePrivate::ShowNodeDebugData)
+		if (ShowNodeDebugData())
 		{
-			OutHoverText = FString::Format(TEXT("{0}\nVertex Name: {1}\nDataType: {2}\nID: {3}"),
+			OutHoverText = FString::Format(TEXT("Description: {0}\nVertex Name: {1}\nDataType: {2}\nID: {3}"),
 			{
 				OutHoverText,
 				OutputHandle->GetName().ToString(),
@@ -741,9 +743,10 @@ FText UMetasoundEditorGraphNode::GetTooltipText() const
 
 	FConstNodeHandle Node = GetConstNodeHandle();
 	FText Description = Node->GetDescription();
-	if (GraphNodePrivate::ShowNodeDebugData)
+
+	if (ShowNodeDebugData())
 	{
-		Description = FText::Format(LOCTEXT("Metasound_DebugNodeTooltipText", "{0}\nClass Name: {1}\nNode ID: {2}"),
+		Description = FText::Format(LOCTEXT("Metasound_DebugNodeTooltipText", "Description: {0}\nClass Name: {1}\nNode ID: {2}"),
 			Description,
 			FText::FromString(Node->GetClassMetadata().GetClassName().ToString()),
 			FText::FromString(Node->GetID().ToString())
