Remove specialized pin literal synchronization no longer necessary in UMetasoundEditorGraphNode::PostEditUndo that was causing crash due to the builder cache not being reloaded yet in this position of transaction rollback (and query looking up invalid index)

#rnx

diff --git a/Engine/Plugins/Runtime/Metasound/Source/MetasoundEditor/Private/MetasoundEditorGraphNode.cpp b/Engine/Plugins/Runtime/Metasound/Source/MetasoundEditor/Private/MetasoundEditorGraphNode.cpp
--- a/Engine/Plugins/Runtime/Metasound/Source/MetasoundEditor/Private/MetasoundEditorGraphNode.cpp
+++ b/Engine/Plugins/Runtime/Metasound/Source/MetasoundEditor/Private/MetasoundEditorGraphNode.cpp
@@ -677,19 +677,6 @@ void UMetasoundEditorGraphNode::PostEditUndo()
 	using namespace Metasound::Frontend;
 
 	UEdGraphPin::ResolveAllPinReferences();
-
-	// Test for frontend node existence safely as it can no longer be so after rolling transactions back or forward.
-	const FMetaSoundFrontendDocumentBuilder& Builder = GetBuilderChecked().GetConstBuilder();
-	if (const FMetasoundFrontendNode* FrontendNode = Builder.FindNode(GetNodeID()))
-	{
-		for (UEdGraphPin* Pin : Pins)
-		{
-			if (Pin && Pin->Direction == EGPD_Input)
-			{
-				FGraphBuilder::SynchronizePinLiteral(Builder, *Pin);
-			}
-		}
-	}
 }
 
 void UMetasoundEditorGraphNode::PostDuplicate(bool bDuplicateForPIE)
