Comment clarification, & only show RemoveUnusedDependencies as an advanced option when removing node
#rnx

diff --git a/Engine/Plugins/Runtime/Metasound/Source/MetasoundEngine/Public/MetasoundBuilderBase.h b/Engine/Plugins/Runtime/Metasound/Source/MetasoundEngine/Public/MetasoundBuilderBase.h
--- a/Engine/Plugins/Runtime/Metasound/Source/MetasoundEngine/Public/MetasoundBuilderBase.h
+++ b/Engine/Plugins/Runtime/Metasound/Source/MetasoundEngine/Public/MetasoundBuilderBase.h
@@ -362,8 +362,9 @@ public:
 	UFUNCTION(BlueprintCallable, Category = "Audio|MetaSound|Builder", meta = (ExpandEnumAsExecs = "OutResult"))
 	void RemoveInterface(FName InterfaceName, EMetaSoundBuilderResult& OutResult);
 
-	// Removes node and any associated connections from the builder's MetaSound.
-	UFUNCTION(BlueprintCallable, Category = "Audio|MetaSound|Builder", meta = (ExpandEnumAsExecs = "OutResult"))
+	// Removes node and any associated connections from the builder's MetaSound. (Advanced) Optionally, remove unused dependencies
+	// from the internal dependendency list on successful removal of node.
+	UFUNCTION(BlueprintCallable, Category = "Audio|MetaSound|Builder", meta = (ExpandEnumAsExecs = "OutResult", AdvancedDisplay = "2"))
 	void RemoveNode(const FMetaSoundNodeHandle& NodeHandle, EMetaSoundBuilderResult& OutResult, bool bRemoveUnusedDependencies = true);
 
 	// Removes node input literal default if set, reverting the value to be whatever the node class defaults the value to.
