Incorrect wording on tooltip when converting Types via Drag and Drop within a MetaSound

#rnx

diff --git a/Engine/Plugins/Runtime/Metasound/Source/MetasoundEditor/Private/MetasoundEditor.cpp b/Engine/Plugins/Runtime/Metasound/Source/MetasoundEditor/Private/MetasoundEditor.cpp
--- a/Engine/Plugins/Runtime/Metasound/Source/MetasoundEditor/Private/MetasoundEditor.cpp
+++ b/Engine/Plugins/Runtime/Metasound/Source/MetasoundEditor/Private/MetasoundEditor.cpp
@@ -625,12 +625,12 @@ namespace Metasound
 									if (Connectability.Connectable == Frontend::FConnectability::EConnectable::No)
 									{
 										PrimarySymbol = FAppStyle::GetBrush(TEXT("Graph.ConnectorFeedback.Error"));
-										Message = FText::Format(IncompatibleText, FText::FromName(DataType), FText::FromName(OtherDataType));
+										Message = FText::Format(IncompatibleText, FText::FromName(OtherDataType), FText::FromName(DataType));
 									}
 									else if (Connectability.Connectable == Frontend::FConnectability::EConnectable::YesWithConverterNode)
 									{
 										PrimarySymbol = FAppStyle::GetBrush(TEXT("Graph.ConnectorFeedback.ViaCast"));
-										Message = FText::Format(CompatibleText, FText::FromName(DataType), FText::FromName(OtherDataType));
+										Message = FText::Format(CompatibleText, FText::FromName(OtherDataType), FText::FromName(DataType));
 									}
 								}
 							}
@@ -727,12 +727,12 @@ namespace Metasound
 									if (Connectability.Connectable == Frontend::FConnectability::EConnectable::No)
 									{
 										PrimarySymbol = FAppStyle::GetBrush(TEXT("Graph.ConnectorFeedback.Error"));
-										Message = FText::Format(IncompatibleText, FText::FromName(DataType), FText::FromName(OtherDataType));
+										Message = FText::Format(IncompatibleText, FText::FromName(OtherDataType), FText::FromName(DataType));
 									}
 									else if (Connectability.Connectable == Frontend::FConnectability::EConnectable::YesWithConverterNode)
 									{
 										PrimarySymbol = FAppStyle::GetBrush(TEXT("Graph.ConnectorFeedback.ViaCast"));
-										Message = FText::Format(CompatibleText, FText::FromName(DataType), FText::FromName(OtherDataType));
+										Message = FText::Format(CompatibleText, FText::FromName(OtherDataType), FText::FromName(DataType));
 									}
 								}
 							}
