Trivial: Add display names for subsystems/MetaSound asset types ("Meta Sound" to "MetaSound"). Remove invalid comment about AssetSubsystem being deprecated in entirety (still used for loading/unloading MetaSounds by directory)

#rnx

diff --git a/Engine/Plugins/Runtime/Metasound/Source/MetasoundEditor/Public/MetasoundEditorSubsystem.h b/Engine/Plugins/Runtime/Metasound/Source/MetasoundEditor/Public/MetasoundEditorSubsystem.h
--- a/Engine/Plugins/Runtime/Metasound/Source/MetasoundEditor/Public/MetasoundEditorSubsystem.h
+++ b/Engine/Plugins/Runtime/Metasound/Source/MetasoundEditor/Public/MetasoundEditorSubsystem.h
@@ -20,7 +20,7 @@ struct FMetaSoundPageSettings;
 
 
 /** The subsystem in charge of editor MetaSound functionality */
-UCLASS()
+UCLASS(meta = (DisplayName = "MetaSound Editor Subsystem"))
 class METASOUNDEDITOR_API UMetaSoundEditorSubsystem : public UEditorSubsystem
 {
 	GENERATED_BODY()
diff --git a/Engine/Plugins/Runtime/Metasound/Source/MetasoundEngine/Public/Metasound.h b/Engine/Plugins/Runtime/Metasound/Source/MetasoundEngine/Public/Metasound.h
--- a/Engine/Plugins/Runtime/Metasound/Source/MetasoundEngine/Public/Metasound.h
+++ b/Engine/Plugins/Runtime/Metasound/Source/MetasoundEngine/Public/Metasound.h
@@ -68,7 +68,7 @@ public:
  * This asset type is used for Metasound assets that can only be used as nodes in other Metasound graphs.
  * Because of this, they contain no required inputs or outputs.
  */
-UCLASS(hidecategories = object, BlueprintType)
+UCLASS(hidecategories = object, BlueprintType, meta = (DisplayName = "MetaSound Patch"))
 class METASOUNDENGINE_API UMetaSoundPatch : public UObject, public FMetasoundAssetBase, public IMetaSoundDocumentInterface
 {
 	GENERATED_BODY()
diff --git a/Engine/Plugins/Runtime/Metasound/Source/MetasoundEngine/Public/MetasoundAssetSubsystem.h b/Engine/Plugins/Runtime/Metasound/Source/MetasoundEngine/Public/MetasoundAssetSubsystem.h
--- a/Engine/Plugins/Runtime/Metasound/Source/MetasoundEngine/Public/MetasoundAssetSubsystem.h
+++ b/Engine/Plugins/Runtime/Metasound/Source/MetasoundEngine/Public/MetasoundAssetSubsystem.h
@@ -55,8 +55,8 @@ namespace Metasound::Engine
 	void InitializeAssetManager();
 } // namespace Metasound::Engine
 
-/** DEPRECATED: Subsystem in charge of the MetaSound asset registry */
-UCLASS()
+
+UCLASS(meta = (DisplayName = "MetaSound Asset Subsystem"))
 class METASOUNDENGINE_API UMetaSoundAssetSubsystem : public UEngineSubsystem
 {
 	GENERATED_BODY()
diff --git a/Engine/Plugins/Runtime/Metasound/Source/MetasoundEngine/Public/MetasoundBuilderBase.h b/Engine/Plugins/Runtime/Metasound/Source/MetasoundEngine/Public/MetasoundBuilderBase.h
--- a/Engine/Plugins/Runtime/Metasound/Source/MetasoundEngine/Public/MetasoundBuilderBase.h
+++ b/Engine/Plugins/Runtime/Metasound/Source/MetasoundEngine/Public/MetasoundBuilderBase.h
@@ -58,7 +58,7 @@ public:
 	}
 };
 
-USTRUCT(BlueprintType)
+USTRUCT(BlueprintType, meta = (DisplayName = "MetaSound Node Handle"))
 struct METASOUNDENGINE_API FMetaSoundNodeHandle
 {
 	GENERATED_BODY()
@@ -107,7 +107,7 @@ struct METASOUNDENGINE_API FMetaSoundBuilderOptions
 	TScriptInterface<IMetaSoundDocumentInterface> ExistingMetaSound;
 };
 
-UENUM(BlueprintType)
+UENUM(BlueprintType, meta = (DisplayName = "MetaSound Builder Result"))
 enum class EMetaSoundBuilderResult : uint8
 {
 	Succeeded,
@@ -115,7 +115,7 @@ enum class EMetaSoundBuilderResult : uint8
 };
 
 /** Base implementation of MetaSound builder */
-UCLASS(Abstract, BlueprintType, Transient)
+UCLASS(Abstract, BlueprintType, Transient, meta = (DisplayName = "MetaSound Builder Base"))
 class METASOUNDENGINE_API UMetaSoundBuilderBase : public UObject
 {
 	GENERATED_BODY()
diff --git a/Engine/Plugins/Runtime/Metasound/Source/MetasoundEngine/Public/MetasoundBuilderSubsystem.h b/Engine/Plugins/Runtime/Metasound/Source/MetasoundEngine/Public/MetasoundBuilderSubsystem.h
--- a/Engine/Plugins/Runtime/Metasound/Source/MetasoundEngine/Public/MetasoundBuilderSubsystem.h
+++ b/Engine/Plugins/Runtime/Metasound/Source/MetasoundEngine/Public/MetasoundBuilderSubsystem.h
@@ -47,7 +47,7 @@ namespace Metasound::Engine
 DECLARE_DYNAMIC_DELEGATE_OneParam(FOnCreateAuditionGeneratorHandleDelegate, UMetasoundGeneratorHandle*, GeneratorHandle);
 
 /** Builder in charge of building a MetaSound Patch */
-UCLASS(Transient, BlueprintType)
+UCLASS(Transient, BlueprintType, meta = (DisplayName = "MetaSound Patch Builder"))
 class METASOUNDENGINE_API UMetaSoundPatchBuilder : public UMetaSoundBuilderBase
 {
 	GENERATED_BODY()
@@ -66,7 +66,7 @@ protected:
 };
 
 /** Builder in charge of building a MetaSound Source */
-UCLASS(Transient, BlueprintType)
+UCLASS(Transient, BlueprintType, meta = (DisplayName = "MetaSound Source Builder"))
 class METASOUNDENGINE_API UMetaSoundSourceBuilder : public UMetaSoundBuilderBase
 {
 	GENERATED_BODY()
@@ -150,7 +150,7 @@ private:
 };
 
 /** The subsystem in charge of tracking MetaSound builders */
-UCLASS()
+UCLASS(meta = (DisplayName = "MetaSound Builder Subsystem"))
 class METASOUNDENGINE_API UMetaSoundBuilderSubsystem : public UEngineSubsystem
 {
 	GENERATED_BODY()
diff --git a/Engine/Plugins/Runtime/Metasound/Source/MetasoundEngine/Public/MetasoundGeneratorHandle.h b/Engine/Plugins/Runtime/Metasound/Source/MetasoundEngine/Public/MetasoundGeneratorHandle.h
--- a/Engine/Plugins/Runtime/Metasound/Source/MetasoundEngine/Public/MetasoundGeneratorHandle.h
+++ b/Engine/Plugins/Runtime/Metasound/Source/MetasoundEngine/Public/MetasoundGeneratorHandle.h
@@ -411,7 +411,7 @@ namespace Metasound
 /**
  * Blueprint-facing interface to a FMetasoundGenerator on a UAudioComponent.
  */
-UCLASS(BlueprintType,Category="MetaSound")
+UCLASS(BlueprintType, Category = "MetaSound", meta = (DisplayName = "MetaSound Generator Handle"))
 class METASOUNDENGINE_API UMetasoundGeneratorHandle : public UObject
 {
 	GENERATED_BODY()
diff --git a/Engine/Plugins/Runtime/Metasound/Source/MetasoundEngine/Public/MetasoundSource.h b/Engine/Plugins/Runtime/Metasound/Source/MetasoundEngine/Public/MetasoundSource.h
--- a/Engine/Plugins/Runtime/Metasound/Source/MetasoundEngine/Public/MetasoundSource.h
+++ b/Engine/Plugins/Runtime/Metasound/Source/MetasoundEngine/Public/MetasoundSource.h
@@ -71,7 +71,7 @@ DECLARE_TS_MULTICAST_DELEGATE_TwoParams(FOnGeneratorInstanceDestroyed, uint64, T
 /**
  * This Metasound type can be played as an audio source.
  */
-UCLASS(hidecategories = object, BlueprintType)
+UCLASS(hidecategories = object, BlueprintType, meta = (DisplayName = "MetaSound Source"))
 class METASOUNDENGINE_API UMetaSoundSource : public USoundWaveProcedural, public FMetasoundAssetBase, public IMetaSoundDocumentInterface
 {
 	GENERATED_BODY()
diff --git a/Engine/Plugins/Runtime/Metasound/Source/MetasoundFrontend/Public/MetasoundDocumentInterface.h b/Engine/Plugins/Runtime/Metasound/Source/MetasoundFrontend/Public/MetasoundDocumentInterface.h
--- a/Engine/Plugins/Runtime/Metasound/Source/MetasoundFrontend/Public/MetasoundDocumentInterface.h
+++ b/Engine/Plugins/Runtime/Metasound/Source/MetasoundFrontend/Public/MetasoundDocumentInterface.h
@@ -24,7 +24,7 @@ namespace Metasound::Frontend
 // UInterface for all MetaSound UClasses that implement a MetaSound document
 // as a means for accessing data via code, scripting, execution, or node
 // class generation.
-UINTERFACE(BlueprintType, meta = (CannotImplementInterfaceInBlueprint))
+UINTERFACE(BlueprintType, meta = (DisplayName = "MetaSound Document Interface", CannotImplementInterfaceInBlueprint))
 class METASOUNDFRONTEND_API UMetaSoundDocumentInterface : public UInterface
 {
 	GENERATED_BODY()
diff --git a/Engine/Plugins/Runtime/Metasound/Source/MetasoundFrontend/Public/MetasoundFrontendDocument.h b/Engine/Plugins/Runtime/Metasound/Source/MetasoundFrontend/Public/MetasoundFrontendDocument.h
--- a/Engine/Plugins/Runtime/Metasound/Source/MetasoundFrontend/Public/MetasoundFrontendDocument.h
+++ b/Engine/Plugins/Runtime/Metasound/Source/MetasoundFrontend/Public/MetasoundFrontendDocument.h
@@ -1331,7 +1331,7 @@ struct METASOUNDFRONTEND_API FMetasoundFrontendInterface : public FMetasoundFron
 
 
 // Name of a Metasound class
-USTRUCT(BlueprintType)
+USTRUCT(BlueprintType, meta = (DisplayName = "MetaSound Class Name"))
 struct METASOUNDFRONTEND_API FMetasoundFrontendClassName
 {
 	GENERATED_BODY()
