Fix missing refactor of bool literal customization.  Now properly shows paged defaults property array member instead of deprecated single default field in editor
#rnx

diff --git a/Engine/Plugins/Runtime/Metasound/Source/MetasoundEditor/Private/MetasoundNodeDetailCustomization.cpp b/Engine/Plugins/Runtime/Metasound/Source/MetasoundEditor/Private/MetasoundNodeDetailCustomization.cpp
--- a/Engine/Plugins/Runtime/Metasound/Source/MetasoundEditor/Private/MetasoundNodeDetailCustomization.cpp
+++ b/Engine/Plugins/Runtime/Metasound/Source/MetasoundEditor/Private/MetasoundNodeDetailCustomization.cpp
@@ -389,52 +389,55 @@ namespace Metasound
 		{
 			check(DefaultCategoryBuilder);
 
-			DefaultProperties.Reset();
-
-			UMetasoundEditorGraphMemberDefaultBool* DefaultBool = Cast<UMetasoundEditorGraphMemberDefaultBool>(&InLiteral);
-			if (!ensure(DefaultBool))
+			TArray<UObject*> EditLiterals;
+			if (UMetasoundEditorGraphMemberDefaultBool* CastLiteral = Cast<UMetasoundEditorGraphMemberDefaultBool>(&InLiteral); ensure(CastLiteral))
+			{
+				EditLiterals = { CastLiteral };
+				BoolLiteral = CastLiteral;
+			}
+			else
 			{
+				BoolLiteral.Reset();
 				return;
 			}
-			BoolLiteral = DefaultBool;
 
-			TArray<TSharedPtr<IPropertyHandle>> DefaultValueHandles;
-			const FName DefaultName = GET_MEMBER_NAME_CHECKED(UMetasoundEditorGraphMemberDefaultBool, Default);
-			if (IDetailPropertyRow* Row = DefaultCategoryBuilder->AddExternalObjectProperty(TArray<UObject*>({ DefaultBool }), DefaultName); ensure(Row))
-			{
-				DefaultValueHandles = { Row->GetPropertyHandle() };
-				Row->Visibility(GetDefaultVisibility());
+			FMetasoundDefaultLiteralCustomizationBase::CustomizeDefaults(InLiteral, InDetailLayout);
+
+			TAttribute<EVisibility> DefaultVisibility = GetDefaultVisibility();
 
-				// Enable widget options for editable inputs only 
-				bool bShowWidgetOptions = false;
-				if (const UMetasoundEditorGraphInput* ParentMember = Cast <UMetasoundEditorGraphInput>(InLiteral.FindMember()))
+			IDetailCategoryBuilder& EditorOptionsBuilder = InDetailLayout.EditCategory("DefaultEditorOptions");
+			auto AddOptionPropRow = [&](FName PropertyName)
+			{
+				IDetailPropertyRow* Row = EditorOptionsBuilder.AddExternalObjectProperty(EditLiterals, PropertyName);
+				if (ensure(Row))
 				{
-					if (const UMetasoundEditorGraph* OwningGraph = ParentMember->GetOwningGraph())
-					{
-						bShowWidgetOptions = OwningGraph->IsEditable();
-					}
+					Row->Visibility(DefaultVisibility);
 				}
+				return Row;
+			};
 
-				// add input widget properties
-				if (bShowWidgetOptions)
+			// Enable widget options for editable inputs only
+			bool bShowWidgetOptions = false;
+			if (const UMetasoundEditorGraphInput* ParentMember = Cast<UMetasoundEditorGraphInput>(InLiteral.FindMember()))
+			{
+				if (const UMetasoundEditorGraph* OwningGraph = ParentMember->GetOwningGraph())
 				{
-					Frontend::FDataTypeRegistryInfo DataTypeInfo;
-					MemberCustomizationPrivate::GetDataTypeFromElementPropertyHandle(DefaultValueHandles.Last(), DataTypeInfo);
+					bShowWidgetOptions = OwningGraph->IsEditable();
+				}
+			}
 
-					const UMetasoundEditorSettings* EditorSettings = GetDefault<UMetasoundEditorSettings>();
-					check(EditorSettings);
+			// add input widget properties
+			if (bShowWidgetOptions)
+			{
+				const UMetasoundEditorSettings* EditorSettings = GetDefault<UMetasoundEditorSettings>();
+				check(EditorSettings);
 
-					if (EditorSettings->bUseAudioMaterialWidgets)
+				if (EditorSettings->bUseAudioMaterialWidgets)
+				{
+					if (InLiteral.GetDataType() != Metasound::GetMetasoundDataTypeName<Metasound::FTrigger>())
 					{
-						if (MemberCustomizationPrivate::GetPrimitiveTypeName(DataTypeInfo) != Metasound::GetMetasoundDataTypeName<Metasound::FTrigger>())
-						{
-							IDetailCategoryBuilder& EditorOptionsBuilder = InDetailLayout.EditCategory("DefaultEditorOptions");
-							if (IDetailPropertyRow* TypeRow = EditorOptionsBuilder.AddExternalObjectProperty(TArray<UObject*>({ DefaultBool }), GET_MEMBER_NAME_CHECKED(UMetasoundEditorGraphMemberDefaultBool, WidgetType)))
-							{
-								DefaultValueHandles.Add(Row->GetPropertyHandle());
-								TypeRow->Visibility(GetDefaultVisibility());
-							}
-						}
+
+						AddOptionPropRow(GET_MEMBER_NAME_CHECKED(UMetasoundEditorGraphMemberDefaultBool, WidgetType));
 					}
 				}
 			}
