Tooltip displayed when renaming a default interface member doesn't display correlated interface name

#rnx

diff --git a/Engine/Plugins/Runtime/Metasound/Source/MetasoundEditor/Private/MetasoundEditorGraph.cpp b/Engine/Plugins/Runtime/Metasound/Source/MetasoundEditor/Private/MetasoundEditorGraph.cpp
--- a/Engine/Plugins/Runtime/Metasound/Source/MetasoundEditor/Private/MetasoundEditorGraph.cpp
+++ b/Engine/Plugins/Runtime/Metasound/Source/MetasoundEditor/Private/MetasoundEditorGraph.cpp
@@ -526,7 +526,8 @@ bool UMetasoundEditorGraphVertex::CanRename(const FText& InNewText, FText& OutEr
 
 	if (IsInterfaceMember())
 	{
-		OutError = FText::Format(LOCTEXT("GraphVertexRenameInvalid_GraphVertexRequired", "{0} is interface member and cannot be renamed."), InNewText);
+		const FText CurrentMemberName = FText::FromName(GetMemberName());
+		OutError = FText::Format(LOCTEXT("GraphVertexRenameInvalid_GraphVertexRequired", "{0} is interface member and cannot be renamed."), CurrentMemberName);
 		return false;
 	}
 
