Fix for crash when using new MetaSound material boolean widget
#rnx

diff --git a/Engine/Plugins/Runtime/Metasound/Source/MetasoundEditor/Private/SMetasoundGraphNode.cpp b/Engine/Plugins/Runtime/Metasound/Source/MetasoundEditor/Private/SMetasoundGraphNode.cpp
--- a/Engine/Plugins/Runtime/Metasound/Source/MetasoundEditor/Private/SMetasoundGraphNode.cpp
+++ b/Engine/Plugins/Runtime/Metasound/Source/MetasoundEditor/Private/SMetasoundGraphNode.cpp
@@ -1137,7 +1137,7 @@ namespace Metasound
 									DefaultBool->Modify();
 
 									FMetasoundFrontendLiteral OutputLiteral;
-									OutputLiteral.Set(FloatInputWidget->GetOutputValue(Value));
+									OutputLiteral.Set(Value);
 
 									const FGuid BuildPageID = Builder.GetBuildPageID();
 									DefaultBool->SetFromLiteral(OutputLiteral, BuildPageID);
