Fix MetaSound Editor transaction stack corruption caused by copying graph node data prior to setting modify state

#rnx

diff --git a/Engine/Plugins/Runtime/Metasound/Source/MetasoundEditor/Private/MetasoundEditor.h b/Engine/Plugins/Runtime/Metasound/Source/MetasoundEditor/Private/MetasoundEditor.h
--- a/Engine/Plugins/Runtime/Metasound/Source/MetasoundEditor/Private/MetasoundEditor.h
+++ b/Engine/Plugins/Runtime/Metasound/Source/MetasoundEditor/Private/MetasoundEditor.h
@@ -178,7 +178,7 @@ namespace Metasound::Editor
 
 		virtual bool IsPlaying() const override;
 
-		/** Whether pasting the currently selected nodes is permissible */
+		/** Whether pasting the current data on the clipboard to the focused graph is permissible */
 		bool CanPasteNodes();
 
 		/** Duplicates the selected node(s) in the graph */
