Safeguard against referenced weak member becoming stale when node customization is being rebuilt and attribute lambdas are still being evaluated

#rnx
#tests MetaSoundEditor adding & removing interfaces & members

diff --git a/Engine/Plugins/Runtime/Metasound/Source/MetasoundEditor/Private/MetasoundNodeDetailCustomization.cpp b/Engine/Plugins/Runtime/Metasound/Source/MetasoundEditor/Private/MetasoundNodeDetailCustomization.cpp
--- a/Engine/Plugins/Runtime/Metasound/Source/MetasoundEditor/Private/MetasoundNodeDetailCustomization.cpp
+++ b/Engine/Plugins/Runtime/Metasound/Source/MetasoundEditor/Private/MetasoundNodeDetailCustomization.cpp
@@ -1282,68 +1282,78 @@ PRAGMA_ENABLE_DEPRECATION_WARNINGS
 
 		void FMetasoundVertexDetailCustomization::AddConstructorPinRow(IDetailLayoutBuilder& InDetailLayout)
 		{
-			if (UMetasoundEditorGraphVertex* Vertex = Cast<UMetasoundEditorGraphVertex>(GraphMember.Get()))
-			{
-				InDetailLayout.EditCategory("General").AddCustomRow(MemberCustomizationPrivate::ConstructorPinText)
-				.IsEnabled(IsGraphEditable() && !IsInterfaceMember())
-				.NameContent()
-				[
-					SNew(STextBlock)
-					.Text(MemberCustomizationPrivate::ConstructorPinText)
-					.ToolTipText(MemberCustomizationPrivate::ConstructorPinTooltip)
-					.Font(IDetailLayoutBuilder::GetDetailFontBold())
-				]
-				.ValueContent()
-				[
-					SAssignNew(ConstructorPinCheckbox, SCheckBox)
-					.IsChecked_Lambda([this, Vertex]()
+			InDetailLayout.EditCategory("General").AddCustomRow(MemberCustomizationPrivate::ConstructorPinText)
+			.IsEnabled(IsGraphEditable() && !IsInterfaceMember())
+			.NameContent()
+			[
+				SNew(STextBlock)
+				.Text(MemberCustomizationPrivate::ConstructorPinText)
+				.ToolTipText(MemberCustomizationPrivate::ConstructorPinTooltip)
+				.Font(IDetailLayoutBuilder::GetDetailFontBold())
+			]
+			.ValueContent()
+			[
+				SAssignNew(ConstructorPinCheckbox, SCheckBox)
+				.IsChecked_Lambda([this]()
+				{
+					if (UMetasoundEditorGraphVertex* Vertex = Cast<UMetasoundEditorGraphVertex>(GraphMember.Get()))
 					{
 						return OnGetConstructorPinCheckboxState(Vertex);
-					})
-					.OnCheckStateChanged_Lambda([this, Vertex](ECheckBoxState InNewState)
+					}
+
+					return ECheckBoxState::Undetermined;
+				})
+				.OnCheckStateChanged_Lambda([this](ECheckBoxState InNewState)
+				{
+					if (UMetasoundEditorGraphVertex* Vertex = Cast<UMetasoundEditorGraphVertex>(GraphMember.Get()))
 					{
 						OnConstructorPinStateChanged(Vertex, InNewState);
-					})
-					[
-						SNew(STextBlock)
-						.Font(IDetailLayoutBuilder::GetDetailFont())
-					]
-				];
-			}
+					}
+				})
+				[
+					SNew(STextBlock)
+					.Font(IDetailLayoutBuilder::GetDetailFont())
+				]
+			];
 		}
 
 #if WITH_EDITORONLY_DATA
 		void FMetasoundVertexDetailCustomization::AddAdvancedPinRow(IDetailLayoutBuilder& InDetailLayout)
 		{
 			//only add row if input or output
-			if (UMetasoundEditorGraphVertex* Vertex = Cast<UMetasoundEditorGraphVertex>(GraphMember.Get()))
-			{
-				InDetailLayout.EditCategory("General").AddCustomRow(MemberCustomizationPrivate::AdvancedPinText)
-					.IsEnabled(IsGraphEditable() && !IsInterfaceMember())
-					.NameContent()
-					[
-						SNew(STextBlock)
-							.Text(MemberCustomizationPrivate::AdvancedPinText)
-							.ToolTipText(MemberCustomizationPrivate::AdvancedPinTooltip)
-							.Font(IDetailLayoutBuilder::GetDetailFontBold())
-					]
-				.ValueContent()
+			InDetailLayout.EditCategory("General").AddCustomRow(MemberCustomizationPrivate::AdvancedPinText)
+				.IsEnabled(IsGraphEditable() && !IsInterfaceMember())
+				.NameContent()
 				[
-					SAssignNew(AdvancedPinCheckbox, SCheckBox)
-					.IsChecked_Lambda([this, Vertex]()
+					SNew(STextBlock)
+						.Text(MemberCustomizationPrivate::AdvancedPinText)
+						.ToolTipText(MemberCustomizationPrivate::AdvancedPinTooltip)
+						.Font(IDetailLayoutBuilder::GetDetailFontBold())
+				]
+			.ValueContent()
+			[
+				SAssignNew(AdvancedPinCheckbox, SCheckBox)
+				.IsChecked_Lambda([this]()
+				{
+					if (UMetasoundEditorGraphVertex* Vertex = Cast<UMetasoundEditorGraphVertex>(GraphMember.Get()))
 					{
 						return OnGetAdvancedPinCheckboxState(Vertex);
-					})
-					.OnCheckStateChanged_Lambda([this, Vertex](ECheckBoxState InNewState)
+					}
+
+					return ECheckBoxState::Undetermined;
+				})
+				.OnCheckStateChanged_Lambda([this](ECheckBoxState InNewState)
+				{
+					if (UMetasoundEditorGraphVertex* Vertex = Cast<UMetasoundEditorGraphVertex>(GraphMember.Get()))
 					{
 						OnAdvancedPinStateChanged(Vertex, InNewState);
-					})
-					[
-						SNew(STextBlock)
-						.Font(IDetailLayoutBuilder::GetDetailFont())
-					]
-				];
-			}
+					}
+				})
+				[
+					SNew(STextBlock)
+					.Font(IDetailLayoutBuilder::GetDetailFont())
+				]
+			];
 		}
 #endif // WITH_EDITORONLY_DATA
 
@@ -1364,9 +1374,7 @@ PRAGMA_ENABLE_DEPRECATION_WARNINGS
 				AddConstructorPinRow(InDetailLayout);
 			}
 			
-#if WITH_EDITORONLY_DATA		
-				AddAdvancedPinRow(InDetailLayout);
-#endif // WITH_EDITORONLY_DATA
+			AddAdvancedPinRow(InDetailLayout);
 
 			// Sort order
 			IDetailCategoryBuilder& CategoryBuilder = FMetasoundMemberDetailCustomization::GetGeneralCategoryBuilder(InDetailLayout);
@@ -1386,7 +1394,14 @@ PRAGMA_ENABLE_DEPRECATION_WARNINGS
 				.ValueContent()
 				[
 					SNew(SNumericEntryBox<int32>)
-					.Value_Lambda([VertexPtr]() { return VertexPtr->GetSortOrderIndex(); })
+					.Value_Lambda([VertexPtr]()
+					{	
+						if (VertexPtr.IsValid())
+						{
+							return VertexPtr->GetSortOrderIndex();
+						}
+						return 0;
+					})
 					.AllowSpin(false)
 					.UndeterminedString(LOCTEXT("Vertex_SortOrder_MultipleValues", "Multiple"))
 					.OnValueCommitted_Lambda([VertexPtr](int32 NewValue, ETextCommit::Type CommitInfo)
@@ -1633,23 +1648,29 @@ PRAGMA_ENABLE_DEPRECATION_WARNINGS
 			}
 			else
 			{
-				// Make default value uneditable while playing for constructor inputs
-				const UMetasoundEditorGraphInput* Input = Cast<UMetasoundEditorGraphInput>(MemberDefaultLiteral->FindMember());
-				if (Input)
+				Enabled = TAttribute<bool>::CreateLambda([this]
 				{
-					Enabled = TAttribute<bool>::CreateLambda([this, Input]
+					if (GraphMember.IsValid())
 					{
-						if (Input->GetVertexAccessType() == EMetasoundFrontendVertexAccessType::Value)
+						if (UMetasoundEditorGraphMemberDefaultLiteral* MemberDefaultLiteral = GraphMember->GetLiteral())
 						{
-							UObject* MetaSoundObject = Input->GetOutermostObject();
-							if (TSharedPtr<FEditor> MetaSoundEditor = FGraphBuilder::GetEditorForMetasound(*MetaSoundObject))
+							// Make default value uneditable while playing for constructor inputs
+							const UMetasoundEditorGraphInput* Input = Cast<UMetasoundEditorGraphInput>(MemberDefaultLiteral->FindMember());
+							if (Input)
 							{
-								return !MetaSoundEditor->IsPlaying();
+								if (Input->GetVertexAccessType() == EMetasoundFrontendVertexAccessType::Value)
+								{
+									UObject* MetaSoundObject = Input->GetOutermostObject();
+									if (TSharedPtr<FEditor> MetaSoundEditor = FGraphBuilder::GetEditorForMetasound(*MetaSoundObject))
+									{
+										return !MetaSoundEditor->IsPlaying();
+									}
+								}
 							}
 						}
-						return true;
-					});
-				}
+					}
+					return true;
+				});
 			}
 
 			CustomizeDefaultCategory(InDetailLayout);
