[MetaSounds] Support copy/cut/paste in the graph members menu
#rnx

diff --git a/Engine/Plugins/Runtime/Metasound/Source/MetasoundEditor/Private/MetasoundEditor.cpp b/Engine/Plugins/Runtime/Metasound/Source/MetasoundEditor/Private/MetasoundEditor.cpp
--- a/Engine/Plugins/Runtime/Metasound/Source/MetasoundEditor/Private/MetasoundEditor.cpp
+++ b/Engine/Plugins/Runtime/Metasound/Source/MetasoundEditor/Private/MetasoundEditor.cpp
@@ -2376,6 +2376,18 @@ namespace Metasound
 				FExecuteAction::CreateSP(this, &FEditor::DuplicateSelectedMemberItems),
 				FCanExecuteAction::CreateSP(this, &FEditor::CanDuplicateSelectedMemberItems));
 
+			ToolkitCommands->MapAction(FGenericCommands::Get().Copy,
+				FExecuteAction::CreateSP(this, &FEditor::CopySelectedMemberItems),
+				FCanExecuteAction::CreateSP(this, &FEditor::CanCopySelectedMemberItems));
+
+			ToolkitCommands->MapAction(FGenericCommands::Get().Cut,
+				FExecuteAction::CreateSP(this, &FEditor::CutSelectedMemberItems),
+				FCanExecuteAction::CreateSP(this, &FEditor::CanCutSelectedMemberItems));
+
+			ToolkitCommands->MapAction(FGenericCommands::Get().Paste,
+				FExecuteAction::CreateSP(this, &FEditor::PasteSelectedMemberItems),
+				FCanExecuteAction::CreateSP(this, &FEditor::CanPasteSelectedMemberItems));
+
 			ToolkitCommands->MapAction(
 				FEditorCommands::Get().UpdateNodeClass,
 				FExecuteAction::CreateSP(this, &FEditor::UpdateSelectedNodeClasses));
@@ -3847,6 +3859,175 @@ namespace Metasound
 			}
 		}
 
+		bool FEditor::CanCopySelectedMemberItems() const
+		{
+			if (!GraphMembersMenu.IsValid())
+			{
+				return false;
+			}
+
+			if (!IsGraphEditable())
+			{
+				return false;
+			}
+
+			TArray<TSharedPtr<FEdGraphSchemaAction>> Actions;
+			GraphMembersMenu->GetSelectedActions(Actions);
+
+			if (Actions.IsEmpty())
+			{
+				return false;
+			}
+
+			for (const TSharedPtr<FEdGraphSchemaAction>& Action : Actions)
+			{
+				TSharedPtr<FMetasoundGraphMemberSchemaAction> MetasoundAction = StaticCastSharedPtr<FMetasoundGraphMemberSchemaAction>(Action);
+				if (MetasoundAction.IsValid())
+				{
+					if (const UMetasoundEditorGraphVertex* GraphVertex = Cast<UMetasoundEditorGraphVertex>(MetasoundAction->GetGraphMember()))
+					{
+						if (GraphVertex->IsInterfaceMember())
+						{
+							return false;
+						}
+					}
+				}
+			}
+
+			return true;
+		}
+
+		void FEditor::CopySelectedMemberItems()
+		{
+			if (!GraphMembersMenu.IsValid())
+			{
+				return;
+			}
+
+			TArray<TSharedPtr<FEdGraphSchemaAction>> Actions;
+			GraphMembersMenu->GetSelectedActions(Actions);
+
+			if (Actions.IsEmpty())
+			{
+				return;
+			}
+			
+			TSharedPtr<FMetasoundGraphMemberSchemaAction> MetasoundAction = StaticCastSharedPtr<FMetasoundGraphMemberSchemaAction>(Actions.Last());
+			MetasoundAction->GetGraphMember()->CacheBreadcrumb();
+			FDocumentClipboardUtils::CopyMemberToClipboard(MetasoundAction->GetGraphMember());
+		}
+
+		bool FEditor::CanCutSelectedMemberItems() const
+		{
+			return CanCopySelectedMemberItems() && CanDeleteInterfaceItems();
+		}
+
+		void FEditor::CutSelectedMemberItems()
+		{
+			CopySelectedMemberItems();
+			DeleteSelectedInterfaceItems();
+		}
+
+		bool FEditor::CanPasteSelectedMemberItems() const
+		{
+			if (!GraphMembersMenu.IsValid())
+			{
+				return false;
+			}
+			
+			if (!IsGraphEditable())
+			{
+				return false;
+			}
+
+			FString ClipboardText;			
+			FPlatformApplicationMisc::ClipboardPaste(ClipboardText);
+
+			if (ClipboardText.IsEmpty())
+			{
+				return false;
+			}
+
+			if (FDocumentClipboardUtils::CanImportMemberFromText(ClipboardText))
+			{
+				return true;
+			}
+
+			return false;
+		}
+
+		void FEditor::PasteSelectedMemberItems()
+		{
+			const UMetasoundEditorGraphMember* Member = FDocumentClipboardUtils::GetMemberFromClipboard();		
+			if (Member)
+			{
+				if (const UMetasoundEditorGraphVertex* SourceGraphVertex = Cast<UMetasoundEditorGraphVertex>(Member))
+				{
+					UObject* MetaSound = GetMetasoundObject();
+					if (!MetaSound)
+					{
+						return;
+					}
+
+					UMetasoundEditorGraph& Graph = GetMetaSoundGraphChecked();
+					MetaSound->Modify();
+					Graph.Modify();
+
+					FMetaSoundFrontendDocumentBuilder& DocumentBuilder = Frontend::IDocumentBuilderRegistry::GetChecked().FindOrBeginBuilding(&Graph.GetMetasoundChecked());				
+					const FMetasoundEditorGraphVertexBreadcrumb Breadcrumb = SourceGraphVertex->GetBreadcrumb();
+
+					FCreateNodeVertexParams VertexParams;
+					VertexParams.DataType = SourceGraphVertex->GetDataType();
+					VertexParams.AccessType = Breadcrumb.AccessType;
+
+					const EMetasoundFrontendClassType ClassType = SourceGraphVertex->GetClassType();
+
+					if (ClassType == EMetasoundFrontendClassType::Input)
+					{				
+						TArray<FMetasoundFrontendClassInputDefault> InputDefaults;
+						Algo::Transform(Breadcrumb.DefaultLiterals, InputDefaults, [](const TPair<FGuid, FMetasoundFrontendLiteral>& Pair)
+							{
+								return FMetasoundFrontendClassInputDefault(Pair.Key, Pair.Value);
+							});
+
+						FMetasoundFrontendClassInput ClassInput = FGraphBuilder::CreateUniqueClassInput(*MetaSound, VertexParams, InputDefaults, &Breadcrumb.MemberName);
+						ClassInput.Metadata.SetDescription(Breadcrumb.Description);
+						ClassInput.Metadata.SetIsAdvancedDisplay(Breadcrumb.bIsAdvancedDisplay);
+						ClassInput.Metadata.SortOrderIndex = Breadcrumb.SortOrderIndex;
+						
+						DocumentBuilder.AddGraphInput(ClassInput);	
+					}
+					else if (ClassType == EMetasoundFrontendClassType::Output)
+					{
+						FMetasoundFrontendClassOutput ClassOutput = FGraphBuilder::CreateUniqueClassOutput(*MetaSound, VertexParams, &Breadcrumb.MemberName);
+						ClassOutput.Metadata.SetDescription(Breadcrumb.Description);
+						ClassOutput.Metadata.SetIsAdvancedDisplay(Breadcrumb.bIsAdvancedDisplay);
+						ClassOutput.Metadata.SortOrderIndex = Breadcrumb.SortOrderIndex;					
+
+						DocumentBuilder.AddGraphOutput(ClassOutput);
+					}
+				}
+				else if (const UMetasoundEditorGraphVariable* SourceGraphVariable = Cast<UMetasoundEditorGraphVariable>(Member))
+				{
+					UObject* MetaSound = GetMetasoundObject();
+
+					const FMetasoundEditorGraphVariableBreadcrumb Breadcrumb = SourceGraphVariable->GetBreadcrumb();
+
+					Frontend::FVariableHandle FrontendVariable = FGraphBuilder::AddVariableHandle(*MetaSound, SourceGraphVariable->GetDataType(), &Breadcrumb.MemberName);
+					FrontendVariable->SetDescription(Breadcrumb.Description);
+					FrontendVariable->SetLiteral(Breadcrumb.DefaultLiteral);
+
+					if (ensure(FrontendVariable->IsValid()))
+					{
+						UMetasoundEditorGraph& Graph = GetMetaSoundGraphChecked();
+						Graph.FindOrAddVariable(FrontendVariable);
+					}
+				}
+
+				return;
+			}
+		}
+
 		void FEditor::RefreshDetails()
 		{
 			TRACE_CPUPROFILER_EVENT_SCOPE(Metasound::Editor::FEditor::RefreshDetails);
@@ -4387,6 +4568,23 @@ namespace Metasound
 				if (!Builder->IsPreset())
 				{
 					MenuBuilder.BeginSection("GraphActionMenuSectionActions", LOCTEXT("SectionActionsMenuHeader", "Section Actions"));
+
+					if (const UMetasoundEditorGraphMember* Member = FDocumentClipboardUtils::GetMemberFromClipboard())
+					{
+						FText PasteTooltip;
+
+						if (Member->IsA<UMetasoundEditorGraphVertex>())
+						{
+							PasteTooltip = FText::Format(LOCTEXT("PasteVertexTooltip", "Pastes the {0} to this MetaSound"), Member->GetGraphMemberLabel());
+						}
+						else if (Member->IsA<UMetasoundEditorGraphVariable>())
+						{
+							PasteTooltip = LOCTEXT("PasteVariableTooltip", "Pastes the Variable to this Graph");
+						}
+
+						MenuBuilder.AddMenuEntry(FGenericCommands::Get().Paste, NAME_None, FText::Format(LOCTEXT("PasteMember","Paste {0}"), Member->GetGraphMemberLabel()), PasteTooltip);
+					}
+
 					MenuBuilder.AddMenuEntry(
 						LOCTEXT("DeleteAllUnusedInSection", "Delete Unused Members"),
 						LOCTEXT("DeleteAllUnusedInSectionTooltip", "Delete all Unused Members under this Section"),
@@ -4400,9 +4598,7 @@ namespace Metasound
 			else //Member is selected
 			{
 				MenuBuilder.BeginSection("GraphActionMenuMemberActions", LOCTEXT("MemberActionsMenuHeader", "Member Actions"));
-				MenuBuilder.AddMenuEntry(FGenericCommands::Get().Delete);
 				MenuBuilder.AddMenuEntry(FGenericCommands::Get().Rename);
-				MenuBuilder.AddMenuEntry(FGenericCommands::Get().Duplicate);
 				MenuBuilder.AddMenuEntry(
 					LOCTEXT("JumpToNodesMenuEntry", "Jump to Node(s) in Graph"),
 					LOCTEXT("JumpToNodesMenuEntryTooltip", "Jump to the corresponding node(s) in the MetaSound graph"),
@@ -4410,6 +4606,10 @@ namespace Metasound
 					FUIAction(
 						FExecuteAction::CreateSP(this, &FEditor::JumpToNodesForSelectedInterfaceItem), 
 						FCanExecuteAction::CreateSP(this, &FEditor::CanJumpToNodesForSelectedInterfaceItem)));
+				MenuBuilder.AddMenuEntry(FGenericCommands::Get().Cut);
+				MenuBuilder.AddMenuEntry(FGenericCommands::Get().Copy);
+				MenuBuilder.AddMenuEntry(FGenericCommands::Get().Duplicate);
+				MenuBuilder.AddMenuEntry(FGenericCommands::Get().Delete);
 				MenuBuilder.EndSection();
 			}
 
diff --git a/Engine/Plugins/Runtime/Metasound/Source/MetasoundEditor/Private/MetasoundEditor.h b/Engine/Plugins/Runtime/Metasound/Source/MetasoundEditor/Private/MetasoundEditor.h
--- a/Engine/Plugins/Runtime/Metasound/Source/MetasoundEditor/Private/MetasoundEditor.h
+++ b/Engine/Plugins/Runtime/Metasound/Source/MetasoundEditor/Private/MetasoundEditor.h
@@ -343,6 +343,24 @@ namespace Metasound::Editor
 		/** Duplicate selected Member items. */
 		void DuplicateSelectedMemberItems();
 
+		/** Whether the currently selected Member item(s) can be copied. */
+		bool CanCopySelectedMemberItems() const;
+
+		/** Copy selected Member items. */
+		void CopySelectedMemberItems();		
+		
+		/** Whether the currently selected Member item(s) can be cut. */
+		bool CanCutSelectedMemberItems() const;
+
+		/** Cut selected Member items. */
+		void CutSelectedMemberItems();
+
+		/** Whether the currently selected Member item(s) can be pasted. */
+		bool CanPasteSelectedMemberItems() const;
+
+		/** Paste selected Member items. */
+		void PasteSelectedMemberItems();
+
 		/** Whether there are nodes to jump to for the currently selected interface item. */
 		bool CanJumpToNodesForSelectedInterfaceItem() const;
 
diff --git a/Engine/Plugins/Runtime/Metasound/Source/MetasoundEditor/Private/MetasoundEditorDocumentClipboardUtils.cpp b/Engine/Plugins/Runtime/Metasound/Source/MetasoundEditor/Private/MetasoundEditorDocumentClipboardUtils.cpp
--- a/Engine/Plugins/Runtime/Metasound/Source/MetasoundEditor/Private/MetasoundEditorDocumentClipboardUtils.cpp
+++ b/Engine/Plugins/Runtime/Metasound/Source/MetasoundEditor/Private/MetasoundEditorDocumentClipboardUtils.cpp
@@ -4,12 +4,15 @@
 
 #include "Algo/Transform.h"
 #include "EdGraphUtilities.h"
+#include "Exporters/Exporter.h"
+#include "Factories.h"
+#include "HAL/PlatformApplicationMisc.h"
 #include "Logging/TokenizedMessage.h"
 #include "MetasoundDocumentBuilderRegistry.h"
 #include "MetasoundFrontendSearchEngine.h"
 #include "NodeTemplates/MetasoundFrontendNodeTemplateInput.h"
 #include "ScopedTransaction.h"
-
+#include "UnrealExporter.h"
 
 namespace Metasound::Editor
 {
@@ -397,6 +400,38 @@ namespace Metasound::Editor
 		}
 	}
 
+	class FMemberClipboardObjectTextFactory : public FCustomizableTextObjectFactory
+	{
+	public:
+
+		FMemberClipboardObjectTextFactory()
+			: FCustomizableTextObjectFactory(GWarn)
+			, Member(nullptr)
+		{
+		}
+
+		class UMetasoundEditorGraphMember* Member;
+
+	protected:
+
+		virtual bool CanCreateClass(UClass* InObjectClass, bool& bOmitSubObjs) const override
+		{
+			if (InObjectClass->IsChildOf(UMetasoundEditorGraphMember::StaticClass()))
+			{
+				return true;
+			}
+			return false;
+		}
+
+		virtual void ProcessConstructedObject(UObject* CreatedObject) override
+		{
+			if (CreatedObject->IsA<UMetasoundEditorGraphMember>())
+			{
+				Member = CastChecked<UMetasoundEditorGraphMember>(CreatedObject);
+			}
+		}
+	};
+
 	void FDocumentClipboardUtils::ProcessPastedNodeConnections(FMetasoundAssetBase& OutAsset, TArray<UMetasoundEditorGraphNode*>& OutPastedNodes)
 	{
 		using namespace Frontend;
@@ -485,4 +520,38 @@ namespace Metasound::Editor
 
 		return PastedNodes;
 	}
+
+	void FDocumentClipboardUtils::CopyMemberToClipboard(UMetasoundEditorGraphMember* Content)
+	{
+		// Clear the mark state for saving.
+		UnMarkAllObjects(EObjectMark(OBJECTMARK_TagExp | OBJECTMARK_TagImp));
+
+		// Export the clipboard to text.
+		FStringOutputDevice Archive;
+		const FExportObjectInnerContext Context;
+		UExporter::ExportToOutputDevice(&Context, Content, nullptr, Archive, TEXT("copy"), 0, PPF_ExportsNotFullyQualified | PPF_Copy | PPF_Delimited, false, Content->GetOuter());
+		FPlatformApplicationMisc::ClipboardCopy(*Archive);
+	}
+
+	const UMetasoundEditorGraphMember* FDocumentClipboardUtils::GetMemberFromClipboard()
+	{
+		// Get the text from the clipboard.
+		FString ClipboardText;
+		FPlatformApplicationMisc::ClipboardPaste(ClipboardText);
+
+		FMemberClipboardObjectTextFactory Factory;
+		if (Factory.CanCreateObjectsFromText(ClipboardText))
+		{
+			Factory.ProcessBuffer(GetTransientPackage(), RF_Transactional, ClipboardText);
+			return Factory.Member;
+		}
+		
+		return nullptr;
+	}
+
+	const bool FDocumentClipboardUtils::CanImportMemberFromText(const FString& TextToImport)
+	{
+		FMemberClipboardObjectTextFactory Factory;
+		return Factory.CanCreateObjectsFromText(TextToImport);
+	}
 } // namespace Metasound::Editor
diff --git a/Engine/Plugins/Runtime/Metasound/Source/MetasoundEditor/Private/MetasoundEditorDocumentClipboardUtils.h b/Engine/Plugins/Runtime/Metasound/Source/MetasoundEditor/Private/MetasoundEditorDocumentClipboardUtils.h
--- a/Engine/Plugins/Runtime/Metasound/Source/MetasoundEditor/Private/MetasoundEditorDocumentClipboardUtils.h
+++ b/Engine/Plugins/Runtime/Metasound/Source/MetasoundEditor/Private/MetasoundEditorDocumentClipboardUtils.h
@@ -25,6 +25,15 @@ namespace Metasound::Editor
 	{
 		static TArray<UEdGraphNode*> PasteClipboardString(const FText& InTransactionText, const FString& InClipboardString, const FVector2D& InLocation, UObject& OutMetaSound, FDocumentPasteNotifications& OutNotifications);
 
+		/** Copy MetasoundMember to Clipboard*/
+		static void CopyMemberToClipboard(class UMetasoundEditorGraphMember* Content);
+
+		/** Return MetasoundMember that is in Clipboard, if any*/
+		static const UMetasoundEditorGraphMember* GetMemberFromClipboard();
+
+		/** Whether the string can be imported to a MetasoundMember */
+		static const bool CanImportMemberFromText(const FString& TextToImport);
+
 	private:
 		static void ProcessPastedCommentNodes(FMetasoundAssetBase& OutAsset, const TArrayView<UMetasoundEditorGraphCommentNode*> CommentNodes);
 
diff --git a/Engine/Plugins/Runtime/Metasound/Source/MetasoundEditor/Private/MetasoundEditorGraph.cpp b/Engine/Plugins/Runtime/Metasound/Source/MetasoundEditor/Private/MetasoundEditorGraph.cpp
--- a/Engine/Plugins/Runtime/Metasound/Source/MetasoundEditor/Private/MetasoundEditorGraph.cpp
+++ b/Engine/Plugins/Runtime/Metasound/Source/MetasoundEditor/Private/MetasoundEditorGraph.cpp
@@ -437,6 +437,17 @@ bool UMetasoundEditorGraphVertex::SetIsAdvancedDisplay(const bool IsAdvancedDisp
 }
 #endif // WITH_EDITORONLY_DATA
 
+void UMetasoundEditorGraphVertex::CacheBreadcrumb()
+{
+	Breadcrumb = { };
+
+	Breadcrumb.MemberName = GetMemberName();
+	Breadcrumb.AccessType = GetVertexAccessType();
+	Breadcrumb.Description = GetDescription();
+	Breadcrumb.SortOrderIndex = GetSortOrderIndex();
+	Breadcrumb.bIsAdvancedDisplay = GetFrontendClassVertex()->GetIsAdvancedDisplay();
+}
+
 Metasound::Frontend::FNodeHandle UMetasoundEditorGraphVertex::GetNodeHandle()
 {
 	using namespace Metasound;
@@ -1015,6 +1026,20 @@ bool UMetasoundEditorGraphInput::Synchronize()
 	return bModified;
 }
 
+void UMetasoundEditorGraphInput::CacheBreadcrumb()
+{
+	Super::CacheBreadcrumb();
+
+	const FMetaSoundFrontendDocumentBuilder& Builder = GetFrontendBuilderChecked();
+	if (const FMetasoundFrontendClassInput* Input = Builder.FindGraphInput(GetMemberName()))
+	{
+		Input->IterateDefaults([this](const FGuid& InPageID, const FMetasoundFrontendLiteral& InLiteral)
+		{
+			Breadcrumb.DefaultLiterals.Add(InPageID, InLiteral);
+		});
+	}
+}
+
 void UMetasoundEditorGraphInput::UpdateFrontendDefaultLiteral(bool bPostTransaction, const FGuid* InPageID)
 {
 	using namespace Metasound::Editor;
@@ -1369,6 +1394,15 @@ const FText& UMetasoundEditorGraphVariable::GetGraphMemberLabel() const
 	return Label;
 }
 
+void UMetasoundEditorGraphVariable::CacheBreadcrumb()
+{
+	Breadcrumb = { };
+
+	Breadcrumb.MemberName = GetMemberName();
+	Breadcrumb.DefaultLiteral = GetVariableHandle()->GetLiteral();
+	Breadcrumb.Description = GetDescription();
+}
+
 Metasound::Frontend::FVariableHandle UMetasoundEditorGraphVariable::GetVariableHandle()
 {
 	using namespace Metasound;
diff --git a/Engine/Plugins/Runtime/Metasound/Source/MetasoundEditor/Public/MetasoundEditorGraph.h b/Engine/Plugins/Runtime/Metasound/Source/MetasoundEditor/Public/MetasoundEditorGraph.h
--- a/Engine/Plugins/Runtime/Metasound/Source/MetasoundEditor/Public/MetasoundEditorGraph.h
+++ b/Engine/Plugins/Runtime/Metasound/Source/MetasoundEditor/Public/MetasoundEditorGraph.h
@@ -43,6 +43,45 @@ namespace Metasound
 	} // namespace Editor
 } // namespace Metasound
 
+USTRUCT()
+struct FMetasoundEditorGraphMemberBreadcrumb
+{
+	GENERATED_BODY()
+
+	UPROPERTY()
+	FName MemberName;
+
+	UPROPERTY()
+	FText Description;
+};
+
+USTRUCT()
+struct FMetasoundEditorGraphVertexBreadcrumb : public FMetasoundEditorGraphMemberBreadcrumb
+{
+	GENERATED_BODY()
+
+	UPROPERTY()
+	EMetasoundFrontendVertexAccessType AccessType = EMetasoundFrontendVertexAccessType::Unset;
+
+	UPROPERTY()
+	TMap<FGuid, FMetasoundFrontendLiteral> DefaultLiterals;
+
+	UPROPERTY()
+	bool bIsAdvancedDisplay = false;
+
+	UPROPERTY()
+	int32 SortOrderIndex = 0;
+};
+
+USTRUCT()
+struct FMetasoundEditorGraphVariableBreadcrumb : public FMetasoundEditorGraphMemberBreadcrumb
+{
+	GENERATED_BODY()
+
+	UPROPERTY()
+	FMetasoundFrontendLiteral DefaultLiteral;
+};
+
 DECLARE_MULTICAST_DELEGATE_OneParam(FOnMetasoundMemberNameChanged, FGuid /* NodeID */);
 DECLARE_MULTICAST_DELEGATE(FOnMetasoundMemberRenameRequested);
 
@@ -199,6 +238,9 @@ public:
 	/** Whether the displayed default supports paged values or is a single default value (i.e. characterized per page). */
 	virtual bool IsDefaultPaged() const { return false; }
 
+	/** Cache breadcrumb data before copying member to the clipboard*/
+	virtual void CacheBreadcrumb() { };
+
 #if WITH_EDITOR
 	virtual void PostEditUndo() override;
 #endif // WITH_EDITOR
@@ -271,8 +313,12 @@ public:
 	virtual void SetMemberName(const FName& InNewName, bool bPostTransaction) override;
 	virtual void SetDisplayName(const FText& InNewName, bool bPostTransaction) override;
 	virtual bool Synchronize() override;
+	virtual void CacheBreadcrumb() override;
 	/* ~End UMetasoundEditorGraphMember interface */
 
+	/** Returns Breadcrumb data of this Vertex*/
+	const FMetasoundEditorGraphVertexBreadcrumb& GetBreadcrumb() const { return Breadcrumb; };
+
 	/** Version of interface membership, or invalid version if not an interface member. */
 	virtual const FMetasoundFrontendVersion& GetInterfaceVersion() const;
 
@@ -312,6 +358,11 @@ public:
 	virtual bool CanRename() const override;
 
 	friend class Metasound::Editor::FEditor;
+
+protected:
+
+	UPROPERTY()
+	FMetasoundEditorGraphVertexBreadcrumb Breadcrumb;
 };
 
 UCLASS()
@@ -334,6 +385,7 @@ public:
 	virtual void SetVertexAccessType(EMetasoundFrontendVertexAccessType InNewAccessType, bool bPostTransaction) override;
 	virtual void SetMemberName(const FName& InNewName, bool bPostTransaction) override;
 	virtual bool Synchronize() override;
+	virtual void CacheBreadcrumb() override;
 	virtual void UpdateFrontendDefaultLiteral(bool bPostTransaction, const FGuid* InPageID = nullptr) override;
 	virtual EMetasoundFrontendVertexAccessType GetVertexAccessType() const override;
 
@@ -407,6 +459,7 @@ public:
 
 	virtual void ResetToClassDefault() override;
 	virtual bool Synchronize() override;
+	virtual void CacheBreadcrumb() override;
 	virtual void UpdateFrontendDefaultLiteral(bool bPostTransaction, const FGuid* InPageID = nullptr) override;
 
 	virtual TArray<UMetasoundEditorGraphMemberNode*> GetNodes() const override;
@@ -416,11 +469,19 @@ public:
 
 	const FGuid& GetVariableID() const;
 
+	/**Returns Breadcrumb data of this Variable*/
+	const FMetasoundEditorGraphVariableBreadcrumb& GetBreadcrumb() const { return Breadcrumb; };
+
 	Metasound::Frontend::FVariableHandle GetVariableHandle();
 	Metasound::Frontend::FConstVariableHandle GetConstVariableHandle() const;
 
 	virtual bool CanRename() const override;
 
+protected:
+
+	UPROPERTY()
+	FMetasoundEditorGraphVariableBreadcrumb Breadcrumb;
+
 protected:
 	UE_DEPRECATED(5.5, "No longer supported in favor of internal Editor::FGraphBuilder API which is actively being transitioned to using Document Builder API")
 	virtual UMetasoundEditorGraphNode* AddNode(Metasound::Frontend::FNodeHandle InNodeHandle, bool bInSelectNewNode) override;
