Fixes for pin inspector not showing input template node analyzer data
- Move edge style calls to builder from controllers
- Fix leaks where removing an edge did not remove  editor-only edge style
- Conform EditorOnly Data markup in BuilderAPI
- Make FindBuildGraphChecked private now that page work is complete (was public during transition between 5.4 & 5.5)

#rnx

diff --git a/Engine/Plugins/Runtime/Metasound/Source/MetasoundEditor/Private/SMetasoundPinValueInspector.cpp b/Engine/Plugins/Runtime/Metasound/Source/MetasoundEditor/Private/SMetasoundPinValueInspector.cpp
--- a/Engine/Plugins/Runtime/Metasound/Source/MetasoundEditor/Private/SMetasoundPinValueInspector.cpp
+++ b/Engine/Plugins/Runtime/Metasound/Source/MetasoundEditor/Private/SMetasoundPinValueInspector.cpp
@@ -6,6 +6,7 @@
 #include "Layout/Geometry.h"
 #include "MetasoundEditor.h"
 #include "MetasoundEditorGraphBuilder.h"
+#include "MetasoundEditorGraphInputNode.h"
 #include "MetasoundEditorGraphNode.h"
 #include "MetasoundEditorSettings.h"
 #include "MetasoundFrontendDocument.h"
@@ -125,9 +126,9 @@ namespace Metasound
 			using namespace Frontend;
 
 			FConstOutputHandle OutputHandle = GetReroutedOutputHandle();
-			ColorLiterals = PinValueInspectorPrivate::InitializeLiterals(GetReroutedOutputHandle()->GetDataType());
+			ColorLiterals = PinValueInspectorPrivate::InitializeLiterals(ReroutedOutputHandle->GetDataType());
 			Update();
-			DisplayName = FGraphBuilder::GetDisplayName(*GetReroutedOutputHandle());
+			DisplayName = FGraphBuilder::GetDisplayName(*ReroutedOutputHandle);
 		}
 
 		FColorPickerArgs FMetasoundNumericDebugLineItem::InitPickerArgs()
@@ -145,7 +146,7 @@ namespace Metasound
 		{
 			if (UEdGraphPin* ReroutedOutputPin = FGraphBuilder::FindReroutedOutputPin(GraphPinObj))
 			{
-				return Cast<UMetasoundEditorGraphNode>(GraphPinObj->GetOwningNode());
+				return Cast<UMetasoundEditorGraphNode>(ReroutedOutputPin->GetOwningNode());
 			}
 
 			return nullptr;
@@ -181,9 +182,47 @@ namespace Metasound
 			return FGraphBuilder::FindReroutedOutputHandleFromPin(GraphPinObj);
 		}
 
+		bool FMetasoundNumericDebugLineItem::GetNodeIDAndOutputName(FGuid& OutNodeID, FName& OutName) const
+		{
+			if (UEdGraphPin* ReroutedOutputPin = FGraphBuilder::FindReroutedOutputPin(GraphPinObj))
+			{
+				const UMetasoundEditorGraphNode& OutputsNode = GetReroutedNodeChecked();
+				const FMetaSoundFrontendDocumentBuilder& Builder = OutputsNode.GetBuilderChecked().GetConstBuilder();
+				if (const UMetasoundEditorGraphInputNode* InputNode = Cast<UMetasoundEditorGraphInputNode>(&OutputsNode))
+				{
+					if (const UMetasoundEditorGraphMember* Member = InputNode->GetMember())
+					{
+						OutNodeID = Member->GetMemberID();
+						if (const FMetasoundFrontendNode* RedirectedInputNode = Builder.FindNode(OutNodeID))
+						{
+							OutName = RedirectedInputNode->Interface.Outputs.Last().Name;
+							return true;
+						}
+					}
+				}
+				else
+				{
+					const FMetasoundFrontendVertexHandle VertexHandle = FGraphBuilder::GetPinVertexHandle(Builder, ReroutedOutputPin);
+					if (VertexHandle.IsSet())
+					{
+						OutNodeID = VertexHandle.NodeID;
+						if (const FMetasoundFrontendVertex* Vertex = FGraphBuilder::GetPinVertex(Builder, ReroutedOutputPin))
+						{
+							OutName = Vertex->Name;
+							return true;
+						}
+					}
+				}
+			}
+
+			OutNodeID = { };
+			OutName = { };
+			return false;
+		}
+
 		FLinearColor FMetasoundNumericDebugLineItem::GetEdgeStyleColorAtIndex(int32 InIndex) const
 		{
-			if (const FMetasoundFrontendEdgeStyle* EdgeStyle = GetEdgeStyle())
+			if (const FMetasoundFrontendEdgeStyle* EdgeStyle = FindConstEdgeStyle())
 			{
 				if (ensure(EdgeStyle->LiteralColorPairs.Num() > InIndex))
 				{
@@ -203,19 +242,17 @@ namespace Metasound
 			check(ParentObject);
 			ParentObject->Modify();
 
-			FGraphHandle GraphHandle = GetGraphHandle();
-			FMetasoundFrontendGraphStyle Style = GraphHandle->GetGraphStyle();
-
-			const FConstOutputHandle OutputHandle = GetReroutedOutputHandle();
-			const FGuid NodeID = OutputHandle->GetOwningNodeID();
-			const FName OutputName = OutputHandle->GetName();
-			Style.EdgeStyles.RemoveAllSwap([&NodeID, &OutputName](FMetasoundFrontendEdgeStyle& EdgeStyle)
+			FGuid NodeID;
+			FName OutputName;
+			if (GetNodeIDAndOutputName(NodeID, OutputName))
 			{
-				return EdgeStyle.NodeID == NodeID && EdgeStyle.OutputName == OutputName;
-			});
-
-			GraphHandle->SetGraphStyle(Style);
-			bIsValueColorizationEnabled = false;
+				if (const UMetasoundEditorGraphNode* Node = GetReroutedNode(); ensure(Node))
+				{
+					FMetaSoundFrontendDocumentBuilder& Builder = Node->GetBuilderChecked().GetBuilder();
+					Builder.RemoveEdgeStyle(NodeID, OutputName);
+					bIsValueColorizationEnabled = false;
+				}
+			}
 		}
 
 		void FMetasoundNumericDebugLineItem::EnableValueColorization()
@@ -234,27 +271,11 @@ namespace Metasound
 				{ ColorLiterals[1], FLinearColor::White }
 			};
 
-			FGraphHandle GraphHandle = GetGraphHandle();
-			FMetasoundFrontendGraphStyle Style = GraphHandle->GetGraphStyle();
-			FConstOutputHandle OutputHandle = GetReroutedOutputHandle();
-
-			const FGuid NodeID = OutputHandle->GetOwningNodeID();
-			const FName OutputName = OutputHandle->GetName();
-			FMetasoundFrontendEdgeStyle* EdgeStyle = Style.EdgeStyles.FindByPredicate([&NodeID, &OutputName](FMetasoundFrontendEdgeStyle& EdgeStyle)
-			{
-				return EdgeStyle.NodeID == NodeID && EdgeStyle.OutputName == OutputName;
-			});
-
-			if (EdgeStyle)
+			if (FMetasoundFrontendEdgeStyle* EdgeStyle = FindOrAddEdgeStyle())
 			{
 				EdgeStyle->LiteralColorPairs = MoveTemp(DefaultPairs);
 			}
-			else
-			{
-				Style.EdgeStyles.Emplace(FMetasoundFrontendEdgeStyle{ NodeID, OutputName, MoveTemp(DefaultPairs) });
-			}
 
-			GraphHandle->SetGraphStyle(Style);
 			bIsValueColorizationEnabled = true;
 		}
 
@@ -272,25 +293,13 @@ namespace Metasound
 			check(ParentObject);
 			ParentObject->Modify();
 
-			FGraphHandle GraphHandle = GetGraphHandle();
-			FMetasoundFrontendGraphStyle Style = GraphHandle->GetGraphStyle();
-
-			const FGuid NodeID = OutputHandle->GetOwningNodeID();
-			const FName OutputName = GetReroutedOutputHandle()->GetName();
-			FMetasoundFrontendEdgeStyle* EdgeStyle = Style.EdgeStyles.FindByPredicate([&NodeID, &OutputName](FMetasoundFrontendEdgeStyle& EdgeStyle)
-			{
-				return EdgeStyle.NodeID == NodeID && EdgeStyle.OutputName == OutputName;
-			});
-
-			if (EdgeStyle)
+			if (FMetasoundFrontendEdgeStyle* EdgeStyle = FindOrAddEdgeStyle())
 			{
 				if (ensure(InIndex < EdgeStyle->LiteralColorPairs.Num()))
 				{
 					EdgeStyle->LiteralColorPairs[InIndex].Color = InColor;
 				}
 			}
-
-			GraphHandle->SetGraphStyle(Style);
 		}
 
 		void FMetasoundNumericDebugLineItem::OnValueCommitted(const FText& InValueText, FMetasoundFrontendLiteral& OutNewLiteral, int32 InIndex)
@@ -307,27 +316,15 @@ namespace Metasound
 			check(ParentObject);
 			ParentObject->Modify();
 
-			const FGuid NodeID = OutputHandle->GetOwningNodeID();
-			const FName OutputName = OutputHandle->GetName();
-
-			FGraphHandle OwningGraph = GetGraphHandle();
-			FMetasoundFrontendGraphStyle Style = OwningGraph->GetGraphStyle();
-
-			bool bUpdated = false;
-			for (FMetasoundFrontendEdgeStyle& EdgeStyle : Style.EdgeStyles)
+			if (FMetasoundFrontendEdgeStyle* EdgeStyle = FindOrAddEdgeStyle())
 			{
-				if (EdgeStyle.NodeID == NodeID && EdgeStyle.OutputName == OutputName)
+				if (ensure(EdgeStyle->LiteralColorPairs.Num() > InIndex))
 				{
-					if (ensure(EdgeStyle.LiteralColorPairs.Num() > InIndex))
-					{
-						const FName DataType = OutputHandle->GetDataType();
-						PinValueInspectorPrivate::SetLiteralFromText(DataType, InValueText, OutNewLiteral);
-						EdgeStyle.LiteralColorPairs[InIndex].Value = OutNewLiteral;
-					}
+					const FName DataType = OutputHandle->GetDataType();
+					PinValueInspectorPrivate::SetLiteralFromText(DataType, InValueText, OutNewLiteral);
+					EdgeStyle->LiteralColorPairs[InIndex].Value = OutNewLiteral;
 				}
 			}
-
-			OwningGraph->SetGraphStyle(Style);
 		}
 
 		Frontend::FGraphHandle FMetasoundNumericDebugLineItem::GetGraphHandle()
@@ -356,25 +353,58 @@ namespace Metasound
 			return nullptr;
 		}
 
-		const FMetasoundFrontendEdgeStyle* FMetasoundNumericDebugLineItem::GetEdgeStyle() const
+		const FMetasoundFrontendEdgeStyle* FMetasoundNumericDebugLineItem::FindConstEdgeStyle() const
 		{
 			using namespace Frontend;
 
-			FConstGraphHandle OwningGraph = GetConstGraphHandle();
-			if (!OwningGraph->IsValid())
+			FGuid NodeID;
+			FName OutputName;
+			if (GetNodeIDAndOutputName(NodeID, OutputName))
 			{
-				return nullptr;
+				if (const UMetasoundEditorGraphNode* Node = GetReroutedNode())
+				{
+					const FMetaSoundFrontendDocumentBuilder& Builder = Node->GetBuilderChecked().GetConstBuilder();
+					return Builder.FindConstEdgeStyle(NodeID, OutputName);
+				}
 			}
 
-			const FConstOutputHandle OutputHandle = GetConstReroutedOutputHandle();
-			const FGuid NodeID = OutputHandle->GetOwningNodeID();
-			const FName OutputName = OutputHandle->GetName();
-			const FMetasoundFrontendGraphStyle& Style = OwningGraph->GetGraphStyle();
+			return nullptr;
+		}
 
-			return Style.EdgeStyles.FindByPredicate([&NodeID, &OutputName](FMetasoundFrontendEdgeStyle& EdgeStyle)
+		FMetasoundFrontendEdgeStyle* FMetasoundNumericDebugLineItem::FindEdgeStyle()
+		{
+			using namespace Frontend;
+
+			FGuid NodeID;
+			FName OutputName;
+			if (GetNodeIDAndOutputName(NodeID, OutputName))
 			{
-				return EdgeStyle.NodeID == NodeID && EdgeStyle.OutputName == OutputName;
-			});
+				if (const UMetasoundEditorGraphNode* Node = GetReroutedNode())
+				{
+					FMetaSoundFrontendDocumentBuilder& Builder = Node->GetBuilderChecked().GetBuilder();
+					return Builder.FindEdgeStyle(NodeID, OutputName);
+				}
+			}
+
+			return nullptr;
+		}
+
+		FMetasoundFrontendEdgeStyle* FMetasoundNumericDebugLineItem::FindOrAddEdgeStyle()
+		{
+			using namespace Frontend;
+
+			FGuid NodeID;
+			FName OutputName;
+			if (GetNodeIDAndOutputName(NodeID, OutputName))
+			{
+				if (const UMetasoundEditorGraphNode* Node = GetReroutedNode())
+				{
+					FMetaSoundFrontendDocumentBuilder& Builder = Node->GetBuilderChecked().GetBuilder();
+					return &Builder.FindOrAddEdgeStyle(NodeID, OutputName);
+				}
+			}
+
+			return nullptr;
 		}
 
 		void FMetasoundNumericDebugLineItem::Update()
@@ -383,24 +413,26 @@ namespace Metasound
 
 			bIsValueColorizationEnabled = false;
 
-			const FConstOutputHandle OutputHandle = GetConstReroutedOutputHandle();
-			const FGuid NodeID = OutputHandle->GetOwningNodeID();
-			const FName OutputName = OutputHandle->GetName();
-			if (const FMetasoundFrontendEdgeStyle* EdgeStyle = GetEdgeStyle())
+			FGuid NodeID;
+			FName OutputName;
+			if (GetNodeIDAndOutputName(NodeID, OutputName))
 			{
-				if (ensure(EdgeStyle->LiteralColorPairs.Num() == 2))
+				if (const FMetasoundFrontendEdgeStyle* EdgeStyle = FindConstEdgeStyle())
 				{
-					const FMetasoundFrontendLiteral& CurrentMin = EdgeStyle->LiteralColorPairs[0].Value;
-					ColorLiterals[0] = CurrentMin;
-					const FMetasoundFrontendLiteral& CurrentMax = EdgeStyle->LiteralColorPairs[1].Value;
-					ColorLiterals[1] = CurrentMax;
-					bIsValueColorizationEnabled = true;
+					if (ensure(EdgeStyle->LiteralColorPairs.Num() == 2))
+					{
+						const FMetasoundFrontendLiteral& CurrentMin = EdgeStyle->LiteralColorPairs[0].Value;
+						ColorLiterals[0] = CurrentMin;
+						const FMetasoundFrontendLiteral& CurrentMax = EdgeStyle->LiteralColorPairs[1].Value;
+						ColorLiterals[1] = CurrentMax;
+						bIsValueColorizationEnabled = true;
+					}
 				}
-			}
 
-			if (FGraphConnectionManager* ConnectionManager = GetConnectionManager())
-			{
-				Message = FString::Format(TEXT("Value: {0}"), { GetValueStringFunction(*ConnectionManager, NodeID, OutputName) });
+				if (FGraphConnectionManager* ConnectionManager = GetConnectionManager())
+				{
+					Message = FString::Format(TEXT("Value: {0}"), { GetValueStringFunction(*ConnectionManager, NodeID, OutputName) });
+				}
 			}
 		}
 
@@ -413,7 +445,7 @@ namespace Metasound
 				return FReply::Unhandled();
 			}
 
-			const FMetasoundFrontendEdgeStyle* EdgeStyle = GetEdgeStyle();
+			const FMetasoundFrontendEdgeStyle* EdgeStyle = FindEdgeStyle();
 			if (!EdgeStyle)
 			{
 				return FReply::Unhandled();
diff --git a/Engine/Plugins/Runtime/Metasound/Source/MetasoundEditor/Private/SMetasoundPinValueInspector.h b/Engine/Plugins/Runtime/Metasound/Source/MetasoundEditor/Private/SMetasoundPinValueInspector.h
--- a/Engine/Plugins/Runtime/Metasound/Source/MetasoundEditor/Private/SMetasoundPinValueInspector.h
+++ b/Engine/Plugins/Runtime/Metasound/Source/MetasoundEditor/Private/SMetasoundPinValueInspector.h
@@ -71,7 +71,7 @@ namespace Metasound
 
 			FGraphConnectionManager* GetConnectionManager();
 
-			const FMetasoundFrontendEdgeStyle* GetEdgeStyle() const;
+			const FMetasoundFrontendEdgeStyle* FindConstEdgeStyle() const;
 
 			void Update();
 
@@ -92,6 +92,9 @@ namespace Metasound
 
 		private:
 			FLinearColor GetEdgeStyleColorAtIndex(int32 InIndex) const;
+			FMetasoundFrontendEdgeStyle* FindOrAddEdgeStyle();
+			FMetasoundFrontendEdgeStyle* FindEdgeStyle();
+			bool GetNodeIDAndOutputName(FGuid& OutNodeID, FName& OutName) const;
 
 			void DisableValueColorization();
 			void EnableValueColorization();
diff --git a/Engine/Plugins/Runtime/Metasound/Source/MetasoundFrontend/Private/MetasoundFrontendDocumentAccessPtr.cpp b/Engine/Plugins/Runtime/Metasound/Source/MetasoundFrontend/Private/MetasoundFrontendDocumentAccessPtr.cpp
--- a/Engine/Plugins/Runtime/Metasound/Source/MetasoundFrontend/Private/MetasoundFrontendDocumentAccessPtr.cpp
+++ b/Engine/Plugins/Runtime/Metasound/Source/MetasoundFrontend/Private/MetasoundFrontendDocumentAccessPtr.cpp
@@ -18,7 +18,8 @@ namespace Metasound
 				{
 					if (FMetaSoundFrontendDocumentBuilder* Builder = BuilderRegistry->FindBuilder(InGraphClass.Metadata.GetClassName(), { }))
 					{
-						return Builder->FindBuildGraphChecked();
+						// Const cast and not exposed to public API to dissuade from direct manipulation of graphs as this can corrupt the internal document cache
+						return const_cast<FMetasoundFrontendGraph&>(Builder->FindConstBuildGraphChecked());
 					}
 				}
 
diff --git a/Engine/Plugins/Runtime/Metasound/Source/MetasoundFrontend/Private/MetasoundFrontendDocumentBuilder.cpp b/Engine/Plugins/Runtime/Metasound/Source/MetasoundFrontend/Private/MetasoundFrontendDocumentBuilder.cpp
--- a/Engine/Plugins/Runtime/Metasound/Source/MetasoundFrontend/Private/MetasoundFrontendDocumentBuilder.cpp
+++ b/Engine/Plugins/Runtime/Metasound/Source/MetasoundFrontend/Private/MetasoundFrontendDocumentBuilder.cpp
@@ -1627,7 +1627,7 @@ const FMetasoundFrontendNode* FMetaSoundFrontendDocumentBuilder::DuplicateGraphO
 	return AddGraphOutput(ClassOutput, &PageID);
 }
 
-FMetasoundFrontendGraph& FMetaSoundFrontendDocumentBuilder::FindBuildGraphChecked()
+FMetasoundFrontendGraph& FMetaSoundFrontendDocumentBuilder::FindBuildGraphChecked() const
 {
 	return GetDocumentChecked().RootGraph.FindGraphChecked(BuildPageID);
 }
@@ -1688,7 +1688,51 @@ TArray<const FMetasoundFrontendEdge*> FMetaSoundFrontendDocumentBuilder::FindEdg
 	return EdgeCache.FindEdges(InNodeID, InVertexID);
 }
 
-#if WITH_EDITOR
+#if WITH_EDITORONLY_DATA
+const FMetasoundFrontendEdgeStyle* FMetaSoundFrontendDocumentBuilder::FindConstEdgeStyle(const FGuid& InNodeID, FName OutputName, const FGuid* InPageID) const
+{
+	auto IsEdgeStyle = [&InNodeID, &OutputName](const FMetasoundFrontendEdgeStyle& EdgeStyle)
+	{
+		return EdgeStyle.NodeID == InNodeID && EdgeStyle.OutputName == OutputName;
+	};
+
+	const FGuid& PageID = InPageID ? *InPageID : BuildPageID;
+	const FMetasoundFrontendDocument& Document = GetConstDocumentChecked();
+	const FMetasoundFrontendGraph& Graph = Document.RootGraph.FindConstGraphChecked(PageID);
+	return Graph.Style.EdgeStyles.FindByPredicate(IsEdgeStyle);
+}
+
+FMetasoundFrontendEdgeStyle* FMetaSoundFrontendDocumentBuilder::FindEdgeStyle(const FGuid& InNodeID, FName OutputName, const FGuid* InPageID)
+{
+	auto IsEdgeStyle = [&InNodeID, &OutputName](const FMetasoundFrontendEdgeStyle& EdgeStyle)
+	{
+		return EdgeStyle.NodeID == InNodeID && EdgeStyle.OutputName == OutputName;
+	};
+
+	const FGuid& PageID = InPageID ? *InPageID : BuildPageID;
+	FMetasoundFrontendDocument& Document = GetDocumentChecked();
+	FMetasoundFrontendGraph& Graph = Document.RootGraph.FindGraphChecked(PageID);
+	return Graph.Style.EdgeStyles.FindByPredicate(IsEdgeStyle);
+}
+
+FMetasoundFrontendEdgeStyle& FMetaSoundFrontendDocumentBuilder::FindOrAddEdgeStyle(const FGuid& InNodeID, FName OutputName, const FGuid* InPageID)
+{
+	if (FMetasoundFrontendEdgeStyle* Style = FindEdgeStyle(InNodeID, OutputName, InPageID))
+	{
+		return *Style;
+	}
+
+	const FGuid& PageID = InPageID ? *InPageID : BuildPageID;
+	FMetasoundFrontendDocument& Document = GetDocumentChecked();
+	FMetasoundFrontendGraph& Graph = Document.RootGraph.FindGraphChecked(PageID);
+	FMetasoundFrontendEdgeStyle& EdgeStyle = Graph.Style.EdgeStyles.AddDefaulted_GetRef();
+
+	checkf(ContainsNode(InNodeID), TEXT("Cannot add edge style for node that does not exist"));
+	EdgeStyle.NodeID = InNodeID;
+	EdgeStyle.OutputName = OutputName;
+	return EdgeStyle;
+}
+
 const FMetaSoundFrontendGraphComment* FMetaSoundFrontendDocumentBuilder::FindGraphComment(const FGuid& InCommentID, const FGuid* InPageID) const
 {
 	check(InCommentID.IsValid());
@@ -1705,7 +1749,7 @@ FMetaSoundFrontendGraphComment* FMetaSoundFrontendDocumentBuilder::FindGraphComm
 	TMap<FGuid, FMetaSoundFrontendGraphComment>& Comments = Document.RootGraph.FindGraphChecked(PageID).Style.Comments;
 	return Comments.Find(InCommentID);
 }
-#endif // WITH_EDITOR
+#endif // WITH_EDITORONLY_DATA
 
 bool FMetaSoundFrontendDocumentBuilder::FindInterfaceInputNodes(FName InterfaceName, TArray<const FMetasoundFrontendNode*>& OutInputs, const FGuid* InPageID) const
 {
@@ -2087,7 +2131,7 @@ const FMetasoundFrontendVertex* FMetaSoundFrontendDocumentBuilder::FindNodeOutpu
 	return nullptr;
 }
 
-#if WITH_EDITOR
+#if WITH_EDITORONLY_DATA
 FMetaSoundFrontendGraphComment& FMetaSoundFrontendDocumentBuilder::FindOrAddGraphComment(const FGuid& InCommentID, const FGuid* InPageID)
 {
 	check(InCommentID.IsValid());
@@ -2096,7 +2140,7 @@ FMetaSoundFrontendGraphComment& FMetaSoundFrontendDocumentBuilder::FindOrAddGrap
 	TMap<FGuid, FMetaSoundFrontendGraphComment>& Comments = Document.RootGraph.FindGraphChecked(PageID).Style.Comments;
 	return Comments.FindOrAdd(InCommentID);
 }
-#endif // WITH_EDITOR
+#endif // WITH_EDITORONLY_DATA
 
 FMetasoundFrontendClassName FMetaSoundFrontendDocumentBuilder::GenerateNewClassName()
 {
@@ -2962,6 +3006,21 @@ bool FMetaSoundFrontendDocumentBuilder::RemoveEdge(const FMetasoundFrontendEdge&
 	return false;
 }
 
+#if WITH_EDITORONLY_DATA
+bool FMetaSoundFrontendDocumentBuilder::RemoveEdgeStyle(const FGuid& InNodeID, FName OutputName, const FGuid* InPageID)
+{
+	auto IsEdgeStyle = [&InNodeID, &OutputName](const FMetasoundFrontendEdgeStyle& EdgeStyle)
+	{
+		return EdgeStyle.NodeID == InNodeID && EdgeStyle.OutputName == OutputName;
+	};
+
+	const FGuid& PageID = InPageID ? *InPageID : BuildPageID;
+	FMetasoundFrontendDocument& Document = GetDocumentChecked();
+	FMetasoundFrontendGraph& Graph = Document.RootGraph.FindGraphChecked(PageID);
+	return Graph.Style.EdgeStyles.RemoveAllSwap(IsEdgeStyle) > 0;
+}
+#endif // WITH_EDITORONLY_DATA
+
 bool FMetaSoundFrontendDocumentBuilder::RemoveNamedEdges(const TSet<Metasound::Frontend::FNamedEdge>& InNamedEdgesToRemove, TArray<FMetasoundFrontendEdge>* OutRemovedEdges, const FGuid* InPageID)
 {
 	using namespace Metasound::Frontend;
@@ -3178,6 +3237,17 @@ bool FMetaSoundFrontendDocumentBuilder::RemoveEdgesFromNodeOutput(const FGuid& I
 		FEdgeModifyDelegates& EdgeDelegates = DocumentDelegates->FindEdgeDelegatesChecked(PageID);
 		for (int32 Index : IndicesCopy)
 		{
+#if WITH_EDITORONLY_DATA
+			if (const FMetasoundFrontendVertex* Vertex = FindNodeOutput(InNodeID, InVertexID))
+			{
+				auto IsEdgeStyle = [&InNodeID, OutputName = Vertex->Name](const FMetasoundFrontendEdgeStyle& EdgeStyle)
+				{
+					return EdgeStyle.NodeID == InNodeID && EdgeStyle.OutputName == OutputName;
+				};
+				Graph.Style.EdgeStyles.RemoveAllSwap(IsEdgeStyle);
+			}
+#endif // WITH_EDITORONLY_DATA
+
 			const int32 LastIndex = Graph.Edges.Num() - 1;
 			EdgeDelegates.OnRemoveSwappingEdge.Broadcast(Index, LastIndex);
 			Graph.Edges.RemoveAtSwap(Index, EAllowShrinking::No);
@@ -3198,8 +3268,20 @@ bool FMetaSoundFrontendDocumentBuilder::RemoveEdgeToNodeInput(const FGuid& InNod
 	if (const int32* IndexPtr = EdgeCache.FindEdgeIndexToNodeInput(InNodeID, InVertexID))
 	{
 		FMetasoundFrontendGraph& Graph = GetDocumentChecked().RootGraph.FindGraphChecked(PageID);
-		const FEdgeModifyDelegates& EdgeDelegates = DocumentDelegates->FindEdgeDelegatesChecked(PageID);
 		const int32 Index = *IndexPtr; // Copy off indices as the pointer may be modified when notifying the cache below
+
+#if WITH_EDITORONLY_DATA
+		if (const FMetasoundFrontendVertex* Vertex = FindNodeOutput(InNodeID, Graph.Edges[Index].FromVertexID))
+		{
+			auto IsEdgeStyle = [&InNodeID, OutputName = Vertex->Name](const FMetasoundFrontendEdgeStyle& EdgeStyle)
+			{
+				return EdgeStyle.NodeID == InNodeID && EdgeStyle.OutputName == OutputName;
+			};
+			Graph.Style.EdgeStyles.RemoveAllSwap(IsEdgeStyle);
+		}
+#endif // WITH_EDITORONLY_DATA
+
+		const FEdgeModifyDelegates& EdgeDelegates = DocumentDelegates->FindEdgeDelegatesChecked(PageID);
 		const int32 LastIndex = Graph.Edges.Num() - 1;
 		EdgeDelegates.OnRemoveSwappingEdge.Broadcast(Index, LastIndex);
 		Graph.Edges.RemoveAtSwap(Index, EAllowShrinking::No);
@@ -3214,23 +3296,21 @@ bool FMetaSoundFrontendDocumentBuilder::RemoveEdgeToNodeInput(const FGuid& InNod
 	return false;
 }
 
-#if WITH_EDITOR
+#if WITH_EDITORONLY_DATA
 bool FMetaSoundFrontendDocumentBuilder::RemoveGraphComment(const FGuid& InCommentID, const FGuid* InPageID)
 {
 	FMetasoundFrontendDocument& Document = GetDocumentChecked();
 	FMetasoundFrontendGraph& Graph = Document.RootGraph.FindGraphChecked(InPageID ? *InPageID : BuildPageID);
 	if (Graph.Style.Comments.Remove(InCommentID) > 0)
 	{
-#if WITH_EDITORONLY_DATA
 		Document.Metadata.ModifyContext.SetDocumentModified();
-#endif // WITH_EDITORONLY_DATA
 
 		return true;
 	}
 
 	return false;
 }
-#endif // WITH_EDITOR
+#endif // WITH_EDITORONLY_DATA
 
 bool FMetaSoundFrontendDocumentBuilder::RemoveGraphInput(FName InInputName)
 {
@@ -3478,7 +3558,7 @@ bool FMetaSoundFrontendDocumentBuilder::RemoveNode(const FGuid& InNodeID, const
 	return false;
 }
 
-#if WITH_EDITOR
+#if WITH_EDITORONLY_DATA
 int32 FMetaSoundFrontendDocumentBuilder::RemoveNodeLocation(const FGuid& InNodeID, const FGuid* InLocationGuid, const FGuid* InPageID)
 {
 	using namespace Metasound;
@@ -3505,7 +3585,7 @@ int32 FMetaSoundFrontendDocumentBuilder::RemoveNodeLocation(const FGuid& InNodeI
 
 	return 0;
 }
-#endif // WITH_EDITOR
+#endif // WITH_EDITORONLY_DATA
 
 bool FMetaSoundFrontendDocumentBuilder::RemoveUnusedDependencies()
 {
diff --git a/Engine/Plugins/Runtime/Metasound/Source/MetasoundFrontend/Private/MetasoundFrontendGraphController.cpp b/Engine/Plugins/Runtime/Metasound/Source/MetasoundFrontend/Private/MetasoundFrontendGraphController.cpp
--- a/Engine/Plugins/Runtime/Metasound/Source/MetasoundFrontend/Private/MetasoundFrontendGraphController.cpp
+++ b/Engine/Plugins/Runtime/Metasound/Source/MetasoundFrontend/Private/MetasoundFrontendGraphController.cpp
@@ -48,7 +48,8 @@ namespace Metasound
 			{
 				if (FMetaSoundFrontendDocumentBuilder* Builder = BuilderRegistry->FindBuilder(InGraphClass.Metadata.GetClassName(), { }))
 				{
-					return Builder->FindBuildGraphChecked();
+					// const cast to dissuade mutable accessor for external API as the controller API is actively deprecated across engine releases.
+					return const_cast<FMetasoundFrontendGraph&>(Builder->FindConstBuildGraphChecked());
 				}
 			}
 
@@ -777,8 +778,9 @@ namespace Metasound
 		{
 			if (FMetasoundFrontendGraphClass* GraphClass = GraphClassPtr.Get())
 			{
-				FindBuildGraphChecked(*GraphClass).Nodes.Reset();
-				FindBuildGraphChecked(*GraphClass).Edges.Reset();
+				FMetasoundFrontendGraph& Graph = FindBuildGraphChecked(*GraphClass);
+				Graph.Nodes.Reset();
+				Graph.Edges.Reset();
 				GraphClass->Interface.Inputs.Reset();
 				GraphClass->Interface.Outputs.Reset();
 				GraphClass->PresetOptions.InputsInheritingDefault.Reset();
@@ -790,6 +792,8 @@ namespace Metasound
 				FMetasoundFrontendDocument* Document = DocumentPtr.Get();
 				check(Document);
 				Document->Metadata.MemberMetadata.Reset();
+
+				Graph.Style.EdgeStyles.Reset();
 #endif // WITH_EDITORONLY_DATA
 			}
 		}
@@ -1708,14 +1712,22 @@ namespace Metasound
 		{
 			if (FMetasoundFrontendGraphClass* GraphClass = GraphClassPtr.Get())
 			{
-				auto IsEdgeForThisNode = [&](const FMetasoundFrontendEdge& ConDesc) { return (ConDesc.FromNodeID == InDesc.GetID()) || (ConDesc.ToNodeID == InDesc.GetID()); };
-
 				// Remove any reference connections
-				int32 NumRemoved = FindBuildGraphChecked(*GraphClass).Edges.RemoveAll(IsEdgeForThisNode);
+				FMetasoundFrontendGraph& Graph = FindBuildGraphChecked(*GraphClass);
+
+#if WITH_EDITORONLY_DATA
+				{
+					auto IsStyleForThisNode = [&](const FMetasoundFrontendEdgeStyle& EdgeStyle) { return (EdgeStyle.NodeID == InDesc.GetID()); };
+					Graph.Style.EdgeStyles.RemoveAllSwap(IsStyleForThisNode);
+				}
+#endif // WITH_EDITORONLY_DATA
+
+				auto IsEdgeForThisNode = [&](const FMetasoundFrontendEdge& ConDesc) { return (ConDesc.FromNodeID == InDesc.GetID()) || (ConDesc.ToNodeID == InDesc.GetID()); };
+				int32 NumRemoved = Graph.Edges.RemoveAll(IsEdgeForThisNode);
 
 				auto IsNodeWithID = [&](const FMetasoundFrontendNode& Desc) { return InDesc.GetID() == Desc.GetID(); };
 
-				NumRemoved += FindBuildGraphChecked(*GraphClass).Nodes.RemoveAll(IsNodeWithID);
+				NumRemoved += Graph.Nodes.RemoveAll(IsNodeWithID);
 				OwningDocument->RemoveUnreferencedDependencies();
 
 #if WITH_EDITOR
diff --git a/Engine/Plugins/Runtime/Metasound/Source/MetasoundFrontend/Private/MetasoundFrontendInputController.cpp b/Engine/Plugins/Runtime/Metasound/Source/MetasoundFrontend/Private/MetasoundFrontendInputController.cpp
--- a/Engine/Plugins/Runtime/Metasound/Source/MetasoundFrontend/Private/MetasoundFrontendInputController.cpp
+++ b/Engine/Plugins/Runtime/Metasound/Source/MetasoundFrontend/Private/MetasoundFrontendInputController.cpp
@@ -361,7 +361,7 @@ namespace Metasound
 			return false;
 		}
 
-		bool FBaseInputController::Disconnect(IOutputController& InController) 
+		bool FBaseInputController::Disconnect(IOutputController& InController)
 		{
 			if (FMetasoundFrontendGraph* Graph = GraphPtr.Get())
 			{
@@ -396,14 +396,22 @@ namespace Metasound
 			if (FMetasoundFrontendGraph* Graph = GraphPtr.Get())
 			{
 				const FGuid NodeID = GetOwningNodeID();
+
+#if WITH_EDITORONLY_DATA
+				{
+					const FName OutputName = GetConnectedOutput()->GetName();
+					auto IsStyleForThisNode = [&](const FMetasoundFrontendEdgeStyle& EdgeStyle) { return EdgeStyle.NodeID == NodeID && OutputName == EdgeStyle.OutputName; };
+					Graph->Style.EdgeStyles.RemoveAllSwap(IsStyleForThisNode);
+				}
+#endif // WITH_EDITORONLY_DATA
+
 				FGuid VertexID = GetID();
 
 				auto EdgeHasMatchingDestination = [&](const FMetasoundFrontendEdge& Edge)
 				{
 					return (Edge.ToNodeID == NodeID) && (Edge.ToVertexID == VertexID);
 				};
-
-				int32 NumRemoved = Graph->Edges.RemoveAllSwap(EdgeHasMatchingDestination);
+				const int32 NumRemoved = Graph->Edges.RemoveAllSwap(EdgeHasMatchingDestination);
 				return NumRemoved > 0;
 			}
 
diff --git a/Engine/Plugins/Runtime/Metasound/Source/MetasoundFrontend/Public/MetasoundFrontendDocumentBuilder.h b/Engine/Plugins/Runtime/Metasound/Source/MetasoundFrontend/Public/MetasoundFrontendDocumentBuilder.h
--- a/Engine/Plugins/Runtime/Metasound/Source/MetasoundFrontend/Public/MetasoundFrontendDocumentBuilder.h
+++ b/Engine/Plugins/Runtime/Metasound/Source/MetasoundFrontend/Public/MetasoundFrontendDocumentBuilder.h
@@ -211,6 +211,9 @@ public:
 	const FMetasoundFrontendNode* DuplicateGraphOutput(const FMetasoundFrontendClassOutput& InClassOutput, const FName InName, const FGuid* InPageID = nullptr);
 
 #if WITH_EDITORONLY_DATA
+	const FMetasoundFrontendEdgeStyle* FindConstEdgeStyle(const FGuid& InNodeID, FName OutputName, const FGuid* InPageID = nullptr) const;
+	FMetasoundFrontendEdgeStyle* FindEdgeStyle(const FGuid& InNodeID, FName OutputName, const FGuid* InPageID = nullptr);
+	FMetasoundFrontendEdgeStyle& FindOrAddEdgeStyle(const FGuid& InNodeID, FName OutputName, const FGuid* InPageID = nullptr);
 	const FMetaSoundFrontendGraphComment* FindGraphComment(const FGuid& InCommentID, const FGuid* InPageID = nullptr) const;
 	FMetaSoundFrontendGraphComment* FindGraphComment(const FGuid& InCommentID, const FGuid* InPageID = nullptr);
 	FMetaSoundFrontendGraphComment& FindOrAddGraphComment(const FGuid& InCommentID, const FGuid* InPageID = nullptr);
@@ -234,7 +237,7 @@ public:
 	bool FindInterfaceInputNodes(FName InterfaceName, TArray<const FMetasoundFrontendNode*>& OutInputs, const FGuid* InPageID = nullptr) const;
 	bool FindInterfaceOutputNodes(FName InterfaceName, TArray<const FMetasoundFrontendNode*>& OutOutputs, const FGuid* InPageID = nullptr) const;
 
-	FMetasoundFrontendGraph& FindBuildGraphChecked();
+	// Accessor for the currently set build graph.
 	const FMetasoundFrontendGraph& FindConstBuildGraphChecked() const;
 
 	const FMetasoundFrontendNode* FindNode(const FGuid& InNodeID, const FGuid* InPageID = nullptr) const;
@@ -363,9 +366,10 @@ public:
 	bool RemoveEdgesFromNodeOutput(const FGuid& InNodeID, const FGuid& InVertexID, const FGuid* InPageID = nullptr);
 	bool RemoveEdgeToNodeInput(const FGuid& InNodeID, const FGuid& InVertexID, const FGuid* InPageID = nullptr);
 
-#if WITH_EDITOR
+#if WITH_EDITORONLY_DATA
+	bool RemoveEdgeStyle(const FGuid& InNodeID, FName OutputName, const FGuid* InPageID = nullptr);
 	bool RemoveGraphComment(const FGuid& InCommentID, const FGuid* InPageID = nullptr);
-#endif // WITH_EDITOR
+#endif // WITH_EDITORONLY_DATA
 
 	bool RemoveGraphInput(FName InInputName);
 	bool RemoveGraphOutput(FName InOutputName);
@@ -378,7 +382,7 @@ public:
 	bool RemoveNamedEdges(const TSet<Metasound::Frontend::FNamedEdge>& InNamedEdgesToRemove, TArray<FMetasoundFrontendEdge>* OutRemovedEdges = nullptr, const FGuid* InPageID = nullptr);
 	bool RemoveNode(const FGuid& InNodeID, const FGuid* InPageID = nullptr);
 
-#if WITH_EDITOR
+#if WITH_EDITORONLY_DATA
 	int32 RemoveNodeLocation(const FGuid& InNodeID, const FGuid* InLocationGuid = nullptr, const FGuid* InPageID = nullptr);
 #endif // WITH_EDITOR
 
@@ -530,6 +534,8 @@ private:
 	// edges, edges may then be invalid due to access type/DataType incompatibility).
 	bool ConformGraphOutputNodeToClass(const FMetasoundFrontendClassOutput& GraphOutput);
 
+	FMetasoundFrontendGraph& FindBuildGraphChecked() const;
+
 	bool FindNodeClassInterfaces(const FGuid& InNodeID, TSet<FMetasoundFrontendVersion>& OutInterfaces, const FGuid& InPageID) const;
 	FMetasoundFrontendNode* FindNodeInternal(const FGuid& InNodeID, const FGuid* InPageID = nullptr);
 
