Fix build (bad local merge)
#rnx

diff --git a/Engine/Plugins/Runtime/Metasound/Source/MetasoundEditor/Private/SMetasoundPinValueInspector.cpp b/Engine/Plugins/Runtime/Metasound/Source/MetasoundEditor/Private/SMetasoundPinValueInspector.cpp
--- a/Engine/Plugins/Runtime/Metasound/Source/MetasoundEditor/Private/SMetasoundPinValueInspector.cpp
+++ b/Engine/Plugins/Runtime/Metasound/Source/MetasoundEditor/Private/SMetasoundPinValueInspector.cpp
@@ -125,7 +125,7 @@ namespace Metasound
 		{
 			using namespace Frontend;
 
-			FConstOutputHandle OutputHandle = GetReroutedOutputHandle();
+			FConstOutputHandle ReroutedOutputHandle = GetReroutedOutputHandle();
 			ColorLiterals = PinValueInspectorPrivate::InitializeLiterals(ReroutedOutputHandle->GetDataType());
 			Update();
 			DisplayName = FGraphBuilder::GetDisplayName(*ReroutedOutputHandle);
