Part 1: Fix for default config having two default page values
#rnx

diff --git a/Engine/Plugins/Runtime/Metasound/Source/MetasoundEngine/Private/MetasoundSettings.cpp b/Engine/Plugins/Runtime/Metasound/Source/MetasoundEngine/Private/MetasoundSettings.cpp
--- a/Engine/Plugins/Runtime/Metasound/Source/MetasoundEngine/Private/MetasoundSettings.cpp
+++ b/Engine/Plugins/Runtime/Metasound/Source/MetasoundEngine/Private/MetasoundSettings.cpp
@@ -218,6 +218,15 @@ void UMetaSoundSettings::ConformPageSettings(bool bNotifyDefaultRenamed)
 
 	TargetPageNameOverride.Reset();
 
+	for (int32 Index = PageSettings.Num() - 1; Index >= 0; --Index)
+	{
+		FMetaSoundPageSettings& PageSetting = PageSettings[Index];
+		if (PageSetting.UniqueId == Metasound::Frontend::DefaultPageID || PageSetting.Name == Metasound::Frontend::DefaultPageName)
+		{
+			PageSettings.RemoveAt(Index);
+		}
+	}
+
 	if (bNotifyDefaultRenamed && bInvalidDefaultRenamed)
 	{
 		OnDefaultRenamed.Broadcast();
