Add missing static cast to inform literal of proper type initialization when initializing TriggerOnValueChangeNode template input

#rnx

diff --git a/Engine/Plugins/Runtime/Metasound/Source/MetasoundStandardNodes/Private/MetasoundTriggerOnValueChangeNode.cpp b/Engine/Plugins/Runtime/Metasound/Source/MetasoundStandardNodes/Private/MetasoundTriggerOnValueChangeNode.cpp
--- a/Engine/Plugins/Runtime/Metasound/Source/MetasoundStandardNodes/Private/MetasoundTriggerOnValueChangeNode.cpp
+++ b/Engine/Plugins/Runtime/Metasound/Source/MetasoundStandardNodes/Private/MetasoundTriggerOnValueChangeNode.cpp
@@ -154,7 +154,7 @@ namespace Metasound
 
 		static const FVertexInterface Interface(
 			FInputVertexInterface(
-				TInputDataVertex<ValueType>(METASOUND_GET_PARAM_NAME_AND_METADATA(InputValue), 1)
+				TInputDataVertex<ValueType>(METASOUND_GET_PARAM_NAME_AND_METADATA(InputValue), static_cast<ValueType>(1))
 			),
 			FOutputVertexInterface(
 				TOutputDataVertex<FTrigger>(METASOUND_GET_PARAM_NAME_AND_METADATA(OutputOnChange))
