Fix renaming MetaSound inputs/outputs not applying to non-focused graph pages

#rnx

diff --git a/Engine/Plugins/Runtime/Metasound/Source/MetasoundEditor/Private/MetasoundEditorGraph.cpp b/Engine/Plugins/Runtime/Metasound/Source/MetasoundEditor/Private/MetasoundEditorGraph.cpp
--- a/Engine/Plugins/Runtime/Metasound/Source/MetasoundEditor/Private/MetasoundEditorGraph.cpp
+++ b/Engine/Plugins/Runtime/Metasound/Source/MetasoundEditor/Private/MetasoundEditorGraph.cpp
@@ -317,10 +317,15 @@ void UMetasoundEditorGraphVertex::SetMemberName(const FName& InNewName, bool bPo
 		return;
 	}
 
-	FNodeHandle NodeHandle = GetNodeHandle();
-	if (NodeHandle->GetNodeName() == InNewName)
+	FMetaSoundFrontendDocumentBuilder& DocBuilder = GetFrontendBuilderChecked();
+	FName OldName;
+	if (const FMetasoundFrontendNode* Node = DocBuilder.FindNode(NodeID); ensure(Node))
 	{
-		return;
+		if (Node->Name == InNewName)
+		{
+			return;
+		}
+		OldName = Node->Name;
 	}
 
 	const FText TransactionLabel = FText::Format(LOCTEXT("RenameGraphVertexMemberNameFormat", "Set Metasound {0} MemberName"), GetGraphMemberLabel());
@@ -329,7 +334,7 @@ void UMetasoundEditorGraphVertex::SetMemberName(const FName& InNewName, bool bPo
 	Graph->Modify();
 	Graph->GetMetasoundChecked().Modify();
 
-	NodeHandle->SetNodeName(InNewName);
+	RenameFrontendMemberInternal(DocBuilder, OldName, InNewName);
 
 	const TArray<UMetasoundEditorGraphMemberNode*> Nodes = GetNodes();
 	for (UMetasoundEditorGraphMemberNode* Node : Nodes)
@@ -345,7 +350,6 @@ void UMetasoundEditorGraphVertex::SetMemberName(const FName& InNewName, bool bPo
 	}
 
 	Graph->RegisterGraphWithFrontend();
-	FGraphBuilder::GetOutermostMetaSoundChecked(*Graph).GetModifyContext().AddMemberIDsModified({ GetMemberID() });
 }
 
 FText UMetasoundEditorGraphVertex::GetDisplayName() const
@@ -1014,6 +1018,11 @@ void UMetasoundEditorGraphInput::SetMemberName(const FName& InNewName, bool bPos
 	Super::SetMemberName(InNewName, bPostTransaction);
 }
 
+bool UMetasoundEditorGraphInput::RenameFrontendMemberInternal(FMetaSoundFrontendDocumentBuilder& Builder, FName OldName, FName InNewName) const
+{
+	return Builder.SetGraphInputName(OldName, InNewName);
+}
+
 bool UMetasoundEditorGraphInput::Synchronize()
 {
 	bool bModified = Super::Synchronize();
@@ -1375,6 +1384,11 @@ bool UMetasoundEditorGraphOutput::Synchronize()
 	return bModified;
 }
 
+bool UMetasoundEditorGraphOutput::RenameFrontendMemberInternal(FMetaSoundFrontendDocumentBuilder& Builder, FName OldName, FName InNewName) const
+{
+	return Builder.SetGraphOutputName(OldName, InNewName);
+}
+
 void UMetasoundEditorGraphVariable::InitMember(FName InDataType, const FMetasoundFrontendLiteral& InDefaultLiteral, FGuid InVariableID)
 {
 	TypeName = InDataType;
diff --git a/Engine/Plugins/Runtime/Metasound/Source/MetasoundEditor/Public/MetasoundEditorGraph.h b/Engine/Plugins/Runtime/Metasound/Source/MetasoundEditor/Public/MetasoundEditorGraph.h
--- a/Engine/Plugins/Runtime/Metasound/Source/MetasoundEditor/Public/MetasoundEditorGraph.h
+++ b/Engine/Plugins/Runtime/Metasound/Source/MetasoundEditor/Public/MetasoundEditorGraph.h
@@ -290,6 +290,9 @@ protected:
 	/** If the vertex's name contains an interface namespace, return true and set OutInterface to that interface if provided. */
 	bool NameContainsInterfaceNamespace(FMetasoundFrontendInterface* OutInterface) const;
 
+	/** Implements frontend call to rename member */
+	virtual bool RenameFrontendMemberInternal(FMetaSoundFrontendDocumentBuilder& Builder, FName OldName, FName InNewName) const PURE_VIRTUAL(UMetasoundEditorGraphVertex::RenameFrontendMemberInternal, return false;)
+
 public:
 	/** Initializes all properties with the given parameters required to identify the frontend member from this editor graph member. */
 	void InitMember(FName InDataType, const FMetasoundFrontendLiteral& InDefaultLiteral, FGuid InNodeID, FMetasoundFrontendClassName&& InClassName);
@@ -396,6 +399,8 @@ protected:
 	UE_DEPRECATED(5.5, "EditorGraph vertices no longer generates node handles (use Builder API instead)")
 	virtual Metasound::Frontend::FNodeHandle AddNodeHandle(const FName& InNodeName, const Metasound::Editor::FCreateNodeVertexParams& InParams) override;
 
+	virtual bool RenameFrontendMemberInternal(FMetaSoundFrontendDocumentBuilder& Builder, FName OldName, FName InNewName) const override;
+
 	virtual EMetasoundFrontendClassType GetClassType() const override { return EMetasoundFrontendClassType::Input; }
 	virtual Metasound::Editor::ENodeSection GetSectionID() const override;
 };
@@ -427,6 +432,8 @@ protected:
 	UE_DEPRECATED(5.5, "No longer supported in favor of internal Editor::FGraphBuilder API which is actively being transitioned to using Document Builder API")
 	virtual Metasound::Frontend::FNodeHandle AddNodeHandle(const FName& InNodeName, const Metasound::Editor::FCreateNodeVertexParams& InParams) override;
 
+	virtual bool RenameFrontendMemberInternal(FMetaSoundFrontendDocumentBuilder& Builder, FName OldName, FName InNewName) const override;
+
 	virtual EMetasoundFrontendClassType GetClassType() const override { return EMetasoundFrontendClassType::Output; }
 	virtual Metasound::Editor::ENodeSection GetSectionID() const override;
 };
diff --git a/Engine/Plugins/Runtime/Metasound/Source/MetasoundEngine/Public/MetasoundBuilderBase.cpp b/Engine/Plugins/Runtime/Metasound/Source/MetasoundEngine/Public/MetasoundBuilderBase.cpp
--- a/Engine/Plugins/Runtime/Metasound/Source/MetasoundEngine/Public/MetasoundBuilderBase.cpp
+++ b/Engine/Plugins/Runtime/Metasound/Source/MetasoundEngine/Public/MetasoundBuilderBase.cpp
@@ -1032,6 +1032,12 @@ void UMetaSoundBuilderBase::SetGraphInputDefault(FName InputName, const FMetasou
 	OutResult = bSet ? EMetaSoundBuilderResult::Succeeded : EMetaSoundBuilderResult::Failed;
 }
 
+void UMetaSoundBuilderBase::SetGraphInputName(FName InputName, FName NewName, EMetaSoundBuilderResult& OutResult)
+{
+	const bool bSet = Builder.SetGraphInputName(InputName, NewName);
+	OutResult = bSet ? EMetaSoundBuilderResult::Succeeded : EMetaSoundBuilderResult::Failed;
+}
+
 void UMetaSoundBuilderBase::SetGraphOutputAccessType(FName OutputName, EMetasoundFrontendVertexAccessType AccessType, EMetaSoundBuilderResult& OutResult)
 {
 	const bool bSet = Builder.SetGraphOutputAccessType(OutputName, AccessType);
@@ -1041,6 +1047,11 @@ void UMetaSoundBuilderBase::SetGraphOutputAccessType(FName OutputName, EMetasoun
 void UMetaSoundBuilderBase::SetGraphOutputDataType(FName OutputName, FName DataType, EMetaSoundBuilderResult& OutResult)
 {
 	const bool bSet = Builder.SetGraphOutputDataType(OutputName, DataType);
+}
+
+void UMetaSoundBuilderBase::SetGraphOutputName(FName OutputName, FName NewName, EMetaSoundBuilderResult& OutResult)
+{
+	const bool bSet = Builder.SetGraphOutputName(OutputName, NewName);
 	OutResult = bSet ? EMetaSoundBuilderResult::Succeeded : EMetaSoundBuilderResult::Failed;
 }
 
diff --git a/Engine/Plugins/Runtime/Metasound/Source/MetasoundEngine/Public/MetasoundBuilderBase.h b/Engine/Plugins/Runtime/Metasound/Source/MetasoundEngine/Public/MetasoundBuilderBase.h
--- a/Engine/Plugins/Runtime/Metasound/Source/MetasoundEngine/Public/MetasoundBuilderBase.h
+++ b/Engine/Plugins/Runtime/Metasound/Source/MetasoundEngine/Public/MetasoundBuilderBase.h
@@ -409,6 +409,11 @@ public:
 	// Sets the input node's default value, overriding the default provided by the referenced graph if the graph is a preset.
 	UFUNCTION(BlueprintCallable, Category = "Audio|MetaSound|Builder", meta = (ExpandEnumAsExecs = "OutResult"))
 	void SetGraphInputDefault(FName InputName, const FMetasoundFrontendLiteral& Literal, EMetaSoundBuilderResult& OutResult);
+	
+	// Sets the given graph input's name to the new name. 
+	// Result succeeds if the name was successfully changed or the new name is the same as the old name, and fails if the given input name doesn't exist.
+	UFUNCTION(BlueprintCallable, Category = "Audio|MetaSound|Builder", meta = (ExpandEnumAsExecs = "OutResult"))
+	void SetGraphInputName(FName InputName, FName NewName, EMetaSoundBuilderResult& OutResult);
 
 	// Disconnects the given graph output's respective template nodes and sets the graph output's AccessType should it not match the current AccessType.
 	// Result succeeds if the AccessType was successfully changed or if the provided AccessType is already the output's current AccessType.
@@ -420,6 +425,11 @@ public:
 	UFUNCTION(BlueprintCallable, Category = "Audio|MetaSound|Builder", meta = (ExpandEnumAsExecs = "OutResult"))
 	void SetGraphOutputDataType(FName OutputName, FName DataType, EMetaSoundBuilderResult& OutResult);
 
+	// Sets the given graph output's name to the new name. 
+	// Result succeeds if the name was successfully changed or the new name is the same as the old name, and fails if the given output name doesn't exist.
+	UFUNCTION(BlueprintCallable, Category = "Audio|MetaSound|Builder", meta = (ExpandEnumAsExecs = "OutResult"))
+	void SetGraphOutputName(FName OutputName, FName NewName, EMetaSoundBuilderResult& OutResult);
+
 	void SetMemberMetadata(UMetaSoundFrontendMemberMetadata& NewMetadata);
 
 	UE_DEPRECATED(5.5, "Call directly on Frontend builder using 'GetBuilder'")
diff --git a/Engine/Plugins/Runtime/Metasound/Source/MetasoundFrontend/Private/MetasoundFrontendDocumentBuilder.cpp b/Engine/Plugins/Runtime/Metasound/Source/MetasoundFrontend/Private/MetasoundFrontendDocumentBuilder.cpp
--- a/Engine/Plugins/Runtime/Metasound/Source/MetasoundFrontend/Private/MetasoundFrontendDocumentBuilder.cpp
+++ b/Engine/Plugins/Runtime/Metasound/Source/MetasoundFrontend/Private/MetasoundFrontendDocumentBuilder.cpp
@@ -3710,7 +3710,7 @@ bool FMetaSoundFrontendDocumentBuilder::SetGraphInputAdvancedDisplay(const FName
 
 	if (const int32* Index = DocumentCache->GetInterfaceCache().FindInputIndex(InputName))
 	{
-		FMetasoundFrontendClassInput& GraphInput = RootGraph.Interface.Inputs[*Index];			
+		FMetasoundFrontendClassInput& GraphInput = RootGraph.Interface.Inputs[*Index];
 		if (GraphInput.Metadata.bIsAdvancedDisplay != InAdvancedDisplay)
 		{
 			GraphInput.Metadata.SetIsAdvancedDisplay(InAdvancedDisplay);
@@ -3964,7 +3964,7 @@ bool FMetaSoundFrontendDocumentBuilder::SetGraphOutputAdvancedDisplay(const FNam
 		if (GraphOutput.Metadata.bIsAdvancedDisplay != InAdvancedDisplay)
 		{
 			GraphOutput.Metadata.SetIsAdvancedDisplay(InAdvancedDisplay);
-			Document.Metadata.ModifyContext.AddMemberIDModified(GraphOutput.VertexID);	
+			Document.Metadata.ModifyContext.AddMemberIDModified(GraphOutput.VertexID);
 			return true;
 		}
 	}
@@ -3994,6 +3994,103 @@ bool FMetaSoundFrontendDocumentBuilder::SetGraphInputInheritsDefault(FName InNam
 	return false;
 }
 
+bool FMetaSoundFrontendDocumentBuilder::SetGraphInputName(FName InputName, FName NewName)
+{
+	using namespace Metasound::Frontend;
+
+	if (InputName == NewName)
+	{
+		return true;
+	}
+
+	const int32* Index = DocumentCache->GetInterfaceCache().FindInputIndex(InputName);
+	if (!Index)
+	{
+		return false;
+	}
+
+	FMetasoundFrontendDocument& Document = GetDocumentChecked();
+	FMetasoundFrontendGraphClass& RootGraph = Document.RootGraph;
+
+	FMetasoundFrontendClassInput& GraphInput = RootGraph.Interface.Inputs[*Index];
+	GraphInput.Name = NewName;
+
+	RootGraph.IterateGraphPages([this, &GraphInput, &NewName](FMetasoundFrontendGraph& Graph)
+	{
+		const IDocumentGraphNodeCache& NodeCache = DocumentCache->GetNodeCache(Graph.PageID);
+		if (const int32* NodeIndex = NodeCache.FindNodeIndex(GraphInput.NodeID))
+		{
+			FMetasoundFrontendNode& Node = Graph.Nodes[*NodeIndex];
+			Node.Name = NewName;
+			for (FMetasoundFrontendVertex& Vertex : Node.Interface.Inputs)
+			{
+				Vertex.Name = NewName;
+			}
+			for (FMetasoundFrontendVertex& Vertex : Node.Interface.Outputs)
+			{
+				Vertex.Name = NewName;
+			}
+		}
+	});
+
+	DocumentDelegates->InterfaceDelegates.OnInputNameChanged.Broadcast(InputName, NewName);
+
+#if WITH_EDITORONLY_DATA
+	Document.Metadata.ModifyContext.AddMemberIDModified(GraphInput.NodeID);
+#endif // WITH_EDITORONLY_DATA
+
+	return true;
+}
+
+bool FMetaSoundFrontendDocumentBuilder::SetGraphOutputName(FName OutputName, FName NewName)
+{
+	using namespace Metasound::Frontend;
+
+	if (OutputName == NewName)
+	{
+		return true;
+	}
+
+	const int32* Index = DocumentCache->GetInterfaceCache().FindOutputIndex(OutputName);
+	if (!Index)
+	{
+		return false;
+	}
+
+	FMetasoundFrontendDocument& Document = GetDocumentChecked();
+	FMetasoundFrontendGraphClass& GraphClass = Document.RootGraph;
+	FMetasoundFrontendClassInterface& Interface = GraphClass.Interface;
+	Interface.UpdateChangeID();
+
+	FMetasoundFrontendClassOutput& GraphOutput = Interface.Outputs[*Index];
+	GraphOutput.Name = NewName;
+	
+	GraphClass.IterateGraphPages([this, &GraphOutput, &NewName](FMetasoundFrontendGraph& Graph)
+	{
+		const IDocumentGraphNodeCache& NodeCache = DocumentCache->GetNodeCache(Graph.PageID);
+		if (const int32* NodeIndex = NodeCache.FindNodeIndex(GraphOutput.NodeID))
+		{
+			FMetasoundFrontendNode& Node = Graph.Nodes[*NodeIndex];
+			Node.Name = NewName;
+			for (FMetasoundFrontendVertex& Vertex : Node.Interface.Inputs)
+			{
+				Vertex.Name = NewName;
+			}
+			for (FMetasoundFrontendVertex& Vertex : Node.Interface.Outputs)
+			{
+				Vertex.Name = NewName;
+			}
+		}
+	});
+	DocumentDelegates->InterfaceDelegates.OnOutputNameChanged.Broadcast(OutputName, NewName);
+	
+#if WITH_EDITORONLY_DATA
+	Document.Metadata.ModifyContext.AddMemberIDModified(GraphOutput.NodeID);
+#endif // WITH_EDITORONLY_DATA
+
+	return true;
+}
+
 bool FMetaSoundFrontendDocumentBuilder::SetGraphOutputAccessType(FName OutputName, EMetasoundFrontendVertexAccessType AccessType)
 {
 	using namespace Metasound::Frontend;
diff --git a/Engine/Plugins/Runtime/Metasound/Source/MetasoundFrontend/Private/MetasoundFrontendDocumentCache.cpp b/Engine/Plugins/Runtime/Metasound/Source/MetasoundFrontend/Private/MetasoundFrontendDocumentCache.cpp
--- a/Engine/Plugins/Runtime/Metasound/Source/MetasoundFrontend/Private/MetasoundFrontendDocumentCache.cpp
+++ b/Engine/Plugins/Runtime/Metasound/Source/MetasoundFrontend/Private/MetasoundFrontendDocumentCache.cpp
@@ -407,6 +407,8 @@ namespace Metasound::Frontend
 		OutDelegates.OnOutputAdded.AddSP(this, &FDocumentGraphInterfaceCache::OnOutputAdded);
 		OutDelegates.OnRemovingInput.AddSP(this, &FDocumentGraphInterfaceCache::OnRemovingInput);
 		OutDelegates.OnRemovingOutput.AddSP(this, &FDocumentGraphInterfaceCache::OnRemovingOutput);
+		OutDelegates.OnInputNameChanged.AddSP(this, &FDocumentGraphInterfaceCache::OnInputNameChanged);
+		OutDelegates.OnOutputNameChanged.AddSP(this, &FDocumentGraphInterfaceCache::OnOutputNameChanged);
 	}
 
 	void FDocumentGraphInterfaceCache::OnInputAdded(int32 NewIndex)
@@ -424,6 +426,16 @@ namespace Metasound::Frontend
 		++TransactionCount;
 	}
 
+	void FDocumentGraphInterfaceCache::OnInputNameChanged(FName OldName, FName NewName)
+	{
+		int32* Index = InputNameToIndex.Find(OldName);
+		check(Index);
+		InputNameToIndex.Remove(OldName);
+		InputNameToIndex.Add(NewName, *Index);
+
+		++TransactionCount;
+	}
+
 	void FDocumentGraphInterfaceCache::OnOutputAdded(int32 NewIndex)
 	{
 		const FMetasoundFrontendDocument& Document = Parent->GetDocument();
@@ -434,6 +446,16 @@ namespace Metasound::Frontend
 		++TransactionCount;
 	}
 
+	void FDocumentGraphInterfaceCache::OnOutputNameChanged(FName OldName, FName NewName)
+	{
+		int32* Index = OutputNameToIndex.Find(OldName);
+		check(Index);
+		OutputNameToIndex.Remove(OldName);
+		OutputNameToIndex.Add(NewName, *Index);
+
+		++TransactionCount;
+	}
+
 	void FDocumentGraphInterfaceCache::OnRemovingInput(int32 IndexBeingRemoved)
 	{
 		const FMetasoundFrontendDocument& Document = Parent->GetDocument();
@@ -450,8 +472,6 @@ namespace Metasound::Frontend
 		const FMetasoundFrontendGraphClass& GraphClass = Document.RootGraph;
 		const FMetasoundFrontendClassOutput& Output = GraphClass.Interface.Outputs[IndexBeingRemoved];
 		OutputNameToIndex.Remove(Output.Name);
-
-		++TransactionCount;
 	}
 
 	FDocumentGraphNodeCache::FDocumentGraphNodeCache()
diff --git a/Engine/Plugins/Runtime/Metasound/Source/MetasoundFrontend/Private/MetasoundFrontendDocumentCache.h b/Engine/Plugins/Runtime/Metasound/Source/MetasoundFrontend/Private/MetasoundFrontendDocumentCache.h
--- a/Engine/Plugins/Runtime/Metasound/Source/MetasoundFrontend/Private/MetasoundFrontendDocumentCache.h
+++ b/Engine/Plugins/Runtime/Metasound/Source/MetasoundFrontend/Private/MetasoundFrontendDocumentCache.h
@@ -134,9 +134,13 @@ namespace Metasound::Frontend
 
 	private:
 		void Init(FInterfaceModifyDelegates& OutDelegates);
+
 		void OnInputAdded(int32 NewIndex);
 		void OnInputDefaultChanged(int32 NewIndex);
+		void OnInputNameChanged(FName OldName, FName NewName);
+
 		void OnOutputAdded(int32 NewIndex);
+		void OnOutputNameChanged(FName OldName, FName NewName);
 		void OnRemovingInput(int32 IndexBeingRemoved);
 		void OnRemovingOutput(int32 IndexBeingRemoved);
 
diff --git a/Engine/Plugins/Runtime/Metasound/Source/MetasoundFrontend/Public/MetasoundFrontendDocumentBuilder.h b/Engine/Plugins/Runtime/Metasound/Source/MetasoundFrontend/Public/MetasoundFrontendDocumentBuilder.h
--- a/Engine/Plugins/Runtime/Metasound/Source/MetasoundFrontend/Public/MetasoundFrontendDocumentBuilder.h
+++ b/Engine/Plugins/Runtime/Metasound/Source/MetasoundFrontend/Public/MetasoundFrontendDocumentBuilder.h
@@ -438,6 +438,9 @@ public:
 	bool SetGraphInputDefault(FName InputName, FMetasoundFrontendLiteral InDefaultLiteral, const FGuid* InPageID = nullptr);
 	bool SetGraphInputDefaults(FName InputName, TArray<FMetasoundFrontendClassInputDefault> Defaults);
 
+	// Sets a given graph input's name to a new name. Succeeds if the graph output exists and the new name is set (or is the same as the old name).
+	bool SetGraphInputName(FName InputName, FName InName);
+
 #if WITH_EDITORONLY_DATA
 	// Sets the given output`s IsAdvancedDisplay state. AdvancedDisplay pins are hidden in the node by default.
 	// returns true if state was changed.
@@ -454,6 +457,9 @@ public:
 	// already the given DataType.
 	bool SetGraphOutputDataType(FName OutputName, FName DataType);
 
+	// Sets a given graph output's name to a new name. Succeeds if the graph output exists and the new name is set (or is the same as the old name).
+	bool SetGraphOutputName(FName InputName, FName InName);
+
 	// Sets the given graph variable's default.
 	bool SetGraphVariableDefault(FName VariableName, FMetasoundFrontendLiteral InDefaultLiteral, const FGuid* InPageID = nullptr);
 
diff --git a/Engine/Plugins/Runtime/Metasound/Source/MetasoundFrontend/Public/MetasoundFrontendDocumentModifyDelegates.h b/Engine/Plugins/Runtime/Metasound/Source/MetasoundFrontend/Public/MetasoundFrontendDocumentModifyDelegates.h
--- a/Engine/Plugins/Runtime/Metasound/Source/MetasoundFrontend/Public/MetasoundFrontendDocumentModifyDelegates.h
+++ b/Engine/Plugins/Runtime/Metasound/Source/MetasoundFrontend/Public/MetasoundFrontendDocumentModifyDelegates.h
@@ -12,7 +12,7 @@ DECLARE_MULTICAST_DELEGATE_OneParam(FOnMetaSoundFrontendDocumentMutateInterfaceA
 DECLARE_MULTICAST_DELEGATE_TwoParams(FOnMetaSoundFrontendDocumentRemoveSwappingArray, int32 /* Index */, int32 /* LastIndex */);
 DECLARE_MULTICAST_DELEGATE_TwoParams(FOnMetaSoundFrontendDocumentRenameClass, const int32 /* Index */, const FMetasoundFrontendClassName& /* NewName */);
 DECLARE_MULTICAST_DELEGATE_ThreeParams(FOnMetaSoundFrontendDocumentMutateNodeInputLiteralArray, int32 /* NodeIndex */, int32 /* VertexIndex */, int32 /* LiteralIndex */);
-
+DECLARE_MULTICAST_DELEGATE_TwoParams(FOnMetaSoundFrontendDocumentRenameVertex, FName /* OldName */, FName /* NewName */);
 
 namespace Metasound::Frontend
 {
@@ -39,9 +39,11 @@ namespace Metasound::Frontend
 
 		FOnMetaSoundFrontendDocumentMutateArray OnInputAdded;
 		FOnMetaSoundFrontendDocumentMutateArray OnInputDefaultChanged;
+		FOnMetaSoundFrontendDocumentRenameVertex OnInputNameChanged;
 		FOnMetaSoundFrontendDocumentMutateArray OnRemovingInput;
 
 		FOnMetaSoundFrontendDocumentMutateArray OnOutputAdded;
+		FOnMetaSoundFrontendDocumentRenameVertex OnOutputNameChanged;
 		FOnMetaSoundFrontendDocumentMutateArray OnRemovingOutput;
 	};
 
