Trivial fix to regression where widget options are now visible for MetaSound triggers

#rnx

diff --git a/Engine/Plugins/Runtime/Metasound/Source/MetasoundEditor/Private/MetasoundNodeDetailCustomization.cpp b/Engine/Plugins/Runtime/Metasound/Source/MetasoundEditor/Private/MetasoundNodeDetailCustomization.cpp
--- a/Engine/Plugins/Runtime/Metasound/Source/MetasoundEditor/Private/MetasoundNodeDetailCustomization.cpp
+++ b/Engine/Plugins/Runtime/Metasound/Source/MetasoundEditor/Private/MetasoundNodeDetailCustomization.cpp
@@ -416,29 +416,30 @@ namespace Metasound
 				return Row;
 			};
 
-			// Enable widget options for editable inputs only
 			bool bShowWidgetOptions = false;
-			if (const UMetasoundEditorGraphInput* ParentMember = Cast<UMetasoundEditorGraphInput>(InLiteral.FindMember()))
+
+			const UMetasoundEditorSettings* EditorSettings = GetDefault<UMetasoundEditorSettings>();
+			check(EditorSettings);
+
+			if (EditorSettings->bUseAudioMaterialWidgets)
 			{
-				if (const UMetasoundEditorGraph* OwningGraph = ParentMember->GetOwningGraph())
+				if (const UMetasoundEditorGraphInput* Member = Cast<UMetasoundEditorGraphInput>(InLiteral.FindMember()))
 				{
-					bShowWidgetOptions = OwningGraph->IsEditable();
+					if (Member->GetDataType() != GetMetasoundDataTypeName<FTrigger>())
+					{
+						if (const UMetasoundEditorGraph* OwningGraph = Member->GetOwningGraph())
+						{
+							bShowWidgetOptions = OwningGraph->IsEditable();
+						}
+					}
 				}
 			}
 
-			// add input widget properties
 			if (bShowWidgetOptions)
 			{
-				const UMetasoundEditorSettings* EditorSettings = GetDefault<UMetasoundEditorSettings>();
-				check(EditorSettings);
-
-				if (EditorSettings->bUseAudioMaterialWidgets)
+				if (InLiteral.GetDataType() != Metasound::GetMetasoundDataTypeName<Metasound::FTrigger>())
 				{
-					if (InLiteral.GetDataType() != Metasound::GetMetasoundDataTypeName<Metasound::FTrigger>())
-					{
-
-						AddOptionPropRow(GET_MEMBER_NAME_CHECKED(UMetasoundEditorGraphMemberDefaultBool, WidgetType));
-					}
+					AddOptionPropRow(GET_MEMBER_NAME_CHECKED(UMetasoundEditorGraphMemberDefaultBool, WidgetType));
 				}
 			}
 		}
