Version document to accommodate potential multiple page input default values being added. This could happen before a bug was fixed in the MetaSoundSettings conform step ensuring no duplicate default page setting existed in the project page settings array (default was moved to a separate member).

#rnx

diff --git a/Engine/Plugins/Runtime/Metasound/Source/MetasoundFrontend/Private/MetasoundFrontendDocumentVersioning.cpp b/Engine/Plugins/Runtime/Metasound/Source/MetasoundFrontend/Private/MetasoundFrontendDocumentVersioning.cpp
--- a/Engine/Plugins/Runtime/Metasound/Source/MetasoundFrontend/Private/MetasoundFrontendDocumentVersioning.cpp
+++ b/Engine/Plugins/Runtime/Metasound/Source/MetasoundFrontend/Private/MetasoundFrontendDocumentVersioning.cpp
@@ -634,8 +634,6 @@ PRAGMA_ENABLE_DEPRECATION_WARNINGS
 			const FSoftObjectPath* Path = nullptr;
 
 		public:
-			static const FMetasoundFrontendVersionNumber TargetVersion;
-
 			FVersionDocument_1_12(FName InName, const FSoftObjectPath& InAssetPath)
 				: Name(InName)
 				, Path(&InAssetPath)
@@ -665,12 +663,10 @@ PRAGMA_ENABLE_DEPRECATION_WARNINGS
 			}
 		};
 
-			/** Versions document from 1.12 to 1.13. */
+		/** Versions document from 1.12 to 1.13. */
 		class FVersionDocument_1_13 : public FVersionDocumentTransform
 		{
 		public:
-			static const FMetasoundFrontendVersionNumber TargetVersion;
-
 			virtual ~FVersionDocument_1_13() = default;
 
 			FMetasoundFrontendVersionNumber GetTargetVersion() const override
@@ -684,6 +680,55 @@ PRAGMA_ENABLE_DEPRECATION_WARNINGS
 			}
 		};
 
+		/** Versions document from 1.13 to 1.14. */
+		class FVersionDocument_1_14 : public FVersionDocumentTransform
+		{
+		public:
+			virtual ~FVersionDocument_1_14() = default;
+
+			FMetasoundFrontendVersionNumber GetTargetVersion() const override
+			{
+				return { 1, 14 };
+			}
+
+			void TransformInternal(FMetaSoundFrontendDocumentBuilder& OutBuilder) const override
+			{
+				// Between 1.13 and 1.14, it was possible to add multiple default input page values
+				// due to missing versioning logic. This fixes that issue if any data was serialized
+				// to a MetaSound Asset by removing any extraneous default data (early values in the
+				// array were stale).
+				FMetasoundFrontendDocument& Document = const_cast<FMetasoundFrontendDocument&>(OutBuilder.GetConstDocumentChecked());
+				for (FMetasoundFrontendClassInput& Input : Document.RootGraph.Interface.Inputs)
+				{
+					int32 PageIDIndex = INDEX_NONE;
+					TArray<FMetasoundFrontendClassInputDefault>& Defaults = const_cast<TArray<FMetasoundFrontendClassInputDefault>&>(Input.GetDefaults());
+					for (int32 Index = 0; Index < Defaults.Num(); ++Index)
+					{
+						FMetasoundFrontendClassInputDefault& Default = Defaults[Index];
+						const bool bIsDefault = Default.PageID == Frontend::DefaultPageID;
+						if (bIsDefault)
+						{
+							if (PageIDIndex == INDEX_NONE)
+							{
+								PageIDIndex = Index;
+							}
+							else
+							{
+								Defaults.RemoveAt(PageIDIndex);
+								break;
+							}
+						}
+					}
+				}
+
+				// Safeguards against prior fix-up corrupting any cached data
+				if (IDocumentBuilderRegistry* BuilderRegistry = IDocumentBuilderRegistry::Get())
+				{
+					BuilderRegistry->ReloadBuilder(Document.RootGraph.Metadata.GetClassName());
+				}
+			}
+		};
+
 		bool VersionBuilderDocument(FMetaSoundFrontendDocumentBuilder& Builder)
 		{
 			UObject& DocObject = Builder.CastDocumentObjectChecked<UObject>();
@@ -693,6 +738,7 @@ PRAGMA_ENABLE_DEPRECATION_WARNINGS
 			bool bWasUpdated = false;
 			bWasUpdated |= FVersionDocument_1_12(Name, Path).Transform(Builder);
 			bWasUpdated |= FVersionDocument_1_13().Transform(Builder);
+			bWasUpdated |= FVersionDocument_1_14().Transform(Builder);
 
 			return bWasUpdated;
 		}
diff --git a/Engine/Plugins/Runtime/Metasound/Source/MetasoundFrontend/Private/MetasoundFrontendDocumentVersioning.h b/Engine/Plugins/Runtime/Metasound/Source/MetasoundFrontend/Private/MetasoundFrontendDocumentVersioning.h
--- a/Engine/Plugins/Runtime/Metasound/Source/MetasoundFrontend/Private/MetasoundFrontendDocumentVersioning.h
+++ b/Engine/Plugins/Runtime/Metasound/Source/MetasoundFrontend/Private/MetasoundFrontendDocumentVersioning.h
@@ -16,7 +16,7 @@ namespace Metasound::Frontend
 {
 	static FMetasoundFrontendVersionNumber GetMaxDocumentVersion()
 	{
-		return FMetasoundFrontendVersionNumber { 1, 13 };
+		return FMetasoundFrontendVersionNumber { 1, 14 };
 	}
 
 	// Versions Frontend Document. Passed as AssetBase for backward compat to
