Fix bad merge (Feedback from prior submission)
#rnx

diff --git a/Engine/Plugins/Runtime/Metasound/Source/MetasoundFrontend/Private/MetasoundFrontendDocumentCache.cpp b/Engine/Plugins/Runtime/Metasound/Source/MetasoundFrontend/Private/MetasoundFrontendDocumentCache.cpp
--- a/Engine/Plugins/Runtime/Metasound/Source/MetasoundFrontend/Private/MetasoundFrontendDocumentCache.cpp
+++ b/Engine/Plugins/Runtime/Metasound/Source/MetasoundFrontend/Private/MetasoundFrontendDocumentCache.cpp
@@ -472,6 +472,8 @@ namespace Metasound::Frontend
 		const FMetasoundFrontendGraphClass& GraphClass = Document.RootGraph;
 		const FMetasoundFrontendClassOutput& Output = GraphClass.Interface.Outputs[IndexBeingRemoved];
 		OutputNameToIndex.Remove(Output.Name);
+
+		++TransactionCount;
 	}
 
 	FDocumentGraphNodeCache::FDocumentGraphNodeCache()
