Put metasound operator pool csv stats behind a cvar that's disabled by default.

Disabled by default to reduce csv bloat. Can be enabled in replays to get the data when desired.
Eventually we'll convert this to csv events to reduce the overhead.

#rnx

diff --git a/Engine/Plugins/Runtime/Metasound/Source/MetasoundGenerator/Private/MetasoundOperatorCacheStatTracker.cpp b/Engine/Plugins/Runtime/Metasound/Source/MetasoundGenerator/Private/MetasoundOperatorCacheStatTracker.cpp
--- a/Engine/Plugins/Runtime/Metasound/Source/MetasoundGenerator/Private/MetasoundOperatorCacheStatTracker.cpp
+++ b/Engine/Plugins/Runtime/Metasound/Source/MetasoundGenerator/Private/MetasoundOperatorCacheStatTracker.cpp
@@ -20,6 +20,13 @@ namespace Metasound::Engine
 
 	namespace Private
 	{
+		static bool bMetasoundOperatorPoolCsvStatsEnabled = false;
+		static FAutoConsoleVariableRef CVarMetasoundOperatorPoolCsvStatsEnabled(
+			TEXT("au.MetaSound.OperatorPool.CsvStatsEnabled"),
+			bMetasoundOperatorPoolCsvStatsEnabled,
+			TEXT("If we should record operator pool stats to the csv.")
+		);
+
 		static TAutoConsoleVariable<bool> CVarCacheMissCsvStatsEnabled(
 			TEXT("au.MetaSound.OperatorPool.CacheMissCsvStatsEnabled"),
 			true,
@@ -34,17 +41,17 @@ namespace Metasound::Engine
 
 	FOperatorCacheStatTracker::FOperatorCacheStatTracker()
 	{
-#if CSV_PROFILER
+#if CSV_PROFILER_STATS 
 		CsvEndFrameDelegateHandle = FCsvProfiler::Get()->OnCSVProfileEndFrame().AddRaw(this, &FOperatorCacheStatTracker::OnCsvProfileEndFrame);
-#endif // CSV_PROFILER
+#endif // CSV_PROFILER_STATS 
 	}
 
 	FOperatorCacheStatTracker::~FOperatorCacheStatTracker()
 	{
-#if CSV_PROFILER
+#if CSV_PROFILER_STATS 
 		FCsvProfiler::Get()->OnCSVProfileEndFrame().Remove(CsvEndFrameDelegateHandle);
 		CsvEndFrameDelegateHandle.Reset();
-#endif // CSV_PROFILER
+#endif // CSV_PROFILER_STATS 
 	}
 
 	void FOperatorCacheStatTracker::RecordPreCacheRequest(const FOperatorBuildData& BuildData, int32 NumInstancesToBuild, int32 NumInstancesInCache)
@@ -148,14 +155,15 @@ namespace Metasound::Engine
 	{
 		if (!bCacheHit)
 		{
-#if CSV_PROFILER
-			if (Private::CVarCacheMissCsvStatsEnabled->GetBool() &&
+#if CSV_PROFILER_STATS 
+			if (Private::bMetasoundOperatorPoolCsvStatsEnabled &&
+				Private::CVarCacheMissCsvStatsEnabled->GetBool() &&
 				Context.GraphInstanceName != NAME_Name)
 			{
 				const FString GraphName = Private::GetGraphName(Context.GraphInstanceName);
 				FCsvProfiler::Get()->RecordCustomStat(*GraphName, CSV_CATEGORY_INDEX(Metasound_OperatorCacheMiss), 1, ECsvCustomStatOp::Accumulate);
 			}
-#endif // CSV_PROFILER
+#endif // CSV_PROFILER_STATS 
 			return;
 		}
 
@@ -227,7 +235,12 @@ namespace Metasound::Engine
 
 	void FOperatorCacheStatTracker::OnCsvProfileEndFrame()
 	{
-#if CSV_PROFILER
+#if CSV_PROFILER_STATS 
+		if (!Private::bMetasoundOperatorPoolCsvStatsEnabled)
+		{
+			return;
+		}
+
 		QUICK_SCOPE_CYCLE_COUNTER(OperatorCacheStatTracker_RecordStats)
 
 		FScopeLock Lock(&CriticalSection);
@@ -257,7 +270,7 @@ namespace Metasound::Engine
 				FCsvProfiler::Get()->RecordCustomStat(Entry.GraphName, CSV_CATEGORY_INDEX(MetaSound_OperatorCacheUtilization), UtilizationRatio, ECsvCustomStatOp::Set);
 			}
 		}
-#endif // CSV_PROFILER
+#endif // CSV_PROFILER_STATS 
 	}
 } // namespace Metasound::Private
 #endif // METASOUND_OPERATORCACHEPROFILER_ENABLED
