Remove private dependency from MetaSoundEditor module

#rnx

diff --git a/Engine/Plugins/Runtime/Metasound/Source/MetasoundEditor/MetasoundEditor.Build.cs b/Engine/Plugins/Runtime/Metasound/Source/MetasoundEditor/MetasoundEditor.Build.cs
--- a/Engine/Plugins/Runtime/Metasound/Source/MetasoundEditor/MetasoundEditor.Build.cs
+++ b/Engine/Plugins/Runtime/Metasound/Source/MetasoundEditor/MetasoundEditor.Build.cs
@@ -9,12 +9,6 @@ namespace UnrealBuildTool.Rules
 	{
 		public MetasoundEditor(ReadOnlyTargetRules Target) : base(Target)
 		{
-			PrivateIncludePaths.AddRange(
-				new string[] {
-					Path.Combine(GetModuleDirectory("MetasoundFrontend"), "Private"),
-				}
-			);
-
 			PrivateDependencyModuleNames.AddRange
 			(
 				new string[]
diff --git a/Engine/Plugins/Runtime/Metasound/Source/MetasoundEditor/Private/MetasoundEditorGraphBuilder.cpp b/Engine/Plugins/Runtime/Metasound/Source/MetasoundEditor/Private/MetasoundEditorGraphBuilder.cpp
--- a/Engine/Plugins/Runtime/Metasound/Source/MetasoundEditor/Private/MetasoundEditorGraphBuilder.cpp
+++ b/Engine/Plugins/Runtime/Metasound/Source/MetasoundEditor/Private/MetasoundEditorGraphBuilder.cpp
@@ -32,7 +32,6 @@
 #include "MetasoundFrontendDataTypeRegistry.h"
 #include "MetasoundFrontendDocumentAccessPtr.h"
 #include "MetasoundFrontendDocumentBuilder.h"
-#include "MetasoundFrontendDocumentVersioning.h"
 #include "MetasoundFrontendQuery.h"
 #include "MetasoundFrontendQuerySteps.h"
 #include "MetasoundFrontendRegistries.h"
