Trivial: Deprecate 'DefaultName' MetaSound EdGraph function no longer used

#rnx

diff --git a/Engine/Plugins/Runtime/Metasound/Source/MetasoundEditor/Private/MetasoundEditorGraph.cpp b/Engine/Plugins/Runtime/Metasound/Source/MetasoundEditor/Private/MetasoundEditorGraph.cpp
--- a/Engine/Plugins/Runtime/Metasound/Source/MetasoundEditor/Private/MetasoundEditorGraph.cpp
+++ b/Engine/Plugins/Runtime/Metasound/Source/MetasoundEditor/Private/MetasoundEditorGraph.cpp
@@ -701,11 +701,14 @@ void UMetasoundEditorGraphMemberDefaultLiteral::PostEditChangeChainProperty(FPro
 		if (const FProperty* ChildProperty = MemberNode->GetValue())
 		{
 			const FName ChildPropertyName = ChildProperty->GetFName();
+
+PRAGMA_DISABLE_DEPRECATION_WARNINGS
 			if (ChildPropertyName.IsEqual(GetDefaultsPropertyName()) || ChildPropertyName.IsEqual(GetDefaultPropertyName()))
 			{
 				ResolvePageDefaults();
 				SortPageDefaults();
 			}
+PRAGMA_ENABLE_DEPRECATION_WARNINGS
 		}
 	}
 
diff --git a/Engine/Plugins/Runtime/Metasound/Source/MetasoundEditor/Public/MetasoundEditorGraph.h b/Engine/Plugins/Runtime/Metasound/Source/MetasoundEditor/Public/MetasoundEditorGraph.h
--- a/Engine/Plugins/Runtime/Metasound/Source/MetasoundEditor/Public/MetasoundEditorGraph.h
+++ b/Engine/Plugins/Runtime/Metasound/Source/MetasoundEditor/Public/MetasoundEditorGraph.h
@@ -135,6 +135,8 @@ public:
 
 	UMetasoundEditorGraphMember* FindMember() const;
 
+	UE_DEPRECATED(5.5, "Due to serialization optimization, literals now inherit from LiteralMetadata and are no longer parented under members. "
+		"Use 'GetDefaultsPropertyName' instead to get transient paged defaults property name")
 	static FName GetDefaultPropertyName()
 	{
 		return "Default";
