[MetaSounds] - Editor crashes when promoting all to input on some nodes

#rnx

diff --git a/Engine/Plugins/Runtime/Metasound/Source/MetasoundEditor/Private/MetasoundEditor.cpp b/Engine/Plugins/Runtime/Metasound/Source/MetasoundEditor/Private/MetasoundEditor.cpp
--- a/Engine/Plugins/Runtime/Metasound/Source/MetasoundEditor/Private/MetasoundEditor.cpp
+++ b/Engine/Plugins/Runtime/Metasound/Source/MetasoundEditor/Private/MetasoundEditor.cpp
@@ -5415,7 +5415,8 @@ namespace Metasound
 					EMetaSoundBuilderResult Result = EMetaSoundBuilderResult::Failed;
 					FMetasoundFrontendLiteral Literal;
 					FGraphBuilder::GetPinLiteral(*Pin, Literal);
-					FMetaSoundBuilderNodeOutputHandle OutputHandle = Builder->AddGraphInputNode(Name, InputVertex->TypeName, Literal, Result, false);
+					bool bIsConstructorInput = DocBuilder.GetNodeInputAccessType(InputVertexHandle.NodeID, InputVertexHandle.VertexID) == EMetasoundFrontendVertexAccessType::Value;
+					FMetaSoundBuilderNodeOutputHandle OutputHandle = Builder->AddGraphInputNode(Name, InputVertex->TypeName, Literal, Result, bIsConstructorInput);
 					check(Result == EMetaSoundBuilderResult::Succeeded);
 
 					FVector2D Location = FVector2D(EdGraphNode->NodePosX, EdGraphNode->NodePosY);
@@ -5509,13 +5510,14 @@ namespace Metasound
 				FName PinName = Pair.Key.Key;
 				FName TypeName = Pair.Key.Value;
 				UEdGraphPin* SourcePin = Pair.Value[0];
-				
+				FMetasoundFrontendVertexHandle InputVertexHandle = FGraphBuilder::GetPinVertexHandle(DocBuilder, SourcePin);
 				FName InputName = FGraphBuilder::GenerateUniqueNameByClassType(ParentMetasound, EMetasoundFrontendClassType::Input, PinName.ToString());
 				
 				EMetaSoundBuilderResult Result = EMetaSoundBuilderResult::Failed;
 				FMetasoundFrontendLiteral Literal;
 				FGraphBuilder::GetPinLiteral(*SourcePin, Literal);
-				FMetaSoundBuilderNodeOutputHandle OutputHandle = Builder->AddGraphInputNode(InputName, TypeName, Literal, Result, false);
+				bool bIsConstructorInput = DocBuilder.GetNodeInputAccessType(InputVertexHandle.NodeID, InputVertexHandle.VertexID) == EMetasoundFrontendVertexAccessType::Value;
+				FMetaSoundBuilderNodeOutputHandle OutputHandle = Builder->AddGraphInputNode(InputName, TypeName, Literal, Result, bIsConstructorInput);
 				check(Result == EMetaSoundBuilderResult::Succeeded);
 
 				FVector2D* NodeOffset = NodeOffsets.Find(SourcePin->GetOwningNode()->NodeGuid);
diff --git a/Engine/Plugins/Runtime/Metasound/Source/MetasoundEditor/Private/MetasoundEditorCommands.cpp b/Engine/Plugins/Runtime/Metasound/Source/MetasoundEditor/Private/MetasoundEditorCommands.cpp
--- a/Engine/Plugins/Runtime/Metasound/Source/MetasoundEditor/Private/MetasoundEditorCommands.cpp
+++ b/Engine/Plugins/Runtime/Metasound/Source/MetasoundEditor/Private/MetasoundEditorCommands.cpp
@@ -33,8 +33,8 @@ namespace Metasound
 			UI_COMMAND(PromoteToVariable, "Promote To Graph Variable", "Promotes pin to graph variable.", EUserInterfaceActionType::Button, FInputChord());
 			UI_COMMAND(PromoteToDeferredVariable, "Promote To Deferred Graph Variable", "Promotes pin to deferred graph variable.", EUserInterfaceActionType::Button, FInputChord());
 
-			UI_COMMAND(PromoteAllToInput, "Promote All To Input", "Promotes unconnected node inputs to graph inputs.", EUserInterfaceActionType::Button, FInputChord());
-			UI_COMMAND(PromoteAllToCommonInputs, "Promote All To Common Inputs", "Promotes unconnected node input pins to graph inputs, sharing inputs if possible.", EUserInterfaceActionType::Button, FInputChord());
+			UI_COMMAND(PromoteAllToInput, "Promote All To Graph Input", "Promotes unconnected node inputs to graph inputs.", EUserInterfaceActionType::Button, FInputChord());
+			UI_COMMAND(PromoteAllToCommonInputs, "Promote All To Common Graph Inputs", "Promotes unconnected input pins from selected nodes to graph inputs, sharing inputs if possible.", EUserInterfaceActionType::Button, FInputChord());
 		}
 	} // namespace Editor
 } // namespace Metasound
