Fix MetaSound presets not copying sound wave properties on creation

#rnx

diff --git a/Engine/Plugins/Runtime/Metasound/Source/MetasoundEditor/Private/MetasoundEditorSubsystem.cpp b/Engine/Plugins/Runtime/Metasound/Source/MetasoundEditor/Private/MetasoundEditorSubsystem.cpp
--- a/Engine/Plugins/Runtime/Metasound/Source/MetasoundEditor/Private/MetasoundEditorSubsystem.cpp
+++ b/Engine/Plugins/Runtime/Metasound/Source/MetasoundEditor/Private/MetasoundEditorSubsystem.cpp
@@ -69,15 +69,25 @@ TScriptInterface<IMetaSoundDocumentInterface> UMetaSoundEditorSubsystem::BuildTo
 				InBuilder->Build(Parent, BuilderOptions);
 			}
 
-			// Apply template SoundWave settings
+			UMetaSoundBuilderBase& NewDocBuilder = FDocumentBuilderRegistry::GetChecked().FindOrBeginBuilding(*NewMetaSound);
+
+			EMetaSoundBuilderResult InjectResult = EMetaSoundBuilderResult::Failed;
+			constexpr bool bForceNodeCreation = true;
+			NewDocBuilder.InjectInputTemplateNodes(bForceNodeCreation, InjectResult);
+
+			FMetasoundAssetBase& Asset = NewDocBuilder.GetBuilder().GetMetasoundAsset();
+			Asset.RebuildReferencedAssetClasses();
+
+			// Apply template SoundWave settings 
+			// (must be post rebuilding referenced asset classes to find referenced asset to copy settings from)
 			{
 				const bool bIsSource = &MetaSoundUClass == UMetaSoundSource::StaticClass();
-				if (InBuilder->IsPreset())
+				if (NewDocBuilder.IsPreset())
 				{
 					// Only use referenced UObject's SoundWave settings for sources if not overridden 
 					if (TemplateSoundWave == nullptr && bIsSource)
 					{
-						if (const UObject* ReferencedObject = InBuilder->GetReferencedPresetAsset())
+						if (const UObject* ReferencedObject = NewDocBuilder.GetReferencedPresetAsset())
 						{
 							TemplateSoundWave = CastChecked<USoundWave>(ReferencedObject);
 						}
@@ -91,15 +101,6 @@ TScriptInterface<IMetaSoundDocumentInterface> UMetaSoundEditorSubsystem::BuildTo
 				}
 			}
 
-			UMetaSoundBuilderBase& NewDocBuilder = FDocumentBuilderRegistry::GetChecked().FindOrBeginBuilding(*NewMetaSound);
-
-			EMetaSoundBuilderResult InjectResult = EMetaSoundBuilderResult::Failed;
-			constexpr bool bForceNodeCreation = true;
-			NewDocBuilder.InjectInputTemplateNodes(bForceNodeCreation, InjectResult);
-
-			FMetasoundAssetBase& Asset = NewDocBuilder.GetBuilder().GetMetasoundAsset();
-			Asset.RebuildReferencedAssetClasses();
-
 			if (!bWasRooted)
 			{
 				InBuilder->RemoveFromRoot();
