- Uniform color fix-ups for MetaSound analyzers
- Move live MetaSound stats/audition settings to top left of graph
- Fix ordering of playback widget in input details panel to follow same order as page details panel

#rnx

diff --git a/Engine/Plugins/Runtime/AudioWidgets/Source/AudioWidgets/Private/AudioMeter.cpp b/Engine/Plugins/Runtime/AudioWidgets/Source/AudioWidgets/Private/AudioMeter.cpp
--- a/Engine/Plugins/Runtime/AudioWidgets/Source/AudioWidgets/Private/AudioMeter.cpp
+++ b/Engine/Plugins/Runtime/AudioWidgets/Source/AudioWidgets/Private/AudioMeter.cpp
@@ -160,6 +160,12 @@ const FText UAudioMeter::GetPaletteCategory()
 }
 #endif
 
+const FAudioMeterDefaultColorStyle& FAudioMeterDefaultColorStyle::GetDefault()
+{
+	static FAudioMeterDefaultColorStyle Style;
+	return Style;
+}
+
 namespace AudioWidgets
 {
 	FAudioMeter::FAudioMeter(int32 InNumChannels, UWorld& InWorld, TObjectPtr<UAudioBus> InExternalAudioBus)
@@ -167,19 +173,18 @@ namespace AudioWidgets
 	{
 	}
 
-	FAudioMeter::FAudioMeter(const int32 InNumChannels, const Audio::FDeviceId InAudioDeviceId, const TObjectPtr<UAudioBus> InExternalAudioBus)
+	FAudioMeter::FAudioMeter(const int32 InNumChannels, const Audio::FDeviceId InAudioDeviceId, const TObjectPtr<UAudioBus> InExternalAudioBus, const FAudioMeterDefaultColorStyle* AudioMeterWidgetStyle)
 	{
+		const FAudioMeterDefaultColorStyle* WidgetStyle = AudioMeterWidgetStyle ? AudioMeterWidgetStyle : &FAudioMeterDefaultColorStyle::GetDefault();
 		Widget = SNew(SAudioMeter)
 			.Orientation(EOrientation::Orient_Vertical)
 			.BackgroundColor(FLinearColor::Transparent)
-
-			// TODO: Move to editor style
-			.MeterBackgroundColor(FLinearColor(0.031f, 0.031f, 0.031f, 1.0f))
-			.MeterValueColor(FLinearColor(0.025719f, 0.208333f, 0.069907f, 1.0f))
-			.MeterPeakColor(FLinearColor(0.24349f, 0.708333f, 0.357002f, 1.0f))
-			.MeterClippingColor(FLinearColor(1.0f, 0.0f, 0.112334f, 1.0f))
-			.MeterScaleColor(FLinearColor(0.017642f, 0.017642f, 0.017642f, 1.0f))
-			.MeterScaleLabelColor(FLinearColor(0.442708f, 0.442708f, 0.442708f, 1.0f));
+			.MeterBackgroundColor(WidgetStyle->MeterBackgroundColor)
+			.MeterValueColor(WidgetStyle->MeterValueColor)
+			.MeterPeakColor(WidgetStyle->MeterPeakColor)
+			.MeterClippingColor(WidgetStyle->MeterClippingColor)
+			.MeterScaleColor(WidgetStyle->MeterScaleColor)
+			.MeterScaleLabelColor(WidgetStyle->MeterScaleLabelColor);
 
 		Init(InNumChannels, InAudioDeviceId, InExternalAudioBus);
 	}
diff --git a/Engine/Plugins/Runtime/AudioWidgets/Source/AudioWidgets/Private/AudioOscilloscope.cpp b/Engine/Plugins/Runtime/AudioWidgets/Source/AudioWidgets/Private/AudioOscilloscope.cpp
--- a/Engine/Plugins/Runtime/AudioWidgets/Source/AudioWidgets/Private/AudioOscilloscope.cpp
+++ b/Engine/Plugins/Runtime/AudioWidgets/Source/AudioWidgets/Private/AudioOscilloscope.cpp
@@ -9,9 +9,10 @@ namespace AudioWidgets
 		const float InTimeWindowMs, 
 		const float InMaxTimeWindowMs, 
 		const float InAnalysisPeriodMs, 
-		const EAudioPanelLayoutType InPanelLayoutType)
-		: OscilloscopePanelStyle(FAudioOscilloscopePanelStyle::GetDefault())
+		const EAudioPanelLayoutType InPanelLayoutType,
+		const FAudioOscilloscopePanelStyle* InOscilloscopePanelStyle)
 	{
+		OscilloscopePanelStyle = InOscilloscopePanelStyle ? *InOscilloscopePanelStyle : FAudioOscilloscopePanelStyle::GetDefault();
 		CreateAudioBus(InNumChannels);
 		CreateDataProvider(InAudioDeviceId, InTimeWindowMs, InMaxTimeWindowMs, InAnalysisPeriodMs, InPanelLayoutType);
 		CreateOscilloscopeWidget(InNumChannels, InPanelLayoutType);
@@ -35,12 +36,17 @@ namespace AudioWidgets
 		AudioSamplesDataProvider = MakeShared<FWaveformAudioSamplesDataProvider>(InAudioDeviceId, AudioBus.Get(), NumChannelsToProvide, InTimeWindowMs, InMaxTimeWindowMs, InAnalysisPeriodMs);
 	}
 
-	void FAudioOscilloscope::CreateOscilloscopeWidget(const uint32 InNumChannels, const EAudioPanelLayoutType InPanelLayoutType)
+	void FAudioOscilloscope::CreateOscilloscopeWidget(const uint32 InNumChannels, const EAudioPanelLayoutType InPanelLayoutType, const FAudioOscilloscopePanelStyle* InOscilloscopePanelStyle)
 	{
 		check(AudioSamplesDataProvider);
 
 		const FFixedSampledSequenceView SequenceView = AudioSamplesDataProvider->GetDataView();
 
+		if (InOscilloscopePanelStyle)
+		{
+			OscilloscopePanelStyle = *InOscilloscopePanelStyle;
+		}
+
 		if (!OscilloscopePanelWidget.IsValid())
 		{
 			OscilloscopePanelWidget = SNew(SAudioOscilloscopePanelWidget, SequenceView, InNumChannels)
diff --git a/Engine/Plugins/Runtime/AudioWidgets/Source/AudioWidgets/Private/AudioSpectrumAnalyzer.cpp b/Engine/Plugins/Runtime/AudioWidgets/Source/AudioWidgets/Private/AudioSpectrumAnalyzer.cpp
--- a/Engine/Plugins/Runtime/AudioWidgets/Source/AudioWidgets/Private/AudioSpectrumAnalyzer.cpp
+++ b/Engine/Plugins/Runtime/AudioWidgets/Source/AudioWidgets/Private/AudioSpectrumAnalyzer.cpp
@@ -47,6 +47,7 @@ namespace AudioWidgets
 		: SpectrumAnalysisSettings(NewObject<USynesthesiaSpectrumAnalysisSettings>())
 		, ConstantQSettings(NewObject<UConstantQSettings>())
 		, Widget(SNew(SAudioSpectrumPlot)
+			.Style(Params.PlotStyle ? Params.PlotStyle : &FAudioSpectrumPlotStyle::GetDefault())
 			.Clipping(EWidgetClipping::ClipToBounds)
 			.TiltExponent(Params.TiltExponent)
 			.DisplayCrosshair(true)
diff --git a/Engine/Plugins/Runtime/AudioWidgets/Source/AudioWidgets/Private/AudioVectorscope.cpp b/Engine/Plugins/Runtime/AudioWidgets/Source/AudioWidgets/Private/AudioVectorscope.cpp
--- a/Engine/Plugins/Runtime/AudioWidgets/Source/AudioWidgets/Private/AudioVectorscope.cpp
+++ b/Engine/Plugins/Runtime/AudioWidgets/Source/AudioWidgets/Private/AudioVectorscope.cpp
@@ -7,13 +7,14 @@
 namespace AudioWidgets
 {
 	FAudioVectorscope::FAudioVectorscope(Audio::FDeviceId InAudioDeviceId,
-		const uint32 InNumChannels, 
-		const float InTimeWindowMs, 
-		const float InMaxTimeWindowMs, 
-		const float InAnalysisPeriodMs, 
-		const EAudioPanelLayoutType InPanelLayoutType)
-		: VectorscopePanelStyle(FAudioVectorscopePanelStyle::GetDefault())
+		const uint32 InNumChannels,
+		const float InTimeWindowMs,
+		const float InMaxTimeWindowMs,
+		const float InAnalysisPeriodMs,
+		const EAudioPanelLayoutType InPanelLayoutType,
+		const FAudioVectorscopePanelStyle* InPanelStyle)
 	{
+		VectorscopePanelStyle = InPanelStyle ? *InPanelStyle : FAudioVectorscopePanelStyle::GetDefault();
 		CreateAudioBus(InNumChannels);
 		CreateDataProvider(InAudioDeviceId, InTimeWindowMs, InMaxTimeWindowMs, InAnalysisPeriodMs);
 		CreateVectorscopeWidget(InPanelLayoutType);
@@ -32,15 +33,21 @@ namespace AudioWidgets
 		AudioSamplesDataProvider = MakeShared<FWaveformAudioSamplesDataProvider>(InAudioDeviceId, AudioBus.Get(), AudioBus->GetNumChannels(), InTimeWindowMs, InMaxTimeWindowMs, InAnalysisPeriodMs);
 	}
 
-	void FAudioVectorscope::CreateVectorscopeWidget(const EAudioPanelLayoutType InPanelLayoutType)
+	void FAudioVectorscope::CreateVectorscopeWidget(const EAudioPanelLayoutType InPanelLayoutType, const FAudioVectorscopePanelStyle* PanelStyle)
 	{
 		check(AudioSamplesDataProvider);
 
 		const FFixedSampledSequenceView SequenceView = AudioSamplesDataProvider->GetDataView();
 
+		if (PanelStyle)
+		{
+			VectorscopePanelStyle = *PanelStyle;
+		}
+
 		if (!VectorscopePanelWidget.IsValid())
 		{
 			VectorscopePanelWidget = SNew(SAudioVectorscopePanelWidget, SequenceView)
+				.PanelStyle(PanelStyle ? PanelStyle : &FAudioVectorscopePanelStyle::GetDefault())
 				.PanelLayoutType(InPanelLayoutType)
 				.PanelStyle(&VectorscopePanelStyle);
 		}
diff --git a/Engine/Plugins/Runtime/AudioWidgets/Source/AudioWidgets/Public/AudioMeter.h b/Engine/Plugins/Runtime/AudioWidgets/Source/AudioWidgets/Public/AudioMeter.h
--- a/Engine/Plugins/Runtime/AudioWidgets/Source/AudioWidgets/Public/AudioMeter.h
+++ b/Engine/Plugins/Runtime/AudioWidgets/Source/AudioWidgets/Public/AudioMeter.h
@@ -13,6 +13,7 @@
 #include "Sound/AudioBus.h"
 #include "Styling/SlateTypes.h"
 #include "Templates/SharedPointer.h"
+#include "UObject/NoExportTypes.h"
 #include "UObject/ObjectMacros.h"
 #include "UObject/WeakObjectPtrTemplates.h"
 #include "UObject/StrongObjectPtr.h"
@@ -27,6 +28,34 @@ class UWorld;
 
 struct FAudioMaterialMeterStyle;
 
+
+USTRUCT(BlueprintType)
+struct AUDIOWIDGETS_API FAudioMeterDefaultColorStyle : public FSlateWidgetStyle
+{
+	GENERATED_USTRUCT_BODY()
+
+public:
+	static const FAudioMeterDefaultColorStyle& GetDefault();
+
+	UPROPERTY(EditAnywhere, BlueprintReadWrite, Category = "Style")
+	FLinearColor MeterBackgroundColor = FLinearColor(0.031f, 0.031f, 0.031f, 1.0f);
+
+	UPROPERTY(EditAnywhere, BlueprintReadWrite, Category = "Style")
+	FLinearColor MeterValueColor = FLinearColor(0.025719f, 0.208333f, 0.069907f, 1.0f);
+
+	UPROPERTY(EditAnywhere, BlueprintReadWrite, Category = "Style")
+	FLinearColor MeterPeakColor = FLinearColor(0.24349f, 0.708333f, 0.357002f, 1.0f);
+
+	UPROPERTY(EditAnywhere, BlueprintReadWrite, Category = "Style")
+	FLinearColor MeterClippingColor = FLinearColor(1.0f, 0.0f, 0.112334f, 1.0f);
+
+	UPROPERTY(EditAnywhere, BlueprintReadWrite, Category = "Style")
+	FLinearColor MeterScaleColor = FLinearColor(0.017642f, 0.017642f, 0.017642f, 1.0f);
+
+	UPROPERTY(EditAnywhere, BlueprintReadWrite, Category = "Style")
+	FLinearColor MeterScaleLabelColor = FLinearColor(0.442708f, 0.442708f, 0.442708f, 1.0f);
+};
+
 /**
  * An audio meter widget.
  *
@@ -162,7 +191,7 @@ namespace AudioWidgets
 		FAudioMeter(int32 InNumChannels, UWorld& InWorld, TObjectPtr<UAudioBus> InExternalAudioBus = nullptr); 
 		
 		//** OPTIONAL PARAM InExternalAudioBus: An audio meter can be constructed from this audio bus.*/
-		FAudioMeter(const int32 InNumChannels, const Audio::FDeviceId InAudioDeviceId, TObjectPtr<UAudioBus> InExternalAudioBus = nullptr); 
+		FAudioMeter(const int32 InNumChannels, const Audio::FDeviceId InAudioDeviceId, TObjectPtr<UAudioBus> InExternalAudioBus = nullptr, const FAudioMeterDefaultColorStyle* AudioMeterColorStyle = nullptr);
 
 		//** Constructs the Meter using AudioMaterialMeter with the given style. OPTIONAL PARAM InExternalAudioBus: An audio meter can be constructed from this audio bus.
 		FAudioMeter(const int32 InNumChannels, const Audio::FDeviceId InAudioDeviceId, const FAudioMaterialMeterStyle& AudioMaterialMeterStyle, TObjectPtr<UAudioBus> InExternalAudioBus = nullptr);
diff --git a/Engine/Plugins/Runtime/AudioWidgets/Source/AudioWidgets/Public/AudioOscilloscope.h b/Engine/Plugins/Runtime/AudioWidgets/Source/AudioWidgets/Public/AudioOscilloscope.h
--- a/Engine/Plugins/Runtime/AudioWidgets/Source/AudioWidgets/Public/AudioOscilloscope.h
+++ b/Engine/Plugins/Runtime/AudioWidgets/Source/AudioWidgets/Public/AudioOscilloscope.h
@@ -20,11 +20,12 @@ namespace AudioWidgets
 	{
 	public:
 		FAudioOscilloscope(Audio::FDeviceId InAudioDeviceId,
-			const uint32 InNumChannels, 
-			const float InTimeWindowMs, 
-			const float InMaxTimeWindowMs, 
-			const float InAnalysisPeriodMs, 
-			const EAudioPanelLayoutType InPanelLayoutType);
+			const uint32 InNumChannels,
+			const float InTimeWindowMs,
+			const float InMaxTimeWindowMs,
+			const float InAnalysisPeriodMs,
+			const EAudioPanelLayoutType InPanelLayoutType,
+			const FAudioOscilloscopePanelStyle* InOscilloscopePanelStyle = nullptr);
 
 		void CreateAudioBus(const uint32 InNumChannels);
 
@@ -34,7 +35,7 @@ namespace AudioWidgets
 			const float InAnalysisPeriodMs,
 			const EAudioPanelLayoutType InPanelLayoutType);
 
-		void CreateOscilloscopeWidget(const uint32 InNumChannels, const EAudioPanelLayoutType InPanelLayoutType);
+		void CreateOscilloscopeWidget(const uint32 InNumChannels, const EAudioPanelLayoutType InPanelLayoutType, const FAudioOscilloscopePanelStyle* InOscilloscopePanelStyle = nullptr);
 
 		void StartProcessing();
 		void StopProcessing();
diff --git a/Engine/Plugins/Runtime/AudioWidgets/Source/AudioWidgets/Public/AudioSpectrumAnalyzer.h b/Engine/Plugins/Runtime/AudioWidgets/Source/AudioWidgets/Public/AudioSpectrumAnalyzer.h
--- a/Engine/Plugins/Runtime/AudioWidgets/Source/AudioWidgets/Public/AudioSpectrumAnalyzer.h
+++ b/Engine/Plugins/Runtime/AudioWidgets/Source/AudioWidgets/Public/AudioSpectrumAnalyzer.h
@@ -34,7 +34,7 @@ namespace AudioWidgets
 	/**
 	 * Constructor parameters for the analyzer.
 	 */
-	struct FAudioSpectrumAnalyzerParams
+	struct AUDIOWIDGETS_API FAudioSpectrumAnalyzerParams
 	{
 		int32 NumChannels = 1;
 		Audio::FDeviceId AudioDeviceId = INDEX_NONE;
@@ -59,6 +59,8 @@ namespace AudioWidgets
 		FOnFrequencyAxisScaleMenuEntryClicked OnFrequencyAxisScaleMenuEntryClicked;
 		FOnDisplayAxisLabelsButtonToggled OnDisplayFrequencyAxisLabelsButtonToggled;
 		FOnDisplayAxisLabelsButtonToggled OnDisplaySoundLevelAxisLabelsButtonToggled;
+
+		const FAudioSpectrumPlotStyle* PlotStyle = nullptr;
 	};
 
 	/**
diff --git a/Engine/Plugins/Runtime/AudioWidgets/Source/AudioWidgets/Public/AudioVectorscope.h b/Engine/Plugins/Runtime/AudioWidgets/Source/AudioWidgets/Public/AudioVectorscope.h
--- a/Engine/Plugins/Runtime/AudioWidgets/Source/AudioWidgets/Public/AudioVectorscope.h
+++ b/Engine/Plugins/Runtime/AudioWidgets/Source/AudioWidgets/Public/AudioVectorscope.h
@@ -23,11 +23,12 @@ namespace AudioWidgets
 			const float InTimeWindowMs, 
 			const float InMaxTimeWindowMs, 
 			const float InAnalysisPeriodMs, 
-			const EAudioPanelLayoutType InPanelLayoutType);
+			const EAudioPanelLayoutType InPanelLayoutType,
+			const FAudioVectorscopePanelStyle* PanelStyle = nullptr);
 
 		void CreateAudioBus(const uint32 InNumChannels);
 		void CreateDataProvider(Audio::FDeviceId InAudioDeviceId, const float InTimeWindowMs, const float InMaxTimeWindowMs, const float InAnalysisPeriodMs);
-		void CreateVectorscopeWidget(const EAudioPanelLayoutType InPanelLayoutType);
+		void CreateVectorscopeWidget(const EAudioPanelLayoutType InPanelLayoutType, const FAudioVectorscopePanelStyle* PanelStyle = nullptr);
 
 		void StartProcessing();
 		void StopProcessing();
diff --git a/Engine/Plugins/Runtime/Metasound/Source/MetasoundEditor/Private/MetasoundDefaultLiteralCustomization.cpp b/Engine/Plugins/Runtime/Metasound/Source/MetasoundEditor/Private/MetasoundDefaultLiteralCustomization.cpp
--- a/Engine/Plugins/Runtime/Metasound/Source/MetasoundEditor/Private/MetasoundDefaultLiteralCustomization.cpp
+++ b/Engine/Plugins/Runtime/Metasound/Source/MetasoundEditor/Private/MetasoundDefaultLiteralCustomization.cpp
@@ -224,18 +224,14 @@ namespace Metasound::Editor
 			return SNullWidget::NullWidget;
 		}
 
-		const FGuid PageID = LiteralCustomizationPrivate::GetGuidPropertyValue(ElementProperty->GetChildHandle(GET_MEMBER_NAME_CHECKED(FMetasoundEditorMemberPageDefault, PageID)));
-		FName PageName;
-		PageNameProperty->GetValue(PageName);
-
 		TSharedRef<SHorizontalBox> NameBox = SNew(SHorizontalBox)
-			+ SHorizontalBox::Slot()
-			.AutoWidth()
-			.VAlign(VAlign_Center)
-			.Padding(2, 0)
-			[
-				PageNameProperty->CreatePropertyValueWidget()
-			];
+		+ SHorizontalBox::Slot()
+		.AutoWidth()
+		.VAlign(VAlign_Center)
+		.Padding(2, 0)
+		[
+			PageNameProperty->CreatePropertyValueWidget()
+		];
 
 		const UMetasoundEditorGraphMember* Member = Literal.FindMember();
 		if (!ensure(Member))
@@ -243,43 +239,14 @@ namespace Metasound::Editor
 			return SNullWidget::NullWidget;
 		}
 
-		if (Member->IsDefaultPaged())
-		{
-			TAttribute<FText> ToolTip = LOCTEXT("MetaSound_ExecutingInputPageDefaultTooltip", "Initial default used upon executing preview.");
-			TAttribute<EVisibility> ExecVisibility = TAttribute<EVisibility>::CreateLambda([MemberPtr = TWeakObjectPtr<const UMetasoundEditorGraphMember>(Member), PageID]()
-			{
-				if (const UMetasoundEditorGraphMember* GraphMember = MemberPtr.Get())
-				{
-					const FMetaSoundFrontendDocumentBuilder& Builder = GraphMember->GetFrontendBuilderChecked();
-					if (const FMetasoundFrontendClassInput* ClassInput = Builder.FindGraphInput(GraphMember->GetMemberName()))
-					{
-						const bool bIsPreviewing = Editor::IsPreviewingPageInputDefault(Builder, *ClassInput, PageID);
-						return bIsPreviewing ? EVisibility::Visible : EVisibility::Collapsed;
-					}
-					return EVisibility::Collapsed;
-				}
-
-				return EVisibility::Collapsed;
-			});
-			TSharedRef<SWidget> ExecImageWidget = SNew(SImage)
-				.Image(Style::CreateSlateIcon("MetasoundEditor.Page.Executing").GetIcon())
-				.ColorAndOpacity(FStyleColors::AccentGreen)
-				.Visibility(MoveTemp(ExecVisibility));
-			NameBox->AddSlot()
-				.FillWidth(1)
-				.HAlign(HAlign_Left)
-				.VAlign(VAlign_Center)
-				.Padding(2, 0)
-				[
-					ExecImageWidget
-				];
-		}
-
 		// Can't delete the default page, so only show remove page field if non-default/project defined page.
+		const FGuid PageID = LiteralCustomizationPrivate::GetGuidPropertyValue(ElementProperty->GetChildHandle(GET_MEMBER_NAME_CHECKED(FMetasoundEditorMemberPageDefault, PageID)));
 		if (PageID != Frontend::DefaultPageID)
 		{
 			TWeakObjectPtr<UMetasoundEditorGraphMemberDefaultLiteral> LiteralPtr(&Literal);
 
+			FName PageName;
+			PageNameProperty->GetValue(PageName);
 			const FText RemoveDescription = FText::Format(LOCTEXT("RemovePageDefaultTransactionFormat", "Remove '{0}' Page '{1}' Default Value"), FText::FromName(Member->GetMemberName()), FText::FromName(PageName));
 			TSharedRef<SWidget> RemovePageDefaultButton = PropertyCustomizationHelpers::MakeDeleteButton(FSimpleDelegate::CreateLambda([this, RemoveDescription, PageID, LiteralPtr]()
 			{
@@ -329,6 +296,39 @@ namespace Metasound::Editor
 			];
 		}
 
+		if (Member->IsDefaultPaged())
+		{
+			TAttribute<EVisibility> ExecVisibility = TAttribute<EVisibility>::CreateLambda([MemberPtr = TWeakObjectPtr<const UMetasoundEditorGraphMember>(Member), PageID]()
+			{
+				if (const UMetasoundEditorGraphMember* GraphMember = MemberPtr.Get())
+				{
+					const FMetaSoundFrontendDocumentBuilder& Builder = GraphMember->GetFrontendBuilderChecked();
+					if (const FMetasoundFrontendClassInput* ClassInput = Builder.FindGraphInput(GraphMember->GetMemberName()))
+					{
+						const bool bIsPreviewing = Editor::IsPreviewingPageInputDefault(Builder, *ClassInput, PageID);
+						return bIsPreviewing ? EVisibility::Visible : EVisibility::Collapsed;
+					}
+					return EVisibility::Collapsed;
+				}
+
+				return EVisibility::Collapsed;
+			});
+
+			TSharedRef<SWidget> ExecImageWidget = SNew(SImage)
+			.Image(Style::CreateSlateIcon("MetasoundEditor.Page.Executing").GetIcon())
+			.ColorAndOpacity(Style::GetPageExecutingColor())
+			.Visibility(MoveTemp(ExecVisibility));
+
+			NameBox->AddSlot()
+			.FillWidth(1)
+			.HAlign(HAlign_Left)
+			.VAlign(VAlign_Center)
+			.Padding(2, 0)
+			[
+				ExecImageWidget
+			];
+		}
+
 		return NameBox;
 	}
 
diff --git a/Engine/Plugins/Runtime/Metasound/Source/MetasoundEditor/Private/MetasoundDetailCustomization.cpp b/Engine/Plugins/Runtime/Metasound/Source/MetasoundEditor/Private/MetasoundDetailCustomization.cpp
--- a/Engine/Plugins/Runtime/Metasound/Source/MetasoundEditor/Private/MetasoundDetailCustomization.cpp
+++ b/Engine/Plugins/Runtime/Metasound/Source/MetasoundEditor/Private/MetasoundDetailCustomization.cpp
@@ -575,7 +575,7 @@ namespace Metasound::Editor
 				});
 				TSharedRef<SWidget> ExecImageWidget = SNew(SImage)
 					.Image(Style::CreateSlateIcon("MetasoundEditor.Page.Executing").GetIcon())
-					.ColorAndOpacity(FStyleColors::AccentGreen)
+					.ColorAndOpacity(Style::GetPageExecutingColor())
 					.Visibility(MoveTemp(Visibility));
 
 				EntryWidget->AddSlot()
diff --git a/Engine/Plugins/Runtime/Metasound/Source/MetasoundEditor/Private/MetasoundEditor.cpp b/Engine/Plugins/Runtime/Metasound/Source/MetasoundEditor/Private/MetasoundEditor.cpp
--- a/Engine/Plugins/Runtime/Metasound/Source/MetasoundEditor/Private/MetasoundEditor.cpp
+++ b/Engine/Plugins/Runtime/Metasound/Source/MetasoundEditor/Private/MetasoundEditor.cpp
@@ -928,50 +928,43 @@ namespace Metasound
 			InTabManager->RegisterTabSpawner(TabNamesPrivate::GraphCanvas, FOnSpawnTab::CreateLambda(
 				[
 					bShowPageTab = ShowPageGraphDetails(),
-					InFocusPageWidget = FocusPageWidget,
-					InPlayTimeWidget = PlayTimeWidget,
+					InPageStatsWidget = PageStatsWidget,
 					InMetasoundGraphEditor = MetasoundGraphEditor,
 					InRenderStatsWidget = RenderStatsWidget
 				](const FSpawnTabArgs& Args)
 			{
 				TSharedRef<SDockTab> SpawnedTab = SNew(SDockTab).Label(LOCTEXT("MetasoundGraphCanvasTitle", "MetaSound Graph"));
 
-				TSharedPtr<SVerticalBox> StatsWidget = SNew(SVerticalBox)
-					+ SVerticalBox::Slot()
-					.HAlign(HAlign_Left)
-					.VAlign(VAlign_Center)
-					.Padding(2.f, 1.f)
+				TSharedRef<SOverlay> Overlay = SNew(SOverlay)
+					+ SOverlay::Slot()
+					[
+						InMetasoundGraphEditor.ToSharedRef()
+					]
+					+ SOverlay::Slot()
+					.VAlign(VAlign_Top)
 					[
 						InRenderStatsWidget.ToSharedRef()
-					];
+					]
+					.Padding(5.0f, 5.0f);
 
-				if (bShowPageTab)
+				if (bShowPageTab && InPageStatsWidget.IsValid())
 				{
-					StatsWidget->AddSlot()
+					TSharedRef<SVerticalBox> GraphStatsWidget = SNew(SVerticalBox)
+						+ SVerticalBox::Slot()
 						.HAlign(HAlign_Left)
 						.VAlign(VAlign_Center)
 						.AutoHeight()
 						[
-							InFocusPageWidget.ToSharedRef()
+							InPageStatsWidget.ToSharedRef()
 						];
-				}
-
-				SpawnedTab->SetContent(SNew(SOverlay)
-					+ SOverlay::Slot()
-					[
-						InMetasoundGraphEditor.ToSharedRef()
-					]
-					+ SOverlay::Slot()
-					[
-						InPlayTimeWidget.ToSharedRef()
-					]
-					+ SOverlay::Slot()
+					Overlay->AddSlot()
 					.VAlign(VAlign_Bottom)
 					[
-						StatsWidget.ToSharedRef()
-					]
-					.Padding(5.0f, 5.0f)
-				);
+						GraphStatsWidget
+					];
+				}
+
+				SpawnedTab->SetContent(Overlay);
 				return SpawnedTab;
 			}))
 			.SetDisplayName(LOCTEXT("GraphCanvasTab", "Viewport"))
@@ -1087,20 +1080,6 @@ namespace Metasound
 				BackgroundColor = MetaSoundStyle->GetColor("MetasoundEditor.Analyzers.BackgroundColor");
 			}
 
-			TSharedRef<SVerticalBox> PageWidget = SNew(SVerticalBox)
-				.Visibility(TAttribute<EVisibility>::Create([this]()
-				{
-					const bool bMakeVisible = IsPlaying() && AssetEditorPrivate::EnablePageEditor;
-					return bMakeVisible ? EVisibility::Visible : EVisibility::Collapsed;
-				}))
-				+ SVerticalBox::Slot()
-				.HAlign(HAlign_Left)
-				.VAlign(VAlign_Center)
-				.AutoHeight()
-				[
-					RenderPageWidget.ToSharedRef()
-				];
-
 			return SNew(SOverlay)
 			+ SOverlay::Slot()
 			[
@@ -1111,11 +1090,6 @@ namespace Metasound
 			[
 				SNew(SSplitter)
 				.Orientation(Orient_Vertical)
-				+ SSplitter::Slot()
-					.Value(0.05f)
-					[
-						PageWidget
-					]
 				+ SSplitter::Slot()
 				.Value(0.5f)
 				[
@@ -1205,10 +1179,7 @@ namespace Metasound
 			Builder->AddTransactionListener(DocListener->AsShared());
 
 			// Stat widgets are potentially intractable with transaction listener, so create then here
-			SAssignNew(FocusPageWidget, SPageStats)
-				.Visibility(EVisibility::HitTestInvisible);
-
-			SAssignNew(RenderPageWidget, SPageStats)
+			SAssignNew(PageStatsWidget, SPageStats)
 				.Visibility(EVisibility::HitTestInvisible);
 
 			SAssignNew(RenderStatsWidget, SRenderStats)
@@ -1355,7 +1326,9 @@ namespace Metasound
 			}
 			else
 			{
-				UpdatePageInfo(false);
+				constexpr bool bIsPlaying = false;
+				UpdatePageInfo(bIsPlaying);
+				UpdateRenderInfo(bIsPlaying);
 			}
 
 			RefreshExecVisibility(Builder->GetConstBuilder().GetBuildPageID());
@@ -1770,7 +1743,7 @@ namespace Metasound
 
 					if (!OutputMeter.IsValid())
 					{
-						OutputMeter = MakeShared<AudioWidgets::FAudioMeter>(MetaSoundSource.NumChannels, AudioDeviceId, DefaultBus);
+						OutputMeter = MakeShared<AudioWidgets::FAudioMeter>(MetaSoundSource.NumChannels, AudioDeviceId, DefaultBus, &Style::GetMeterDefaultColorStyle());
 					}
 				}
 				else if (OutputMeter->GetAudioBus()->GetNumChannels() != MetaSoundSource.NumChannels)
@@ -1793,13 +1766,14 @@ namespace Metasound
 						OscilloscopeTimeWindowMs,
 						OscilloscopeMaxTimeWindowMs,
 						OscilloscopeAnalysisPeriodMs,
-						OscilloscopePanelLayoutType);
+						OscilloscopePanelLayoutType,
+						&Style::GetOscilloscopeStyle());
 				}
 				else if (OutputOscilloscope->GetAudioBus()->GetNumChannels() != MetaSoundSource.NumChannels)
 				{
 					OutputOscilloscope->CreateAudioBus(MetaSoundNumChannels);
 					OutputOscilloscope->CreateDataProvider(AudioDeviceId, OscilloscopeTimeWindowMs, OscilloscopeMaxTimeWindowMs, OscilloscopeAnalysisPeriodMs, OscilloscopePanelLayoutType);
-					OutputOscilloscope->CreateOscilloscopeWidget(MetaSoundNumChannels, OscilloscopePanelLayoutType);
+					OutputOscilloscope->CreateOscilloscopeWidget(MetaSoundNumChannels, OscilloscopePanelLayoutType, &Style::GetOscilloscopeStyle());
 				}
 
 				// Init Vectorscope
@@ -1815,13 +1789,14 @@ namespace Metasound
 						VectorscopeTimeWindowMs,
 						VectorscopeMaxTimeWindowMs,
 						VectorscopeAnalysisPeriodMs,
-						VectorscopePanelLayoutType);
+						VectorscopePanelLayoutType,
+						&Style::GetVectorscopeStyle());
 				}
 				else if (OutputVectorscope->GetAudioBus()->GetNumChannels() != MetaSoundSource.NumChannels)
 				{
 					OutputVectorscope->CreateAudioBus(MetaSoundNumChannels);
 					OutputVectorscope->CreateDataProvider(AudioDeviceId, VectorscopeTimeWindowMs, VectorscopeMaxTimeWindowMs, VectorscopeAnalysisPeriodMs);
-					OutputVectorscope->CreateVectorscopeWidget(VectorscopePanelLayoutType);
+					OutputVectorscope->CreateVectorscopeWidget(VectorscopePanelLayoutType, &Style::GetVectorscopeStyle());
 				}
 
 				if (!OutputSpectrumAnalyzer.IsValid())
@@ -1915,7 +1890,7 @@ namespace Metasound
 							SpectrumAnalyzerSettings.bDisplaySoundLevelAxisLabels = !SpectrumAnalyzerSettings.bDisplaySoundLevelAxisLabels;
 							GetMutableDefault<UMetasoundEditorSettings>()->SaveConfig();
 						});
-
+					Params.PlotStyle = &Style::GetSpectrumPlotStyle();
 					OutputSpectrumAnalyzer = MakeShared<AudioWidgets::FAudioSpectrumAnalyzer>(Params);
 				}
 				else if (OutputSpectrumAnalyzer->GetAudioBus()->GetNumChannels() != MetaSoundSource.NumChannels)
@@ -2592,7 +2567,7 @@ namespace Metasound
 				// Set the send to the audio bus that is used for analyzing the metasound output
 				check(GEditor);
 
-				PlayTime = 0.0;
+				UpdateRenderInfo(true /* bIsPlaying */);
 				UpdatePageInfo(true);
 
 				if (UMetaSoundSource* Source = Cast<UMetaSoundSource>(GetMetasoundObject()))
@@ -2628,8 +2603,8 @@ namespace Metasound
 				MetasoundGraphEditor->RegisterActiveTimer(0.0f,
 					FWidgetActiveTimerDelegate::CreateLambda([this](double InCurrentTime, float InDeltaTime)
 					{
-						const bool bIsPlaying = UpdatePlayTime(InDeltaTime);
-						UpdateRenderInfo();
+						const bool bIsPlaying = IsPlaying();
+						UpdateRenderInfo(bIsPlaying, InDeltaTime);
 
 						if (bIsPlaying)
 						{
@@ -2637,8 +2612,7 @@ namespace Metasound
 						}
 						else
 						{
-							PlayTime = 0.0;
-							UpdatePageInfo(false);
+							UpdatePageInfo(bIsPlaying);
 							GraphConnectionManager = RebuildConnectionManager();
 							return EActiveTimerReturnType::Stop;
 						}
@@ -2714,11 +2688,6 @@ namespace Metasound
 			return true;
 		}
 
-		double FEditor::GetPlayTime() const
-		{
-			return PlayTime;
-		}
-
 		TSharedPtr<SGraphEditor> FEditor::GetGraphEditor() const
 		{
 			return MetasoundGraphEditor;
@@ -3107,11 +3076,6 @@ namespace Metasound
 				.GraphToEdit(MetaSoundAsset->GetGraph())
 				.IsEditable(this, &FEditor::IsGraphEditable)
 				.ShowGraphStateOverlay(false);
-
-			SAssignNew(PlayTimeWidget, STextBlock)
-				.Visibility(EVisibility::HitTestInvisible)
-				.TextStyle(FAppStyle::Get(), "Graph.ZoomText")
-				.ColorAndOpacity(FLinearColor(1, 1, 1, 0.30f));
 		}
 
 		FGraphAppearanceInfo FEditor::GetGraphAppearance() const
@@ -5113,14 +5077,18 @@ namespace Metasound
 			const UMetaSoundSettings* Settings = GetDefault<UMetaSoundSettings>();
 			check(Settings)
 
-			if (FocusPageWidget.IsValid())
+			if (PageStatsWidget.IsValid())
 			{
 				const FSlateColor* Color = nullptr;
-				// For now, this is one and the same, but may change in the future to accommodate conflating
-				// the static nature of the active build graph ID with what the user is visualizing.
-				const FMetaSoundPageSettings* PageSettings = nullptr;
+				const FMetaSoundPageSettings* GraphPageSettings = nullptr;
+				const FMetaSoundPageSettings* AuditionPageSettings = nullptr;
 				if (Builder.IsValid())
 				{
+					if (const UMetasoundEditorSettings* EditorSettings = GetDefault<UMetasoundEditorSettings>())
+					{
+						AuditionPageSettings = Settings->FindPageSettings(EditorSettings->AuditionPage);
+					}
+
 					const FMetaSoundFrontendDocumentBuilder& DocBuilder = Builder->GetConstBuilder();
 					const FGuid& PageID = DocBuilder.GetBuildPageID();
 					if (bIsPlaying)
@@ -5129,61 +5097,26 @@ namespace Metasound
 						const FGuid ResolvePageID = FDocumentBuilderRegistry::GetChecked().ResolveTargetPageID(GraphClass);
 						if (ResolvePageID == PageID)
 						{
-							Color = &FStyleColors::AccentGreen;
+							Color = &Style::GetPageExecutingColor();
 						}
 					}
-					PageSettings = Settings->FindPageSettings(PageID);
+					GraphPageSettings = Settings->FindPageSettings(PageID);
 				}
 
-				const FText GraphPageHeader = LOCTEXT("GraphPageTargetHeader", "Graph Page");
-				FocusPageWidget->Update(PageSettings, GraphPageHeader, Color);
-			}
-
-			if (RenderPageWidget.IsValid())
-			{
-				const FMetaSoundPageSettings* PageSettings = nullptr;
-				if (Builder.IsValid())
-				{
-					if (const UMetasoundEditorSettings* EditorSettings = GetDefault<UMetasoundEditorSettings>())
-					{
-						PageSettings = Settings->FindPageSettings(EditorSettings->AuditionPage);
-					}
-				}
-
-				const FText RenderPageHeader = LOCTEXT("AuditionPageHeader", "Audition Page");
-				RenderPageWidget->Update(PageSettings, RenderPageHeader, &FStyleColors::AccentGreen);
+				PageStatsWidget->Update(AuditionPageSettings, GraphPageSettings, Color);
 			}
 		}
 
-		bool FEditor::UpdatePlayTime(float InDeltaTime)
+		void FEditor::UpdateRenderInfo(bool bIsPlaying, float InDeltaTime)
 		{
-			const bool bIsPlaying = IsPlaying();
-			if (bIsPlaying)
+			if (!bIsPlaying)
 			{
-				if (PlayTimeWidget.IsValid())
-				{
-					PlayTime += InDeltaTime;
-
-					FString PlayTimeString = FTimespan::FromSeconds(PlayTime).ToString();
-					PlayTimeString.ReplaceInline(TEXT("+"), TEXT(""));
-					PlayTimeWidget->SetText(FText::FromString(MoveTemp(PlayTimeString)));
-				}
-
-				return true;
+				SetPreviewID(INDEX_NONE);
 			}
 
-			SetPreviewID(INDEX_NONE);
-			PlayTime = 0.0;
-			PlayTimeWidget->SetText(FText::GetEmpty());
-			return false;
-		}
-
-		void FEditor::UpdateRenderInfo()
-		{
 			if (RenderStatsWidget.IsValid())
 			{
-				const bool bIsPlaying = IsPlaying();
-				RenderStatsWidget->Update(bIsPlaying, Cast<const UMetaSoundSource>(GetMetasoundObject()));
+				RenderStatsWidget->Update(bIsPlaying, InDeltaTime, Cast<const UMetaSoundSource>(GetMetasoundObject()));
 			}
 		}
 
@@ -5205,7 +5138,7 @@ namespace Metasound
 
 		void FEditor::RefreshExecVisibility(const FGuid& InPageID) const
 		{
-			if (FocusPageWidget.IsValid())
+			if (PageStatsWidget.IsValid())
 			{
 				TAttribute<EVisibility> ExecVisibility = TAttribute<EVisibility>::CreateLambda([this, InPageID]()
 				{
@@ -5214,12 +5147,12 @@ namespace Metasound
 					{
 						const FMetaSoundFrontendDocumentBuilder& DocBuilder = Builder->GetConstBuilder();
 						const bool bIsPreviewing = IsPreviewingPageGraph(DocBuilder, InPageID);
-						return bIsPreviewing ? EVisibility::Visible : EVisibility::Collapsed;
+						return bIsPreviewing ? EVisibility::HitTestInvisible : EVisibility::Collapsed;
 					}
 
 					return EVisibility::Collapsed;
 				});
-				FocusPageWidget->SetExecVisibility(MoveTemp(ExecVisibility));
+				PageStatsWidget->SetExecVisibility(MoveTemp(ExecVisibility));
 			}
 		}
 
@@ -5228,13 +5161,9 @@ namespace Metasound
 			if (MetasoundGraphEditor.IsValid())
 			{
 				UMetasoundEditorGraph& Graph = GetMetaSoundGraphChecked();
-
-				UEdGraphPin* TargetPin = MetasoundGraphEditor->GetGraphPinForMenu();
-				check(TargetPin);
-
-				if (TargetPin->Direction == EGPD_Input)
+				if (const UEdGraphPin* TargetPin = MetasoundGraphEditor->GetGraphPinForMenu(); ensure(TargetPin))
 				{
-					return true;
+					return TargetPin->Direction == EGPD_Input;
 				}
 			}
 			return false;
@@ -5262,13 +5191,9 @@ namespace Metasound
 			if (MetasoundGraphEditor.IsValid())
 			{
 				UMetasoundEditorGraph& Graph = GetMetaSoundGraphChecked();
-
-				UEdGraphPin* TargetPin = MetasoundGraphEditor->GetGraphPinForMenu();
-				check(TargetPin);
-
-				if (TargetPin->Direction == EGPD_Output)
+				if (const UEdGraphPin* TargetPin = MetasoundGraphEditor->GetGraphPinForMenu(); ensure(TargetPin))
 				{
-					return true;
+					return TargetPin->Direction == EGPD_Output;
 				}
 			}
 			return false;
diff --git a/Engine/Plugins/Runtime/Metasound/Source/MetasoundEditor/Private/MetasoundEditor.h b/Engine/Plugins/Runtime/Metasound/Source/MetasoundEditor/Private/MetasoundEditor.h
--- a/Engine/Plugins/Runtime/Metasound/Source/MetasoundEditor/Private/MetasoundEditor.h
+++ b/Engine/Plugins/Runtime/Metasound/Source/MetasoundEditor/Private/MetasoundEditor.h
@@ -135,7 +135,6 @@ namespace Metasound::Editor
 		virtual void RegisterTabSpawners(const TSharedRef<FTabManager>& TabManager) override;
 		virtual void UnregisterTabSpawners(const TSharedRef<FTabManager>& TabManager) override;
 
-		double GetPlayTime() const;
 		TSharedPtr<SGraphEditor> GetGraphEditor() const;
 
 		/** Edits the specified Metasound object */
@@ -532,7 +531,7 @@ namespace Metasound::Editor
 		bool UpdatePlayTime(float InDeltaTime);
 
 		/** Updates the render info widget. */
-		void UpdateRenderInfo();
+		void UpdateRenderInfo(bool bIsPlaying, float InDeltaTime = 0.0f);
 
 		/** List of open tool panels; used to ensure only one exists at any one time */
 		TMap<FName, TWeakPtr<SDockableTab>> SpawnedToolPanels;
@@ -572,15 +571,8 @@ namespace Metasound::Editor
 		/** Palette of Node types */
 		TSharedPtr<SMetasoundPalette> Palette;
 
-		/** Widget showing playtime that overlays the graph when previewing */
-		TSharedPtr<STextBlock> PlayTimeWidget;
-		double PlayTime = 0.0;
-
-		/** Widget showing page info regarding the currently focused page that overlays the graph tab content */
-		TSharedPtr<SPageStats> FocusPageWidget;
-
-		/** Widget showing page info regarding the currently rendering page that overlays the analyzer tab content */
-		TSharedPtr<SPageStats> RenderPageWidget;
+		/** Widget showing page info regarding page info that overlays the graph tab content */
+		TSharedPtr<SPageStats> PageStatsWidget;
 
 		/** Widget showing render performance information that overlays the graph */
 		TSharedPtr<SRenderStats> RenderStatsWidget;
diff --git a/Engine/Plugins/Runtime/Metasound/Source/MetasoundEditor/Private/MetasoundEditorModule.cpp b/Engine/Plugins/Runtime/Metasound/Source/MetasoundEditor/Private/MetasoundEditorModule.cpp
--- a/Engine/Plugins/Runtime/Metasound/Source/MetasoundEditor/Private/MetasoundEditorModule.cpp
+++ b/Engine/Plugins/Runtime/Metasound/Source/MetasoundEditor/Private/MetasoundEditorModule.cpp
@@ -2,6 +2,10 @@
 #include "MetasoundEditorModule.h"
 
 #include "AssetTypeActions_Base.h"
+#include "AudioMeter.h"
+#include "AudioOscilloscopePanelStyle.h"
+#include "AudioSpectrumPlotStyle.h"
+#include "AudioVectorscopePanelStyle.h"
 #include "Brushes/SlateImageBrush.h"
 #include "EdGraph/EdGraphNode.h"
 #include "EdGraph/EdGraphPin.h"
@@ -160,13 +164,17 @@ namespace Metasound
 
 					// Analyzers
 					Set("MetasoundEditor.Analyzers.BackgroundColor", FLinearColor(0.0075f, 0.0075f, 0.0075, 1.0f));
+					Set("MetasoundEditor.Analyzers.ForegroundColor", FLinearColor(0.025719f, 0.208333f, 0.069907f, 1.0f)); // "Audio" Green
 
 					// Misc
 					Set("MetasoundEditor.Audition", new IMAGE_BRUSH_SVG(TEXT("Icons/metasound_page"), Icon16));
 					Set("MetasoundEditor.Metasound.Icon", new IMAGE_BRUSH_SVG(TEXT("Icons/metasound_icon"), Icon16));
-					Set("MetasoundEditor.Page.Executing", new IMAGE_BRUSH_SVG(TEXT("Icons/metasound_page_exec"), Icon16));
 					Set("MetasoundEditor.Speaker", new FSlateImageBrush(RootToContentDir(TEXT("/Icons/speaker_144x.png")), FVector2D(144.0f, 144.0f)));
 
+					// Pages
+					Set("MetasoundEditor.Page.Executing.ForegroundColor", FStyleColors::AccentGreen.GetSpecifiedColor());
+					Set("MetasoundEditor.Page.Executing", new IMAGE_BRUSH_SVG(TEXT("Icons/metasound_page_exec"), Icon16));
+
 					// Class Icons
 					auto SetClassIcon = [this, InIcon16 = Icon16, InIcon64 = Icon64](const FString& ClassName)
 					{
@@ -201,6 +209,21 @@ namespace Metasound
 				return { "MetaSoundStyle", InName};
 			}
 
+			const FSlateColor& GetDefaultAnalyzerColor()
+			{
+				auto MakeColor = []() -> FSlateColor
+				{
+					if (const ISlateStyle* MetaSoundStyle = FSlateStyleRegistry::FindSlateStyle("MetaSoundStyle"))
+					{
+						return MetaSoundStyle->GetColor("MetasoundEditor.Analyzers.ForegroundColor");
+					}
+
+					return FStyleColors::AccentWhite;
+				};
+				static const FSlateColor AnalyzerColor = MakeColor();
+				return AnalyzerColor;
+			}
+
 			const FSlateBrush& GetSlateBrushSafe(FName InName)
 			{
 				const ISlateStyle* MetaSoundStyle = FSlateStyleRegistry::FindSlateStyle("MetaSoundStyle");
@@ -221,7 +244,77 @@ namespace Metasound
 				static const FSlateBrush NullBrush;
 				return NullBrush;
 			}
-		}
+
+			const FAudioMeterDefaultColorStyle& GetMeterDefaultColorStyle()
+			{
+				auto MakeStyle = []()
+				{
+					FAudioMeterDefaultColorStyle MeterStyle;
+					MeterStyle.MeterValueColor = GetDefaultAnalyzerColor().GetSpecifiedColor();
+					return MeterStyle;
+				};
+				static const FAudioMeterDefaultColorStyle ThisStyle = MakeStyle();
+				return ThisStyle;
+			}
+
+			const FAudioOscilloscopePanelStyle& GetOscilloscopeStyle()
+			{
+				auto MakeStyle = []()
+				{
+					FAudioOscilloscopePanelStyle OscStyle;
+					FSampledSequenceViewerStyle SampleView;
+					SampleView.SequenceColor = GetDefaultAnalyzerColor();
+					OscStyle.SetWaveViewerStyle(SampleView);
+
+					return OscStyle;
+				};
+				static FAudioOscilloscopePanelStyle ThisStyle = MakeStyle();
+				return ThisStyle;
+			}
+
+			const FSlateColor& GetPageExecutingColor()
+			{
+				auto MakeColor = []() -> FSlateColor
+				{
+					if (const ISlateStyle* MetaSoundStyle = FSlateStyleRegistry::FindSlateStyle("MetaSoundStyle"))
+					{
+						return MetaSoundStyle->GetColor("MetasoundEditor.Page.Executing.ForegroundColor");
+					}
+
+					return FStyleColors::AccentWhite;
+				};
+				static const FSlateColor AnalyzerColor = MakeColor();
+				return AnalyzerColor;
+			}
+
+			const FAudioSpectrumPlotStyle& GetSpectrumPlotStyle()
+			{
+				auto MakeStyle = []()
+				{
+					FAudioSpectrumPlotStyle PlotStyle;
+					const FSlateColor AnalyzerColor = GetDefaultAnalyzerColor();
+					PlotStyle.CrosshairColor = AnalyzerColor.UseSubduedForeground();
+					PlotStyle.SpectrumColor = AnalyzerColor;
+					return PlotStyle;
+				};
+				static FAudioSpectrumPlotStyle ThisStyle = MakeStyle();
+				return ThisStyle;
+			}
+
+			const FAudioVectorscopePanelStyle& GetVectorscopeStyle()
+			{
+				auto MakeStyle = []()
+				{
+					FAudioVectorscopePanelStyle PanelStyle;
+					FSampledSequenceVectorViewerStyle VectorViewStyle;
+					VectorViewStyle.LineColor = GetDefaultAnalyzerColor().GetSpecifiedColor();
+					PanelStyle.SetVectorViewerStyle(VectorViewStyle);
+					return PanelStyle;
+				};
+				static FAudioVectorscopePanelStyle ThisStyle = MakeStyle();
+				return ThisStyle;
+			}
+		} // namespace Style
 
 		// A structure that contains information about registered custom pin types. 
 		struct FGraphPinConfiguration
diff --git a/Engine/Plugins/Runtime/Metasound/Source/MetasoundEditor/Private/MetasoundEditorSettings.h b/Engine/Plugins/Runtime/Metasound/Source/MetasoundEditor/Private/MetasoundEditorSettings.h
--- a/Engine/Plugins/Runtime/Metasound/Source/MetasoundEditor/Private/MetasoundEditorSettings.h
+++ b/Engine/Plugins/Runtime/Metasound/Source/MetasoundEditor/Private/MetasoundEditorSettings.h
@@ -271,23 +271,23 @@ public:
 	EMetasoundActiveDetailView DetailView = EMetasoundActiveDetailView::General;
 
 	/** Whether the AudioMaterialWidgets are used when possible in Metasound Editor*/
-	UPROPERTY(EditAnywhere, config, DisplayName = "Use Audio Material Widgets", Category = WidgetStyling)
+	UPROPERTY(EditAnywhere, config, DisplayName = "Use Audio Material Widgets", Category = "Widget Styling (Experimental)")
 	bool bUseAudioMaterialWidgets = false;
 	
 	/**Override the Knob Style used in the Metasound Editor.*/
-	UPROPERTY(EditAnywhere, config, Category = WidgetStyling, meta = (AllowedClasses = "/Script/SlateCore.SlateWidgetStyleAsset", EditCondition = "bUseAudioMaterialWidgets", DisplayName = "Knob Style"))
+	UPROPERTY(EditAnywhere, config, Category = "Widget Styling (Experimental)", meta = (AllowedClasses = "/Script/SlateCore.SlateWidgetStyleAsset", EditCondition = "bUseAudioMaterialWidgets", DisplayName = "Knob Style"))
 	FSoftObjectPath KnobStyleOverride;
 	
 	/**Override the Slider Style used in the Metasound Editor.*/
-	UPROPERTY(EditAnywhere, config, Category = WidgetStyling, meta = (AllowedClasses = "/Script/SlateCore.SlateWidgetStyleAsset", EditCondition = "bUseAudioMaterialWidgets", DisplayName = "Slider Style"))
+	UPROPERTY(EditAnywhere, config, Category = "Widget Styling (Experimental)", meta = (AllowedClasses = "/Script/SlateCore.SlateWidgetStyleAsset", EditCondition = "bUseAudioMaterialWidgets", DisplayName = "Slider Style"))
 	FSoftObjectPath SliderStyleOverride;
 
 	/**Override the Button Style used in the Metasound Editor.*/
-	UPROPERTY(EditAnywhere, config, Category = WidgetStyling, meta = (AllowedClasses = "/Script/SlateCore.SlateWidgetStyleAsset", EditCondition = "bUseAudioMaterialWidgets", DisplayName = "Button Style"))
+	UPROPERTY(EditAnywhere, config, Category = "Widget Styling (Experimental)", meta = (AllowedClasses = "/Script/SlateCore.SlateWidgetStyleAsset", EditCondition = "bUseAudioMaterialWidgets", DisplayName = "Button Style"))
 	FSoftObjectPath ButtonStyleOverride;
 	
 	/**Override the Meter Style used in the Metasound Editor.*/
-	UPROPERTY(EditAnywhere, config, Category = WidgetStyling, meta = (AllowedClasses = "/Script/SlateCore.SlateWidgetStyleAsset", EditCondition = "bUseAudioMaterialWidgets", DisplayName = "Meter Style"))
+	UPROPERTY(EditAnywhere, config, Category = "Widget Styling (Experimental)", meta = (AllowedClasses = "/Script/SlateCore.SlateWidgetStyleAsset", EditCondition = "bUseAudioMaterialWidgets", DisplayName = "Meter Style"))
 	FSoftObjectPath MeterStyleOverride;
 
 	//UObject
diff --git a/Engine/Plugins/Runtime/Metasound/Source/MetasoundEditor/Private/MetasoundPinAudioInspector.cpp b/Engine/Plugins/Runtime/Metasound/Source/MetasoundEditor/Private/MetasoundPinAudioInspector.cpp
--- a/Engine/Plugins/Runtime/Metasound/Source/MetasoundEditor/Private/MetasoundPinAudioInspector.cpp
+++ b/Engine/Plugins/Runtime/Metasound/Source/MetasoundEditor/Private/MetasoundPinAudioInspector.cpp
@@ -9,115 +9,115 @@
 #include "MetasoundEditor.h"
 #include "MetasoundEditorGraphBuilder.h"
 #include "MetasoundEditorGraphNode.h"
+#include "MetasoundEditorModule.h"
 
 #define LOCTEXT_NAMESPACE "MetasoundEditor"
 
-namespace Metasound
+namespace Metasound::Editor
 {
-	namespace Editor
+	namespace FMetasoundPinAudioInspectorPrivate
 	{
-		namespace FMetasoundPinAudioInspectorPrivate
+		UEdGraphPin* ResolvePinObjectAsOutput(UEdGraphPin* InPin)
 		{
-			UEdGraphPin* ResolvePinObjectAsOutput(UEdGraphPin* InPin)
+			if (InPin && !InPin->LinkedTo.IsEmpty())
 			{
-				if (InPin && !InPin->LinkedTo.IsEmpty())
+				// Swap to show connected output if input (Only ever one)
+				if (InPin->Direction == EGPD_Input)
 				{
-					// Swap to show connected output if input (Only ever one)
-					if (InPin->Direction == EGPD_Input)
-					{
-						InPin = InPin->LinkedTo.Last();
-						check(InPin->Direction == EGPD_Output);
-					}
+					InPin = InPin->LinkedTo.Last();
+					check(InPin->Direction == EGPD_Output);
 				}
-
-				return InPin;
 			}
+
+			return InPin;
 		}
+	}
 
-		FMetasoundPinAudioInspector::FMetasoundPinAudioInspector(FEdGraphPinReference InPinRef)
-			: GraphPinObj(FMetasoundPinAudioInspectorPrivate::ResolvePinObjectAsOutput(InPinRef.Get()))
-		{
-			using namespace Audio;
+	FMetasoundPinAudioInspector::FMetasoundPinAudioInspector(FEdGraphPinReference InPinRef)
+		: GraphPinObj(FMetasoundPinAudioInspectorPrivate::ResolvePinObjectAsOutput(InPinRef.Get()))
+	{
+		using namespace Audio;
 
-			// Initialize Oscilloscope
-			const FDeviceId AudioDeviceId = GEditor->GetMainAudioDeviceID();
+		// Initialize Oscilloscope
+		const FDeviceId AudioDeviceId = GEditor->GetMainAudioDeviceID();
 
-			constexpr int32 NumChannels = 1; // Audio wires are currently mono signals
+		constexpr int32 NumChannels = 1; // Audio wires are currently mono signals
 
-			Oscilloscope = MakeShared<AudioWidgets::FAudioOscilloscope>(AudioDeviceId,
-				NumChannels,
-				/*InTimeWindowMs*/     10.0f,
-				/*InMaxTimeWindowMs*/  10.0f,
-				/*InAnalysisPeriodMs*/ 10.0f,
-				/*InPanelLayoutType*/  EAudioPanelLayoutType::Basic);
+		Oscilloscope = MakeShared<AudioWidgets::FAudioOscilloscope>(AudioDeviceId,
+			NumChannels,
+			/*InTimeWindowMs*/     10.0f,
+			/*InMaxTimeWindowMs*/  10.0f,
+			/*InAnalysisPeriodMs*/ 10.0f,
+			/*InPanelLayoutType*/  EAudioPanelLayoutType::Basic,
+			/*InOscilloscopePanelStyle*/ &Style::GetOscilloscopeStyle()
+		);
 
-			PinAudioInspectorWidget = SNew(SMetasoundPinAudioInspector)
-									  .VisualizationWidget(Oscilloscope->GetPanelWidget());
+		PinAudioInspectorWidget = SNew(SMetasoundPinAudioInspector)
+			.VisualizationWidget(Oscilloscope->GetPanelWidget());
 
-			Oscilloscope->StartProcessing();
+		Oscilloscope->StartProcessing();
 
-			// Analyze Audio Pin
-			if (GraphPinObj && !GraphPinObj->LinkedTo.IsEmpty() && GraphPinObj->PinType.PinCategory == FGraphBuilder::PinCategoryAudio)
+		// Analyze Audio Pin
+		if (GraphPinObj && !GraphPinObj->LinkedTo.IsEmpty() && GraphPinObj->PinType.PinCategory == FGraphBuilder::PinCategoryAudio)
+		{
+			if (FGraphConnectionManager* ConnectionManager = GetConnectionManager())
 			{
-				if (FGraphConnectionManager* ConnectionManager = GetConnectionManager())
-				{
-					const Frontend::FConstOutputHandle OutputHandle = FGraphBuilder::FindReroutedOutputHandleFromPin(GraphPinObj);
-					const FGuid NodeID = OutputHandle->GetOwningNodeID();
-					const FName OutputName = OutputHandle->GetName();
-					AnalyzerInstanceID = ConnectionManager->AddAudioBusWriter(NodeID, OutputName, AudioDeviceId, Oscilloscope->GetAudioBus());
-				}
+				const Frontend::FConstOutputHandle OutputHandle = FGraphBuilder::FindReroutedOutputHandleFromPin(GraphPinObj);
+				const FGuid NodeID = OutputHandle->GetOwningNodeID();
+				const FName OutputName = OutputHandle->GetName();
+				AnalyzerInstanceID = ConnectionManager->AddAudioBusWriter(NodeID, OutputName, AudioDeviceId, Oscilloscope->GetAudioBus());
 			}
 		}
+	}
 
-		FMetasoundPinAudioInspector::~FMetasoundPinAudioInspector()
+	FMetasoundPinAudioInspector::~FMetasoundPinAudioInspector()
+	{
+		// Remove audio pin analyzer
+		if (AnalyzerInstanceID.IsValid())
 		{
-			// Remove audio pin analyzer
-			if (AnalyzerInstanceID.IsValid())
+			if (FGraphConnectionManager* ConnectionManager = GetConnectionManager())
 			{
-				if (FGraphConnectionManager* ConnectionManager = GetConnectionManager())
-				{
-					ConnectionManager->RemoveAudioBusWriter(AnalyzerInstanceID);
-					AnalyzerInstanceID.Invalidate();
-				}
+				ConnectionManager->RemoveAudioBusWriter(AnalyzerInstanceID);
+				AnalyzerInstanceID.Invalidate();
 			}
+		}
 
-			// Stop AudioBus
-			if (UWorld* EditorWorld = GEditor->GetEditorWorldContext().World())
+		// Stop AudioBus
+		if (UWorld* EditorWorld = GEditor->GetEditorWorldContext().World())
+		{
+			if (Audio::FMixerDevice* MixerDevice = static_cast<Audio::FMixerDevice*>(EditorWorld->GetAudioDeviceRaw()))
 			{
-				if (Audio::FMixerDevice* MixerDevice = static_cast<Audio::FMixerDevice*>(EditorWorld->GetAudioDeviceRaw()))
+				if (UAudioBusSubsystem* AudioBusSubsystem = MixerDevice->GetSubsystem<UAudioBusSubsystem>())
 				{
-					if (UAudioBusSubsystem* AudioBusSubsystem = MixerDevice->GetSubsystem<UAudioBusSubsystem>())
-					{
-						const uint32 AudioBusId = Oscilloscope->GetAudioBus()->GetUniqueID();
-						AudioBusSubsystem->StopAudioBus(Audio::FAudioBusKey(AudioBusId));
-					}
+					const uint32 AudioBusId = Oscilloscope->GetAudioBus()->GetUniqueID();
+					AudioBusSubsystem->StopAudioBus(Audio::FAudioBusKey(AudioBusId));
 				}
 			}
 		}
+	}
 
-		TSharedPtr<SMetasoundPinAudioInspector> FMetasoundPinAudioInspector::GetWidget()
-		{
-			return PinAudioInspectorWidget;
-		}
+	TSharedPtr<SMetasoundPinAudioInspector> FMetasoundPinAudioInspector::GetWidget()
+	{
+		return PinAudioInspectorWidget;
+	}
 
-		const UMetasoundEditorGraphNode& FMetasoundPinAudioInspector::GetReroutedNode() const
-		{
-			const UMetasoundEditorGraphNode* Node = FGraphBuilder::FindReroutedOutputPin(GraphPinObj)
-				? Cast<UMetasoundEditorGraphNode>(GraphPinObj->GetOwningNode())
-				: nullptr;
-			check(Node);
+	const UMetasoundEditorGraphNode& FMetasoundPinAudioInspector::GetReroutedNode() const
+	{
+		const UMetasoundEditorGraphNode* Node = FGraphBuilder::FindReroutedOutputPin(GraphPinObj)
+			? Cast<UMetasoundEditorGraphNode>(GraphPinObj->GetOwningNode())
+			: nullptr;
+		check(Node);
 
-			return *Node;
-		}
+		return *Node;
+	}
 
-		FGraphConnectionManager* FMetasoundPinAudioInspector::GetConnectionManager()
-		{
-			const UMetasoundEditorGraphNode& Node = GetReroutedNode();
-			TSharedPtr<FEditor> Editor = FGraphBuilder::GetEditorForNode(Node);
+	FGraphConnectionManager* FMetasoundPinAudioInspector::GetConnectionManager()
+	{
+		const UMetasoundEditorGraphNode& Node = GetReroutedNode();
+		TSharedPtr<FEditor> Editor = FGraphBuilder::GetEditorForNode(Node);
 
-			return Editor.IsValid() ? &Editor->GetConnectionManager() : nullptr;
-		}
-	} // namespace Editor
-} // namespace Metasound
+		return Editor.IsValid() ? &Editor->GetConnectionManager() : nullptr;
+	}
+} // namespace Metasound::Editor
 
 #undef LOCTEXT_NAMESPACE
diff --git a/Engine/Plugins/Runtime/Metasound/Source/MetasoundEditor/Private/SMetasoundFilterFrequencyResponsePlots.cpp b/Engine/Plugins/Runtime/Metasound/Source/MetasoundEditor/Private/SMetasoundFilterFrequencyResponsePlots.cpp
--- a/Engine/Plugins/Runtime/Metasound/Source/MetasoundEditor/Private/SMetasoundFilterFrequencyResponsePlots.cpp
+++ b/Engine/Plugins/Runtime/Metasound/Source/MetasoundEditor/Private/SMetasoundFilterFrequencyResponsePlots.cpp
@@ -107,6 +107,7 @@ namespace Metasound::Editor
 					.TiltExponent_Lambda([] { return 0.0f; }) // Binding this property has the effect of hiding its context menu entry (Tilting the spectrum is not desired here).
 					.FrequencyAxisPixelBucketMode_Lambda([]() { return EAudioSpectrumPlotFrequencyAxisPixelBucketMode::Sample; }) // Binding this property has the effect of hiding its context menu entry (PixelBucketMode is not much use here).
 					.OnGetAudioSpectrumData(this, &SMetaSoundFilterFrequencyResponsePlot::GetAudioSpectrumData)
+					.Style(&Editor::Style::GetSpectrumPlotStyle())
 			];
 
 		ContextMenuExtension = FrequencyResponsePlot->AddContextMenuExtension(EExtensionHook::Before, nullptr, FMenuExtensionDelegate::CreateSP(this, &SMetaSoundFilterFrequencyResponsePlot::ExtendSpectrumPlotContextMenu));
diff --git a/Engine/Plugins/Runtime/Metasound/Source/MetasoundEditor/Private/SMetasoundSpectrumAnalyzerGraphNode.cpp b/Engine/Plugins/Runtime/Metasound/Source/MetasoundEditor/Private/SMetasoundSpectrumAnalyzerGraphNode.cpp
--- a/Engine/Plugins/Runtime/Metasound/Source/MetasoundEditor/Private/SMetasoundSpectrumAnalyzerGraphNode.cpp
+++ b/Engine/Plugins/Runtime/Metasound/Source/MetasoundEditor/Private/SMetasoundSpectrumAnalyzerGraphNode.cpp
@@ -5,6 +5,7 @@
 #include "Editor.h"
 #include "MetasoundEditorGraphBuilder.h"
 #include "MetasoundEditorGraphNode.h"
+#include "MetasoundEditorModule.h"
 #include "Widgets/Layout/SBox.h"
 
 namespace Metasound::Editor
@@ -56,13 +57,17 @@ namespace Metasound::Editor
 
 	void SMetaSoundSpectrumAnalyzerGraphNode::CreateBelowPinControls(TSharedPtr<SVerticalBox> MainBox)
 	{
+		using namespace AudioWidgets;
+
 		if (ensure(GEditor))
 		{
 			if (!SpectrumAnalyzer.IsValid())
 			{
-				constexpr int32 NumChannels = 1;
-				const Audio::FDeviceId AudioDeviceId = GEditor->GetMainAudioDeviceID();
-				SpectrumAnalyzer = MakeShared<AudioWidgets::FAudioSpectrumAnalyzer>(NumChannels, AudioDeviceId);
+				FAudioSpectrumAnalyzerParams Params;
+				Params.NumChannels = 1;
+				Params.AudioDeviceId = GEditor->GetMainAudioDeviceID();
+				Params.PlotStyle = &Style::GetSpectrumPlotStyle();
+				SpectrumAnalyzer = MakeShared<FAudioSpectrumAnalyzer>(Params);
 			}
 
 			MainBox->AddSlot()
diff --git a/Engine/Plugins/Runtime/Metasound/Source/MetasoundEditor/Private/SMetasoundStats.cpp b/Engine/Plugins/Runtime/Metasound/Source/MetasoundEditor/Private/SMetasoundStats.cpp
--- a/Engine/Plugins/Runtime/Metasound/Source/MetasoundEditor/Private/SMetasoundStats.cpp
+++ b/Engine/Plugins/Runtime/Metasound/Source/MetasoundEditor/Private/SMetasoundStats.cpp
@@ -32,30 +32,40 @@ namespace Metasound::Editor
 	{
 		SVerticalBox::Construct(SVerticalBox::FArguments());
 
-		TSharedRef<SHorizontalBox> PageWidgetBox = SNew(SHorizontalBox);
-		PageWidgetBox->AddSlot()
-			.Padding(2.0f)
-			.HAlign(HAlign_Center)
-			[
-				SAssignNew(PageTextWidget, STextBlock)
-				.Visibility(EVisibility::HitTestInvisible)
+		AddSlot().HAlign(HAlign_Left)
+		[
+			SAssignNew(AuditionPageTextWidget, STextBlock)
+				.Visibility(EVisibility::Collapsed)
 				.TextStyle(FAppStyle::Get(), "Graph.ZoomText")
 				.ColorAndOpacity(StatsPrivate::BaseTextColor)
-			];
-
-		ExecImageWidget = SNew(SImage)
-			.Image(Style::CreateSlateIcon("MetasoundEditor.Page.Executing").GetIcon())
-			.DesiredSizeOverride(FVector2D(24.f, 24.f))
-			.ColorAndOpacity(FStyleColors::AccentGreen)
-			.Visibility(EVisibility::Collapsed);
-
-		PageWidgetBox->AddSlot()
-			.Padding(2.0f)
-			.HAlign(HAlign_Center)
-			.VAlign(VAlign_Center)
-			.AutoWidth()[ ExecImageWidget.ToSharedRef()];
+		];
 
-		AddSlot().HAlign(HAlign_Left) [ PageWidgetBox ];
+		AddSlot().HAlign(HAlign_Left)
+		[
+			SNew(SHorizontalBox)
+				+ SHorizontalBox::Slot()
+				.Padding(2.0f)
+				.HAlign(HAlign_Center)
+				[
+					SAssignNew(GraphPageTextWidget, STextBlock)
+					.Visibility(EVisibility::HitTestInvisible)
+					.TextStyle(FAppStyle::Get(), "Graph.ZoomText")
+					.ColorAndOpacity(StatsPrivate::BaseTextColor)
+				]
+
+				+ SHorizontalBox::Slot()
+				.Padding(2.0f)
+				.HAlign(HAlign_Center)
+				.VAlign(VAlign_Center)
+				.AutoWidth()
+				[
+					SAssignNew(ExecImageWidget, SImage)
+						.Image(Style::CreateSlateIcon("MetasoundEditor.Page.Executing").GetIcon())
+						.DesiredSizeOverride(FVector2D(24.f, 24.f))
+						.ColorAndOpacity(Style::GetPageExecutingColor())
+						.Visibility(EVisibility::Collapsed)
+				]
+		];
 	}
 
 	void SPageStats::SetExecVisibility(TAttribute<EVisibility> InVisibility)
@@ -63,87 +73,151 @@ namespace Metasound::Editor
 		ExecImageWidget->SetVisibility(MoveTemp(InVisibility));
 	}
 
-	void SPageStats::Update(const FMetaSoundPageSettings* PageSettings, const FText& Header, const FSlateColor* ColorOverride)
+	void SPageStats::Update(const FMetaSoundPageSettings* AuditionPageSettings, const FMetaSoundPageSettings* GraphPageSettings, const FSlateColor* ActiveColor)
 	{
 		using namespace Engine;
 
-		FText PageInfo;
-		EVisibility Visibility = EVisibility::Collapsed;
-		if (PageTextWidget.IsValid() && PageSettings)
+		const FText PageStatsFormat = LOCTEXT("PageStatsFormat", "{0}: {1}");
+
 		{
-			PageInfo = FText::Format(LOCTEXT("PageStatsFormat", "{0}: {1}"), Header, FText::FromString(PageSettings->Name.ToString()));
-			Visibility = EVisibility::Visible;
+			FText PageInfo;
+			EVisibility Visibility = EVisibility::Collapsed;
+			if (AuditionPageTextWidget.IsValid() && AuditionPageSettings)
+			{
+				const FText Header = LOCTEXT("AuditionPageHeader", "Auditioning Page");
+				PageInfo = FText::Format(PageStatsFormat, Header, FText::FromString(AuditionPageSettings->Name.ToString()));
+				Visibility = EVisibility::Visible;
+			}
+			AuditionPageTextWidget->SetText(PageInfo);
 		}
 
-		PageTextWidget->SetText(PageInfo);
-		PageTextWidget->SetColorAndOpacity(ColorOverride ? *ColorOverride : StatsPrivate::BaseTextColor);
+		{
+			FText PageInfo;
+			EVisibility Visibility = EVisibility::Collapsed;
+			if (GraphPageTextWidget.IsValid() && GraphPageSettings)
+			{
+				const FText Header = LOCTEXT("GraphPageTargetHeader", "Graph Page");
+				PageInfo = FText::Format(PageStatsFormat, Header, FText::FromString(GraphPageSettings->Name.ToString()));
+				Visibility = EVisibility::Visible;
+			}
+			GraphPageTextWidget->SetText(PageInfo);
+			GraphPageTextWidget->SetColorAndOpacity(ActiveColor ? *ActiveColor : StatsPrivate::BaseTextColor);
+		}
 	}
 
 	void SRenderStats::Construct(const FArguments& InArgs)
 	{
 		SVerticalBox::Construct(SVerticalBox::FArguments());
-		AddSlot()
-		.HAlign(HAlign_Left)
-		.AutoHeight()
-		[
-			SAssignNew(RenderStatsCostWidget, STextBlock)
-			.Visibility(EVisibility::HitTestInvisible)
-			.TextStyle(FAppStyle::Get(), "Graph.ZoomText")
-			.ColorAndOpacity(FLinearColor(1, 1, 1, 0.30f))
-		];
 
-		AddSlot()
-		.HAlign(HAlign_Left)
-		.AutoHeight()
-		[
-			SAssignNew(RenderStatsCPUWidget, STextBlock)
-			.Visibility(EVisibility::HitTestInvisible)
-			.TextStyle(FAppStyle::Get(), "Graph.ZoomText")
-			.ColorAndOpacity(FLinearColor(1, 1, 1, 0.30f))
-		];
+		auto AddTextWidgetSlot = [this](TSharedPtr<STextBlock>& OutBlock)
+		{
+			AddSlot()
+			.HAlign(HAlign_Left)
+			.AutoHeight()
+			[
+				SAssignNew(OutBlock, STextBlock)
+				.Visibility(EVisibility::HitTestInvisible)
+				.TextStyle(FAppStyle::Get(), "GraphPreview.CornerText")
+				.ColorAndOpacity(FLinearColor(1, 1, 1, 0.30f))
+			];
+		};
+
+		AddTextWidgetSlot(PlayTimeWidget);
+		AddTextWidgetSlot(RenderStatsCostWidget);
+		AddTextWidgetSlot(RenderStatsCPUWidget);
+		AddTextWidgetSlot(AuditionPageWidget);
+		AddTextWidgetSlot(AuditionPlatformWidget);
 	}
 
-	void SRenderStats::Update(bool bIsPlaying, const UMetaSoundSource* InSource)
+
+	void SRenderStats::Update(bool bIsPlaying, double InDeltaTime, const UMetaSoundSource* InSource)
 	{
 		using namespace Metasound;
 
 		// Reset maximum values when play restarts
-		if (bIsPlaying && !bPreviousIsPlaying)
+		if (bIsPlaying)
+		{
+			if (!bPreviousIsPlaying)
+			{
+				MaxRelativeRenderCost = 0.f;
+				MaxCPUCoreUtilization = 0;
+			}
+			PlayTime += InDeltaTime;
+		}
+		else
 		{
-			MaxRelativeRenderCost = 0.f;
-			MaxCPUCoreUtilization = 0;
+			PlayTime = 0.0;
 		}
+
 		bPreviousIsPlaying = bIsPlaying;
 
-		if (RenderStatsCPUWidget.IsValid() && RenderStatsCostWidget.IsValid())
-		{
-			double CPUCoreUtilization = 0;
-			float RelativeRenderCost = 0.f;
+		if (!RenderStatsCPUWidget.IsValid()
+			|| !RenderStatsCostWidget.IsValid()
+			|| !PlayTimeWidget.IsValid()
+			|| !AuditionPageWidget.IsValid()
+			|| !AuditionPlatformWidget.IsValid())
+			{
+			return;
+		}
 
-			// Find generator for playing preview component. 
-			if (bIsPlaying && InSource)
+		double CPUCoreUtilization = 0;
+		float RelativeRenderCost = 0.f;
+
+		if (bIsPlaying && InSource)
+		{
+			if (const UAudioComponent* PreviewComponent = GEditor->GetPreviewAudioComponent())
 			{
-				if (const UAudioComponent* PreviewComponent = GEditor->GetPreviewAudioComponent())
+				TSharedPtr<FMetasoundGenerator> Generator = InSource->GetGeneratorForAudioComponent(PreviewComponent->GetAudioComponentID()).Pin();
+				if (Generator.IsValid())
 				{
-					TSharedPtr<FMetasoundGenerator> Generator = InSource->GetGeneratorForAudioComponent(PreviewComponent->GetAudioComponentID()).Pin();
-					if (Generator.IsValid())
-					{
-						// Update render stats
-						CPUCoreUtilization = Generator->GetCPUCoreUtilization();
-						MaxCPUCoreUtilization = FMath::Max(MaxCPUCoreUtilization, CPUCoreUtilization);
+					// Update render stats
+					CPUCoreUtilization = Generator->GetCPUCoreUtilization();
+					MaxCPUCoreUtilization = FMath::Max(MaxCPUCoreUtilization, CPUCoreUtilization);
 
-						RelativeRenderCost = Generator->GetRelativeRenderCost();
-						MaxRelativeRenderCost = FMath::Max(MaxRelativeRenderCost, RelativeRenderCost);
-					}
+					RelativeRenderCost = Generator->GetRelativeRenderCost();
+					MaxRelativeRenderCost = FMath::Max(MaxRelativeRenderCost, RelativeRenderCost);
 				}
 			}
+		}
+
+		// Display updated render stats.
+		FString PlayTimeString = FTimespan::FromSeconds(PlayTime).ToString();
+		PlayTimeString.ReplaceInline(TEXT("+"), TEXT(""));
+		PlayTimeWidget->SetText(FText::FromString(MoveTemp(PlayTimeString)));
 
-			// Display updated render stats. 
-			FString CPUCoreUtilizationString = FString::Printf(TEXT("%3.2f%% (%3.2f%% Max) CPU Core"), 100. * CPUCoreUtilization, 100. * MaxCPUCoreUtilization);
-			RenderStatsCPUWidget->SetText(FText::FromString(CPUCoreUtilizationString));
+		FString RenderCostString = FString::Printf(TEXT("Relative Render Cost: %3.2f (%3.2f Max)"), RelativeRenderCost, MaxRelativeRenderCost);
+		RenderStatsCostWidget->SetText(FText::FromString(MoveTemp(RenderCostString)));
+
+		FString CPUCoreUtilizationString = FString::Printf(TEXT("CPU Core: %3.2f%% (%3.2f%% Max)"), 100. * CPUCoreUtilization, 100. * MaxCPUCoreUtilization);
+		RenderStatsCPUWidget->SetText(FText::FromString(MoveTemp(CPUCoreUtilizationString)));
+
+		if (bPreviousIsPlaying != bIsPlaying)
+		{
+			FText AuditionPage;
+			FText AuditionPlatform;
 
-			FString RenderCostString = FString::Printf(TEXT("%3.2f (%3.2f Max) Relative Render Cost"), RelativeRenderCost, MaxRelativeRenderCost);
-			RenderStatsCostWidget->SetText(FText::FromString(RenderCostString));
+			if (bIsPlaying)
+			{
+				if (const UMetasoundEditorSettings* EditorSettings = GetDefault<UMetasoundEditorSettings>())
+				{
+					if (const UMetaSoundSettings* Settings = GetDefault<UMetaSoundSettings>())
+					{
+						if (!Settings->GetProjectPageSettings().IsEmpty())
+						{
+							const FText AuditionPageHeader = LOCTEXT("AuditionPageActive_DebugFormat", "Auditioning Page: {0}");
+							AuditionPage = FText::Format(AuditionPageHeader, FText::FromName(EditorSettings->AuditionPage));
+
+							if (EditorSettings->AuditionPlatform != "Editor")
+							{
+								const FText AuditionPlatformHeader = LOCTEXT("AuditionPlatformActive_DebugFormat", "Auditioning Platform: {0}");
+								AuditionPlatform = FText::Format(AuditionPlatformHeader, FText::FromName(EditorSettings->AuditionPlatform));
+							}
+						}
+					}
+				}
+			}
+			AuditionPageWidget->SetText(AuditionPage);
+			AuditionPlatformWidget->SetText(AuditionPlatform);
 		}
 	}
 } // namespace Metasound::Editor
diff --git a/Engine/Plugins/Runtime/Metasound/Source/MetasoundEditor/Private/SMetasoundStats.h b/Engine/Plugins/Runtime/Metasound/Source/MetasoundEditor/Private/SMetasoundStats.h
--- a/Engine/Plugins/Runtime/Metasound/Source/MetasoundEditor/Private/SMetasoundStats.h
+++ b/Engine/Plugins/Runtime/Metasound/Source/MetasoundEditor/Private/SMetasoundStats.h
@@ -28,11 +28,12 @@ namespace Metasound::Editor
 
 		void Construct(const FArguments& InArgs);
 		void SetExecVisibility(TAttribute<EVisibility> InVisibility);
-		void Update(const FMetaSoundPageSettings* PageSettings, const FText& Header, const FSlateColor* TextColor);
+		void Update(const FMetaSoundPageSettings* AuditionPageSettings, const FMetaSoundPageSettings* GraphPageSettings, const FSlateColor* ActiveColor);
 
 	private:
 		TSharedPtr<SImage> ExecImageWidget;
-		TSharedPtr<STextBlock> PageTextWidget;
+		TSharedPtr<STextBlock> GraphPageTextWidget;
+		TSharedPtr<STextBlock> AuditionPageTextWidget;
 	};
 
 
@@ -44,14 +45,18 @@ namespace Metasound::Editor
 		SLATE_END_ARGS()
 
 		void Construct(const FArguments& InArgs);
-		void Update(bool bIsPlaying, const UMetaSoundSource* InSource);
+		void Update(bool bIsPlaying, double InDeltaTime, const UMetaSoundSource* InSource);
 
 	private:
+		TSharedPtr<STextBlock> AuditionPageWidget;
+		TSharedPtr<STextBlock> AuditionPlatformWidget;
+		TSharedPtr<STextBlock> PlayTimeWidget;
 		TSharedPtr<STextBlock> RenderStatsCostWidget;
 		TSharedPtr<STextBlock> RenderStatsCPUWidget;
 
 		bool bPreviousIsPlaying = false;
-		double MaxCPUCoreUtilization = 0;
+		double MaxCPUCoreUtilization = 0.0;
+		double PlayTime = 0.0;
 		float MaxRelativeRenderCost = 0.f;
 	};
 } // namespace Metasound::Editor
diff --git a/Engine/Plugins/Runtime/Metasound/Source/MetasoundEditor/Public/MetasoundEditorModule.h b/Engine/Plugins/Runtime/Metasound/Source/MetasoundEditor/Public/MetasoundEditorModule.h
--- a/Engine/Plugins/Runtime/Metasound/Source/MetasoundEditor/Public/MetasoundEditorModule.h
+++ b/Engine/Plugins/Runtime/Metasound/Source/MetasoundEditor/Public/MetasoundEditorModule.h
@@ -15,6 +15,11 @@
 #include "Templates/Function.h"
 
 // Forward Declarations
+struct FAudioMeterDefaultColorStyle;
+struct FAudioOscilloscopePanelStyle;
+struct FAudioSpectrumPlotStyle;
+struct FAudioVectorscopePanelStyle;
+
 class IDetailLayoutBuilder;
 class IDetailPropertyRow;
 class IPropertyHandle;
@@ -41,6 +46,14 @@ namespace Metasound
 		{
 			METASOUNDEDITOR_API FSlateIcon CreateSlateIcon(FName InName);
 			METASOUNDEDITOR_API const FSlateBrush& GetSlateBrushSafe(FName InName);
+
+			const FSlateColor& GetDefaultAnalyzerColor();
+			const FSlateColor& GetPageExecutingColor();
+
+			const FAudioMeterDefaultColorStyle& GetMeterDefaultColorStyle();
+			const FAudioOscilloscopePanelStyle& GetOscilloscopeStyle();
+			const FAudioSpectrumPlotStyle& GetSpectrumPlotStyle();
+			const FAudioVectorscopePanelStyle& GetVectorscopeStyle();
 		} // namespace Style
 
 		using EAssetScanStatus = Metasound::Engine::EAssetScanStatus;
