Fix state change issue causing audition stats to not show (missed in last submitted shelved CL)
#rnx

diff --git a/Engine/Plugins/Runtime/Metasound/Source/MetasoundEditor/Private/SMetasoundStats.cpp b/Engine/Plugins/Runtime/Metasound/Source/MetasoundEditor/Private/SMetasoundStats.cpp
--- a/Engine/Plugins/Runtime/Metasound/Source/MetasoundEditor/Private/SMetasoundStats.cpp
+++ b/Engine/Plugins/Runtime/Metasound/Source/MetasoundEditor/Private/SMetasoundStats.cpp
@@ -135,6 +135,7 @@ namespace Metasound::Editor
 		using namespace Metasound;
 
 		// Reset maximum values when play restarts
+		const bool bPlayStateChanged = bIsPlaying != bPreviousIsPlaying;
 		if (bIsPlaying)
 		{
 			if (!bPreviousIsPlaying)
@@ -191,7 +192,7 @@ namespace Metasound::Editor
 		FString CPUCoreUtilizationString = FString::Printf(TEXT("CPU Core: %3.2f%% (%3.2f%% Max)"), 100. * CPUCoreUtilization, 100. * MaxCPUCoreUtilization);
 		RenderStatsCPUWidget->SetText(FText::FromString(MoveTemp(CPUCoreUtilizationString)));
 
-		if (bPreviousIsPlaying != bIsPlaying)
+		if (bPlayStateChanged)
 		{
 			FText AuditionPage;
 			FText AuditionPlatform;
