[Backout] - CL36443855
Original CL Desc
-----------------------------------------------------------------
- Uniform color fix-ups for MetaSound analyzers
- Move live MetaSound stats/audition settings to top left of graph
- Fix ordering of playback widget in input details panel to follow same order as page details panel

#rnx

diff --git a/Engine/Plugins/Runtime/AudioWidgets/Source/AudioWidgets/Private/AudioMeter.cpp b/Engine/Plugins/Runtime/AudioWidgets/Source/AudioWidgets/Private/AudioMeter.cpp
--- a/Engine/Plugins/Runtime/AudioWidgets/Source/AudioWidgets/Private/AudioMeter.cpp
+++ b/Engine/Plugins/Runtime/AudioWidgets/Source/AudioWidgets/Private/AudioMeter.cpp
@@ -160,12 +160,6 @@ const FText UAudioMeter::GetPaletteCategory()
 }
 #endif
 
-const FAudioMeterDefaultColorStyle& FAudioMeterDefaultColorStyle::GetDefault()
-{
-	static FAudioMeterDefaultColorStyle Style;
-	return Style;
-}
-
 namespace AudioWidgets
 {
 	FAudioMeter::FAudioMeter(int32 InNumChannels, UWorld& InWorld, TObjectPtr<UAudioBus> InExternalAudioBus)
@@ -173,18 +167,19 @@ namespace AudioWidgets
 	{
 	}
 
-	FAudioMeter::FAudioMeter(const int32 InNumChannels, const Audio::FDeviceId InAudioDeviceId, const TObjectPtr<UAudioBus> InExternalAudioBus, const FAudioMeterDefaultColorStyle* AudioMeterWidgetStyle)
+	FAudioMeter::FAudioMeter(const int32 InNumChannels, const Audio::FDeviceId InAudioDeviceId, const TObjectPtr<UAudioBus> InExternalAudioBus)
 	{
-		const FAudioMeterDefaultColorStyle* WidgetStyle = AudioMeterWidgetStyle ? AudioMeterWidgetStyle : &FAudioMeterDefaultColorStyle::GetDefault();
 		Widget = SNew(SAudioMeter)
 			.Orientation(EOrientation::Orient_Vertical)
 			.BackgroundColor(FLinearColor::Transparent)
-			.MeterBackgroundColor(WidgetStyle->MeterBackgroundColor)
-			.MeterValueColor(WidgetStyle->MeterValueColor)
-			.MeterPeakColor(WidgetStyle->MeterPeakColor)
-			.MeterClippingColor(WidgetStyle->MeterClippingColor)
-			.MeterScaleColor(WidgetStyle->MeterScaleColor)
-			.MeterScaleLabelColor(WidgetStyle->MeterScaleLabelColor);
+
+			// TODO: Move to editor style
+			.MeterBackgroundColor(FLinearColor(0.031f, 0.031f, 0.031f, 1.0f))
+			.MeterValueColor(FLinearColor(0.025719f, 0.208333f, 0.069907f, 1.0f))
+			.MeterPeakColor(FLinearColor(0.24349f, 0.708333f, 0.357002f, 1.0f))
+			.MeterClippingColor(FLinearColor(1.0f, 0.0f, 0.112334f, 1.0f))
+			.MeterScaleColor(FLinearColor(0.017642f, 0.017642f, 0.017642f, 1.0f))
+			.MeterScaleLabelColor(FLinearColor(0.442708f, 0.442708f, 0.442708f, 1.0f));
 
 		Init(InNumChannels, InAudioDeviceId, InExternalAudioBus);
 	}
diff --git a/Engine/Plugins/Runtime/AudioWidgets/Source/AudioWidgets/Private/AudioOscilloscope.cpp b/Engine/Plugins/Runtime/AudioWidgets/Source/AudioWidgets/Private/AudioOscilloscope.cpp
--- a/Engine/Plugins/Runtime/AudioWidgets/Source/AudioWidgets/Private/AudioOscilloscope.cpp
+++ b/Engine/Plugins/Runtime/AudioWidgets/Source/AudioWidgets/Private/AudioOscilloscope.cpp
@@ -9,10 +9,9 @@ namespace AudioWidgets
 		const float InTimeWindowMs, 
 		const float InMaxTimeWindowMs, 
 		const float InAnalysisPeriodMs, 
-		const EAudioPanelLayoutType InPanelLayoutType,
-		const FAudioOscilloscopePanelStyle* InOscilloscopePanelStyle)
+		const EAudioPanelLayoutType InPanelLayoutType)
+		: OscilloscopePanelStyle(FAudioOscilloscopePanelStyle::GetDefault())
 	{
-		OscilloscopePanelStyle = InOscilloscopePanelStyle ? *InOscilloscopePanelStyle : FAudioOscilloscopePanelStyle::GetDefault();
 		CreateAudioBus(InNumChannels);
 		CreateDataProvider(InAudioDeviceId, InTimeWindowMs, InMaxTimeWindowMs, InAnalysisPeriodMs, InPanelLayoutType);
 		CreateOscilloscopeWidget(InNumChannels, InPanelLayoutType);
@@ -36,17 +35,12 @@ namespace AudioWidgets
 		AudioSamplesDataProvider = MakeShared<FWaveformAudioSamplesDataProvider>(InAudioDeviceId, AudioBus.Get(), NumChannelsToProvide, InTimeWindowMs, InMaxTimeWindowMs, InAnalysisPeriodMs);
 	}
 
-	void FAudioOscilloscope::CreateOscilloscopeWidget(const uint32 InNumChannels, const EAudioPanelLayoutType InPanelLayoutType, const FAudioOscilloscopePanelStyle* InOscilloscopePanelStyle)
+	void FAudioOscilloscope::CreateOscilloscopeWidget(const uint32 InNumChannels, const EAudioPanelLayoutType InPanelLayoutType)
 	{
 		check(AudioSamplesDataProvider);
 
 		const FFixedSampledSequenceView SequenceView = AudioSamplesDataProvider->GetDataView();
 
-		if (InOscilloscopePanelStyle)
-		{
-			OscilloscopePanelStyle = *InOscilloscopePanelStyle;
-		}
-
 		if (!OscilloscopePanelWidget.IsValid())
 		{
 			OscilloscopePanelWidget = SNew(SAudioOscilloscopePanelWidget, SequenceView, InNumChannels)
diff --git a/Engine/Plugins/Runtime/AudioWidgets/Source/AudioWidgets/Private/AudioSpectrumAnalyzer.cpp b/Engine/Plugins/Runtime/AudioWidgets/Source/AudioWidgets/Private/AudioSpectrumAnalyzer.cpp
--- a/Engine/Plugins/Runtime/AudioWidgets/Source/AudioWidgets/Private/AudioSpectrumAnalyzer.cpp
+++ b/Engine/Plugins/Runtime/AudioWidgets/Source/AudioWidgets/Private/AudioSpectrumAnalyzer.cpp
@@ -47,7 +47,6 @@ namespace AudioWidgets
 		: SpectrumAnalysisSettings(NewObject<USynesthesiaSpectrumAnalysisSettings>())
 		, ConstantQSettings(NewObject<UConstantQSettings>())
 		, Widget(SNew(SAudioSpectrumPlot)
-			.Style(Params.PlotStyle ? Params.PlotStyle : &FAudioSpectrumPlotStyle::GetDefault())
 			.Clipping(EWidgetClipping::ClipToBounds)
 			.TiltExponent(Params.TiltExponent)
 			.DisplayCrosshair(true)
diff --git a/Engine/Plugins/Runtime/AudioWidgets/Source/AudioWidgets/Private/AudioVectorscope.cpp b/Engine/Plugins/Runtime/AudioWidgets/Source/AudioWidgets/Private/AudioVectorscope.cpp
--- a/Engine/Plugins/Runtime/AudioWidgets/Source/AudioWidgets/Private/AudioVectorscope.cpp
+++ b/Engine/Plugins/Runtime/AudioWidgets/Source/AudioWidgets/Private/AudioVectorscope.cpp
@@ -7,14 +7,13 @@
 namespace AudioWidgets
 {
 	FAudioVectorscope::FAudioVectorscope(Audio::FDeviceId InAudioDeviceId,
-		const uint32 InNumChannels,
-		const float InTimeWindowMs,
-		const float InMaxTimeWindowMs,
-		const float InAnalysisPeriodMs,
-		const EAudioPanelLayoutType InPanelLayoutType,
-		const FAudioVectorscopePanelStyle* InPanelStyle)
+		const uint32 InNumChannels, 
+		const float InTimeWindowMs, 
+		const float InMaxTimeWindowMs, 
+		const float InAnalysisPeriodMs, 
+		const EAudioPanelLayoutType InPanelLayoutType)
+		: VectorscopePanelStyle(FAudioVectorscopePanelStyle::GetDefault())
 	{
-		VectorscopePanelStyle = InPanelStyle ? *InPanelStyle : FAudioVectorscopePanelStyle::GetDefault();
 		CreateAudioBus(InNumChannels);
 		CreateDataProvider(InAudioDeviceId, InTimeWindowMs, InMaxTimeWindowMs, InAnalysisPeriodMs);
 		CreateVectorscopeWidget(InPanelLayoutType);
@@ -33,21 +32,15 @@ namespace AudioWidgets
 		AudioSamplesDataProvider = MakeShared<FWaveformAudioSamplesDataProvider>(InAudioDeviceId, AudioBus.Get(), AudioBus->GetNumChannels(), InTimeWindowMs, InMaxTimeWindowMs, InAnalysisPeriodMs);
 	}
 
-	void FAudioVectorscope::CreateVectorscopeWidget(const EAudioPanelLayoutType InPanelLayoutType, const FAudioVectorscopePanelStyle* PanelStyle)
+	void FAudioVectorscope::CreateVectorscopeWidget(const EAudioPanelLayoutType InPanelLayoutType)
 	{
 		check(AudioSamplesDataProvider);
 
 		const FFixedSampledSequenceView SequenceView = AudioSamplesDataProvider->GetDataView();
 
-		if (PanelStyle)
-		{
-			VectorscopePanelStyle = *PanelStyle;
-		}
-
 		if (!VectorscopePanelWidget.IsValid())
 		{
 			VectorscopePanelWidget = SNew(SAudioVectorscopePanelWidget, SequenceView)
-				.PanelStyle(PanelStyle ? PanelStyle : &FAudioVectorscopePanelStyle::GetDefault())
 				.PanelLayoutType(InPanelLayoutType)
 				.PanelStyle(&VectorscopePanelStyle);
 		}
diff --git a/Engine/Plugins/Runtime/AudioWidgets/Source/AudioWidgets/Public/AudioMeter.h b/Engine/Plugins/Runtime/AudioWidgets/Source/AudioWidgets/Public/AudioMeter.h
--- a/Engine/Plugins/Runtime/AudioWidgets/Source/AudioWidgets/Public/AudioMeter.h
+++ b/Engine/Plugins/Runtime/AudioWidgets/Source/AudioWidgets/Public/AudioMeter.h
@@ -13,7 +13,6 @@
 #include "Sound/AudioBus.h"
 #include "Styling/SlateTypes.h"
 #include "Templates/SharedPointer.h"
-#include "UObject/NoExportTypes.h"
 #include "UObject/ObjectMacros.h"
 #include "UObject/WeakObjectPtrTemplates.h"
 #include "UObject/StrongObjectPtr.h"
@@ -28,34 +27,6 @@ class UWorld;
 
 struct FAudioMaterialMeterStyle;
 
-
-USTRUCT(BlueprintType)
-struct AUDIOWIDGETS_API FAudioMeterDefaultColorStyle : public FSlateWidgetStyle
-{
-	GENERATED_USTRUCT_BODY()
-
-public:
-	static const FAudioMeterDefaultColorStyle& GetDefault();
-
-	UPROPERTY(EditAnywhere, BlueprintReadWrite, Category = "Style")
-	FLinearColor MeterBackgroundColor = FLinearColor(0.031f, 0.031f, 0.031f, 1.0f);
-
-	UPROPERTY(EditAnywhere, BlueprintReadWrite, Category = "Style")
-	FLinearColor MeterValueColor = FLinearColor(0.025719f, 0.208333f, 0.069907f, 1.0f);
-
-	UPROPERTY(EditAnywhere, BlueprintReadWrite, Category = "Style")
-	FLinearColor MeterPeakColor = FLinearColor(0.24349f, 0.708333f, 0.357002f, 1.0f);
-
-	UPROPERTY(EditAnywhere, BlueprintReadWrite, Category = "Style")
-	FLinearColor MeterClippingColor = FLinearColor(1.0f, 0.0f, 0.112334f, 1.0f);
-
-	UPROPERTY(EditAnywhere, BlueprintReadWrite, Category = "Style")
-	FLinearColor MeterScaleColor = FLinearColor(0.017642f, 0.017642f, 0.017642f, 1.0f);
-
-	UPROPERTY(EditAnywhere, BlueprintReadWrite, Category = "Style")
-	FLinearColor MeterScaleLabelColor = FLinearColor(0.442708f, 0.442708f, 0.442708f, 1.0f);
-};
-
 /**
  * An audio meter widget.
  *
@@ -191,7 +162,7 @@ namespace AudioWidgets
 		FAudioMeter(int32 InNumChannels, UWorld& InWorld, TObjectPtr<UAudioBus> InExternalAudioBus = nullptr); 
 		
 		//** OPTIONAL PARAM InExternalAudioBus: An audio meter can be constructed from this audio bus.*/
-		FAudioMeter(const int32 InNumChannels, const Audio::FDeviceId InAudioDeviceId, TObjectPtr<UAudioBus> InExternalAudioBus = nullptr, const FAudioMeterDefaultColorStyle* AudioMeterColorStyle = nullptr);
+		FAudioMeter(const int32 InNumChannels, const Audio::FDeviceId InAudioDeviceId, TObjectPtr<UAudioBus> InExternalAudioBus = nullptr); 
 
 		//** Constructs the Meter using AudioMaterialMeter with the given style. OPTIONAL PARAM InExternalAudioBus: An audio meter can be constructed from this audio bus.
 		FAudioMeter(const int32 InNumChannels, const Audio::FDeviceId InAudioDeviceId, const FAudioMaterialMeterStyle& AudioMaterialMeterStyle, TObjectPtr<UAudioBus> InExternalAudioBus = nullptr);
diff --git a/Engine/Plugins/Runtime/AudioWidgets/Source/AudioWidgets/Public/AudioOscilloscope.h b/Engine/Plugins/Runtime/AudioWidgets/Source/AudioWidgets/Public/AudioOscilloscope.h
--- a/Engine/Plugins/Runtime/AudioWidgets/Source/AudioWidgets/Public/AudioOscilloscope.h
+++ b/Engine/Plugins/Runtime/AudioWidgets/Source/AudioWidgets/Public/AudioOscilloscope.h
@@ -20,12 +20,11 @@ namespace AudioWidgets
 	{
 	public:
 		FAudioOscilloscope(Audio::FDeviceId InAudioDeviceId,
-			const uint32 InNumChannels,
-			const float InTimeWindowMs,
-			const float InMaxTimeWindowMs,
-			const float InAnalysisPeriodMs,
-			const EAudioPanelLayoutType InPanelLayoutType,
-			const FAudioOscilloscopePanelStyle* InOscilloscopePanelStyle = nullptr);
+			const uint32 InNumChannels, 
+			const float InTimeWindowMs, 
+			const float InMaxTimeWindowMs, 
+			const float InAnalysisPeriodMs, 
+			const EAudioPanelLayoutType InPanelLayoutType);
 
 		void CreateAudioBus(const uint32 InNumChannels);
 
@@ -35,7 +34,7 @@ namespace AudioWidgets
 			const float InAnalysisPeriodMs,
 			const EAudioPanelLayoutType InPanelLayoutType);
 
-		void CreateOscilloscopeWidget(const uint32 InNumChannels, const EAudioPanelLayoutType InPanelLayoutType, const FAudioOscilloscopePanelStyle* InOscilloscopePanelStyle = nullptr);
+		void CreateOscilloscopeWidget(const uint32 InNumChannels, const EAudioPanelLayoutType InPanelLayoutType);
 
 		void StartProcessing();
 		void StopProcessing();
diff --git a/Engine/Plugins/Runtime/AudioWidgets/Source/AudioWidgets/Public/AudioSpectrumAnalyzer.h b/Engine/Plugins/Runtime/AudioWidgets/Source/AudioWidgets/Public/AudioSpectrumAnalyzer.h
--- a/Engine/Plugins/Runtime/AudioWidgets/Source/AudioWidgets/Public/AudioSpectrumAnalyzer.h
+++ b/Engine/Plugins/Runtime/AudioWidgets/Source/AudioWidgets/Public/AudioSpectrumAnalyzer.h
@@ -34,7 +34,7 @@ namespace AudioWidgets
 	/**
 	 * Constructor parameters for the analyzer.
 	 */
-	struct AUDIOWIDGETS_API FAudioSpectrumAnalyzerParams
+	struct FAudioSpectrumAnalyzerParams
 	{
 		int32 NumChannels = 1;
 		Audio::FDeviceId AudioDeviceId = INDEX_NONE;
@@ -59,8 +59,6 @@ namespace AudioWidgets
 		FOnFrequencyAxisScaleMenuEntryClicked OnFrequencyAxisScaleMenuEntryClicked;
 		FOnDisplayAxisLabelsButtonToggled OnDisplayFrequencyAxisLabelsButtonToggled;
 		FOnDisplayAxisLabelsButtonToggled OnDisplaySoundLevelAxisLabelsButtonToggled;
-
-		const FAudioSpectrumPlotStyle* PlotStyle = nullptr;
 	};
 
 	/**
diff --git a/Engine/Plugins/Runtime/AudioWidgets/Source/AudioWidgets/Public/AudioVectorscope.h b/Engine/Plugins/Runtime/AudioWidgets/Source/AudioWidgets/Public/AudioVectorscope.h
--- a/Engine/Plugins/Runtime/AudioWidgets/Source/AudioWidgets/Public/AudioVectorscope.h
+++ b/Engine/Plugins/Runtime/AudioWidgets/Source/AudioWidgets/Public/AudioVectorscope.h
@@ -23,12 +23,11 @@ namespace AudioWidgets
 			const float InTimeWindowMs, 
 			const float InMaxTimeWindowMs, 
 			const float InAnalysisPeriodMs, 
-			const EAudioPanelLayoutType InPanelLayoutType,
-			const FAudioVectorscopePanelStyle* PanelStyle = nullptr);
+			const EAudioPanelLayoutType InPanelLayoutType);
 
 		void CreateAudioBus(const uint32 InNumChannels);
 		void CreateDataProvider(Audio::FDeviceId InAudioDeviceId, const float InTimeWindowMs, const float InMaxTimeWindowMs, const float InAnalysisPeriodMs);
-		void CreateVectorscopeWidget(const EAudioPanelLayoutType InPanelLayoutType, const FAudioVectorscopePanelStyle* PanelStyle = nullptr);
+		void CreateVectorscopeWidget(const EAudioPanelLayoutType InPanelLayoutType);
 
 		void StartProcessing();
 		void StopProcessing();
diff --git a/Engine/Plugins/Runtime/Metasound/Source/MetasoundEditor/Private/MetasoundDefaultLiteralCustomization.cpp b/Engine/Plugins/Runtime/Metasound/Source/MetasoundEditor/Private/MetasoundDefaultLiteralCustomization.cpp
--- a/Engine/Plugins/Runtime/Metasound/Source/MetasoundEditor/Private/MetasoundDefaultLiteralCustomization.cpp
+++ b/Engine/Plugins/Runtime/Metasound/Source/MetasoundEditor/Private/MetasoundDefaultLiteralCustomization.cpp
@@ -224,14 +224,18 @@ namespace Metasound::Editor
 			return SNullWidget::NullWidget;
 		}
 
+		const FGuid PageID = LiteralCustomizationPrivate::GetGuidPropertyValue(ElementProperty->GetChildHandle(GET_MEMBER_NAME_CHECKED(FMetasoundEditorMemberPageDefault, PageID)));
+		FName PageName;
+		PageNameProperty->GetValue(PageName);
+
 		TSharedRef<SHorizontalBox> NameBox = SNew(SHorizontalBox)
-		+ SHorizontalBox::Slot()
-		.AutoWidth()
-		.VAlign(VAlign_Center)
-		.Padding(2, 0)
-		[
-			PageNameProperty->CreatePropertyValueWidget()
-		];
+			+ SHorizontalBox::Slot()
+			.AutoWidth()
+			.VAlign(VAlign_Center)
+			.Padding(2, 0)
+			[
+				PageNameProperty->CreatePropertyValueWidget()
+			];
 
 		const UMetasoundEditorGraphMember* Member = Literal.FindMember();
 		if (!ensure(Member))
@@ -239,14 +243,43 @@ namespace Metasound::Editor
 			return SNullWidget::NullWidget;
 		}
 
+		if (Member->IsDefaultPaged())
+		{
+			TAttribute<FText> ToolTip = LOCTEXT("MetaSound_ExecutingInputPageDefaultTooltip", "Initial default used upon executing preview.");
+			TAttribute<EVisibility> ExecVisibility = TAttribute<EVisibility>::CreateLambda([MemberPtr = TWeakObjectPtr<const UMetasoundEditorGraphMember>(Member), PageID]()
+			{
+				if (const UMetasoundEditorGraphMember* GraphMember = MemberPtr.Get())
+				{
+					const FMetaSoundFrontendDocumentBuilder& Builder = GraphMember->GetFrontendBuilderChecked();
+					if (const FMetasoundFrontendClassInput* ClassInput = Builder.FindGraphInput(GraphMember->GetMemberName()))
+					{
+						const bool bIsPreviewing = Editor::IsPreviewingPageInputDefault(Builder, *ClassInput, PageID);
+						return bIsPreviewing ? EVisibility::Visible : EVisibility::Collapsed;
+					}
+					return EVisibility::Collapsed;
+				}
+
+				return EVisibility::Collapsed;
+			});
+			TSharedRef<SWidget> ExecImageWidget = SNew(SImage)
+				.Image(Style::CreateSlateIcon("MetasoundEditor.Page.Executing").GetIcon())
+				.ColorAndOpacity(FStyleColors::AccentGreen)
+				.Visibility(MoveTemp(ExecVisibility));
+			NameBox->AddSlot()
+				.FillWidth(1)
+				.HAlign(HAlign_Left)
+				.VAlign(VAlign_Center)
+				.Padding(2, 0)
+				[
+					ExecImageWidget
+				];
+		}
+
 		// Can't delete the default page, so only show remove page field if non-default/project defined page.
-		const FGuid PageID = LiteralCustomizationPrivate::GetGuidPropertyValue(ElementProperty->GetChildHandle(GET_MEMBER_NAME_CHECKED(FMetasoundEditorMemberPageDefault, PageID)));
 		if (PageID != Frontend::DefaultPageID)
 		{
 			TWeakObjectPtr<UMetasoundEditorGraphMemberDefaultLiteral> LiteralPtr(&Literal);
 
-			FName PageName;
-			PageNameProperty->GetValue(PageName);
 			const FText RemoveDescription = FText::Format(LOCTEXT("RemovePageDefaultTransactionFormat", "Remove '{0}' Page '{1}' Default Value"), FText::FromName(Member->GetMemberName()), FText::FromName(PageName));
 			TSharedRef<SWidget> RemovePageDefaultButton = PropertyCustomizationHelpers::MakeDeleteButton(FSimpleDelegate::CreateLambda([this, RemoveDescription, PageID, LiteralPtr]()
 			{
@@ -296,39 +329,6 @@ namespace Metasound::Editor
 			];
 		}
 
-		if (Member->IsDefaultPaged())
-		{
-			TAttribute<EVisibility> ExecVisibility = TAttribute<EVisibility>::CreateLambda([MemberPtr = TWeakObjectPtr<const UMetasoundEditorGraphMember>(Member), PageID]()
-			{
-				if (const UMetasoundEditorGraphMember* GraphMember = MemberPtr.Get())
-				{
-					const FMetaSoundFrontendDocumentBuilder& Builder = GraphMember->GetFrontendBuilderChecked();
-					if (const FMetasoundFrontendClassInput* ClassInput = Builder.FindGraphInput(GraphMember->GetMemberName()))
-					{
-						const bool bIsPreviewing = Editor::IsPreviewingPageInputDefault(Builder, *ClassInput, PageID);
-						return bIsPreviewing ? EVisibility::Visible : EVisibility::Collapsed;
-					}
-					return EVisibility::Collapsed;
-				}
-
-				return EVisibility::Collapsed;
-			});
-
-			TSharedRef<SWidget> ExecImageWidget = SNew(SImage)
-			.Image(Style::CreateSlateIcon("MetasoundEditor.Page.Executing").GetIcon())
-			.ColorAndOpacity(Style::GetPageExecutingColor())
-			.Visibility(MoveTemp(ExecVisibility));
-
-			NameBox->AddSlot()
-			.FillWidth(1)
-			.HAlign(HAlign_Left)
-			.VAlign(VAlign_Center)
-			.Padding(2, 0)
-			[
-				ExecImageWidget
-			];
-		}
-
 		return NameBox;
 	}
 
diff --git a/Engine/Plugins/Runtime/Metasound/Source/MetasoundEditor/Private/MetasoundDetailCustomization.cpp b/Engine/Plugins/Runtime/Metasound/Source/MetasoundEditor/Private/MetasoundDetailCustomization.cpp
--- a/Engine/Plugins/Runtime/Metasound/Source/MetasoundEditor/Private/MetasoundDetailCustomization.cpp
+++ b/Engine/Plugins/Runtime/Metasound/Source/MetasoundEditor/Private/MetasoundDetailCustomization.cpp
@@ -575,7 +575,7 @@ namespace Metasound::Editor
 				});
 				TSharedRef<SWidget> ExecImageWidget = SNew(SImage)
 					.Image(Style::CreateSlateIcon("MetasoundEditor.Page.Executing").GetIcon())
-					.ColorAndOpacity(Style::GetPageExecutingColor())
+					.ColorAndOpacity(FStyleColors::AccentGreen)
 					.Visibility(MoveTemp(Visibility));
 
 				EntryWidget->AddSlot()
diff --git a/Engine/Plugins/Runtime/Metasound/Source/MetasoundEditor/Private/MetasoundEditor.h b/Engine/Plugins/Runtime/Metasound/Source/MetasoundEditor/Private/MetasoundEditor.h
--- a/Engine/Plugins/Runtime/Metasound/Source/MetasoundEditor/Private/MetasoundEditor.h
+++ b/Engine/Plugins/Runtime/Metasound/Source/MetasoundEditor/Private/MetasoundEditor.h
@@ -135,6 +135,7 @@ namespace Metasound::Editor
 		virtual void RegisterTabSpawners(const TSharedRef<FTabManager>& TabManager) override;
 		virtual void UnregisterTabSpawners(const TSharedRef<FTabManager>& TabManager) override;
 
+		double GetPlayTime() const;
 		TSharedPtr<SGraphEditor> GetGraphEditor() const;
 
 		/** Edits the specified Metasound object */
@@ -531,7 +532,7 @@ namespace Metasound::Editor
 		bool UpdatePlayTime(float InDeltaTime);
 
 		/** Updates the render info widget. */
-		void UpdateRenderInfo(bool bIsPlaying, float InDeltaTime = 0.0f);
+		void UpdateRenderInfo();
 
 		/** List of open tool panels; used to ensure only one exists at any one time */
 		TMap<FName, TWeakPtr<SDockableTab>> SpawnedToolPanels;
@@ -571,8 +572,15 @@ namespace Metasound::Editor
 		/** Palette of Node types */
 		TSharedPtr<SMetasoundPalette> Palette;
 
-		/** Widget showing page info regarding page info that overlays the graph tab content */
-		TSharedPtr<SPageStats> PageStatsWidget;
+		/** Widget showing playtime that overlays the graph when previewing */
+		TSharedPtr<STextBlock> PlayTimeWidget;
+		double PlayTime = 0.0;
+
+		/** Widget showing page info regarding the currently focused page that overlays the graph tab content */
+		TSharedPtr<SPageStats> FocusPageWidget;
+
+		/** Widget showing page info regarding the currently rendering page that overlays the analyzer tab content */
+		TSharedPtr<SPageStats> RenderPageWidget;
 
 		/** Widget showing render performance information that overlays the graph */
 		TSharedPtr<SRenderStats> RenderStatsWidget;
diff --git a/Engine/Plugins/Runtime/Metasound/Source/MetasoundEditor/Private/MetasoundEditorModule.cpp b/Engine/Plugins/Runtime/Metasound/Source/MetasoundEditor/Private/MetasoundEditorModule.cpp
--- a/Engine/Plugins/Runtime/Metasound/Source/MetasoundEditor/Private/MetasoundEditorModule.cpp
+++ b/Engine/Plugins/Runtime/Metasound/Source/MetasoundEditor/Private/MetasoundEditorModule.cpp
@@ -2,10 +2,6 @@
 #include "MetasoundEditorModule.h"
 
 #include "AssetTypeActions_Base.h"
-#include "AudioMeter.h"
-#include "AudioOscilloscopePanelStyle.h"
-#include "AudioSpectrumPlotStyle.h"
-#include "AudioVectorscopePanelStyle.h"
 #include "Brushes/SlateImageBrush.h"
 #include "EdGraph/EdGraphNode.h"
 #include "EdGraph/EdGraphPin.h"
@@ -164,16 +160,12 @@ namespace Metasound
 
 					// Analyzers
 					Set("MetasoundEditor.Analyzers.BackgroundColor", FLinearColor(0.0075f, 0.0075f, 0.0075, 1.0f));
-					Set("MetasoundEditor.Analyzers.ForegroundColor", FLinearColor(0.025719f, 0.208333f, 0.069907f, 1.0f)); // "Audio" Green
 
 					// Misc
 					Set("MetasoundEditor.Audition", new IMAGE_BRUSH_SVG(TEXT("Icons/metasound_page"), Icon16));
 					Set("MetasoundEditor.Metasound.Icon", new IMAGE_BRUSH_SVG(TEXT("Icons/metasound_icon"), Icon16));
-					Set("MetasoundEditor.Speaker", new FSlateImageBrush(RootToContentDir(TEXT("/Icons/speaker_144x.png")), FVector2D(144.0f, 144.0f)));
-
-					// Pages
-					Set("MetasoundEditor.Page.Executing.ForegroundColor", FStyleColors::AccentGreen.GetSpecifiedColor());
 					Set("MetasoundEditor.Page.Executing", new IMAGE_BRUSH_SVG(TEXT("Icons/metasound_page_exec"), Icon16));
+					Set("MetasoundEditor.Speaker", new FSlateImageBrush(RootToContentDir(TEXT("/Icons/speaker_144x.png")), FVector2D(144.0f, 144.0f)));
 
 					// Class Icons
 					auto SetClassIcon = [this, InIcon16 = Icon16, InIcon64 = Icon64](const FString& ClassName)
@@ -209,21 +201,6 @@ namespace Metasound
 				return { "MetaSoundStyle", InName};
 			}
 
-			const FSlateColor& GetDefaultAnalyzerColor()
-			{
-				auto MakeColor = []() -> FSlateColor
-				{
-					if (const ISlateStyle* MetaSoundStyle = FSlateStyleRegistry::FindSlateStyle("MetaSoundStyle"))
-					{
-						return MetaSoundStyle->GetColor("MetasoundEditor.Analyzers.ForegroundColor");
-					}
-
-					return FStyleColors::AccentWhite;
-				};
-				static const FSlateColor AnalyzerColor = MakeColor();
-				return AnalyzerColor;
-			}
-
 			const FSlateBrush& GetSlateBrushSafe(FName InName)
 			{
 				const ISlateStyle* MetaSoundStyle = FSlateStyleRegistry::FindSlateStyle("MetaSoundStyle");
@@ -244,77 +221,7 @@ namespace Metasound
 				static const FSlateBrush NullBrush;
 				return NullBrush;
 			}
-
-			const FAudioMeterDefaultColorStyle& GetMeterDefaultColorStyle()
-			{
-				auto MakeStyle = []()
-				{
-					FAudioMeterDefaultColorStyle MeterStyle;
-					MeterStyle.MeterValueColor = GetDefaultAnalyzerColor().GetSpecifiedColor();
-					return MeterStyle;
-				};
-				static const FAudioMeterDefaultColorStyle ThisStyle = MakeStyle();
-				return ThisStyle;
-			}
-
-			const FAudioOscilloscopePanelStyle& GetOscilloscopeStyle()
-			{
-				auto MakeStyle = []()
-				{
-					FAudioOscilloscopePanelStyle OscStyle;
-					FSampledSequenceViewerStyle SampleView;
-					SampleView.SequenceColor = GetDefaultAnalyzerColor();
-					OscStyle.SetWaveViewerStyle(SampleView);
-
-					return OscStyle;
-				};
-				static FAudioOscilloscopePanelStyle ThisStyle = MakeStyle();
-				return ThisStyle;
-			}
-
-			const FSlateColor& GetPageExecutingColor()
-			{
-				auto MakeColor = []() -> FSlateColor
-				{
-					if (const ISlateStyle* MetaSoundStyle = FSlateStyleRegistry::FindSlateStyle("MetaSoundStyle"))
-					{
-						return MetaSoundStyle->GetColor("MetasoundEditor.Page.Executing.ForegroundColor");
-					}
-
-					return FStyleColors::AccentWhite;
-				};
-				static const FSlateColor AnalyzerColor = MakeColor();
-				return AnalyzerColor;
-			}
-
-			const FAudioSpectrumPlotStyle& GetSpectrumPlotStyle()
-			{
-				auto MakeStyle = []()
-				{
-					FAudioSpectrumPlotStyle PlotStyle;
-					const FSlateColor AnalyzerColor = GetDefaultAnalyzerColor();
-					PlotStyle.CrosshairColor = AnalyzerColor.UseSubduedForeground();
-					PlotStyle.SpectrumColor = AnalyzerColor;
-					return PlotStyle;
-				};
-				static FAudioSpectrumPlotStyle ThisStyle = MakeStyle();
-				return ThisStyle;
-			}
-
-			const FAudioVectorscopePanelStyle& GetVectorscopeStyle()
-			{
-				auto MakeStyle = []()
-				{
-					FAudioVectorscopePanelStyle PanelStyle;
-					FSampledSequenceVectorViewerStyle VectorViewStyle;
-					VectorViewStyle.LineColor = GetDefaultAnalyzerColor().GetSpecifiedColor();
-					PanelStyle.SetVectorViewerStyle(VectorViewStyle);
-					return PanelStyle;
-				};
-				static FAudioVectorscopePanelStyle ThisStyle = MakeStyle();
-				return ThisStyle;
-			}
-		} // namespace Style
+		}
 
 		// A structure that contains information about registered custom pin types. 
 		struct FGraphPinConfiguration
diff --git a/Engine/Plugins/Runtime/Metasound/Source/MetasoundEditor/Private/MetasoundEditorSettings.h b/Engine/Plugins/Runtime/Metasound/Source/MetasoundEditor/Private/MetasoundEditorSettings.h
--- a/Engine/Plugins/Runtime/Metasound/Source/MetasoundEditor/Private/MetasoundEditorSettings.h
+++ b/Engine/Plugins/Runtime/Metasound/Source/MetasoundEditor/Private/MetasoundEditorSettings.h
@@ -271,23 +271,23 @@ public:
 	EMetasoundActiveDetailView DetailView = EMetasoundActiveDetailView::General;
 
 	/** Whether the AudioMaterialWidgets are used when possible in Metasound Editor*/
-	UPROPERTY(EditAnywhere, config, DisplayName = "Use Audio Material Widgets", Category = "Widget Styling (Experimental)")
+	UPROPERTY(EditAnywhere, config, DisplayName = "Use Audio Material Widgets", Category = WidgetStyling)
 	bool bUseAudioMaterialWidgets = false;
 	
 	/**Override the Knob Style used in the Metasound Editor.*/
-	UPROPERTY(EditAnywhere, config, Category = "Widget Styling (Experimental)", meta = (AllowedClasses = "/Script/SlateCore.SlateWidgetStyleAsset", EditCondition = "bUseAudioMaterialWidgets", DisplayName = "Knob Style"))
+	UPROPERTY(EditAnywhere, config, Category = WidgetStyling, meta = (AllowedClasses = "/Script/SlateCore.SlateWidgetStyleAsset", EditCondition = "bUseAudioMaterialWidgets", DisplayName = "Knob Style"))
 	FSoftObjectPath KnobStyleOverride;
 	
 	/**Override the Slider Style used in the Metasound Editor.*/
-	UPROPERTY(EditAnywhere, config, Category = "Widget Styling (Experimental)", meta = (AllowedClasses = "/Script/SlateCore.SlateWidgetStyleAsset", EditCondition = "bUseAudioMaterialWidgets", DisplayName = "Slider Style"))
+	UPROPERTY(EditAnywhere, config, Category = WidgetStyling, meta = (AllowedClasses = "/Script/SlateCore.SlateWidgetStyleAsset", EditCondition = "bUseAudioMaterialWidgets", DisplayName = "Slider Style"))
 	FSoftObjectPath SliderStyleOverride;
 
 	/**Override the Button Style used in the Metasound Editor.*/
-	UPROPERTY(EditAnywhere, config, Category = "Widget Styling (Experimental)", meta = (AllowedClasses = "/Script/SlateCore.SlateWidgetStyleAsset", EditCondition = "bUseAudioMaterialWidgets", DisplayName = "Button Style"))
+	UPROPERTY(EditAnywhere, config, Category = WidgetStyling, meta = (AllowedClasses = "/Script/SlateCore.SlateWidgetStyleAsset", EditCondition = "bUseAudioMaterialWidgets", DisplayName = "Button Style"))
 	FSoftObjectPath ButtonStyleOverride;
 	
 	/**Override the Meter Style used in the Metasound Editor.*/
-	UPROPERTY(EditAnywhere, config, Category = "Widget Styling (Experimental)", meta = (AllowedClasses = "/Script/SlateCore.SlateWidgetStyleAsset", EditCondition = "bUseAudioMaterialWidgets", DisplayName = "Meter Style"))
+	UPROPERTY(EditAnywhere, config, Category = WidgetStyling, meta = (AllowedClasses = "/Script/SlateCore.SlateWidgetStyleAsset", EditCondition = "bUseAudioMaterialWidgets", DisplayName = "Meter Style"))
 	FSoftObjectPath MeterStyleOverride;
 
 	//UObject
diff --git a/Engine/Plugins/Runtime/Metasound/Source/MetasoundEditor/Private/MetasoundPinAudioInspector.cpp b/Engine/Plugins/Runtime/Metasound/Source/MetasoundEditor/Private/MetasoundPinAudioInspector.cpp
--- a/Engine/Plugins/Runtime/Metasound/Source/MetasoundEditor/Private/MetasoundPinAudioInspector.cpp
+++ b/Engine/Plugins/Runtime/Metasound/Source/MetasoundEditor/Private/MetasoundPinAudioInspector.cpp
@@ -9,115 +9,115 @@
 #include "MetasoundEditor.h"
 #include "MetasoundEditorGraphBuilder.h"
 #include "MetasoundEditorGraphNode.h"
-#include "MetasoundEditorModule.h"
 
 #define LOCTEXT_NAMESPACE "MetasoundEditor"
 
-namespace Metasound::Editor
+namespace Metasound
 {
-	namespace FMetasoundPinAudioInspectorPrivate
+	namespace Editor
 	{
-		UEdGraphPin* ResolvePinObjectAsOutput(UEdGraphPin* InPin)
+		namespace FMetasoundPinAudioInspectorPrivate
 		{
-			if (InPin && !InPin->LinkedTo.IsEmpty())
+			UEdGraphPin* ResolvePinObjectAsOutput(UEdGraphPin* InPin)
 			{
-				// Swap to show connected output if input (Only ever one)
-				if (InPin->Direction == EGPD_Input)
+				if (InPin && !InPin->LinkedTo.IsEmpty())
 				{
-					InPin = InPin->LinkedTo.Last();
-					check(InPin->Direction == EGPD_Output);
+					// Swap to show connected output if input (Only ever one)
+					if (InPin->Direction == EGPD_Input)
+					{
+						InPin = InPin->LinkedTo.Last();
+						check(InPin->Direction == EGPD_Output);
+					}
 				}
-			}
 
-			return InPin;
+				return InPin;
+			}
 		}
-	}
 
-	FMetasoundPinAudioInspector::FMetasoundPinAudioInspector(FEdGraphPinReference InPinRef)
-		: GraphPinObj(FMetasoundPinAudioInspectorPrivate::ResolvePinObjectAsOutput(InPinRef.Get()))
-	{
-		using namespace Audio;
+		FMetasoundPinAudioInspector::FMetasoundPinAudioInspector(FEdGraphPinReference InPinRef)
+			: GraphPinObj(FMetasoundPinAudioInspectorPrivate::ResolvePinObjectAsOutput(InPinRef.Get()))
+		{
+			using namespace Audio;
 
-		// Initialize Oscilloscope
-		const FDeviceId AudioDeviceId = GEditor->GetMainAudioDeviceID();
+			// Initialize Oscilloscope
+			const FDeviceId AudioDeviceId = GEditor->GetMainAudioDeviceID();
 
-		constexpr int32 NumChannels = 1; // Audio wires are currently mono signals
+			constexpr int32 NumChannels = 1; // Audio wires are currently mono signals
 
-		Oscilloscope = MakeShared<AudioWidgets::FAudioOscilloscope>(AudioDeviceId,
-			NumChannels,
-			/*InTimeWindowMs*/     10.0f,
-			/*InMaxTimeWindowMs*/  10.0f,
-			/*InAnalysisPeriodMs*/ 10.0f,
-			/*InPanelLayoutType*/  EAudioPanelLayoutType::Basic,
-			/*InOscilloscopePanelStyle*/ &Style::GetOscilloscopeStyle()
-		);
+			Oscilloscope = MakeShared<AudioWidgets::FAudioOscilloscope>(AudioDeviceId,
+				NumChannels,
+				/*InTimeWindowMs*/     10.0f,
+				/*InMaxTimeWindowMs*/  10.0f,
+				/*InAnalysisPeriodMs*/ 10.0f,
+				/*InPanelLayoutType*/  EAudioPanelLayoutType::Basic);
 
-		PinAudioInspectorWidget = SNew(SMetasoundPinAudioInspector)
-			.VisualizationWidget(Oscilloscope->GetPanelWidget());
+			PinAudioInspectorWidget = SNew(SMetasoundPinAudioInspector)
+									  .VisualizationWidget(Oscilloscope->GetPanelWidget());
 
-		Oscilloscope->StartProcessing();
+			Oscilloscope->StartProcessing();
 
-		// Analyze Audio Pin
-		if (GraphPinObj && !GraphPinObj->LinkedTo.IsEmpty() && GraphPinObj->PinType.PinCategory == FGraphBuilder::PinCategoryAudio)
-		{
-			if (FGraphConnectionManager* ConnectionManager = GetConnectionManager())
+			// Analyze Audio Pin
+			if (GraphPinObj && !GraphPinObj->LinkedTo.IsEmpty() && GraphPinObj->PinType.PinCategory == FGraphBuilder::PinCategoryAudio)
 			{
-				const Frontend::FConstOutputHandle OutputHandle = FGraphBuilder::FindReroutedOutputHandleFromPin(GraphPinObj);
-				const FGuid NodeID = OutputHandle->GetOwningNodeID();
-				const FName OutputName = OutputHandle->GetName();
-				AnalyzerInstanceID = ConnectionManager->AddAudioBusWriter(NodeID, OutputName, AudioDeviceId, Oscilloscope->GetAudioBus());
+				if (FGraphConnectionManager* ConnectionManager = GetConnectionManager())
+				{
+					const Frontend::FConstOutputHandle OutputHandle = FGraphBuilder::FindReroutedOutputHandleFromPin(GraphPinObj);
+					const FGuid NodeID = OutputHandle->GetOwningNodeID();
+					const FName OutputName = OutputHandle->GetName();
+					AnalyzerInstanceID = ConnectionManager->AddAudioBusWriter(NodeID, OutputName, AudioDeviceId, Oscilloscope->GetAudioBus());
+				}
 			}
 		}
-	}
 
-	FMetasoundPinAudioInspector::~FMetasoundPinAudioInspector()
-	{
-		// Remove audio pin analyzer
-		if (AnalyzerInstanceID.IsValid())
+		FMetasoundPinAudioInspector::~FMetasoundPinAudioInspector()
 		{
-			if (FGraphConnectionManager* ConnectionManager = GetConnectionManager())
+			// Remove audio pin analyzer
+			if (AnalyzerInstanceID.IsValid())
 			{
-				ConnectionManager->RemoveAudioBusWriter(AnalyzerInstanceID);
-				AnalyzerInstanceID.Invalidate();
+				if (FGraphConnectionManager* ConnectionManager = GetConnectionManager())
+				{
+					ConnectionManager->RemoveAudioBusWriter(AnalyzerInstanceID);
+					AnalyzerInstanceID.Invalidate();
+				}
 			}
-		}
 
-		// Stop AudioBus
-		if (UWorld* EditorWorld = GEditor->GetEditorWorldContext().World())
-		{
-			if (Audio::FMixerDevice* MixerDevice = static_cast<Audio::FMixerDevice*>(EditorWorld->GetAudioDeviceRaw()))
+			// Stop AudioBus
+			if (UWorld* EditorWorld = GEditor->GetEditorWorldContext().World())
 			{
-				if (UAudioBusSubsystem* AudioBusSubsystem = MixerDevice->GetSubsystem<UAudioBusSubsystem>())
+				if (Audio::FMixerDevice* MixerDevice = static_cast<Audio::FMixerDevice*>(EditorWorld->GetAudioDeviceRaw()))
 				{
-					const uint32 AudioBusId = Oscilloscope->GetAudioBus()->GetUniqueID();
-					AudioBusSubsystem->StopAudioBus(Audio::FAudioBusKey(AudioBusId));
+					if (UAudioBusSubsystem* AudioBusSubsystem = MixerDevice->GetSubsystem<UAudioBusSubsystem>())
+					{
+						const uint32 AudioBusId = Oscilloscope->GetAudioBus()->GetUniqueID();
+						AudioBusSubsystem->StopAudioBus(Audio::FAudioBusKey(AudioBusId));
+					}
 				}
 			}
 		}
-	}
 
-	TSharedPtr<SMetasoundPinAudioInspector> FMetasoundPinAudioInspector::GetWidget()
-	{
-		return PinAudioInspectorWidget;
-	}
+		TSharedPtr<SMetasoundPinAudioInspector> FMetasoundPinAudioInspector::GetWidget()
+		{
+			return PinAudioInspectorWidget;
+		}
 
-	const UMetasoundEditorGraphNode& FMetasoundPinAudioInspector::GetReroutedNode() const
-	{
-		const UMetasoundEditorGraphNode* Node = FGraphBuilder::FindReroutedOutputPin(GraphPinObj)
-			? Cast<UMetasoundEditorGraphNode>(GraphPinObj->GetOwningNode())
-			: nullptr;
-		check(Node);
+		const UMetasoundEditorGraphNode& FMetasoundPinAudioInspector::GetReroutedNode() const
+		{
+			const UMetasoundEditorGraphNode* Node = FGraphBuilder::FindReroutedOutputPin(GraphPinObj)
+				? Cast<UMetasoundEditorGraphNode>(GraphPinObj->GetOwningNode())
+				: nullptr;
+			check(Node);
 
-		return *Node;
-	}
+			return *Node;
+		}
 
-	FGraphConnectionManager* FMetasoundPinAudioInspector::GetConnectionManager()
-	{
-		const UMetasoundEditorGraphNode& Node = GetReroutedNode();
-		TSharedPtr<FEditor> Editor = FGraphBuilder::GetEditorForNode(Node);
+		FGraphConnectionManager* FMetasoundPinAudioInspector::GetConnectionManager()
+		{
+			const UMetasoundEditorGraphNode& Node = GetReroutedNode();
+			TSharedPtr<FEditor> Editor = FGraphBuilder::GetEditorForNode(Node);
 
-		return Editor.IsValid() ? &Editor->GetConnectionManager() : nullptr;
-	}
-} // namespace Metasound::Editor
+			return Editor.IsValid() ? &Editor->GetConnectionManager() : nullptr;
+		}
+	} // namespace Editor
+} // namespace Metasound
 
 #undef LOCTEXT_NAMESPACE
diff --git a/Engine/Plugins/Runtime/Metasound/Source/MetasoundEditor/Private/SMetasoundFilterFrequencyResponsePlots.cpp b/Engine/Plugins/Runtime/Metasound/Source/MetasoundEditor/Private/SMetasoundFilterFrequencyResponsePlots.cpp
--- a/Engine/Plugins/Runtime/Metasound/Source/MetasoundEditor/Private/SMetasoundFilterFrequencyResponsePlots.cpp
+++ b/Engine/Plugins/Runtime/Metasound/Source/MetasoundEditor/Private/SMetasoundFilterFrequencyResponsePlots.cpp
@@ -107,7 +107,6 @@ namespace Metasound::Editor
 					.TiltExponent_Lambda([] { return 0.0f; }) // Binding this property has the effect of hiding its context menu entry (Tilting the spectrum is not desired here).
 					.FrequencyAxisPixelBucketMode_Lambda([]() { return EAudioSpectrumPlotFrequencyAxisPixelBucketMode::Sample; }) // Binding this property has the effect of hiding its context menu entry (PixelBucketMode is not much use here).
 					.OnGetAudioSpectrumData(this, &SMetaSoundFilterFrequencyResponsePlot::GetAudioSpectrumData)
-					.Style(&Editor::Style::GetSpectrumPlotStyle())
 			];
 
 		ContextMenuExtension = FrequencyResponsePlot->AddContextMenuExtension(EExtensionHook::Before, nullptr, FMenuExtensionDelegate::CreateSP(this, &SMetaSoundFilterFrequencyResponsePlot::ExtendSpectrumPlotContextMenu));
diff --git a/Engine/Plugins/Runtime/Metasound/Source/MetasoundEditor/Private/SMetasoundSpectrumAnalyzerGraphNode.cpp b/Engine/Plugins/Runtime/Metasound/Source/MetasoundEditor/Private/SMetasoundSpectrumAnalyzerGraphNode.cpp
--- a/Engine/Plugins/Runtime/Metasound/Source/MetasoundEditor/Private/SMetasoundSpectrumAnalyzerGraphNode.cpp
+++ b/Engine/Plugins/Runtime/Metasound/Source/MetasoundEditor/Private/SMetasoundSpectrumAnalyzerGraphNode.cpp
@@ -5,7 +5,6 @@
 #include "Editor.h"
 #include "MetasoundEditorGraphBuilder.h"
 #include "MetasoundEditorGraphNode.h"
-#include "MetasoundEditorModule.h"
 #include "Widgets/Layout/SBox.h"
 
 namespace Metasound::Editor
@@ -57,17 +56,13 @@ namespace Metasound::Editor
 
 	void SMetaSoundSpectrumAnalyzerGraphNode::CreateBelowPinControls(TSharedPtr<SVerticalBox> MainBox)
 	{
-		using namespace AudioWidgets;
-
 		if (ensure(GEditor))
 		{
 			if (!SpectrumAnalyzer.IsValid())
 			{
-				FAudioSpectrumAnalyzerParams Params;
-				Params.NumChannels = 1;
-				Params.AudioDeviceId = GEditor->GetMainAudioDeviceID();
-				Params.PlotStyle = &Style::GetSpectrumPlotStyle();
-				SpectrumAnalyzer = MakeShared<FAudioSpectrumAnalyzer>(Params);
+				constexpr int32 NumChannels = 1;
+				const Audio::FDeviceId AudioDeviceId = GEditor->GetMainAudioDeviceID();
+				SpectrumAnalyzer = MakeShared<AudioWidgets::FAudioSpectrumAnalyzer>(NumChannels, AudioDeviceId);
 			}
 
 			MainBox->AddSlot()
diff --git a/Engine/Plugins/Runtime/Metasound/Source/MetasoundEditor/Private/SMetasoundStats.cpp b/Engine/Plugins/Runtime/Metasound/Source/MetasoundEditor/Private/SMetasoundStats.cpp
--- a/Engine/Plugins/Runtime/Metasound/Source/MetasoundEditor/Private/SMetasoundStats.cpp
+++ b/Engine/Plugins/Runtime/Metasound/Source/MetasoundEditor/Private/SMetasoundStats.cpp
@@ -32,40 +32,30 @@ namespace Metasound::Editor
 	{
 		SVerticalBox::Construct(SVerticalBox::FArguments());
 
-		AddSlot().HAlign(HAlign_Left)
-		[
-			SAssignNew(AuditionPageTextWidget, STextBlock)
-				.Visibility(EVisibility::Collapsed)
+		TSharedRef<SHorizontalBox> PageWidgetBox = SNew(SHorizontalBox);
+		PageWidgetBox->AddSlot()
+			.Padding(2.0f)
+			.HAlign(HAlign_Center)
+			[
+				SAssignNew(PageTextWidget, STextBlock)
+				.Visibility(EVisibility::HitTestInvisible)
 				.TextStyle(FAppStyle::Get(), "Graph.ZoomText")
 				.ColorAndOpacity(StatsPrivate::BaseTextColor)
-		];
+			];
 
-		AddSlot().HAlign(HAlign_Left)
-		[
-			SNew(SHorizontalBox)
-				+ SHorizontalBox::Slot()
-				.Padding(2.0f)
-				.HAlign(HAlign_Center)
-				[
-					SAssignNew(GraphPageTextWidget, STextBlock)
-					.Visibility(EVisibility::HitTestInvisible)
-					.TextStyle(FAppStyle::Get(), "Graph.ZoomText")
-					.ColorAndOpacity(StatsPrivate::BaseTextColor)
-				]
-
-				+ SHorizontalBox::Slot()
-				.Padding(2.0f)
-				.HAlign(HAlign_Center)
-				.VAlign(VAlign_Center)
-				.AutoWidth()
-				[
-					SAssignNew(ExecImageWidget, SImage)
-						.Image(Style::CreateSlateIcon("MetasoundEditor.Page.Executing").GetIcon())
-						.DesiredSizeOverride(FVector2D(24.f, 24.f))
-						.ColorAndOpacity(Style::GetPageExecutingColor())
-						.Visibility(EVisibility::Collapsed)
-				]
-		];
+		ExecImageWidget = SNew(SImage)
+			.Image(Style::CreateSlateIcon("MetasoundEditor.Page.Executing").GetIcon())
+			.DesiredSizeOverride(FVector2D(24.f, 24.f))
+			.ColorAndOpacity(FStyleColors::AccentGreen)
+			.Visibility(EVisibility::Collapsed);
+
+		PageWidgetBox->AddSlot()
+			.Padding(2.0f)
+			.HAlign(HAlign_Center)
+			.VAlign(VAlign_Center)
+			.AutoWidth()[ ExecImageWidget.ToSharedRef()];
+
+		AddSlot().HAlign(HAlign_Left) [ PageWidgetBox ];
 	}
 
 	void SPageStats::SetExecVisibility(TAttribute<EVisibility> InVisibility)
@@ -73,152 +63,87 @@ namespace Metasound::Editor
 		ExecImageWidget->SetVisibility(MoveTemp(InVisibility));
 	}
 
-	void SPageStats::Update(const FMetaSoundPageSettings* AuditionPageSettings, const FMetaSoundPageSettings* GraphPageSettings, const FSlateColor* ActiveColor)
+	void SPageStats::Update(const FMetaSoundPageSettings* PageSettings, const FText& Header, const FSlateColor* ColorOverride)
 	{
 		using namespace Engine;
 
-		const FText PageStatsFormat = LOCTEXT("PageStatsFormat", "{0}: {1}");
-
+		FText PageInfo;
+		EVisibility Visibility = EVisibility::Collapsed;
+		if (PageTextWidget.IsValid() && PageSettings)
 		{
-			FText PageInfo;
-			EVisibility Visibility = EVisibility::Collapsed;
-			if (AuditionPageTextWidget.IsValid() && AuditionPageSettings)
-			{
-				const FText Header = LOCTEXT("AuditionPageHeader", "Auditioning Page");
-				PageInfo = FText::Format(PageStatsFormat, Header, FText::FromString(AuditionPageSettings->Name.ToString()));
-				Visibility = EVisibility::Visible;
-			}
-			AuditionPageTextWidget->SetText(PageInfo);
+			PageInfo = FText::Format(LOCTEXT("PageStatsFormat", "{0}: {1}"), Header, FText::FromString(PageSettings->Name.ToString()));
+			Visibility = EVisibility::Visible;
 		}
 
-		{
-			FText PageInfo;
-			EVisibility Visibility = EVisibility::Collapsed;
-			if (GraphPageTextWidget.IsValid() && GraphPageSettings)
-			{
-				const FText Header = LOCTEXT("GraphPageTargetHeader", "Graph Page");
-				PageInfo = FText::Format(PageStatsFormat, Header, FText::FromString(GraphPageSettings->Name.ToString()));
-				Visibility = EVisibility::Visible;
-			}
-			GraphPageTextWidget->SetText(PageInfo);
-			GraphPageTextWidget->SetColorAndOpacity(ActiveColor ? *ActiveColor : StatsPrivate::BaseTextColor);
-		}
+		PageTextWidget->SetText(PageInfo);
+		PageTextWidget->SetColorAndOpacity(ColorOverride ? *ColorOverride : StatsPrivate::BaseTextColor);
 	}
 
 	void SRenderStats::Construct(const FArguments& InArgs)
 	{
 		SVerticalBox::Construct(SVerticalBox::FArguments());
+		AddSlot()
+		.HAlign(HAlign_Left)
+		.AutoHeight()
+		[
+			SAssignNew(RenderStatsCostWidget, STextBlock)
+			.Visibility(EVisibility::HitTestInvisible)
+			.TextStyle(FAppStyle::Get(), "Graph.ZoomText")
+			.ColorAndOpacity(FLinearColor(1, 1, 1, 0.30f))
+		];
 
-		auto AddTextWidgetSlot = [this](TSharedPtr<STextBlock>& OutBlock)
-		{
-			AddSlot()
-			.HAlign(HAlign_Left)
-			.AutoHeight()
-			[
-				SAssignNew(OutBlock, STextBlock)
-				.Visibility(EVisibility::HitTestInvisible)
-				.TextStyle(FAppStyle::Get(), "GraphPreview.CornerText")
-				.ColorAndOpacity(FLinearColor(1, 1, 1, 0.30f))
-			];
-		};
-
-		AddTextWidgetSlot(PlayTimeWidget);
-		AddTextWidgetSlot(RenderStatsCostWidget);
-		AddTextWidgetSlot(RenderStatsCPUWidget);
-		AddTextWidgetSlot(AuditionPageWidget);
-		AddTextWidgetSlot(AuditionPlatformWidget);
+		AddSlot()
+		.HAlign(HAlign_Left)
+		.AutoHeight()
+		[
+			SAssignNew(RenderStatsCPUWidget, STextBlock)
+			.Visibility(EVisibility::HitTestInvisible)
+			.TextStyle(FAppStyle::Get(), "Graph.ZoomText")
+			.ColorAndOpacity(FLinearColor(1, 1, 1, 0.30f))
+		];
 	}
 
-
-	void SRenderStats::Update(bool bIsPlaying, double InDeltaTime, const UMetaSoundSource* InSource)
+	void SRenderStats::Update(bool bIsPlaying, const UMetaSoundSource* InSource)
 	{
 		using namespace Metasound;
 
 		// Reset maximum values when play restarts
-		const bool bPlayStateChanged = bIsPlaying != bPreviousIsPlaying;
-		if (bIsPlaying)
-		{
-			if (!bPreviousIsPlaying)
-			{
-				MaxRelativeRenderCost = 0.f;
-				MaxCPUCoreUtilization = 0;
-			}
-			PlayTime += InDeltaTime;
-		}
-		else
+		if (bIsPlaying && !bPreviousIsPlaying)
 		{
-			PlayTime = 0.0;
+			MaxRelativeRenderCost = 0.f;
+			MaxCPUCoreUtilization = 0;
 		}
-
 		bPreviousIsPlaying = bIsPlaying;
 
-		if (!RenderStatsCPUWidget.IsValid()
-			|| !RenderStatsCostWidget.IsValid()
-			|| !PlayTimeWidget.IsValid()
-			|| !AuditionPageWidget.IsValid()
-			|| !AuditionPlatformWidget.IsValid())
-			{
-			return;
-		}
-
-		double CPUCoreUtilization = 0;
-		float RelativeRenderCost = 0.f;
-
-		if (bIsPlaying && InSource)
-		{
-			if (const UAudioComponent* PreviewComponent = GEditor->GetPreviewAudioComponent())
-			{
-				TSharedPtr<FMetasoundGenerator> Generator = InSource->GetGeneratorForAudioComponent(PreviewComponent->GetAudioComponentID()).Pin();
-				if (Generator.IsValid())
-				{
-					// Update render stats
-					CPUCoreUtilization = Generator->GetCPUCoreUtilization();
-					MaxCPUCoreUtilization = FMath::Max(MaxCPUCoreUtilization, CPUCoreUtilization);
-
-					RelativeRenderCost = Generator->GetRelativeRenderCost();
-					MaxRelativeRenderCost = FMath::Max(MaxRelativeRenderCost, RelativeRenderCost);
-				}
-			}
-		}
-
-		// Display updated render stats.
-		FString PlayTimeString = FTimespan::FromSeconds(PlayTime).ToString();
-		PlayTimeString.ReplaceInline(TEXT("+"), TEXT(""));
-		PlayTimeWidget->SetText(FText::FromString(MoveTemp(PlayTimeString)));
-
-		FString RenderCostString = FString::Printf(TEXT("Relative Render Cost: %3.2f (%3.2f Max)"), RelativeRenderCost, MaxRelativeRenderCost);
-		RenderStatsCostWidget->SetText(FText::FromString(MoveTemp(RenderCostString)));
-
-		FString CPUCoreUtilizationString = FString::Printf(TEXT("CPU Core: %3.2f%% (%3.2f%% Max)"), 100. * CPUCoreUtilization, 100. * MaxCPUCoreUtilization);
-		RenderStatsCPUWidget->SetText(FText::FromString(MoveTemp(CPUCoreUtilizationString)));
-
-		if (bPlayStateChanged)
+		if (RenderStatsCPUWidget.IsValid() && RenderStatsCostWidget.IsValid())
 		{
-			FText AuditionPage;
-			FText AuditionPlatform;
+			double CPUCoreUtilization = 0;
+			float RelativeRenderCost = 0.f;
 
-			if (bIsPlaying)
+			// Find generator for playing preview component. 
+			if (bIsPlaying && InSource)
 			{
-				if (const UMetasoundEditorSettings* EditorSettings = GetDefault<UMetasoundEditorSettings>())
+				if (const UAudioComponent* PreviewComponent = GEditor->GetPreviewAudioComponent())
 				{
-					if (const UMetaSoundSettings* Settings = GetDefault<UMetaSoundSettings>())
+					TSharedPtr<FMetasoundGenerator> Generator = InSource->GetGeneratorForAudioComponent(PreviewComponent->GetAudioComponentID()).Pin();
+					if (Generator.IsValid())
 					{
-						if (!Settings->GetProjectPageSettings().IsEmpty())
-						{
-							const FText AuditionPageHeader = LOCTEXT("AuditionPageActive_DebugFormat", "Auditioning Page: {0}");
-							AuditionPage = FText::Format(AuditionPageHeader, FText::FromName(EditorSettings->AuditionPage));
-
-							if (EditorSettings->AuditionPlatform != "Editor")
-							{
-								const FText AuditionPlatformHeader = LOCTEXT("AuditionPlatformActive_DebugFormat", "Auditioning Platform: {0}");
-								AuditionPlatform = FText::Format(AuditionPlatformHeader, FText::FromName(EditorSettings->AuditionPlatform));
-							}
-						}
+						// Update render stats
+						CPUCoreUtilization = Generator->GetCPUCoreUtilization();
+						MaxCPUCoreUtilization = FMath::Max(MaxCPUCoreUtilization, CPUCoreUtilization);
+
+						RelativeRenderCost = Generator->GetRelativeRenderCost();
+						MaxRelativeRenderCost = FMath::Max(MaxRelativeRenderCost, RelativeRenderCost);
 					}
 				}
 			}
-			AuditionPageWidget->SetText(AuditionPage);
-			AuditionPlatformWidget->SetText(AuditionPlatform);
+
+			// Display updated render stats. 
+			FString CPUCoreUtilizationString = FString::Printf(TEXT("%3.2f%% (%3.2f%% Max) CPU Core"), 100. * CPUCoreUtilization, 100. * MaxCPUCoreUtilization);
+			RenderStatsCPUWidget->SetText(FText::FromString(CPUCoreUtilizationString));
+
+			FString RenderCostString = FString::Printf(TEXT("%3.2f (%3.2f Max) Relative Render Cost"), RelativeRenderCost, MaxRelativeRenderCost);
+			RenderStatsCostWidget->SetText(FText::FromString(RenderCostString));
 		}
 	}
 } // namespace Metasound::Editor
diff --git a/Engine/Plugins/Runtime/Metasound/Source/MetasoundEditor/Private/SMetasoundStats.h b/Engine/Plugins/Runtime/Metasound/Source/MetasoundEditor/Private/SMetasoundStats.h
--- a/Engine/Plugins/Runtime/Metasound/Source/MetasoundEditor/Private/SMetasoundStats.h
+++ b/Engine/Plugins/Runtime/Metasound/Source/MetasoundEditor/Private/SMetasoundStats.h
@@ -28,12 +28,11 @@ namespace Metasound::Editor
 
 		void Construct(const FArguments& InArgs);
 		void SetExecVisibility(TAttribute<EVisibility> InVisibility);
-		void Update(const FMetaSoundPageSettings* AuditionPageSettings, const FMetaSoundPageSettings* GraphPageSettings, const FSlateColor* ActiveColor);
+		void Update(const FMetaSoundPageSettings* PageSettings, const FText& Header, const FSlateColor* TextColor);
 
 	private:
 		TSharedPtr<SImage> ExecImageWidget;
-		TSharedPtr<STextBlock> GraphPageTextWidget;
-		TSharedPtr<STextBlock> AuditionPageTextWidget;
+		TSharedPtr<STextBlock> PageTextWidget;
 	};
 
 
@@ -45,18 +44,14 @@ namespace Metasound::Editor
 		SLATE_END_ARGS()
 
 		void Construct(const FArguments& InArgs);
-		void Update(bool bIsPlaying, double InDeltaTime, const UMetaSoundSource* InSource);
+		void Update(bool bIsPlaying, const UMetaSoundSource* InSource);
 
 	private:
-		TSharedPtr<STextBlock> AuditionPageWidget;
-		TSharedPtr<STextBlock> AuditionPlatformWidget;
-		TSharedPtr<STextBlock> PlayTimeWidget;
 		TSharedPtr<STextBlock> RenderStatsCostWidget;
 		TSharedPtr<STextBlock> RenderStatsCPUWidget;
 
 		bool bPreviousIsPlaying = false;
-		double MaxCPUCoreUtilization = 0.0;
-		double PlayTime = 0.0;
+		double MaxCPUCoreUtilization = 0;
 		float MaxRelativeRenderCost = 0.f;
 	};
 } // namespace Metasound::Editor
diff --git a/Engine/Plugins/Runtime/Metasound/Source/MetasoundEditor/Public/MetasoundEditorModule.h b/Engine/Plugins/Runtime/Metasound/Source/MetasoundEditor/Public/MetasoundEditorModule.h
--- a/Engine/Plugins/Runtime/Metasound/Source/MetasoundEditor/Public/MetasoundEditorModule.h
+++ b/Engine/Plugins/Runtime/Metasound/Source/MetasoundEditor/Public/MetasoundEditorModule.h
@@ -15,11 +15,6 @@
 #include "Templates/Function.h"
 
 // Forward Declarations
-struct FAudioMeterDefaultColorStyle;
-struct FAudioOscilloscopePanelStyle;
-struct FAudioSpectrumPlotStyle;
-struct FAudioVectorscopePanelStyle;
-
 class IDetailLayoutBuilder;
 class IDetailPropertyRow;
 class IPropertyHandle;
@@ -46,14 +41,6 @@ namespace Metasound
 		{
 			METASOUNDEDITOR_API FSlateIcon CreateSlateIcon(FName InName);
 			METASOUNDEDITOR_API const FSlateBrush& GetSlateBrushSafe(FName InName);
-
-			const FSlateColor& GetDefaultAnalyzerColor();
-			const FSlateColor& GetPageExecutingColor();
-
-			const FAudioMeterDefaultColorStyle& GetMeterDefaultColorStyle();
-			const FAudioOscilloscopePanelStyle& GetOscilloscopeStyle();
-			const FAudioSpectrumPlotStyle& GetSpectrumPlotStyle();
-			const FAudioVectorscopePanelStyle& GetVectorscopeStyle();
 		} // namespace Style
 
 		using EAssetScanStatus = Metasound::Engine::EAssetScanStatus;
