[Backout] - CL36442802
Original CL Desc
-----------------------------------------------------------------
- Uniform color fix-ups for MetaSound analyzers
- Move live MetaSound stats/audition settings to top left of graph
- Fix ordering of playback widget in input details panel to follow same order as page details panel

#rnx

diff --git a/Engine/Plugins/Runtime/Metasound/Source/MetasoundEditor/Private/MetasoundEditor.cpp b/Engine/Plugins/Runtime/Metasound/Source/MetasoundEditor/Private/MetasoundEditor.cpp
--- a/Engine/Plugins/Runtime/Metasound/Source/MetasoundEditor/Private/MetasoundEditor.cpp
+++ b/Engine/Plugins/Runtime/Metasound/Source/MetasoundEditor/Private/MetasoundEditor.cpp
@@ -928,43 +928,50 @@ namespace Metasound
 			InTabManager->RegisterTabSpawner(TabNamesPrivate::GraphCanvas, FOnSpawnTab::CreateLambda(
 				[
 					bShowPageTab = ShowPageGraphDetails(),
-					InPageStatsWidget = PageStatsWidget,
+					InFocusPageWidget = FocusPageWidget,
+					InPlayTimeWidget = PlayTimeWidget,
 					InMetasoundGraphEditor = MetasoundGraphEditor,
 					InRenderStatsWidget = RenderStatsWidget
 				](const FSpawnTabArgs& Args)
 			{
 				TSharedRef<SDockTab> SpawnedTab = SNew(SDockTab).Label(LOCTEXT("MetasoundGraphCanvasTitle", "MetaSound Graph"));
 
-				TSharedRef<SOverlay> Overlay = SNew(SOverlay)
-					+ SOverlay::Slot()
-					[
-						InMetasoundGraphEditor.ToSharedRef()
-					]
-					+ SOverlay::Slot()
-					.VAlign(VAlign_Top)
+				TSharedPtr<SVerticalBox> StatsWidget = SNew(SVerticalBox)
+					+ SVerticalBox::Slot()
+					.HAlign(HAlign_Left)
+					.VAlign(VAlign_Center)
+					.Padding(2.f, 1.f)
 					[
 						InRenderStatsWidget.ToSharedRef()
-					]
-					.Padding(5.0f, 5.0f);
+					];
 
-				if (bShowPageTab && InPageStatsWidget.IsValid())
+				if (bShowPageTab)
 				{
-					TSharedRef<SVerticalBox> GraphStatsWidget = SNew(SVerticalBox)
-						+ SVerticalBox::Slot()
+					StatsWidget->AddSlot()
 						.HAlign(HAlign_Left)
 						.VAlign(VAlign_Center)
 						.AutoHeight()
 						[
-							InPageStatsWidget.ToSharedRef()
+							InFocusPageWidget.ToSharedRef()
 						];
-					Overlay->AddSlot()
-					.VAlign(VAlign_Bottom)
-					[
-						GraphStatsWidget
-					];
 				}
 
-				SpawnedTab->SetContent(Overlay);
+				SpawnedTab->SetContent(SNew(SOverlay)
+					+ SOverlay::Slot()
+					[
+						InMetasoundGraphEditor.ToSharedRef()
+					]
+					+ SOverlay::Slot()
+					[
+						InPlayTimeWidget.ToSharedRef()
+					]
+					+ SOverlay::Slot()
+					.VAlign(VAlign_Bottom)
+					[
+						StatsWidget.ToSharedRef()
+					]
+					.Padding(5.0f, 5.0f)
+				);
 				return SpawnedTab;
 			}))
 			.SetDisplayName(LOCTEXT("GraphCanvasTab", "Viewport"))
@@ -1080,6 +1087,20 @@ namespace Metasound
 				BackgroundColor = MetaSoundStyle->GetColor("MetasoundEditor.Analyzers.BackgroundColor");
 			}
 
+			TSharedRef<SVerticalBox> PageWidget = SNew(SVerticalBox)
+				.Visibility(TAttribute<EVisibility>::Create([this]()
+				{
+					const bool bMakeVisible = IsPlaying() && AssetEditorPrivate::EnablePageEditor;
+					return bMakeVisible ? EVisibility::Visible : EVisibility::Collapsed;
+				}))
+				+ SVerticalBox::Slot()
+				.HAlign(HAlign_Left)
+				.VAlign(VAlign_Center)
+				.AutoHeight()
+				[
+					RenderPageWidget.ToSharedRef()
+				];
+
 			return SNew(SOverlay)
 			+ SOverlay::Slot()
 			[
@@ -1090,6 +1111,11 @@ namespace Metasound
 			[
 				SNew(SSplitter)
 				.Orientation(Orient_Vertical)
+				+ SSplitter::Slot()
+					.Value(0.05f)
+					[
+						PageWidget
+					]
 				+ SSplitter::Slot()
 				.Value(0.5f)
 				[
@@ -1179,7 +1205,10 @@ namespace Metasound
 			Builder->AddTransactionListener(DocListener->AsShared());
 
 			// Stat widgets are potentially intractable with transaction listener, so create then here
-			SAssignNew(PageStatsWidget, SPageStats)
+			SAssignNew(FocusPageWidget, SPageStats)
+				.Visibility(EVisibility::HitTestInvisible);
+
+			SAssignNew(RenderPageWidget, SPageStats)
 				.Visibility(EVisibility::HitTestInvisible);
 
 			SAssignNew(RenderStatsWidget, SRenderStats)
@@ -1326,9 +1355,7 @@ namespace Metasound
 			}
 			else
 			{
-				constexpr bool bIsPlaying = false;
-				UpdatePageInfo(bIsPlaying);
-				UpdateRenderInfo(bIsPlaying);
+				UpdatePageInfo(false);
 			}
 
 			RefreshExecVisibility(Builder->GetConstBuilder().GetBuildPageID());
@@ -1743,7 +1770,7 @@ namespace Metasound
 
 					if (!OutputMeter.IsValid())
 					{
-						OutputMeter = MakeShared<AudioWidgets::FAudioMeter>(MetaSoundSource.NumChannels, AudioDeviceId, DefaultBus, &Style::GetMeterDefaultColorStyle());
+						OutputMeter = MakeShared<AudioWidgets::FAudioMeter>(MetaSoundSource.NumChannels, AudioDeviceId, DefaultBus);
 					}
 				}
 				else if (OutputMeter->GetAudioBus()->GetNumChannels() != MetaSoundSource.NumChannels)
@@ -1766,14 +1793,13 @@ namespace Metasound
 						OscilloscopeTimeWindowMs,
 						OscilloscopeMaxTimeWindowMs,
 						OscilloscopeAnalysisPeriodMs,
-						OscilloscopePanelLayoutType,
-						&Style::GetOscilloscopeStyle());
+						OscilloscopePanelLayoutType);
 				}
 				else if (OutputOscilloscope->GetAudioBus()->GetNumChannels() != MetaSoundSource.NumChannels)
 				{
 					OutputOscilloscope->CreateAudioBus(MetaSoundNumChannels);
 					OutputOscilloscope->CreateDataProvider(AudioDeviceId, OscilloscopeTimeWindowMs, OscilloscopeMaxTimeWindowMs, OscilloscopeAnalysisPeriodMs, OscilloscopePanelLayoutType);
-					OutputOscilloscope->CreateOscilloscopeWidget(MetaSoundNumChannels, OscilloscopePanelLayoutType, &Style::GetOscilloscopeStyle());
+					OutputOscilloscope->CreateOscilloscopeWidget(MetaSoundNumChannels, OscilloscopePanelLayoutType);
 				}
 
 				// Init Vectorscope
@@ -1789,14 +1815,13 @@ namespace Metasound
 						VectorscopeTimeWindowMs,
 						VectorscopeMaxTimeWindowMs,
 						VectorscopeAnalysisPeriodMs,
-						VectorscopePanelLayoutType,
-						&Style::GetVectorscopeStyle());
+						VectorscopePanelLayoutType);
 				}
 				else if (OutputVectorscope->GetAudioBus()->GetNumChannels() != MetaSoundSource.NumChannels)
 				{
 					OutputVectorscope->CreateAudioBus(MetaSoundNumChannels);
 					OutputVectorscope->CreateDataProvider(AudioDeviceId, VectorscopeTimeWindowMs, VectorscopeMaxTimeWindowMs, VectorscopeAnalysisPeriodMs);
-					OutputVectorscope->CreateVectorscopeWidget(VectorscopePanelLayoutType, &Style::GetVectorscopeStyle());
+					OutputVectorscope->CreateVectorscopeWidget(VectorscopePanelLayoutType);
 				}
 
 				if (!OutputSpectrumAnalyzer.IsValid())
@@ -1890,7 +1915,7 @@ namespace Metasound
 							SpectrumAnalyzerSettings.bDisplaySoundLevelAxisLabels = !SpectrumAnalyzerSettings.bDisplaySoundLevelAxisLabels;
 							GetMutableDefault<UMetasoundEditorSettings>()->SaveConfig();
 						});
-					Params.PlotStyle = &Style::GetSpectrumPlotStyle();
+
 					OutputSpectrumAnalyzer = MakeShared<AudioWidgets::FAudioSpectrumAnalyzer>(Params);
 				}
 				else if (OutputSpectrumAnalyzer->GetAudioBus()->GetNumChannels() != MetaSoundSource.NumChannels)
@@ -2567,7 +2592,7 @@ namespace Metasound
 				// Set the send to the audio bus that is used for analyzing the metasound output
 				check(GEditor);
 
-				UpdateRenderInfo(true /* bIsPlaying */);
+				PlayTime = 0.0;
 				UpdatePageInfo(true);
 
 				if (UMetaSoundSource* Source = Cast<UMetaSoundSource>(GetMetasoundObject()))
@@ -2603,8 +2628,8 @@ namespace Metasound
 				MetasoundGraphEditor->RegisterActiveTimer(0.0f,
 					FWidgetActiveTimerDelegate::CreateLambda([this](double InCurrentTime, float InDeltaTime)
 					{
-						const bool bIsPlaying = IsPlaying();
-						UpdateRenderInfo(bIsPlaying, InDeltaTime);
+						const bool bIsPlaying = UpdatePlayTime(InDeltaTime);
+						UpdateRenderInfo();
 
 						if (bIsPlaying)
 						{
@@ -2612,7 +2637,8 @@ namespace Metasound
 						}
 						else
 						{
-							UpdatePageInfo(bIsPlaying);
+							PlayTime = 0.0;
+							UpdatePageInfo(false);
 							GraphConnectionManager = RebuildConnectionManager();
 							return EActiveTimerReturnType::Stop;
 						}
@@ -2688,6 +2714,11 @@ namespace Metasound
 			return true;
 		}
 
+		double FEditor::GetPlayTime() const
+		{
+			return PlayTime;
+		}
+
 		TSharedPtr<SGraphEditor> FEditor::GetGraphEditor() const
 		{
 			return MetasoundGraphEditor;
@@ -3076,6 +3107,11 @@ namespace Metasound
 				.GraphToEdit(MetaSoundAsset->GetGraph())
 				.IsEditable(this, &FEditor::IsGraphEditable)
 				.ShowGraphStateOverlay(false);
+
+			SAssignNew(PlayTimeWidget, STextBlock)
+				.Visibility(EVisibility::HitTestInvisible)
+				.TextStyle(FAppStyle::Get(), "Graph.ZoomText")
+				.ColorAndOpacity(FLinearColor(1, 1, 1, 0.30f));
 		}
 
 		FGraphAppearanceInfo FEditor::GetGraphAppearance() const
@@ -5077,18 +5113,14 @@ namespace Metasound
 			const UMetaSoundSettings* Settings = GetDefault<UMetaSoundSettings>();
 			check(Settings)
 
-			if (PageStatsWidget.IsValid())
+			if (FocusPageWidget.IsValid())
 			{
 				const FSlateColor* Color = nullptr;
-				const FMetaSoundPageSettings* GraphPageSettings = nullptr;
-				const FMetaSoundPageSettings* AuditionPageSettings = nullptr;
+				// For now, this is one and the same, but may change in the future to accommodate conflating
+				// the static nature of the active build graph ID with what the user is visualizing.
+				const FMetaSoundPageSettings* PageSettings = nullptr;
 				if (Builder.IsValid())
 				{
-					if (const UMetasoundEditorSettings* EditorSettings = GetDefault<UMetasoundEditorSettings>())
-					{
-						AuditionPageSettings = Settings->FindPageSettings(EditorSettings->AuditionPage);
-					}
-
 					const FMetaSoundFrontendDocumentBuilder& DocBuilder = Builder->GetConstBuilder();
 					const FGuid& PageID = DocBuilder.GetBuildPageID();
 					if (bIsPlaying)
@@ -5097,26 +5129,61 @@ namespace Metasound
 						const FGuid ResolvePageID = FDocumentBuilderRegistry::GetChecked().ResolveTargetPageID(GraphClass);
 						if (ResolvePageID == PageID)
 						{
-							Color = &Style::GetPageExecutingColor();
+							Color = &FStyleColors::AccentGreen;
 						}
 					}
-					GraphPageSettings = Settings->FindPageSettings(PageID);
+					PageSettings = Settings->FindPageSettings(PageID);
 				}
 
-				PageStatsWidget->Update(AuditionPageSettings, GraphPageSettings, Color);
+				const FText GraphPageHeader = LOCTEXT("GraphPageTargetHeader", "Graph Page");
+				FocusPageWidget->Update(PageSettings, GraphPageHeader, Color);
+			}
+
+			if (RenderPageWidget.IsValid())
+			{
+				const FMetaSoundPageSettings* PageSettings = nullptr;
+				if (Builder.IsValid())
+				{
+					if (const UMetasoundEditorSettings* EditorSettings = GetDefault<UMetasoundEditorSettings>())
+					{
+						PageSettings = Settings->FindPageSettings(EditorSettings->AuditionPage);
+					}
+				}
+
+				const FText RenderPageHeader = LOCTEXT("AuditionPageHeader", "Audition Page");
+				RenderPageWidget->Update(PageSettings, RenderPageHeader, &FStyleColors::AccentGreen);
 			}
 		}
 
-		void FEditor::UpdateRenderInfo(bool bIsPlaying, float InDeltaTime)
+		bool FEditor::UpdatePlayTime(float InDeltaTime)
 		{
-			if (!bIsPlaying)
+			const bool bIsPlaying = IsPlaying();
+			if (bIsPlaying)
 			{
-				SetPreviewID(INDEX_NONE);
+				if (PlayTimeWidget.IsValid())
+				{
+					PlayTime += InDeltaTime;
+
+					FString PlayTimeString = FTimespan::FromSeconds(PlayTime).ToString();
+					PlayTimeString.ReplaceInline(TEXT("+"), TEXT(""));
+					PlayTimeWidget->SetText(FText::FromString(MoveTemp(PlayTimeString)));
+				}
+
+				return true;
 			}
 
+			SetPreviewID(INDEX_NONE);
+			PlayTime = 0.0;
+			PlayTimeWidget->SetText(FText::GetEmpty());
+			return false;
+		}
+
+		void FEditor::UpdateRenderInfo()
+		{
 			if (RenderStatsWidget.IsValid())
 			{
-				RenderStatsWidget->Update(bIsPlaying, InDeltaTime, Cast<const UMetaSoundSource>(GetMetasoundObject()));
+				const bool bIsPlaying = IsPlaying();
+				RenderStatsWidget->Update(bIsPlaying, Cast<const UMetaSoundSource>(GetMetasoundObject()));
 			}
 		}
 
@@ -5138,7 +5205,7 @@ namespace Metasound
 
 		void FEditor::RefreshExecVisibility(const FGuid& InPageID) const
 		{
-			if (PageStatsWidget.IsValid())
+			if (FocusPageWidget.IsValid())
 			{
 				TAttribute<EVisibility> ExecVisibility = TAttribute<EVisibility>::CreateLambda([this, InPageID]()
 				{
@@ -5147,12 +5214,12 @@ namespace Metasound
 					{
 						const FMetaSoundFrontendDocumentBuilder& DocBuilder = Builder->GetConstBuilder();
 						const bool bIsPreviewing = IsPreviewingPageGraph(DocBuilder, InPageID);
-						return bIsPreviewing ? EVisibility::HitTestInvisible : EVisibility::Collapsed;
+						return bIsPreviewing ? EVisibility::Visible : EVisibility::Collapsed;
 					}
 
 					return EVisibility::Collapsed;
 				});
-				PageStatsWidget->SetExecVisibility(MoveTemp(ExecVisibility));
+				FocusPageWidget->SetExecVisibility(MoveTemp(ExecVisibility));
 			}
 		}
 
@@ -5161,9 +5228,13 @@ namespace Metasound
 			if (MetasoundGraphEditor.IsValid())
 			{
 				UMetasoundEditorGraph& Graph = GetMetaSoundGraphChecked();
-				if (const UEdGraphPin* TargetPin = MetasoundGraphEditor->GetGraphPinForMenu(); ensure(TargetPin))
+
+				UEdGraphPin* TargetPin = MetasoundGraphEditor->GetGraphPinForMenu();
+				check(TargetPin);
+
+				if (TargetPin->Direction == EGPD_Input)
 				{
-					return TargetPin->Direction == EGPD_Input;
+					return true;
 				}
 			}
 			return false;
@@ -5191,9 +5262,13 @@ namespace Metasound
 			if (MetasoundGraphEditor.IsValid())
 			{
 				UMetasoundEditorGraph& Graph = GetMetaSoundGraphChecked();
-				if (const UEdGraphPin* TargetPin = MetasoundGraphEditor->GetGraphPinForMenu(); ensure(TargetPin))
+
+				UEdGraphPin* TargetPin = MetasoundGraphEditor->GetGraphPinForMenu();
+				check(TargetPin);
+
+				if (TargetPin->Direction == EGPD_Output)
 				{
-					return TargetPin->Direction == EGPD_Output;
+					return true;
 				}
 			}
 			return false;
