Fix details panel issues with MetaSound member metadata (ex. float range/clamp) in outputs/variables + presets
- Don't show or apply range on outputs/variables unless Clamp Default is selected
- Fix inconsistencies with Clamp Default and Range displays in presets

#rnx

diff --git a/Engine/Plugins/Runtime/Metasound/Source/MetasoundEditor/Private/MetasoundEditorGraphMemberDefaults.cpp b/Engine/Plugins/Runtime/Metasound/Source/MetasoundEditor/Private/MetasoundEditorGraphMemberDefaults.cpp
--- a/Engine/Plugins/Runtime/Metasound/Source/MetasoundEditor/Private/MetasoundEditorGraphMemberDefaults.cpp
+++ b/Engine/Plugins/Runtime/Metasound/Source/MetasoundEditor/Private/MetasoundEditorGraphMemberDefaults.cpp
@@ -674,6 +674,10 @@ void UMetasoundEditorGraphMemberDefaultFloat::PostEditChangeChainProperty(FPrope
 		}
 
 		// If the widget type is changed to none, we need to refresh clamping the value or not, since if the widget was a slider before, the value was clamped
+		if (WidgetType == EMetasoundMemberDefaultWidget::None)
+		{
+			ClampDefault = true;
+		}
 		OnClampChanged.Broadcast(ClampDefault);
 	}
 	else if (PropertyName.IsEqual(GET_MEMBER_NAME_CHECKED(UMetasoundEditorGraphMemberDefaultFloat, ClampDefault)))
diff --git a/Engine/Plugins/Runtime/Metasound/Source/MetasoundEditor/Private/MetasoundEditorSubsystem.cpp b/Engine/Plugins/Runtime/Metasound/Source/MetasoundEditor/Private/MetasoundEditorSubsystem.cpp
--- a/Engine/Plugins/Runtime/Metasound/Source/MetasoundEditor/Private/MetasoundEditorSubsystem.cpp
+++ b/Engine/Plugins/Runtime/Metasound/Source/MetasoundEditor/Private/MetasoundEditorSubsystem.cpp
@@ -87,7 +87,8 @@ TScriptInterface<IMetaSoundDocumentInterface> UMetaSoundEditorSubsystem::BuildTo
 					// Only use referenced UObject's SoundWave settings for sources if not overridden 
 					if (TemplateSoundWave == nullptr && bIsSource)
 					{
-						if (const UObject* ReferencedObject = NewDocBuilder.GetReferencedPresetAsset())
+						const UObject* ReferencedObject = NewDocBuilder.GetReferencedPresetAsset();
+						if (ensureMsgf(ReferencedObject, TEXT("Preset builder %s does not have a referenced asset to apply soundwave settings from."), *NewDocBuilder.GetName()))
 						{
 							TemplateSoundWave = CastChecked<USoundWave>(ReferencedObject);
 						}
diff --git a/Engine/Plugins/Runtime/Metasound/Source/MetasoundEditor/Private/MetasoundNodeDetailCustomization.cpp b/Engine/Plugins/Runtime/Metasound/Source/MetasoundEditor/Private/MetasoundNodeDetailCustomization.cpp
--- a/Engine/Plugins/Runtime/Metasound/Source/MetasoundEditor/Private/MetasoundNodeDetailCustomization.cpp
+++ b/Engine/Plugins/Runtime/Metasound/Source/MetasoundEditor/Private/MetasoundNodeDetailCustomization.cpp
@@ -306,21 +306,28 @@ namespace Metasound
 				if (ensure(Row))
 				{
 					Row->Visibility(DefaultVisibility);
+					Row->IsEnabled(GetEnabled());
 				}
 				return Row;
 			};
 
 			if (IDetailPropertyRow* ClampRow = AddOptionPropRow(GET_MEMBER_NAME_CHECKED(UMetasoundEditorGraphMemberDefaultFloat, ClampDefault)))
 			{
-				// Apply the clamp range to the default value if not using a widget and ClampDefault is true
+				// Apply the clamp range to the default value if using a widget or ClampDefault is otherwise true
+				// Only show clamp row if not using a widget (widgets always require a clamp and range)
+				// Presets are an exception, because they may have a widget inherited, but that doesn't apply and isn't editable
 				const bool bUsingWidget = FloatLiteral->WidgetType != EMetasoundMemberDefaultWidget::None;
-				const bool bShouldClampDefaultValue = bUsingWidget || (!bUsingWidget && FloatLiteral->ClampDefault);
-				ClampRow->Visibility(DefaultVisibility);
+				const UMetasoundEditorGraphMember* Member = FloatLiteral->FindMember();
+				const bool bIsPreset = Member ? Member->GetFrontendBuilderChecked().IsPreset() : false;
+				const bool bShowClampRow = !bUsingWidget || bIsPreset;
+				const bool bApplyRange = (bUsingWidget && !bIsPreset) || FloatLiteral->ClampDefault;
+
+				ClampRow->Visibility(bShowClampRow ? DefaultVisibility : EVisibility::Hidden);
 				for (TSharedPtr<IPropertyHandle>& DefaultValueHandle : DefaultProperties)
 				{
 					if (DefaultValueHandle.IsValid())
 					{
-						if (bShouldClampDefaultValue)
+						if (bApplyRange)
 						{
 							FVector2D Range = FloatLiteral->GetRange();
 							DefaultValueHandle->SetInstanceMetaData("ClampMin", FString::Printf(TEXT("%f"), Range.X));
@@ -348,7 +355,7 @@ namespace Metasound
 					}
 				});
 
-				if (bShouldClampDefaultValue)
+				if (bApplyRange)
 				{
 					AddOptionPropRow(GET_MEMBER_NAME_CHECKED(UMetasoundEditorGraphMemberDefaultFloat, Range));
 				}
diff --git a/Engine/Plugins/Runtime/Metasound/Source/MetasoundEditor/Public/MetasoundEditorGraphMemberDefaults.h b/Engine/Plugins/Runtime/Metasound/Source/MetasoundEditor/Public/MetasoundEditorGraphMemberDefaults.h
--- a/Engine/Plugins/Runtime/Metasound/Source/MetasoundEditor/Public/MetasoundEditorGraphMemberDefaults.h
+++ b/Engine/Plugins/Runtime/Metasound/Source/MetasoundEditor/Public/MetasoundEditorGraphMemberDefaults.h
@@ -326,7 +326,7 @@ private:
 public:
 	virtual ~UMetasoundEditorGraphMemberDefaultFloat() = default;
 
-	UPROPERTY(EditAnywhere, Category = DefaultValue, meta=(EditCondition = "WidgetType == EMetasoundMemberDefaultWidget::None", EditConditionHides))
+	UPROPERTY(EditAnywhere, Category = DefaultValue)
 	bool ClampDefault = false;
 
 	UPROPERTY(EditAnywhere, Category = DefaultValue)
