Fix for wrong node registry key ctor used when migrating package

#rnx

diff --git a/Engine/Plugins/Runtime/Metasound/Source/MetasoundEditor/Private/MetasoundEditorModule.cpp b/Engine/Plugins/Runtime/Metasound/Source/MetasoundEditor/Private/MetasoundEditorModule.cpp
--- a/Engine/Plugins/Runtime/Metasound/Source/MetasoundEditor/Private/MetasoundEditorModule.cpp
+++ b/Engine/Plugins/Runtime/Metasound/Source/MetasoundEditor/Private/MetasoundEditorModule.cpp
@@ -784,7 +784,7 @@ namespace Metasound
 					TMap<FNodeRegistryKey, FNodeRegistryKey> OldToNewReferenceKeys;
 					for (FMetaSoundFrontendDocumentBuilder& MetaSoundBuilder : NewMetaSoundAssetBuilders)
 					{
-						FNodeRegistryKey OldRegistryKey(MetaSoundBuilder.GetConstDocumentChecked().RootGraph.Metadata);
+						FNodeRegistryKey OldRegistryKey(MetaSoundBuilder.GetConstDocumentChecked().RootGraph);
 						FNodeRegistryKey NewRegistryKey(EMetasoundFrontendClassType::External, MetaSoundBuilder.GenerateNewClassName(), OldRegistryKey.Version);
 						OldToNewReferenceKeys.FindOrAdd(MoveTemp(OldRegistryKey)) = MoveTemp(NewRegistryKey);
 
