Fix for crash deleting metasound node

#tests Tested on Mac with repro steps from JIRA and Windows.
[RN] Bug fix for crash in Metasounds from deleteing a node

diff --git a/Engine/Plugins/Runtime/Metasound/Source/MetasoundEditor/Private/SMetasoundGraphEnumPin.cpp b/Engine/Plugins/Runtime/Metasound/Source/MetasoundEditor/Private/SMetasoundGraphEnumPin.cpp
--- a/Engine/Plugins/Runtime/Metasound/Source/MetasoundEditor/Private/SMetasoundGraphEnumPin.cpp
+++ b/Engine/Plugins/Runtime/Metasound/Source/MetasoundEditor/Private/SMetasoundGraphEnumPin.cpp
@@ -38,14 +38,14 @@ namespace Metasound
 		SMetasoundGraphEnumPin::FindEnumInterfaceFromPin(UEdGraphPin* InPin)
 		{
 			using namespace Frontend;
-
-			auto MetasoundEditorNode = Cast<UMetasoundEditorGraphNode>(InPin->GetOwningNode());
-			FConstNodeHandle NodeHandle = MetasoundEditorNode->GetConstNodeHandle();
-			FConstInputHandle Input = NodeHandle->GetConstInputWithVertexName(InPin->GetFName());
-			if (Input->IsValid())
+			if (const UMetasoundEditorGraphNode* MetasoundEditorNode = Cast<UMetasoundEditorGraphNode>(InPin->GetOwningNode()))
 			{
-				FName DataType = Input->GetDataType();
-				return Frontend::IDataTypeRegistry::Get().GetEnumInterfaceForDataType(DataType);
+				const UMetaSoundBuilderBase& Builder = MetasoundEditorNode->GetBuilderChecked();
+				if (const FMetasoundFrontendVertex* Vertex = FGraphBuilder::GetPinVertex(Builder.GetConstBuilder(), InPin))
+				{
+					const FName DataType = Vertex->TypeName;
+					return IDataTypeRegistry::Get().GetEnumInterfaceForDataType(DataType);
+				}
 			}
 			return nullptr;
 		}
@@ -54,62 +54,65 @@ namespace Metasound
 		{
 			using namespace Frontend;
 
-			TSharedPtr<const IEnumDataTypeInterface> EnumInterface = FindEnumInterfaceFromPin(GraphPinObj);
-			check(EnumInterface.IsValid());
-
-			if (UEdGraphPin* Pin = SGraphPin::GetPinObj())
+			if (const TSharedPtr<const IEnumDataTypeInterface> EnumInterface = FindEnumInterfaceFromPin(GraphPinObj); EnumInterface.IsValid())
 			{
-				int32 SelectedValue = FCString::Atoi(*GraphPinObj->GetDefaultAsString());	// Enums are currently serialized as ints (the value of the enum).
-				if (TOptional<IEnumDataTypeInterface::FGenericInt32Entry> Result = EnumInterface->FindByValue(SelectedValue))
+				if (GetPinObj())
 				{
-					return Result->DisplayName.ToString();
+					const int32 SelectedValue = FCString::Atoi(*GraphPinObj->GetDefaultAsString());	// Enums are currently serialized as ints (the value of the enum).
+					if (TOptional<IEnumDataTypeInterface::FGenericInt32Entry> Result = EnumInterface->FindByValue(SelectedValue); Result.IsSet() )
+					{
+						return Result->DisplayName.ToString();
+					}
 				}
 			}
-
-			return { };
+			return {};
 		}
 
 		void SMetasoundGraphEnumPin::GenerateComboBoxIndexes(TArray<TSharedPtr<int32>>& OutComboBoxIndexes)
 		{
 			using namespace Frontend;
-			TSharedPtr<const IEnumDataTypeInterface> EnumInterface = FindEnumInterfaceFromPin(GraphPinObj);
-			check(EnumInterface.IsValid());
-
-			const TArray<IEnumDataTypeInterface::FGenericInt32Entry>& Entries = EnumInterface->GetAllEntries();
-			for (int32 i = 0; i < Entries.Num(); ++i)
+			if (const TSharedPtr<const IEnumDataTypeInterface> EnumInterface = FindEnumInterfaceFromPin(GraphPinObj); EnumInterface.IsValid())
 			{
-				OutComboBoxIndexes.Add(MakeShared<int32>(i));
+				const TArray<IEnumDataTypeInterface::FGenericInt32Entry>& Entries = EnumInterface->GetAllEntries();
+				for (int32 i = 0; i < Entries.Num(); ++i)
+				{
+					OutComboBoxIndexes.Add(MakeShared<int32>(i));
+				}
 			}
 		}
 
 		void SMetasoundGraphEnumPin::ComboBoxSelectionChanged(TSharedPtr<int32> NewSelection, ESelectInfo::Type SelectInfo)
 		{
 			using namespace Frontend;
-			TSharedPtr<const IEnumDataTypeInterface> EnumInterface = FindEnumInterfaceFromPin(SGraphPin::GetPinObj());
-			check(EnumInterface.IsValid());
-
-			const TArray<IEnumDataTypeInterface::FGenericInt32Entry>& Entries = EnumInterface->GetAllEntries();
 
-			if (NewSelection.IsValid() && Entries.IsValidIndex(*NewSelection))
+			if (const TSharedPtr<const IEnumDataTypeInterface> EnumInterface = FindEnumInterfaceFromPin(SGraphPin::GetPinObj()); EnumInterface.IsValid())
 			{
-				int32 EnumValue = Entries[*NewSelection].Value;
-				FString EnumValueString = FString::FromInt(EnumValue);
-				if (GraphPinObj->GetDefaultAsString() != EnumValueString)
-				{
-					const FScopedTransaction Transaction(NSLOCTEXT("MetaSoundEditor", "ChangeEnumPinValue", "Change MetaSound Node Default Input Enum Value"));
-					GraphPinObj->Modify();
+				const TArray<IEnumDataTypeInterface::FGenericInt32Entry>& Entries = EnumInterface->GetAllEntries();
 
-					if (UMetasoundEditorGraphNode* MetaSoundNode = Cast<UMetasoundEditorGraphNode>(GraphPinObj->GetOwningNode()))
+				if (NewSelection.IsValid() && Entries.IsValidIndex(*NewSelection))
+				{
+					const int32 EnumValue = Entries[*NewSelection].Value;
+					const FString EnumValueString = FString::FromInt(EnumValue);
+					if (GraphPinObj->GetDefaultAsString() != EnumValueString)
 					{
-						if (UMetasoundEditorGraph* Graph = Cast<UMetasoundEditorGraph>(MetaSoundNode->GetGraph()))
+						const FScopedTransaction Transaction(NSLOCTEXT("MetaSoundEditor", "ChangeEnumPinValue", "Change MetaSound Node Default Input Enum Value"));
+						GraphPinObj->Modify();
+
+						if (UMetasoundEditorGraphNode* MetaSoundNode = Cast<UMetasoundEditorGraphNode>(GraphPinObj->GetOwningNode()))
 						{
-							Graph->Modify();
-							Graph->GetMetasoundChecked().Modify();
+							if (UMetasoundEditorGraph* Graph = Cast<UMetasoundEditorGraph>(MetaSoundNode->GetGraph()))
+							{
+								Graph->Modify();
+								Graph->GetMetasoundChecked().Modify();
+							}
 						}
-					}
 
-					//Set new selection
-					GraphPinObj->GetSchema()->TrySetDefaultValue(*GraphPinObj, EnumValueString);
+						//Set new selection
+						if (ensure(GraphPinObj->GetSchema()))
+						{
+							GraphPinObj->GetSchema()->TrySetDefaultValue(*GraphPinObj, EnumValueString);
+						}
+					}
 				}
 			}
 		}
@@ -118,26 +121,28 @@ namespace Metasound
 		{
 			using namespace Frontend;
 
-			TSharedPtr<const IEnumDataTypeInterface> Interface = FindEnumInterfaceFromPin(SGraphPin::GetPinObj());
-			check(Interface.IsValid());
-
-			const TArray<IEnumDataTypeInterface::FGenericInt32Entry>& Entries = Interface->GetAllEntries();
-			check(Entries.IsValidIndex(EnumIndex));
+			if (const TSharedPtr<const IEnumDataTypeInterface> Interface = FindEnumInterfaceFromPin(GetPinObj()); Interface.IsValid())
+			{
+				const TArray<IEnumDataTypeInterface::FGenericInt32Entry>& Entries = Interface->GetAllEntries();
+				check(Entries.IsValidIndex(EnumIndex));
 
-			return Entries[EnumIndex].DisplayName;
+				return Entries[EnumIndex].DisplayName;
+			}
+			return {};
 		}
 
 		FText SMetasoundGraphEnumPin::OnGetTooltip(int32 EnumIndex)
 		{
 			using namespace Frontend;
 
-			TSharedPtr<const IEnumDataTypeInterface> Interface = FindEnumInterfaceFromPin(GraphPinObj);
-			check(Interface.IsValid());
-
-			const TArray<IEnumDataTypeInterface::FGenericInt32Entry>& Entries = Interface->GetAllEntries();
-			check(Entries.IsValidIndex(EnumIndex));
+			if (const TSharedPtr<const IEnumDataTypeInterface> Interface = FindEnumInterfaceFromPin(GraphPinObj))
+			{
+				const TArray<IEnumDataTypeInterface::FGenericInt32Entry>& Entries = Interface->GetAllEntries();
+				check(Entries.IsValidIndex(EnumIndex));
 
-			return Entries[EnumIndex].Tooltip;
+				return Entries[EnumIndex].Tooltip;
+			}
+			return {};
 		}
 	} // namespace Editor
 } // namespace Metasound
