Trivial: Fix ensure when promoting In value to input from Trigger on Threshold
- Move ensure to warning as implicit cast won't necessarily induce a bug

#rnx

diff --git a/Engine/Plugins/Runtime/Metasound/Source/MetasoundEditor/Private/MetasoundEditorGraphMemberDefaults.cpp b/Engine/Plugins/Runtime/Metasound/Source/MetasoundEditor/Private/MetasoundEditorGraphMemberDefaults.cpp
--- a/Engine/Plugins/Runtime/Metasound/Source/MetasoundEditor/Private/MetasoundEditorGraphMemberDefaults.cpp
+++ b/Engine/Plugins/Runtime/Metasound/Source/MetasoundEditor/Private/MetasoundEditorGraphMemberDefaults.cpp
@@ -288,7 +288,13 @@ namespace Metasound::Editor
 									bModified = true;
 
 									LiteralType NewValue;
-									ensure(TestLiteral.TryGet(NewValue));
+									if (!TestLiteral.TryGet(NewValue))
+									{
+										UE_LOG(LogMetaSound, Warning, TEXT("Synchronizing Page Default of member '%s' (type '%s') was supplied literal value '%s' with mismatched type"),
+											*Member->GetMemberName().ToString(),
+											*Member->GetLiteral()->GetClass()->GetName(),
+											*TestLiteral.ToString());
+									}
 									Default.Value = LiteralToMember(NewValue);
 								}
 
diff --git a/Engine/Plugins/Runtime/Metasound/Source/MetasoundStandardNodes/Private/MetasoundTriggerOnThresholdNode.cpp b/Engine/Plugins/Runtime/Metasound/Source/MetasoundStandardNodes/Private/MetasoundTriggerOnThresholdNode.cpp
--- a/Engine/Plugins/Runtime/Metasound/Source/MetasoundStandardNodes/Private/MetasoundTriggerOnThresholdNode.cpp
+++ b/Engine/Plugins/Runtime/Metasound/Source/MetasoundStandardNodes/Private/MetasoundTriggerOnThresholdNode.cpp
@@ -92,7 +92,7 @@ namespace Metasound
 
 			static const FVertexInterface Interface(
 				FInputVertexInterface(
-					TInputDataVertex<ValueType>(METASOUND_GET_PARAM_NAME_AND_METADATA(InPin), 0),
+					TInputDataVertex<ValueType>(METASOUND_GET_PARAM_NAME_AND_METADATA(InPin), static_cast<ValueType>(0)),
 					TInputDataVertex<ThresholdType>(METASOUND_GET_PARAM_NAME_AND_METADATA(InThresholdPin), DefaultThreshold),
 					TInputDataVertex<FEnumBufferTriggerType>(METASOUND_GET_PARAM_NAME_AND_METADATA(InTriggerType), (int32)EBufferTriggerType::RisingEdge)
 				),
