- Add more logging to MetaSound cook indicating page removal
- Fix off-by-one when removing default graph via builder API
- Add missing Super::PostLoad that may be causing config values from getting loaded in MetaSoundSettings

#rnx
#tests local cook

diff --git a/Engine/Plugins/Runtime/Metasound/Source/MetasoundEngine/Private/MetasoundDocumentBuilderRegistry.cpp b/Engine/Plugins/Runtime/Metasound/Source/MetasoundEngine/Private/MetasoundDocumentBuilderRegistry.cpp
--- a/Engine/Plugins/Runtime/Metasound/Source/MetasoundEngine/Private/MetasoundDocumentBuilderRegistry.cpp
+++ b/Engine/Plugins/Runtime/Metasound/Source/MetasoundEngine/Private/MetasoundDocumentBuilderRegistry.cpp
@@ -141,13 +141,20 @@ namespace Metasound::Engine
 				return Builder.RemoveGraphPage(InPageID);
 			};
 
+			const int32 NumInitGraphs = Document.RootGraph.GetConstGraphPages().Num();
 			StripPageEntries(ResolvePageIDs, ResolvedTargetScratchIDs, RemoveGraphPage);
+			const int32 NumRemainingGraphs = Document.RootGraph.GetConstGraphPages().Num();
 
-			checkf(!Document.RootGraph.GetConstGraphPages().IsEmpty(),
+			checkf(NumRemainingGraphs > 0,
 				TEXT("Document in MetaSound asset '%s' had all default values "
 					"cooked away leaving it in an invalid state. "
 					"Graph must always have at least one implementation."),
 				*Builder.GetDebugName());
+
+			if (NumInitGraphs > NumRemainingGraphs)
+			{
+				UE_LOG(LogMetaSound, Display, TEXT("Cook removed %i graph page(s) from '%s'"), NumInitGraphs - NumRemainingGraphs, *Builder.GetDebugName());
+			}
 		}
 
 		{ // Strip default input values
@@ -162,11 +169,19 @@ namespace Metasound::Engine
 					return Builder.RemoveGraphInputDefault(GraphInput.Name, InPageID);
 				};
 
+				const int32 NumInitDefaults = GraphInput.GetDefaults().Num();
 				StripPageEntries(ResolvePageIDs, ResolvedTargetScratchIDs, RemoveDefault);
-				checkf(!GraphInput.GetDefaults().IsEmpty(),
+				const int32 NumRemainingDefaults = GraphInput.GetDefaults().Num();
+
+				checkf(NumRemainingDefaults > 0,
 					TEXT("Input '%s' had all default values stripped leaving it in an invalid state. "
 					"Input must always have at least one default value"),
 					*GraphInput.Name.ToString());
+
+				if (NumInitDefaults > NumRemainingDefaults)
+				{
+					UE_LOG(LogMetaSound, Display, TEXT("Cook removed %i default input page value(s) from input '%s'"), NumInitDefaults - NumRemainingDefaults, *GraphInput.Name.ToString());
+				}
 			}
 		}
 
diff --git a/Engine/Plugins/Runtime/Metasound/Source/MetasoundEngine/Private/MetasoundSettings.cpp b/Engine/Plugins/Runtime/Metasound/Source/MetasoundEngine/Private/MetasoundSettings.cpp
--- a/Engine/Plugins/Runtime/Metasound/Source/MetasoundEngine/Private/MetasoundSettings.cpp
+++ b/Engine/Plugins/Runtime/Metasound/Source/MetasoundEngine/Private/MetasoundSettings.cpp
@@ -350,10 +350,24 @@ const FMetaSoundPageSettings& UMetaSoundSettings::GetTargetPageSettings() const
 {
 	const FName TargetPage = TargetPageNameOverride.IsSet() ? *TargetPageNameOverride : TargetPageName;
 
+#if !NO_LOGGING
+	auto WarnIfUninitialized = [this](const FMetaSoundPageSettings& SettingsSet)
+	{
+		if (bWarnAccessBeforeInit)
+		{
+			UE_LOG(LogMetaSound, Display, TEXT("Target Page Settings accessed prior to 'PostLoad' called.  Uninitialized PageSettings '%s' being returned."), *SettingsSet.Name.ToString());
+			bWarnAccessBeforeInit = false;
+		}
+	};
+#endif // !NO_LOGGING
+
 	if (const FMetaSoundPageSettings* TargetSettings = FindPageSettings(TargetPage))
 	{
 		if (TargetSettings->CanTarget.GetValue())
 		{
+#if !NO_LOGGING
+			WarnIfUninitialized(*TargetSettings);
+#endif // !NO_LOGGING
 			return *TargetSettings;
 		}
 	}
@@ -364,10 +378,16 @@ const FMetaSoundPageSettings& UMetaSoundSettings::GetTargetPageSettings() const
 	{
 		if (Setting.CanTarget.GetValue())
 		{
+#if !NO_LOGGING
+			WarnIfUninitialized(Setting);
+#endif // !NO_LOGGING
 			return Setting;
 		}
 	}
 
+#if !NO_LOGGING
+	WarnIfUninitialized(DefaultPageSettings);
+#endif // !NO_LOGGING
 	return DefaultPageSettings;
 }
 
@@ -399,16 +419,59 @@ void UMetaSoundSettings::PostEditChangeProperty(FPropertyChangedEvent& PropertyC
 
 	DenyListCacheChangeID++;
 }
+#endif // WITH_EDITOR
 
 void UMetaSoundSettings::PostInitProperties()
 {
 	Super::PostInitProperties();
 
+#if WITH_EDITOR
 	constexpr bool bNotifyDefaultRenamed = false;
 	ConformPageSettings(bNotifyDefaultRenamed);
-}
 #endif // WITH_EDITOR
 
+#if !NO_LOGGING
+	bWarnAccessBeforeInit = false;
+#endif // !NO_LOGGING
+
+#if WITH_EDITORONLY_DATA
+	UE_LOG(LogMetaSound, Display, TEXT("MetaSound Page Target Initialized to '%s'"), *GetTargetPageSettings().Name.ToString());
+#else // !WITH_EDITORONLY_DATA
+	if (const FMetaSoundPageSettings* Page = FindPageSettings(TargetPageName))
+	{
+		if (Page->CanTarget.GetValue())
+		{
+			UE_LOG(LogMetaSound, Display, TEXT("MetaSound Page Target Initialized to '%s'"), *GetTargetPageSettings().Name.ToString());
+			return;
+		}
+	}
+
+	UE_LOG(LogMetaSound, Warning, TEXT("TargetPageName on load in MetaSound settings did not correspond to a valid targetable page."));
+	if (DefaultPageSettings.CanTarget.GetValue())
+	{
+		UE_LOG(LogMetaSound, Warning, TEXT("Setting target to 'Default' page settings."));
+		TargetPageNameOverride = DefaultPageSettings.Name;
+		return;
+	}
+	else
+	{
+		for (const FMetaSoundPageSettings& Page : PageSettings)
+		{
+			if (Page.CanTarget.GetValue())
+			{
+				UE_LOG(LogMetaSound, Warning, TEXT("Setting target to '%s' page settings."), *Page.Name.ToString());
+				TargetPageNameOverride = Page.Name;
+				return;
+			}
+		}
+
+		UE_LOG(LogMetaSound, Warning, TEXT("Setting target to 'Default' page settings & forcing default as targetable."));
+		DefaultPageSettings.CanTarget.Default = true;
+		TargetPageNameOverride = DefaultPageSettings.Name;
+	}
+#endif // !WITH_EDITORONLY_DATA
+}
+
 bool UMetaSoundSettings::SetTargetPage(FName PageName)
 {
 	if (const FMetaSoundPageSettings* PageSetting = FindPageSettings(PageName))
@@ -497,41 +560,4 @@ void UMetaSoundSettings::IteratePageSettings(TFunctionRef<void(const FMetaSoundP
 		}
 	}
 }
-
-void UMetaSoundSettings::PostLoad()
-{
-#if !WITH_EDITORONLY_DATA
-	if (const FMetaSoundPageSettings* Page = FindPageSettings(TargetPageName))
-	{
-		if (Page->CanTarget.GetValue())
-		{
-			return;
-		}
-	}
-
-	UE_LOG(LogMetaSound, Warning, TEXT("TargetPageName on load in MetaSound settings did not correspond to a valid targetable page."));
-	if (DefaultPageSettings.CanTarget.GetValue())
-	{
-		UE_LOG(LogMetaSound, Warning, TEXT("Setting target to 'Default' page settings."));
-		TargetPageNameOverride = DefaultPageSettings.Name;
-		return;
-	}
-	else
-	{
-		for (const FMetaSoundPageSettings& Page : PageSettings)
-		{
-			if (Page.CanTarget.GetValue())
-			{
-				UE_LOG(LogMetaSound, Warning, TEXT("Setting target to '%s' page settings."), *Page.Name.ToString());
-				TargetPageNameOverride = Page.Name;
-				return;
-			}
-		}
-
-		UE_LOG(LogMetaSound, Warning, TEXT("Setting target to 'Default' page settings & forcing default as targetable."));
-		DefaultPageSettings.CanTarget.Default = true;
-		TargetPageNameOverride = DefaultPageSettings.Name;
-	}
-#endif // !WITH_EDITORONLY_DATA
-}
 #undef LOCTEXT_NAMESPACE // MetaSound
diff --git a/Engine/Plugins/Runtime/Metasound/Source/MetasoundEngine/Public/MetasoundSettings.h b/Engine/Plugins/Runtime/Metasound/Source/MetasoundEngine/Public/MetasoundSettings.h
--- a/Engine/Plugins/Runtime/Metasound/Source/MetasoundEngine/Public/MetasoundSettings.h
+++ b/Engine/Plugins/Runtime/Metasound/Source/MetasoundEngine/Public/MetasoundSettings.h
@@ -250,8 +250,6 @@ public:
 	// optionally set to reverse, iterates in reverse.
 	void IteratePageSettings(TFunctionRef<void(const FMetaSoundPageSettings&)> Iter, bool bReverse = false) const;
 
-	virtual void PostLoad() override;
-
 	// Sets the target page to the given name. Returns true if associated page settings were found
 	// and target set, false if not found and not set.
 	bool SetTargetPage(FName PageName);
@@ -283,6 +281,11 @@ private:
 
 	virtual void PostEditChangeChainProperty(FPropertyChangedChainEvent& PropertyChangedEvent) override;
 	virtual void PostEditChangeProperty(FPropertyChangedEvent& PropertyChangedEvent) override;
-	virtual void PostInitProperties() override;
 #endif // WITH_EDITOR
+
+	virtual void PostInitProperties() override;
+
+#if !NO_LOGGING
+	mutable bool bWarnAccessBeforeInit = true;
+#endif // !NO_LOGGING
 };
diff --git a/Engine/Plugins/Runtime/Metasound/Source/MetasoundFrontend/Private/MetasoundFrontendDocument.cpp b/Engine/Plugins/Runtime/Metasound/Source/MetasoundFrontend/Private/MetasoundFrontendDocument.cpp
--- a/Engine/Plugins/Runtime/Metasound/Source/MetasoundFrontend/Private/MetasoundFrontendDocument.cpp
+++ b/Engine/Plugins/Runtime/Metasound/Source/MetasoundFrontend/Private/MetasoundFrontendDocument.cpp
@@ -1118,7 +1118,7 @@ bool FMetasoundFrontendGraphClass::RemoveGraphPage(const FGuid& InPageID, FGuid*
 				{
 					*OutAdjacentPageID = PagedGraphs[Index - 1].PageID;
 				}
-				else if (Index < PagedGraphs.Num() - 1)
+				else if (Index < PagedGraphs.Num())
 				{
 					*OutAdjacentPageID = PagedGraphs[0].PageID;
 				}
