Remove cvar requirement to see MetaSound asset page editor bits
#rnx

diff --git a/Engine/Plugins/Runtime/Metasound/Source/MetasoundEditor/Private/MetasoundEditor.cpp b/Engine/Plugins/Runtime/Metasound/Source/MetasoundEditor/Private/MetasoundEditor.cpp
--- a/Engine/Plugins/Runtime/Metasound/Source/MetasoundEditor/Private/MetasoundEditor.cpp
+++ b/Engine/Plugins/Runtime/Metasound/Source/MetasoundEditor/Private/MetasoundEditor.cpp
@@ -106,18 +106,6 @@ namespace Metasound
 {
 	namespace Editor
 	{
-		namespace AssetEditorPrivate
-		{
-			int32 EnablePageEditor = 0;
-
-			FAutoConsoleVariableRef CVarMetaSoundEnablePageEditor(
-				TEXT("au.MetaSound.Experimental.EnablePagesEditor"),
-				EnablePageEditor,
-				TEXT("Enables Page Editor in MetaSound Asset Editor.\n")
-				TEXT("Default: 0"),
-				ECVF_Default);
-		} // namespace AssetEditorPrivate
-
 		bool IsPreviewingMetaSound(const UObject& InMetaSound)
 		{
 			using namespace Engine;
@@ -166,11 +154,6 @@ namespace Metasound
 
 		bool PageEditorEnabled(const FMetaSoundFrontendDocumentBuilder& Builder, bool bHasProjectPageValues, bool bPresetCanEditPageValues)
 		{
-			if (!AssetEditorPrivate::EnablePageEditor)
-			{
-				return false;
-			}
-
 			const UMetaSoundSettings* Settings = GetDefault<UMetaSoundSettings>();
 			if (!Settings)
 			{
@@ -927,11 +910,10 @@ namespace Metasound
 
 			InTabManager->RegisterTabSpawner(TabNamesPrivate::GraphCanvas, FOnSpawnTab::CreateLambda(
 				[
-					bShowPageTab = ShowPageGraphDetails(),
 					InPageStatsWidget = PageStatsWidget,
 					InMetasoundGraphEditor = MetasoundGraphEditor,
 					InRenderStatsWidget = RenderStatsWidget
-				](const FSpawnTabArgs& Args)
+				] (const FSpawnTabArgs& Args)
 			{
 				TSharedRef<SDockTab> SpawnedTab = SNew(SDockTab).Label(LOCTEXT("MetasoundGraphCanvasTitle", "MetaSound Graph"));
 
@@ -947,7 +929,7 @@ namespace Metasound
 					]
 					.Padding(5.0f, 5.0f);
 
-				if (bShowPageTab && InPageStatsWidget.IsValid())
+				if (InPageStatsWidget.IsValid())
 				{
 					TSharedRef<SVerticalBox> GraphStatsWidget = SNew(SVerticalBox)
 						+ SVerticalBox::Slot()
@@ -1056,12 +1038,7 @@ namespace Metasound
 			InTabManager->UnregisterTabSpawner(TabNamesPrivate::GraphCanvas);
 			InTabManager->UnregisterTabSpawner(TabNamesPrivate::Details);
 			InTabManager->UnregisterTabSpawner(TabNamesPrivate::Members);
-
-			if (ShowPageGraphDetails())
-			{
-				InTabManager->UnregisterTabSpawner(TabNamesPrivate::Pages);
-			}
-
+			InTabManager->UnregisterTabSpawner(TabNamesPrivate::Pages);
 			InTabManager->UnregisterTabSpawner(TabNamesPrivate::Interfaces);
 			InTabManager->UnregisterTabSpawner(TabNamesPrivate::Find);
 		}
@@ -1689,24 +1666,26 @@ namespace Metasound
 				InterfacesDetails->HideFilterArea(true);
 			}
 
-			if (ShowPageGraphDetails())
+			PagesDetails = PropertyModule.CreateDetailView(Args);
+			if (PagesDetails.IsValid())
 			{
-				PagesDetails = PropertyModule.CreateDetailView(Args);
-				if (PagesDetails.IsValid())
-				{
-					PagesView = TStrongObjectPtr(NewObject<UMetasoundPagesView>());
-					PagesView->SetMetasound(&MetaSound);
-					const TArray<UObject*> PagesViewObj{ PagesView.Get() };
+				PagesView = TStrongObjectPtr(NewObject<UMetasoundPagesView>());
+				PagesView->SetMetasound(&MetaSound);
+				const TArray<UObject*> PagesViewObj{ PagesView.Get() };
 
-					PagesDetails->SetObjects(PagesViewObj);
-					PagesDetails->HideFilterArea(true);
+				PagesDetails->SetObjects(PagesViewObj);
+				PagesDetails->HideFilterArea(true);
 
-					TAttribute<bool> EnabledAttr = TAttribute<bool>::Create([this]()
-					{
-						return ShowPageGraphDetails();
-					});
-					PagesDetails->SetEnabled(EnabledAttr);
-				}
+				TAttribute<bool> EnabledAttr = TAttribute<bool>::Create([this]()
+				{
+					return ShowPageGraphDetails();
+				});
+				TAttribute<EVisibility> VisibilityAttr = TAttribute<EVisibility>::Create([this]()
+				{
+					return ShowPageGraphDetails() ? EVisibility::Visible : EVisibility::Hidden;
+				});
+				PagesDetails->SetEnabled(EnabledAttr);
+				PagesDetails->SetVisibility(VisibilityAttr);
 			}
 
 			Palette = SNew(SMetasoundPalette);
@@ -2245,7 +2224,10 @@ namespace Metasound
 				})
 			);
 
-			if (AssetEditorPrivate::EnablePageEditor)
+			constexpr bool bHasProjectPageValues = true;
+			constexpr bool bPresetCanEditPageValues = true;
+			const bool bShowAuditionSettings = PageEditorEnabled(Builder->GetConstBuilder(), bHasProjectPageValues, bPresetCanEditPageValues);
+			if (bShowAuditionSettings)
 			{
 				if (UToolMenu* AssetToolbar = UToolMenus::Get()->ExtendMenu(GetToolMenuToolbarName()))
 				{
@@ -5082,7 +5064,8 @@ namespace Metasound
 				const FSlateColor* Color = nullptr;
 				const FMetaSoundPageSettings* GraphPageSettings = nullptr;
 				const FMetaSoundPageSettings* AuditionPageSettings = nullptr;
-				if (Builder.IsValid())
+
+				if (Builder.IsValid() && ShowPageGraphDetails())
 				{
 					if (const UMetasoundEditorSettings* EditorSettings = GetDefault<UMetasoundEditorSettings>())
 					{
@@ -5143,7 +5126,7 @@ namespace Metasound
 				TAttribute<EVisibility> ExecVisibility = TAttribute<EVisibility>::CreateLambda([this, InPageID]()
 				{
 					using namespace Engine;
-					if (Builder.IsValid())
+					if (Builder.IsValid() && ShowPageGraphDetails())
 					{
 						const FMetaSoundFrontendDocumentBuilder& DocBuilder = Builder->GetConstBuilder();
 						const bool bIsPreviewing = IsPreviewingPageGraph(DocBuilder, InPageID);
