Fix bad merge from last week causing false ensure on changing a MetaSound output's data type
#rnx

diff --git a/Engine/Plugins/Runtime/Metasound/Source/MetasoundEngine/Public/MetasoundBuilderBase.cpp b/Engine/Plugins/Runtime/Metasound/Source/MetasoundEngine/Public/MetasoundBuilderBase.cpp
--- a/Engine/Plugins/Runtime/Metasound/Source/MetasoundEngine/Public/MetasoundBuilderBase.cpp
+++ b/Engine/Plugins/Runtime/Metasound/Source/MetasoundEngine/Public/MetasoundBuilderBase.cpp
@@ -1047,6 +1047,7 @@ void UMetaSoundBuilderBase::SetGraphOutputAccessType(FName OutputName, EMetasoun
 void UMetaSoundBuilderBase::SetGraphOutputDataType(FName OutputName, FName DataType, EMetaSoundBuilderResult& OutResult)
 {
 	const bool bSet = Builder.SetGraphOutputDataType(OutputName, DataType);
+	OutResult = bSet ? EMetaSoundBuilderResult::Succeeded : EMetaSoundBuilderResult::Failed;
 }
 
 void UMetaSoundBuilderBase::SetGraphOutputName(FName OutputName, FName NewName, EMetaSoundBuilderResult& OutResult)
