[MetaSounds] - Duplicated Inputs have errors on nodes when used in composition
&&
[MetaSounds] - Inputs made using Promote All Inputs have errors on nodes when used in composition

#rnx

diff --git a/Engine/Plugins/Runtime/Metasound/Source/MetasoundEditor/Private/MetasoundEditor.cpp b/Engine/Plugins/Runtime/Metasound/Source/MetasoundEditor/Private/MetasoundEditor.cpp
--- a/Engine/Plugins/Runtime/Metasound/Source/MetasoundEditor/Private/MetasoundEditor.cpp
+++ b/Engine/Plugins/Runtime/Metasound/Source/MetasoundEditor/Private/MetasoundEditor.cpp
@@ -3793,6 +3793,8 @@ namespace Metasound
 				}
 			}
 
+			FGraphBuilder::RegisterGraphWithFrontend(*MetaSound, true);
+
 			if (GraphMembersMenu.IsValid())
 			{
 				GraphMembersMenu->RefreshAllActions(true);
@@ -5354,6 +5356,8 @@ namespace Metasound
 					}
 				}
 			}
+
+			FGraphBuilder::RegisterGraphWithFrontend(ParentMetasound, true);
 		}
 
 		bool FEditor::CanPromoteAllToCommonInputs()
@@ -5462,6 +5466,8 @@ namespace Metasound
 					}
 				}
 			}
+
+			FGraphBuilder::RegisterGraphWithFrontend(ParentMetasound, true);
 		}
 	}
 }
