Fix MetaSound Build BP call to not take parent object as parameter (Transient asset MetaSounds are now required to be TopLevel in the temp package)

#tests local BP
#rnx

diff --git a/Engine/Plugins/Runtime/Metasound/Source/MetasoundEngine/Private/MetasoundBuilderSubsystem.cpp b/Engine/Plugins/Runtime/Metasound/Source/MetasoundEngine/Private/MetasoundBuilderSubsystem.cpp
--- a/Engine/Plugins/Runtime/Metasound/Source/MetasoundEngine/Private/MetasoundBuilderSubsystem.cpp
+++ b/Engine/Plugins/Runtime/Metasound/Source/MetasoundEngine/Private/MetasoundBuilderSubsystem.cpp
@@ -96,10 +96,31 @@ namespace Metasound::Engine
 	} // namespace BuilderSubsystemPrivate
 } // namespace Metasound::Engine
 
+void UMetaSoundPatchBuilder::BuildAndOverwriteMetaSound(TScriptInterface<IMetaSoundDocumentInterface> ExistingMetaSound, bool bForceUniqueClassName)
+{
+	if (!ExistingMetaSound || !ExistingMetaSound.GetObject()->IsAsset())
+	{
+		UE_LOG(LogMetaSound, Error, TEXT("Failed to build and overwrite MetaSound: No existing MetaSound supplied or is serialized asset "
+			"(use 'BuildNewMetaSound' to create a new, transient MetaSound. Overwriting serialized asset is only supported at edit time via "
+			"UMetaSoundEditorSubsystem::BuildToAsset."));
+		return;
+	}
+
+	FMetaSoundBuilderOptions Options;
+	Options.ExistingMetaSound = ExistingMetaSound;
+	Options.bForceUniqueClassName = bForceUniqueClassName;
+	constexpr UObject* Parent = nullptr;
 
-TScriptInterface<IMetaSoundDocumentInterface> UMetaSoundPatchBuilder::Build(UObject* Parent, const FMetaSoundBuilderOptions& InBuilderOptions) const
+	BuildInternal<UMetaSoundPatch>(Parent, Options);
+}
+
+TScriptInterface<IMetaSoundDocumentInterface> UMetaSoundPatchBuilder::BuildNewMetaSound(FName NameBase) const
 {
-	return &BuildInternal<UMetaSoundPatch>(Parent, InBuilderOptions);
+	FMetaSoundBuilderOptions Options;
+	Options.Name = NameBase;
+	constexpr UObject* Parent = nullptr;
+
+	return &BuildInternal<UMetaSoundPatch>(Parent, Options);
 }
 
 const UClass& UMetaSoundPatchBuilder::GetBaseMetaSoundUClass() const
@@ -249,9 +270,31 @@ bool UMetaSoundSourceBuilder::ExecuteAuditionableTransaction(FAuditionableTransa
 	return false;
 }
 
-TScriptInterface<IMetaSoundDocumentInterface> UMetaSoundSourceBuilder::Build(UObject* Parent, const FMetaSoundBuilderOptions& InBuilderOptions) const
+void UMetaSoundSourceBuilder::BuildAndOverwriteMetaSound(TScriptInterface<IMetaSoundDocumentInterface> ExistingMetaSound, bool bForceUniqueClassName)
 {
-	return &BuildInternal<UMetaSoundSource>(Parent, InBuilderOptions);
+	if (!ExistingMetaSound || !ExistingMetaSound.GetObject()->IsAsset())
+	{
+		UE_LOG(LogMetaSound, Error, TEXT("Failed to build and overwrite MetaSound: No existing MetaSound supplied or is serialized asset "
+			"(use 'BuildNewMetaSound' to create a new, transient MetaSound. Overwriting serialized asset is only supported at edit time via "
+			"UMetaSoundEditorSubsystem::BuildToAsset."));
+		return;
+	}
+
+	FMetaSoundBuilderOptions Options;
+	Options.ExistingMetaSound = ExistingMetaSound;
+	Options.bForceUniqueClassName = bForceUniqueClassName;
+	constexpr UObject* Parent = nullptr;
+
+	BuildInternal<UMetaSoundSource>(Parent, Options);
+}
+
+TScriptInterface<IMetaSoundDocumentInterface> UMetaSoundSourceBuilder::BuildNewMetaSound(FName NameBase) const
+{
+	FMetaSoundBuilderOptions Options;
+	Options.Name = NameBase;
+	constexpr UObject* Parent = nullptr;
+
+	return &BuildInternal<UMetaSoundSource>(Parent, Options);
 }
 
 const Metasound::Engine::FOutputAudioFormatInfoPair* UMetaSoundSourceBuilder::FindOutputAudioFormatInfo() const
diff --git a/Engine/Plugins/Runtime/Metasound/Source/MetasoundEngine/Public/MetasoundBuilderBase.h b/Engine/Plugins/Runtime/Metasound/Source/MetasoundEngine/Public/MetasoundBuilderBase.h
--- a/Engine/Plugins/Runtime/Metasound/Source/MetasoundEngine/Public/MetasoundBuilderBase.h
+++ b/Engine/Plugins/Runtime/Metasound/Source/MetasoundEngine/Public/MetasoundBuilderBase.h
@@ -87,12 +87,13 @@ struct METASOUNDENGINE_API FMetaSoundBuilderOptions
 	GENERATED_BODY()
 
 	// Name of generated object. If object already exists, used as the base name to ensure new object is unique.
+	// If left 'None', creates unique name.
 	UPROPERTY(EditAnywhere, BlueprintReadWrite, Category = "MetaSound|Builder")
 	FName Name;
 
 	// If the resulting MetaSound is building over an existing document, a unique class name will be generated,
 	// invalidating any referencing MetaSounds and registering the MetaSound as a new entry in the Frontend. If
-	// building a new document, option is ignored.
+	// building a new document, option is ignored (new document always generates a unique class name).
 	UPROPERTY(EditAnywhere, BlueprintReadWrite, Category = "MetaSound|Builder", meta = (AdvancedDisplay))
 	bool bForceUniqueClassName = false;
 
@@ -435,7 +436,25 @@ public:
 	UE_DEPRECATED(5.5, "Call directly on Frontend builder using 'GetBuilder'")
 	void UpdateDependencyClassNames(const TMap<FMetasoundFrontendClassName, FMetasoundFrontendClassName>& OldToNewReferencedClassNames);
 
-	virtual TScriptInterface<IMetaSoundDocumentInterface> Build(UObject* Parent, const FMetaSoundBuilderOptions& Options) const PURE_VIRTUAL(UMetaSoundBuilderBase::Build, return { }; );
+	UFUNCTION(BlueprintCallable, Category = "Audio|MetaSound|Builder", meta = (
+		WorldContext = "Parent",
+		DeprecatedFunction,
+		DeprecationMessage = "MetaSounds must now be top-level objects within a package.  Use 'BuildAndOverwriteMetaSound' to copy this builder's implementation over an existing transient MetaSound or 'BuildNewMetaSound' to generate a new, transient MetaSound. Use 'MetaSoundEditorSubsystem::BuildToAsset' to build a copied MetaSound as a new, serialized asset (editor only)."))
+	virtual UPARAM(DisplayName = "MetaSound") TScriptInterface<IMetaSoundDocumentInterface> Build(UObject* Parent, const FMetaSoundBuilderOptions& Options) const { return BuildNewMetaSound(Options.Name); }
+
+	// Copies a transient MetaSound with the provided builder options, copying the underlying MetaSound
+	// managed by this builder and registering it with the MetaSound Node Registry as a unique name.
+	// If 'Force Unique Class Name' is true, registers MetaSound as a new class in the registry, potentially
+	// invalidating existing references in other MetaSounds.
+	UFUNCTION(BlueprintCallable, Category = "Audio|MetaSound|Builder", meta = (DisplayName = "Build And Overwrite MetaSound", AdvancedDisplay = "1"))
+	virtual void BuildAndOverwriteMetaSound(UPARAM(DisplayName = "Existing MetaSound") TScriptInterface<IMetaSoundDocumentInterface> ExistingMetaSound, bool bForceUniqueClassName = false) PURE_VIRTUAL(UMetaSoundBuilderBase::BuildAndOverwriteMetaSound, ;);
+
+	// Builds a transient MetaSound with the provided builder options, copying the underlying MetaSound
+	// managed by this builder and registering it with the MetaSound Node Registry as a unique class. If
+	// existing MetaSound exists with the provided NameBase, will make object with unique name with the given
+	// NameBase as prefix.
+	UFUNCTION(BlueprintCallable, Category = "Audio|MetaSound|Builder", meta = (DisplayName = "Build New MetaSound"))
+	virtual UPARAM(DisplayName = "New MetaSound") TScriptInterface<IMetaSoundDocumentInterface> BuildNewMetaSound(FName NameBase) const PURE_VIRTUAL(UMetaSoundBuilderBase::BuildNewMetaSound, return { }; );
 
 	virtual bool ConformObjectToDocument();
 
diff --git a/Engine/Plugins/Runtime/Metasound/Source/MetasoundEngine/Public/MetasoundBuilderSubsystem.h b/Engine/Plugins/Runtime/Metasound/Source/MetasoundEngine/Public/MetasoundBuilderSubsystem.h
--- a/Engine/Plugins/Runtime/Metasound/Source/MetasoundEngine/Public/MetasoundBuilderSubsystem.h
+++ b/Engine/Plugins/Runtime/Metasound/Source/MetasoundEngine/Public/MetasoundBuilderSubsystem.h
@@ -53,9 +53,8 @@ class METASOUNDENGINE_API UMetaSoundPatchBuilder : public UMetaSoundBuilderBase
 	GENERATED_BODY()
 
 public:
-	UFUNCTION(BlueprintCallable, Category = "Audio|MetaSound|Builder", meta = (WorldContext = "Parent"))
-	virtual UPARAM(DisplayName = "MetaSound") TScriptInterface<IMetaSoundDocumentInterface> Build(UObject* Parent, const FMetaSoundBuilderOptions& Options) const override;
-
+	virtual void BuildAndOverwriteMetaSound(TScriptInterface<IMetaSoundDocumentInterface> ExistingMetaSound, bool bForceUniqueClassName) override;
+	virtual TScriptInterface<IMetaSoundDocumentInterface> BuildNewMetaSound(FName NameBase) const override;
 	virtual const UClass& GetBaseMetaSoundUClass() const override;
 
 protected:
@@ -75,9 +74,6 @@ public:
 	UFUNCTION(BlueprintCallable, Category = "Audio|MetaSound|Builder", meta = (WorldContext = "Parent", AdvancedDisplay = "2"))
 	void Audition(UObject* Parent, UAudioComponent* AudioComponent, FOnCreateAuditionGeneratorHandleDelegate OnCreateGenerator, bool bLiveUpdatesEnabled = false);
 
-	UFUNCTION(BlueprintCallable, Category = "Audio|MetaSound|Builder", meta = (WorldContext = "Parent"))
-	virtual UPARAM(DisplayName = "MetaSound") TScriptInterface<IMetaSoundDocumentInterface> Build(UObject* Parent, const FMetaSoundBuilderOptions& Options) const override;
-
 	// Returns whether or not live updates are both globally enabled (via cvar) and are enabled on this builder's last built sound, which may or may not still be playing.
 	UFUNCTION(BlueprintCallable, Category = "Audio|MetaSound|Builder")
 	bool GetLiveUpdatesEnabled() const;
@@ -96,6 +92,8 @@ public:
 
 	const Metasound::Engine::FOutputAudioFormatInfoPair* FindOutputAudioFormatInfo() const;
 
+	virtual void BuildAndOverwriteMetaSound(TScriptInterface<IMetaSoundDocumentInterface> ExistingMetaSound, bool bForceUniqueClassName) override;
+	virtual TScriptInterface<IMetaSoundDocumentInterface> BuildNewMetaSound(FName NameBase) const override;
 	virtual const UClass& GetBaseMetaSoundUClass() const override;
 
 #if WITH_EDITORONLY_DATA
