- Fix for ensure when copying and pasting a variable from a different page on the same asset
- Fix for variable names not updating when swapping pages in the members detail panel
- Hide audition options if project has no page settings

#rnx

diff --git a/Engine/Plugins/Runtime/Metasound/Source/MetasoundEditor/Private/MetasoundEditor.cpp b/Engine/Plugins/Runtime/Metasound/Source/MetasoundEditor/Private/MetasoundEditor.cpp
--- a/Engine/Plugins/Runtime/Metasound/Source/MetasoundEditor/Private/MetasoundEditor.cpp
+++ b/Engine/Plugins/Runtime/Metasound/Source/MetasoundEditor/Private/MetasoundEditor.cpp
@@ -85,6 +85,7 @@
 #include "ToolMenuEntry.h"
 #include "ToolMenus.h"
 #include "UObject/ScriptInterface.h"
+#include "UObject/WeakObjectPtrTemplates.h"
 #include "Widgets/Docking/SDockTab.h"
 #include "Widgets/Images/SImage.h"
 #include "Widgets/Input/SButton.h"
@@ -1889,7 +1890,21 @@ namespace Metasound
 			constexpr bool bShouldCloseWindowAfterMenuSelection = false;
 			FMenuBuilder MenuBuilder(bShouldCloseWindowAfterMenuSelection, Commands);
 			CreateAuditionPageSubMenuOptions(MenuBuilder);
-			return MenuBuilder.MakeWidget();
+			TSharedRef<SWidget> MenuWidget = MenuBuilder.MakeWidget();
+			TWeakObjectPtr<UMetaSoundBuilderBase> WeakBuilderPtr(Builder.Get());
+			MenuWidget->SetVisibility(TAttribute<EVisibility>::Create([WeakBuilderPtr]()
+			{
+				if (const TStrongObjectPtr<UMetaSoundBuilderBase> BuilderPtr = WeakBuilderPtr.Pin())
+				{
+					constexpr bool bHasProjectPageValues = true;
+					constexpr bool bPresetCanEditPageValues = true;
+					const bool bIsEnabled = PageEditorEnabled(BuilderPtr->GetConstBuilder(), bHasProjectPageValues, bPresetCanEditPageValues);
+					return bIsEnabled ? EVisibility::Visible : EVisibility::Collapsed;
+				}
+
+				return EVisibility::Collapsed;
+			}));
+			return MenuWidget;
 		}
 
 		void FEditor::CreateAuditionPageSubMenuOptions(FMenuBuilder& MenuBuilder)
@@ -5118,6 +5133,11 @@ namespace Metasound
 				ParentPtr->UpdatePageInfo(false);
 				ParentPtr->bRefreshGraph = true;
 				ParentPtr->RefreshExecVisibility(Args.PageID);
+
+				if (ParentPtr->GraphMembersMenu.IsValid())
+				{
+					ParentPtr->GraphMembersMenu->RefreshAllActions(true);
+				}
 			}
 		}
 
diff --git a/Engine/Plugins/Runtime/Metasound/Source/MetasoundEditor/Private/MetasoundEditorDocumentClipboardUtils.cpp b/Engine/Plugins/Runtime/Metasound/Source/MetasoundEditor/Private/MetasoundEditorDocumentClipboardUtils.cpp
--- a/Engine/Plugins/Runtime/Metasound/Source/MetasoundEditor/Private/MetasoundEditorDocumentClipboardUtils.cpp
+++ b/Engine/Plugins/Runtime/Metasound/Source/MetasoundEditor/Private/MetasoundEditorDocumentClipboardUtils.cpp
@@ -35,7 +35,7 @@ namespace Metasound::Editor
 
 			InputNode->CreateNewGuid();
 			TObjectPtr<UMetasoundEditorGraphInput>& Input = InputNode->Input;
-			if (!Input || Input->GetOwningGraph() != &Graph)
+			if (!Input || !Graph.ContainsInput(*Input))
 			{
 				Input = nullptr;
 
@@ -181,7 +181,7 @@ namespace Metasound::Editor
 
 			// Can only have one setter node
 			TObjectPtr<UMetasoundEditorGraphVariable>& Variable = VariableNode->Variable;
-			if (!Variable || Variable->GetOwningGraph() != &Graph)
+			if (!Variable || !Graph.ContainsVariable(*Variable))
 			{
 				const FMetasoundEditorGraphMemberNodeBreadcrumb& Breadcrumb = VariableNode->Breadcrumb;
 				const FName BaseName = Breadcrumb.MemberName;
