Fix various intergraph MetaSound member copy paste issues
- Fix not all fields being copied
- Add output node copying between graphs

#rnx

diff --git a/Engine/Plugins/Runtime/Metasound/Source/MetasoundEditor/Private/MetasoundEditor.cpp b/Engine/Plugins/Runtime/Metasound/Source/MetasoundEditor/Private/MetasoundEditor.cpp
--- a/Engine/Plugins/Runtime/Metasound/Source/MetasoundEditor/Private/MetasoundEditor.cpp
+++ b/Engine/Plugins/Runtime/Metasound/Source/MetasoundEditor/Private/MetasoundEditor.cpp
@@ -1592,6 +1592,16 @@ namespace Metasound
 
 			MetasoundGraphEditor->AddNotification(Info, false /* bSuccess */);
 		}
+		
+		void FEditor::NotifyNodePasteFailure_MultipleOutputs()
+		{
+			FNotificationInfo Info(LOCTEXT("NodePasteFailed_MultipleOutputs", "Node(s) not pasted: Only one output node possible per graph."));
+			Info.bFireAndForget = true;
+			Info.bUseSuccessFailIcons = false;
+			Info.ExpireDuration = 5.0f;
+
+			MetasoundGraphEditor->AddNotification(Info, false /* bSuccess */);
+		}
 
 		void FEditor::NotifyNodePasteFailure_ReferenceLoop()
 		{
@@ -3373,7 +3383,7 @@ namespace Metasound
 			for (FGraphPanelSelectionSet::TConstIterator SelectedIter(SelectedNodes); SelectedIter; ++SelectedIter)
 			{
 				UEdGraphNode* Node = Cast<UEdGraphNode>(*SelectedIter);
-				if (!Node || !Node->CanDuplicateNode())
+				if (!Node)
 				{
 					return false;
 				}
@@ -3498,6 +3508,22 @@ namespace Metasound
 			FDocumentPasteNotifications Notifications;
 			TArray<UEdGraphNode*> PastedNodes = FDocumentClipboardUtils::PasteClipboardString(InTransactionText, NodeTextToPaste, Location, *GetMetasoundObject(), Notifications);
 
+			// Paste notifications 
+			if (Notifications.bPastedNodesCreateLoop)
+			{
+				NotifyNodePasteFailure_ReferenceLoop();
+			}
+
+			if (Notifications.bPastedNodesAddMultipleVariableSetters)
+			{
+				NotifyNodePasteFailure_MultipleVariableSetters();
+			}
+
+			if (Notifications.bPastedNodesAddMultipleOutputNodes)
+			{
+				NotifyNodePasteFailure_MultipleOutputs();
+			}
+
 			// Clear the selection set (newly pasted stuff will be selected)
 			if (!PastedNodes.IsEmpty())
 			{
@@ -3509,16 +3535,6 @@ namespace Metasound
 					MetasoundGraphEditor->SetNodeSelection(GraphNode, true);
 				}
 
-				if (Notifications.bPastedNodesCreateLoop)
-				{
-					NotifyNodePasteFailure_ReferenceLoop();
-				}
-
-				if (Notifications.bPastedNodesAddMultipleVariableSetters)
-				{
-					NotifyNodePasteFailure_MultipleVariableSetters();
-				}
-
 				MetasoundGraphEditor->NotifyGraphChanged();
 			}
 
diff --git a/Engine/Plugins/Runtime/Metasound/Source/MetasoundEditor/Private/MetasoundEditor.h b/Engine/Plugins/Runtime/Metasound/Source/MetasoundEditor/Private/MetasoundEditor.h
--- a/Engine/Plugins/Runtime/Metasound/Source/MetasoundEditor/Private/MetasoundEditor.h
+++ b/Engine/Plugins/Runtime/Metasound/Source/MetasoundEditor/Private/MetasoundEditor.h
@@ -521,6 +521,7 @@ namespace Metasound::Editor
 		void NotifyDocumentVersioned();
 		void NotifyNodePasteFailure_MultipleVariableSetters();
 		void NotifyNodePasteFailure_ReferenceLoop();
+		void NotifyNodePasteFailure_MultipleOutputs();
 
 		TUniquePtr<FGraphConnectionManager> RebuildConnectionManager(UAudioComponent* PreviewComponent = nullptr) const;
 
diff --git a/Engine/Plugins/Runtime/Metasound/Source/MetasoundEditor/Private/MetasoundEditorDocumentClipboardUtils.cpp b/Engine/Plugins/Runtime/Metasound/Source/MetasoundEditor/Private/MetasoundEditorDocumentClipboardUtils.cpp
--- a/Engine/Plugins/Runtime/Metasound/Source/MetasoundEditor/Private/MetasoundEditorDocumentClipboardUtils.cpp
+++ b/Engine/Plugins/Runtime/Metasound/Source/MetasoundEditor/Private/MetasoundEditorDocumentClipboardUtils.cpp
@@ -9,6 +9,7 @@
 #include "HAL/PlatformApplicationMisc.h"
 #include "Logging/TokenizedMessage.h"
 #include "MetasoundDocumentBuilderRegistry.h"
+#include "MetasoundEditorSubsystem.h"
 #include "MetasoundFrontendSearchEngine.h"
 #include "NodeTemplates/MetasoundFrontendNodeTemplateInput.h"
 #include "ScopedTransaction.h"
@@ -69,6 +70,7 @@ namespace Metasound::Editor
 						{
 							FCreateNodeVertexParams VertexParams;
 							VertexParams.DataType = Breadcrumb.DataType;
+							VertexParams.AccessType = Breadcrumb.AccessType;
 
 							TArray<FMetasoundFrontendClassInputDefault> InputDefaults;
 							Algo::Transform(Breadcrumb.DefaultLiterals, InputDefaults, [](const TPair<FGuid, FMetasoundFrontendLiteral>& Pair)
@@ -77,9 +79,22 @@ namespace Metasound::Editor
 							});
 
 							FMetasoundFrontendClassInput ClassInput = FGraphBuilder::CreateUniqueClassInput(*OutAsset.GetOwningAsset(), VertexParams, InputDefaults, &Breadcrumb.MemberName);
+							ClassInput.Metadata = Breadcrumb.VertexMetadata;
+
 							if (const FMetasoundFrontendNode* NewNode = Builder.AddGraphInput(ClassInput))
 							{
 								Input = Graph.FindOrAddInput(NewNode->GetID());
+								if (Breadcrumb.MemberMetadataPath.IsSet())
+								{
+									UObject* MemberMetadata = Breadcrumb.MemberMetadataPath->TryLoad();
+									if (UMetasoundEditorGraphMemberDefaultLiteral* DefaultLiteral = Cast<UMetasoundEditorGraphMemberDefaultLiteral>(MemberMetadata))
+									{
+										Builder.ClearMemberMetadata(ClassInput.NodeID);
+										UMetaSoundEditorSubsystem& MetaSoundEditorSubsystem = UMetaSoundEditorSubsystem::GetChecked();
+										TSubclassOf<UMetasoundEditorGraphMemberDefaultLiteral> LiteralClass = MetaSoundEditorSubsystem.GetLiteralClassForType(Breadcrumb.DataType);
+										MetaSoundEditorSubsystem.BindMemberMetadata(Builder, *Input, LiteralClass, DefaultLiteral);
+									}
+								}
 								MappedGeneratedInputNames.Add(Breadcrumb.MemberName, Input);
 							}
 						}
@@ -114,10 +129,11 @@ namespace Metasound::Editor
 		}
 	}
 
-	void FDocumentClipboardUtils::ProcessPastedOutputNodes(FMetasoundAssetBase& OutAsset, TArray<UMetasoundEditorGraphNode*>& OutPastedNodes)
+	void FDocumentClipboardUtils::ProcessPastedOutputNodes(FMetasoundAssetBase& OutAsset, TArray<UMetasoundEditorGraphNode*>& OutPastedNodes, FDocumentPasteNotifications& OutNotifications)
 	{
 		using namespace Frontend;
-
+		using namespace Engine;
+		FMetaSoundFrontendDocumentBuilder& Builder = FDocumentBuilderRegistry::GetChecked().FindOrBeginBuilding(OutAsset.GetOwningAsset());
 		UMetasoundEditorGraph& Graph = *CastChecked<UMetasoundEditorGraph>(&OutAsset.GetGraphChecked());
 		for (int32 Index = OutPastedNodes.Num() - 1; Index >= 0; --Index)
 		{
@@ -129,7 +145,6 @@ namespace Metasound::Editor
 
 			OutputNode->CreateNewGuid();
 
-			// TODO: Add logic to dynamically add output if missing
 			if (OutputNode->Output && Graph.ContainsOutput(*OutputNode->Output))
 			{
 				auto NodeMatches = [OutputNodeID = OutputNode->GetNodeID()](const TObjectPtr<UEdGraphNode>& EdNode)
@@ -144,14 +159,51 @@ namespace Metasound::Editor
 				// Can only have one output reference node
 				if (TObjectPtr<UEdGraphNode>* MatchingOutput = Graph.Nodes.FindByPredicate(NodeMatches))
 				{
+					OutNotifications.bPastedNodesAddMultipleOutputNodes = true;
 					Graph.RemoveNode(OutputNode);
 					OutPastedNodes.RemoveAtSwap(Index, 1, EAllowShrinking::No);
 				}
 			}
 			else
 			{
-				Graph.RemoveNode(OutputNode);
-				OutPastedNodes.RemoveAtSwap(Index, 1, EAllowShrinking::No);
+				// Add output if doesn't exist 
+				const FMetasoundEditorGraphVertexNodeBreadcrumb& Breadcrumb = OutputNode->GetBreadcrumb();
+
+				FDataTypeRegistryInfo Info;
+				if (IDataTypeRegistry::Get().GetDataTypeInfo(Breadcrumb.DataType, Info))
+				{
+					FCreateNodeVertexParams VertexParams;
+					VertexParams.DataType = Breadcrumb.DataType;
+					VertexParams.AccessType = Breadcrumb.AccessType;
+
+					FMetasoundFrontendClassOutput ClassOutput = FGraphBuilder::CreateUniqueClassOutput(*OutAsset.GetOwningAsset(), VertexParams, &Breadcrumb.MemberName);
+					ClassOutput.Metadata = Breadcrumb.VertexMetadata;
+
+					if (const FMetasoundFrontendNode* NewNode = Builder.AddGraphOutput(ClassOutput))
+					{
+						UMetasoundEditorGraphOutput* Output = Graph.FindOrAddOutput(NewNode->GetID());
+						if (Output)
+						{
+							if (Breadcrumb.MemberMetadataPath.IsSet())
+							{
+								UObject* MemberMetadata = Breadcrumb.MemberMetadataPath->TryLoad();
+								if (UMetasoundEditorGraphMemberDefaultLiteral* DefaultLiteral = Cast<UMetasoundEditorGraphMemberDefaultLiteral>(MemberMetadata))
+								{
+									Builder.ClearMemberMetadata(ClassOutput.NodeID);
+									UMetaSoundEditorSubsystem& MetaSoundEditorSubsystem = UMetaSoundEditorSubsystem::GetChecked();
+									TSubclassOf<UMetasoundEditorGraphMemberDefaultLiteral> LiteralClass = MetaSoundEditorSubsystem.GetLiteralClassForType(Breadcrumb.DataType);
+									MetaSoundEditorSubsystem.BindMemberMetadata(Builder, *Output, LiteralClass, DefaultLiteral);
+								}
+							}
+							OutputNode->Output = Output;
+						}
+						else
+						{
+							Graph.RemoveNode(OutputNode);
+							OutPastedNodes.RemoveAtSwap(Index, 1, EAllowShrinking::No);
+						}
+					}
+				}
 			}
 		}
 	}
@@ -179,7 +231,6 @@ namespace Metasound::Editor
 
 			VariableNode->CreateNewGuid();
 
-			// Can only have one setter node
 			TObjectPtr<UMetasoundEditorGraphVariable>& Variable = VariableNode->Variable;
 			if (!Variable || !Graph.ContainsVariable(*Variable))
 			{
@@ -199,6 +250,15 @@ namespace Metasound::Editor
 				else
 				{
 					FrontendVariable = FGraphBuilder::AddVariableHandle(*OutAsset.GetOwningAsset(), Breadcrumb.DataType);
+					FrontendVariable->SetDescription(Breadcrumb.VertexMetadata.GetDescription());
+					FrontendVariable->SetDisplayName(Breadcrumb.VertexMetadata.GetDisplayName());
+					// Hack to reuse the default literals breadcrumb property for variables, which only have a single (rather than paged) literals
+					const FMetasoundFrontendLiteral* Literal = Breadcrumb.DefaultLiterals.Find(DefaultPageID);
+					if (Literal)
+					{
+						FrontendVariable->SetLiteral(*Literal);
+					}
+					// TODO: copy member metadata when variables have been moved to use the builder API
 					Variable = Graph.FindOrAddVariable(FrontendVariable);
 					check(Variable);
 
@@ -213,6 +273,7 @@ namespace Metasound::Editor
 			FConstVariableHandle VariableHandle = Variable->GetConstVariableHandle();
 			if (VariableHandle->IsValid())
 			{
+				// Can only have one setter node
 				FConstNodeHandle VariableMutatorNodeHandle = VariableHandle->FindMutatorNode();
 				if (VariableNode->GetNodeID() == VariableMutatorNodeHandle->GetID())
 				{
@@ -509,7 +570,7 @@ namespace Metasound::Editor
 
 		ProcessPastedCommentNodes(*Asset, PastedCommentNodes);
 		ProcessPastedInputNodes(*Asset, PastedGraphNodes);
-		ProcessPastedOutputNodes(*Asset, PastedGraphNodes);
+		ProcessPastedOutputNodes(*Asset, PastedGraphNodes, OutNotifications);
 		ProcessPastedVariableNodes(*Asset, PastedGraphNodes, OutNotifications);
 		ProcessPastedExternalNodes(*Asset, PastedGraphNodes, OutNotifications);
 		ProcessPastedNodePositions(*Asset, InLocation, PastedGraphNodes, PastedCommentNodes);
diff --git a/Engine/Plugins/Runtime/Metasound/Source/MetasoundEditor/Private/MetasoundEditorDocumentClipboardUtils.h b/Engine/Plugins/Runtime/Metasound/Source/MetasoundEditor/Private/MetasoundEditorDocumentClipboardUtils.h
--- a/Engine/Plugins/Runtime/Metasound/Source/MetasoundEditor/Private/MetasoundEditorDocumentClipboardUtils.h
+++ b/Engine/Plugins/Runtime/Metasound/Source/MetasoundEditor/Private/MetasoundEditorDocumentClipboardUtils.h
@@ -19,6 +19,7 @@ namespace Metasound::Editor
 	{
 		bool bPastedNodesAddMultipleVariableSetters = false;
 		bool bPastedNodesCreateLoop = false;
+		bool bPastedNodesAddMultipleOutputNodes = false;
 	};
 
 	struct FDocumentClipboardUtils
@@ -45,7 +46,7 @@ namespace Metasound::Editor
 
 		static void ProcessPastedNodeConnections(FMetasoundAssetBase& OutAsset, TArray<UMetasoundEditorGraphNode*>& OutPastedNodes);
 
-		static void ProcessPastedOutputNodes(FMetasoundAssetBase& OutAsset, TArray<UMetasoundEditorGraphNode*>& OutPastedNodes);
+		static void ProcessPastedOutputNodes(FMetasoundAssetBase& OutAsset, TArray<UMetasoundEditorGraphNode*>& OutPastedNodes, FDocumentPasteNotifications& OutNotifications);
 
 		static void ProcessPastedVariableNodes(FMetasoundAssetBase& OutAsset, TArray<UMetasoundEditorGraphNode*>& OutPastedNodes, FDocumentPasteNotifications& OutNotifications);
 	};
diff --git a/Engine/Plugins/Runtime/Metasound/Source/MetasoundEditor/Private/MetasoundEditorGraphInputNode.cpp b/Engine/Plugins/Runtime/Metasound/Source/MetasoundEditor/Private/MetasoundEditorGraphInputNode.cpp
--- a/Engine/Plugins/Runtime/Metasound/Source/MetasoundEditor/Private/MetasoundEditorGraphInputNode.cpp
+++ b/Engine/Plugins/Runtime/Metasound/Source/MetasoundEditor/Private/MetasoundEditorGraphInputNode.cpp
@@ -53,7 +53,7 @@ void UMetasoundEditorGraphInputNode::CacheBreadcrumb()
 	{
 		Breadcrumb.MemberName = Input->GetMemberName();
 
-		const FMetaSoundFrontendDocumentBuilder& Builder = Input->GetFrontendBuilderChecked();
+		FMetaSoundFrontendDocumentBuilder& Builder = Input->GetFrontendBuilderChecked();
 		if (const FMetasoundFrontendClassInput* ClassInput = Builder.FindGraphInput(Breadcrumb.MemberName))
 		{
 			if (const FMetasoundFrontendNode* Node = Builder.FindGraphInputNode(Breadcrumb.MemberName))
@@ -63,11 +63,16 @@ void UMetasoundEditorGraphInputNode::CacheBreadcrumb()
 					Breadcrumb.ClassName = Class->Metadata.GetClassName();
 					Breadcrumb.AccessType = ClassInput->AccessType;
 					Breadcrumb.DataType = ClassInput->TypeName;
-
+					Breadcrumb.VertexMetadata = ClassInput->Metadata;
+					if (UMetaSoundFrontendMemberMetadata* MemberMetadata = Builder.FindMemberMetadata(Node->GetID()))
+					{
+						Breadcrumb.MemberMetadataPath = FSoftObjectPath(MemberMetadata);
+					}
 					ClassInput->IterateDefaults([this](const FGuid& PageID, const FMetasoundFrontendLiteral& Literal)
 					{
 						Breadcrumb.DefaultLiterals.Add(PageID, Literal);
 					});
+					
 				}
 			}
 		}
diff --git a/Engine/Plugins/Runtime/Metasound/Source/MetasoundEditor/Private/MetasoundEditorGraphNode.cpp b/Engine/Plugins/Runtime/Metasound/Source/MetasoundEditor/Private/MetasoundEditorGraphNode.cpp
--- a/Engine/Plugins/Runtime/Metasound/Source/MetasoundEditor/Private/MetasoundEditorGraphNode.cpp
+++ b/Engine/Plugins/Runtime/Metasound/Source/MetasoundEditor/Private/MetasoundEditorGraphNode.cpp
@@ -809,11 +809,6 @@ FString UMetasoundEditorGraphMemberNode::GetFindReferenceSearchString_Impl(EGetF
 	return FString();
 }
 
-bool UMetasoundEditorGraphOutputNode::CanDuplicateNode() const
-{
-	return false;
-}
-
 void UMetasoundEditorGraphOutputNode::PinDefaultValueChanged(UEdGraphPin* InPin)
 {
 	using namespace Metasound::Editor;
@@ -950,8 +945,7 @@ void UMetasoundEditorGraphOutputNode::Validate(Metasound::Editor::FGraphNodeVali
 	}
 }
 
-
-const FMetasoundEditorGraphNodeBreadcrumb& UMetasoundEditorGraphOutputNode::GetBreadcrumb() const
+const FMetasoundEditorGraphVertexNodeBreadcrumb& UMetasoundEditorGraphOutputNode::GetBreadcrumb() const
 {
 	return Breadcrumb;
 }
@@ -962,16 +956,20 @@ void UMetasoundEditorGraphOutputNode::CacheBreadcrumb()
 	{
 		Breadcrumb.MemberName = Output->GetMemberName();
 
-		const FMetaSoundFrontendDocumentBuilder& Builder = Output->GetFrontendBuilderChecked();
+		FMetaSoundFrontendDocumentBuilder& Builder = Output->GetFrontendBuilderChecked();
 		if (const FMetasoundFrontendClassOutput* ClassOutput = Builder.FindGraphOutput(Breadcrumb.MemberName))
 		{
-			if (const FMetasoundFrontendNode* Node = Builder.FindGraphInputNode(Breadcrumb.MemberName))
+			if (const FMetasoundFrontendNode* Node = Builder.FindGraphOutputNode(Breadcrumb.MemberName))
 			{
 				if (const FMetasoundFrontendClass* Class = Builder.FindDependency(Node->ClassID))
 				{
 					Breadcrumb.ClassName = Class->Metadata.GetClassName();
 					Breadcrumb.AccessType = ClassOutput->AccessType;
 					Breadcrumb.DataType = ClassOutput->TypeName;
+					if (UMetaSoundFrontendMemberMetadata* MemberMetadata = Builder.FindMemberMetadata(Node->GetID()))
+					{
+						Breadcrumb.MemberMetadataPath = FSoftObjectPath(MemberMetadata);
+					}
 				}
 			}
 		}
@@ -1514,6 +1512,16 @@ void UMetasoundEditorGraphVariableNode::CacheBreadcrumb()
 
 		Breadcrumb.MemberName = VariableHandle->GetName();
 		Breadcrumb.DataType = VariableHandle->GetDataType();
+		// Hack to reuse the default literals breadcrumb property for variables, which only have a single (rather than paged) literals
+		TMap<FGuid, FMetasoundFrontendLiteral> DefaultLiterals;
+		DefaultLiterals.Add(DefaultPageID, VariableHandle->GetLiteral());
+		Breadcrumb.DefaultLiterals = DefaultLiterals;
+
+		FMetasoundFrontendVertexMetadata VertexMetadata;
+		VertexMetadata.SetDisplayName(VariableHandle->GetDisplayName());
+		VertexMetadata.SetDescription(VariableHandle->GetDescription());
+		
+		Breadcrumb.VertexMetadata = VertexMetadata;
 	}
 }
 
diff --git a/Engine/Plugins/Runtime/Metasound/Source/MetasoundEditor/Private/MetasoundEditorSubsystem.cpp b/Engine/Plugins/Runtime/Metasound/Source/MetasoundEditor/Private/MetasoundEditorSubsystem.cpp
--- a/Engine/Plugins/Runtime/Metasound/Source/MetasoundEditor/Private/MetasoundEditorSubsystem.cpp
+++ b/Engine/Plugins/Runtime/Metasound/Source/MetasoundEditor/Private/MetasoundEditorSubsystem.cpp
@@ -227,18 +227,8 @@ UMetaSoundFrontendMemberMetadata* UMetaSoundEditorSubsystem::FindOrCreateGraphIn
 			else
 			{
 				// Get literal class
-				FDataTypeRegistryInfo DataTypeInfo;
-				IMetasoundEditorModule& EditorModule = FModuleManager::GetModuleChecked<IMetasoundEditorModule>("MetaSoundEditor");
 				const FName TypeName = DocBuilder.FindGraphInput(InputName)->TypeName;
-				IDataTypeRegistry::Get().GetDataTypeInfo(TypeName, DataTypeInfo);
-				const EMetasoundFrontendLiteralType LiteralType = static_cast<EMetasoundFrontendLiteralType>(DataTypeInfo.PreferredLiteralType);
-
-				TSubclassOf<UMetasoundEditorGraphMemberDefaultLiteral> LiteralClass = EditorModule.FindDefaultLiteralClass(LiteralType);
-				if (!LiteralClass)
-				{
-					LiteralClass = UMetasoundEditorGraphMemberDefaultLiteral::StaticClass();
-				}
-
+				TSubclassOf<UMetasoundEditorGraphMemberDefaultLiteral> LiteralClass = GetLiteralClassForType(TypeName);
 
 				if (InBuilder->IsPreset())
 				{
@@ -273,6 +263,25 @@ UMetaSoundFrontendMemberMetadata* UMetaSoundEditorSubsystem::FindOrCreateGraphIn
 	return nullptr;
 }
 
+TSubclassOf<UMetasoundEditorGraphMemberDefaultLiteral> UMetaSoundEditorSubsystem::GetLiteralClassForType(FName TypeName) const
+{
+	using namespace Metasound::Editor;
+	using namespace Metasound::Frontend;
+
+	// Get literal class
+	FDataTypeRegistryInfo DataTypeInfo;
+	IMetasoundEditorModule& EditorModule = FModuleManager::GetModuleChecked<IMetasoundEditorModule>("MetaSoundEditor");
+	IDataTypeRegistry::Get().GetDataTypeInfo(TypeName, DataTypeInfo);
+	const EMetasoundFrontendLiteralType LiteralType = static_cast<EMetasoundFrontendLiteralType>(DataTypeInfo.PreferredLiteralType);
+
+	TSubclassOf<UMetasoundEditorGraphMemberDefaultLiteral> LiteralClass = EditorModule.FindDefaultLiteralClass(LiteralType);
+	if (!LiteralClass)
+	{
+		LiteralClass = UMetasoundEditorGraphMemberDefaultLiteral::StaticClass();
+	}
+	return LiteralClass;
+}
+
 UMetaSoundEditorSubsystem& UMetaSoundEditorSubsystem::GetChecked()
 {
 	checkf(GEditor, TEXT("Cannot access UMetaSoundEditorSubsystem without editor loaded"));
diff --git a/Engine/Plugins/Runtime/Metasound/Source/MetasoundEditor/Private/MetasoundNodeDetailCustomization.cpp b/Engine/Plugins/Runtime/Metasound/Source/MetasoundEditor/Private/MetasoundNodeDetailCustomization.cpp
--- a/Engine/Plugins/Runtime/Metasound/Source/MetasoundEditor/Private/MetasoundNodeDetailCustomization.cpp
+++ b/Engine/Plugins/Runtime/Metasound/Source/MetasoundEditor/Private/MetasoundNodeDetailCustomization.cpp
@@ -315,14 +315,15 @@ namespace Metasound
 			{
 				// Apply the clamp range to the default value if using a widget or ClampDefault is otherwise true
 				// Only show clamp row if not using a widget (widgets always require a clamp and range)
-				// Presets are an exception, because they may have a widget inherited, but that doesn't apply and isn't editable
+				// Presets and non inputs are an exception, because they may have a widget inherited, but that doesn't apply and isn't editable,
 				const bool bUsingWidget = FloatLiteral->WidgetType != EMetasoundMemberDefaultWidget::None;
 				const UMetasoundEditorGraphMember* Member = FloatLiteral->FindMember();
 				const bool bIsPreset = Member ? Member->GetFrontendBuilderChecked().IsPreset() : false;
-				const bool bShowClampRow = !bUsingWidget || bIsPreset;
-				const bool bApplyRange = (bUsingWidget && !bIsPreset) || FloatLiteral->ClampDefault;
+				const bool bIsInput = Member ? Cast<UMetasoundEditorGraphInput>(Member) != nullptr : true;
+				const bool bClampActiveWithoutWidget = !bUsingWidget || bIsPreset || !bIsInput;
+				const bool bApplyRange = FloatLiteral->ClampDefault || !bClampActiveWithoutWidget;
 
-				ClampRow->Visibility(bShowClampRow ? DefaultVisibility : EVisibility::Hidden);
+				ClampRow->Visibility(bClampActiveWithoutWidget ? DefaultVisibility : EVisibility::Hidden);
 				for (TSharedPtr<IPropertyHandle>& DefaultValueHandle : DefaultProperties)
 				{
 					if (DefaultValueHandle.IsValid())
diff --git a/Engine/Plugins/Runtime/Metasound/Source/MetasoundEditor/Public/MetasoundEditorGraphNode.h b/Engine/Plugins/Runtime/Metasound/Source/MetasoundEditor/Public/MetasoundEditorGraphNode.h
--- a/Engine/Plugins/Runtime/Metasound/Source/MetasoundEditor/Public/MetasoundEditorGraphNode.h
+++ b/Engine/Plugins/Runtime/Metasound/Source/MetasoundEditor/Public/MetasoundEditorGraphNode.h
@@ -13,6 +13,7 @@
 #include "Sound/SoundWave.h"
 #include "Textures/SlateIcon.h"
 #include "UObject/ObjectMacros.h"
+#include "UObject/SoftObjectPath.h"
 
 #include "MetasoundEditorGraphNode.generated.h"
 
@@ -68,6 +69,12 @@ struct FMetasoundEditorGraphMemberNodeBreadcrumb : public FMetasoundEditorGraphN
 
 	UPROPERTY()
 	TMap<FGuid, FMetasoundFrontendLiteral> DefaultLiterals;
+	
+	UPROPERTY()
+	FMetasoundFrontendVertexMetadata VertexMetadata;
+
+	UPROPERTY()
+	TOptional<FSoftObjectPath> MemberMetadataPath;
 };
 
 USTRUCT()
@@ -264,8 +271,9 @@ public:
 	UPROPERTY()
 	TObjectPtr<UMetasoundEditorGraphOutput> Output;
 
-	virtual const FMetasoundEditorGraphNodeBreadcrumb& GetBreadcrumb() const override;
-	virtual void CacheBreadcrumb() override;
+	const FMetasoundEditorGraphVertexNodeBreadcrumb& GetBreadcrumb() const;
+
+	virtual void CacheBreadcrumb() override;	
 
 	UE_DEPRECATED(5.4, "Use the Frontend node or Breadcrumb directly to get the class name.")
 	virtual FMetasoundFrontendClassName GetClassName() const override;
@@ -273,8 +281,6 @@ public:
 	virtual FGuid GetNodeID() const override;
 	virtual UMetasoundEditorGraphMember* GetMember() const override;
 
-	// Can't duplicate, only one node per graph.
-	virtual bool CanDuplicateNode() const override;
 
 	// Disallow deleting outputs as they require being connected to some
 	// part of the graph by the Frontend Graph Builder (which is enforced
diff --git a/Engine/Plugins/Runtime/Metasound/Source/MetasoundEditor/Public/MetasoundEditorSubsystem.h b/Engine/Plugins/Runtime/Metasound/Source/MetasoundEditor/Public/MetasoundEditorSubsystem.h
--- a/Engine/Plugins/Runtime/Metasound/Source/MetasoundEditor/Public/MetasoundEditorSubsystem.h
+++ b/Engine/Plugins/Runtime/Metasound/Source/MetasoundEditor/Public/MetasoundEditorSubsystem.h
@@ -66,6 +66,9 @@ public:
 	UFUNCTION(BlueprintCallable, Category = "Audio|MetaSound|Builder", meta = (ExpandEnumAsExecs = "OutResult"))
 	UPARAM(DisplayName = "Input Metadata") UMetaSoundFrontendMemberMetadata* FindOrCreateGraphInputMetadata(UPARAM(DisplayName = "Builder") UMetaSoundBuilderBase* InBuilder, FName InputName, EMetaSoundBuilderResult& OutResult);
 
+	// Returns the corresponding literal class for a given type
+	TSubclassOf<UMetasoundEditorGraphMemberDefaultLiteral> GetLiteralClassForType(FName TypeName) const;
+
 	// Sets the visual location to InLocation of a given node InNode of a given builder's document.
 	UFUNCTION(BlueprintCallable, Category = "Audio|MetaSound|Builder|Editor", meta = (ExpandEnumAsExecs = "OutResult"))
 	void SetNodeLocation(
