Fix warning in MetaSound tests
#rnx

diff --git a/Engine/Plugins/Runtime/Metasound/Source/MetasoundEngine/Private/MetasoundSettings.cpp b/Engine/Plugins/Runtime/Metasound/Source/MetasoundEngine/Private/MetasoundSettings.cpp
--- a/Engine/Plugins/Runtime/Metasound/Source/MetasoundEngine/Private/MetasoundSettings.cpp
+++ b/Engine/Plugins/Runtime/Metasound/Source/MetasoundEngine/Private/MetasoundSettings.cpp
@@ -355,7 +355,7 @@ const FMetaSoundPageSettings& UMetaSoundSettings::GetTargetPageSettings() const
 	{
 		if (bWarnAccessBeforeInit)
 		{
-			UE_LOG(LogMetaSound, Display, TEXT("Target Page Settings accessed prior to 'PostLoad' called.  Uninitialized PageSettings '%s' being returned."), *SettingsSet.Name.ToString());
+			UE_LOG(LogMetaSound, Display, TEXT("Target Page Settings accessed prior to 'PostInitProperties' being called.  Uninitialized PageSettings '%s' being returned."), *SettingsSet.Name.ToString());
 			bWarnAccessBeforeInit = false;
 		}
 	};
diff --git a/Engine/Plugins/Runtime/Metasound/Source/MetasoundEngineTest/Private/EngineTestMetaSoundBuilder.h b/Engine/Plugins/Runtime/Metasound/Source/MetasoundEngineTest/Private/EngineTestMetaSoundBuilder.h
--- a/Engine/Plugins/Runtime/Metasound/Source/MetasoundEngineTest/Private/EngineTestMetaSoundBuilder.h
+++ b/Engine/Plugins/Runtime/Metasound/Source/MetasoundEngineTest/Private/EngineTestMetaSoundBuilder.h
@@ -43,7 +43,7 @@ namespace EngineTestMetaSoundPatchBuilderPrivate
 	{
 		UMetaSoundPatchBuilder& Builder = CreatePatchBuilderChecked(Test, FName(PatchName + TEXT(" Builder")), InterfaceNamesToAdd);
 
-		UMetaSoundPatch* InputPatch = CastChecked<UMetaSoundPatch>(Builder.Build(nullptr, FMetaSoundBuilderOptions { FName(PatchName) }).GetObject());
+		UMetaSoundPatch* InputPatch = CastChecked<UMetaSoundPatch>(Builder.BuildNewMetaSound(FName(PatchName)).GetObject());
 		Test.AddErrorIfFalse(InputPatch != nullptr, FString::Printf(TEXT("Failed to build MetaSound patch '%s'"), *PatchName));
 		return InputPatch;
 	}
