Backout deprecation portion of prior CL until content can be fixed up
#rnx

diff --git a/Engine/Plugins/Runtime/Metasound/Source/MetasoundEngine/Public/MetasoundBuilderBase.h b/Engine/Plugins/Runtime/Metasound/Source/MetasoundEngine/Public/MetasoundBuilderBase.h
--- a/Engine/Plugins/Runtime/Metasound/Source/MetasoundEngine/Public/MetasoundBuilderBase.h
+++ b/Engine/Plugins/Runtime/Metasound/Source/MetasoundEngine/Public/MetasoundBuilderBase.h
@@ -436,10 +436,7 @@ public:
 	UE_DEPRECATED(5.5, "Call directly on Frontend builder using 'GetBuilder'")
 	void UpdateDependencyClassNames(const TMap<FMetasoundFrontendClassName, FMetasoundFrontendClassName>& OldToNewReferencedClassNames);
 
-	UFUNCTION(BlueprintCallable, Category = "Audio|MetaSound|Builder", meta = (
-		WorldContext = "Parent",
-		DeprecatedFunction,
-		DeprecationMessage = "MetaSounds must now be top-level objects within a package.  Use 'BuildAndOverwriteMetaSound' to copy this builder's implementation over an existing transient MetaSound or 'BuildNewMetaSound' to generate a new, transient MetaSound. Use 'MetaSoundEditorSubsystem::BuildToAsset' to build a copied MetaSound as a new, serialized asset (editor only)."))
+	UFUNCTION(BlueprintCallable, Category = "Audio|MetaSound|Builder", meta = (WorldContext = "Parent"))
 	virtual UPARAM(DisplayName = "MetaSound") TScriptInterface<IMetaSoundDocumentInterface> Build(UObject* Parent, const FMetaSoundBuilderOptions& Options) const { return BuildNewMetaSound(Options.Name); }
 
 	// Copies a transient MetaSound with the provided builder options, copying the underlying MetaSound
