Fix crash and errors when building MetaSound presets and renaming referenced MetaSounds

#rnx

diff --git a/Engine/Plugins/Runtime/Metasound/Source/MetasoundEditor/Private/MetasoundEditorSubsystem.cpp b/Engine/Plugins/Runtime/Metasound/Source/MetasoundEditor/Private/MetasoundEditorSubsystem.cpp
--- a/Engine/Plugins/Runtime/Metasound/Source/MetasoundEditor/Private/MetasoundEditorSubsystem.cpp
+++ b/Engine/Plugins/Runtime/Metasound/Source/MetasoundEditor/Private/MetasoundEditorSubsystem.cpp
@@ -62,11 +62,8 @@ TScriptInterface<IMetaSoundDocumentInterface> UMetaSoundEditorSubsystem::BuildTo
 
 			// Initialize and Build
 			{
-				constexpr UObject* Parent = nullptr;
 				constexpr bool bForceUniqueClassName = true;
-				constexpr bool bAddToRegistry = true;
-				const FMetaSoundBuilderOptions BuilderOptions { FName(*AssetName), bForceUniqueClassName, bAddToRegistry, NewMetaSound };
-				InBuilder->Build(Parent, BuilderOptions);
+				InBuilder->BuildAndOverwriteMetaSound(NewMetaSound, bForceUniqueClassName);
 			}
 
 			UMetaSoundBuilderBase& NewDocBuilder = FDocumentBuilderRegistry::GetChecked().FindOrBeginBuilding(*NewMetaSound);
@@ -144,7 +141,7 @@ UMetasoundEditorGraphMemberDefaultLiteral* UMetaSoundEditorSubsystem::CreateMemb
 				{
 					if (UMetaSoundFrontendMemberMetadata* ReferencedMemberMetadata = ReferencedBuilder.FindMemberMetadata(ClassInput->NodeID))
 					{
-						return NewObject<UMetasoundEditorGraphMemberDefaultLiteral>(&Builder.CastDocumentObjectChecked<UObject>(), LiteralClass, FName(), RF_Transactional, ReferencedMemberMetadata);
+						return NewObject<UMetasoundEditorGraphMemberDefaultLiteral>(&Builder.CastDocumentObjectChecked<UObject>(), ReferencedMemberMetadata->GetClass(), FName(), RF_Transactional, ReferencedMemberMetadata);
 					}
 				}
 			}
