Remove checks for CanTarget in calls leveraged in non-editor builds. (PerPlatform settings are not available at runtime when initialized in an .ini and deserialized as CDOs)

#rnx

diff --git a/Engine/Plugins/Runtime/Metasound/Source/MetasoundEngine/Private/MetasoundDocumentBuilderRegistry.cpp b/Engine/Plugins/Runtime/Metasound/Source/MetasoundEngine/Private/MetasoundDocumentBuilderRegistry.cpp
--- a/Engine/Plugins/Runtime/Metasound/Source/MetasoundEngine/Private/MetasoundDocumentBuilderRegistry.cpp
+++ b/Engine/Plugins/Runtime/Metasound/Source/MetasoundEngine/Private/MetasoundDocumentBuilderRegistry.cpp
@@ -517,6 +517,26 @@ namespace Metasound::Engine
 			}
 		}, bReverse);
 
+		if (!bResolved)
+		{
+			const FGuid& AnyPageID = InPageIDsToResolve.Last();
+#if !NO_LOGGING
+			auto GetDisplayPageString = [&Settings](const FGuid& InPageID)
+			{
+				if (const FMetaSoundPageSettings* DisplayPage = Settings.FindPageSettings(InPageID))
+				{
+					return DisplayPage->Name.ToString();
+				}
+				return InPageID.ToString();
+			};
+			UE_LOG(LogMetaSound, Error,
+				TEXT("Failed to resolve PageID for Target '%s': Setting to arbitrary Page '%s' (Target likely overridden by page not set as 'CanTarget/Targetable' for the current platform)"),
+				*GetDisplayPageString(TargetPageID),
+				*GetDisplayPageString(AnyPageID));
+#endif // !NO_LOGGING
+			ResolvedPageID = AnyPageID;
+		}
+
 		return ResolvedPageID;
 	}
 
diff --git a/Engine/Plugins/Runtime/Metasound/Source/MetasoundEngine/Private/MetasoundSettings.cpp b/Engine/Plugins/Runtime/Metasound/Source/MetasoundEngine/Private/MetasoundSettings.cpp
--- a/Engine/Plugins/Runtime/Metasound/Source/MetasoundEngine/Private/MetasoundSettings.cpp
+++ b/Engine/Plugins/Runtime/Metasound/Source/MetasoundEngine/Private/MetasoundSettings.cpp
@@ -363,26 +363,20 @@ const FMetaSoundPageSettings& UMetaSoundSettings::GetTargetPageSettings() const
 
 	if (const FMetaSoundPageSettings* TargetSettings = FindPageSettings(TargetPage))
 	{
-		if (TargetSettings->CanTarget.GetValue())
-		{
 #if !NO_LOGGING
-			WarnIfUninitialized(*TargetSettings);
+		WarnIfUninitialized(*TargetSettings);
 #endif // !NO_LOGGING
-			return *TargetSettings;
-		}
+		return *TargetSettings;
 	}
 
 	// Shouldn't hit this, but if for some reason the target page is in a bad state,
 	// try and return any page setting set as a valid target.
 	for (const FMetaSoundPageSettings& Setting : PageSettings)
 	{
-		if (Setting.CanTarget.GetValue())
-		{
 #if !NO_LOGGING
-			WarnIfUninitialized(Setting);
+		WarnIfUninitialized(Setting);
 #endif // !NO_LOGGING
-			return Setting;
-		}
+		return Setting;
 	}
 
 #if !NO_LOGGING
@@ -434,55 +428,36 @@ void UMetaSoundSettings::PostInitProperties()
 	bWarnAccessBeforeInit = false;
 #endif // !NO_LOGGING
 
-#if WITH_EDITORONLY_DATA
-	UE_LOG(LogMetaSound, Display, TEXT("MetaSound Page Target Initialized to '%s'"), *GetTargetPageSettings().Name.ToString());
-#else // !WITH_EDITORONLY_DATA
 	if (const FMetaSoundPageSettings* Page = FindPageSettings(TargetPageName))
 	{
-		if (Page->CanTarget.GetValue())
-		{
-			UE_LOG(LogMetaSound, Display, TEXT("MetaSound Page Target Initialized to '%s'"), *GetTargetPageSettings().Name.ToString());
-			return;
-		}
-	}
-
-	UE_LOG(LogMetaSound, Warning, TEXT("TargetPageName on load in MetaSound settings did not correspond to a valid targetable page."));
-	if (DefaultPageSettings.CanTarget.GetValue())
-	{
-		UE_LOG(LogMetaSound, Warning, TEXT("Setting target to 'Default' page settings."));
-		TargetPageNameOverride = DefaultPageSettings.Name;
-		return;
+		UE_LOG(LogMetaSound, Display, TEXT("MetaSound Page Target Initialized to '%s'"), *GetTargetPageSettings().Name.ToString());
 	}
 	else
 	{
-		for (const FMetaSoundPageSettings& Page : PageSettings)
+		UE_LOG(LogMetaSound, Warning, TEXT("TargetPageName '%s' at time of 'UMetaSoundSettings::PostInitProperties' did not correspond to a valid page."), *TargetPageName.ToString());
+		if (PageSettings.IsEmpty())
 		{
-			if (Page.CanTarget.GetValue())
-			{
-				UE_LOG(LogMetaSound, Warning, TEXT("Setting target to '%s' page settings."), *Page.Name.ToString());
-				TargetPageNameOverride = Page.Name;
-				return;
-			}
+			UE_LOG(LogMetaSound, Warning, TEXT("Setting target to '%s' page settings."), *DefaultPageSettings.Name.ToString());
+			TargetPageName = DefaultPageSettings.Name;
+		}
+		else
+		{
+			UE_LOG(LogMetaSound, Warning, TEXT("Setting target to highest project page settings '%s'."), *PageSettings.Last().Name.ToString());
+			TargetPageName = PageSettings.Last().Name;
 		}
-
-		UE_LOG(LogMetaSound, Warning, TEXT("Setting target to 'Default' page settings & forcing default as targetable."));
-		DefaultPageSettings.CanTarget.Default = true;
-		TargetPageNameOverride = DefaultPageSettings.Name;
 	}
-#endif // !WITH_EDITORONLY_DATA
 }
 
 bool UMetaSoundSettings::SetTargetPage(FName PageName)
 {
 	if (const FMetaSoundPageSettings* PageSetting = FindPageSettings(PageName))
 	{
-		if (TargetPageName != PageSetting->Name)
+		const FName TargetPage = TargetPageNameOverride.IsSet() ? *TargetPageNameOverride : TargetPageName;
+		if (TargetPage != PageSetting->Name)
 		{
-			if (PageSetting->CanTarget.GetValue())
-			{
-				TargetPageNameOverride = PageSetting->Name;
-				return true;
-			}
+			UE_LOG(LogMetaSound, Display, TEXT("Target page override set to '%s'."), *TargetPage.ToString());
+			TargetPageNameOverride = PageSetting->Name;
+			return true;
 		}
 	}
 
diff --git a/Engine/Plugins/Runtime/Metasound/Source/MetasoundEngine/Public/MetasoundSettings.h b/Engine/Plugins/Runtime/Metasound/Source/MetasoundEngine/Public/MetasoundSettings.h
--- a/Engine/Plugins/Runtime/Metasound/Source/MetasoundEngine/Public/MetasoundSettings.h
+++ b/Engine/Plugins/Runtime/Metasound/Source/MetasoundEngine/Public/MetasoundSettings.h
@@ -70,13 +70,15 @@ struct METASOUNDENGINE_API FMetaSoundPageSettings
 	FName Name = Metasound::Frontend::DefaultPageName;
 
 private:
-	// When true, page can be targeted for the assigned platform(s)/platform group(s).
-	// Is used to determine which asset page data (i.e. graphs and input defaults) can be
-	// excluded from cook if page is not explicitly omitted via 'ExcludeFromCook' setting.
-	UPROPERTY(EditAnywhere, Category = "Pages")
+#if WITH_EDITORONLY_DATA
+	// When true, page asset data (i.e. graphs and input defaults) can be targeted
+	// for the most applicable platform/platform group. If associated asset data is
+	// defined, will always be cooked. If false, asset page data is only cooked if it is
+	// resolved to from a higher-indexed page setting and is not set to explicitly
+	// "ExcludeFromCook".
+	UPROPERTY(EditAnywhere, Category = "Pages", meta = (DisplayName = "Targetable"))
 	FPerPlatformBool CanTarget = true;
 
-#if WITH_EDITORONLY_DATA
 	// Just used to inform edit condition to enable/disable exclude from cook. Maintained by ConformPageSettings on object load/mutation.
 	// EditCondition meta mark-up is hack to avoid boolean being default added to name field
 	UPROPERTY(EditAnywhere, Category = "Pages", meta = (EditCondition = false, EditConditionHides))
@@ -85,7 +87,7 @@ private:
 	// When true, exclude page data when cooking from the assigned platform(s)/platform group(s).
 	// If false, page data may or may not be included in cook depending on whether or not the given
 	// page data is required in order to ensure a value is always resolved for the cook platform target(s).
-	// (Ignored if CanTarget true for corresponding platform/group).
+	// (Ignored if 'Targetable' is true for most applicable platform/platform group).
 	UPROPERTY(EditAnywhere, Category = "Pages", meta = (EditCondition = "!bIsDefaultPage"))
 	FPerPlatformBool ExcludeFromCook = false;
 #endif //WITH_EDITORONLY_DATA
