Fix inverted logic when building & overriding MetaSound via builder API
#rnx

diff --git a/Engine/Plugins/Runtime/Metasound/Source/MetasoundEngine/Private/MetasoundBuilderSubsystem.cpp b/Engine/Plugins/Runtime/Metasound/Source/MetasoundEngine/Private/MetasoundBuilderSubsystem.cpp
--- a/Engine/Plugins/Runtime/Metasound/Source/MetasoundEngine/Private/MetasoundBuilderSubsystem.cpp
+++ b/Engine/Plugins/Runtime/Metasound/Source/MetasoundEngine/Private/MetasoundBuilderSubsystem.cpp
@@ -98,7 +98,7 @@ namespace Metasound::Engine
 
 void UMetaSoundPatchBuilder::BuildAndOverwriteMetaSound(TScriptInterface<IMetaSoundDocumentInterface> ExistingMetaSound, bool bForceUniqueClassName)
 {
-	if (!ExistingMetaSound || !ExistingMetaSound.GetObject()->IsAsset())
+	if (!ExistingMetaSound || ExistingMetaSound.GetObject()->IsAsset())
 	{
 		UE_LOG(LogMetaSound, Error, TEXT("Failed to build and overwrite MetaSound: No existing MetaSound supplied or is serialized asset "
 			"(use 'BuildNewMetaSound' to create a new, transient MetaSound. Overwriting serialized asset is only supported at edit time via "
