[MetaSounds] - Copy/Cut/Paste don't populate to the node in composition immediately
#rnx

diff --git a/Engine/Plugins/Runtime/Metasound/Source/MetasoundEditor/Private/MetasoundEditor.cpp b/Engine/Plugins/Runtime/Metasound/Source/MetasoundEditor/Private/MetasoundEditor.cpp
--- a/Engine/Plugins/Runtime/Metasound/Source/MetasoundEditor/Private/MetasoundEditor.cpp
+++ b/Engine/Plugins/Runtime/Metasound/Source/MetasoundEditor/Private/MetasoundEditor.cpp
@@ -3940,13 +3940,14 @@ namespace Metasound
 			const UMetasoundEditorGraphMember* Member = FDocumentClipboardUtils::GetMemberFromClipboard();		
 			if (Member)
 			{
+				UObject* MetaSound = GetMetasoundObject();
+				if (!MetaSound)
+				{
+					return;
+				}
+
 				if (const UMetasoundEditorGraphVertex* SourceGraphVertex = Cast<UMetasoundEditorGraphVertex>(Member))
 				{
-					UObject* MetaSound = GetMetasoundObject();
-					if (!MetaSound)
-					{
-						return;
-					}
 
 					UMetasoundEditorGraph& Graph = GetMetaSoundGraphChecked();
 					MetaSound->Modify();
@@ -3988,7 +3989,6 @@ namespace Metasound
 				}
 				else if (const UMetasoundEditorGraphVariable* SourceGraphVariable = Cast<UMetasoundEditorGraphVariable>(Member))
 				{
-					UObject* MetaSound = GetMetasoundObject();
 
 					const FMetasoundEditorGraphVariableBreadcrumb Breadcrumb = SourceGraphVariable->GetBreadcrumb();
 
@@ -4003,7 +4003,7 @@ namespace Metasound
 					}
 				}
 
-				return;
+				FGraphBuilder::RegisterGraphWithFrontend(*MetaSound, true);
 			}
 		}
 
