Fix invalid MetaSound validation error after diffing assets

#rnx

diff --git a/Engine/Plugins/Runtime/Metasound/Source/MetasoundEngine/Private/MetasoundAssetSubsystem.cpp b/Engine/Plugins/Runtime/Metasound/Source/MetasoundEngine/Private/MetasoundAssetSubsystem.cpp
--- a/Engine/Plugins/Runtime/Metasound/Source/MetasoundEngine/Private/MetasoundAssetSubsystem.cpp
+++ b/Engine/Plugins/Runtime/Metasound/Source/MetasoundEngine/Private/MetasoundAssetSubsystem.cpp
@@ -351,6 +351,16 @@ namespace Metasound::Engine
 
 		METASOUND_TRACE_CPUPROFILER_EVENT_SCOPE(UMetaSoundAssetSubsystem::AddOrUpdateAsset_AssetData);
 
+		// Invalid ClassID means the node could not be registered.
+		// Let caller report or ensure as necessary.
+		FAssetKey AssetKey = FAssetKey::GetInvalid();
+
+		// Don't add temporary assets used for diffing
+		if (InAssetData.HasAnyPackageFlags(PKG_ForDiffing))
+		{
+			return AssetKey;
+		}
+
 		FNodeClassInfo ClassInfo;
 		bool bClassInfoFound = AssetSubsystemPrivate::GetAssetClassInfo(InAssetData, ClassInfo);
 		if (!bClassInfoFound)
@@ -380,18 +390,13 @@ namespace Metasound::Engine
 
 		if (ClassInfo.AssetClassID.IsValid())
 		{
-			const FAssetKey AssetKey = FAssetKey(ClassInfo.ClassName, ClassInfo.Version);
+			AssetKey = FAssetKey(ClassInfo.ClassName, ClassInfo.Version);
 			if (AssetKey.IsValid())
 			{
 				AssetSubsystemPrivate::AddPath(&PathMapCriticalSection, PathMap, AssetKey, ClassInfo.AssetPath);
 			}
-
-			return AssetKey;
 		}
-
-		// Invalid ClassID means the node could not be registered.
-		// Let caller report or ensure as necessary.
-		return FAssetKey();
+		return AssetKey;
 	}
 
 	bool FMetaSoundAssetManager::CanAutoUpdate(const FMetasoundFrontendClassName& InClassName) const
diff --git a/Engine/Plugins/Runtime/Metasound/Source/MetasoundEngine/Private/MetasoundEngineModule.cpp b/Engine/Plugins/Runtime/Metasound/Source/MetasoundEngine/Private/MetasoundEngineModule.cpp
--- a/Engine/Plugins/Runtime/Metasound/Source/MetasoundEngine/Private/MetasoundEngineModule.cpp
+++ b/Engine/Plugins/Runtime/Metasound/Source/MetasoundEngine/Private/MetasoundEngineModule.cpp
@@ -242,6 +242,11 @@ namespace Metasound::Engine
 		// report the failed load of the class.
 		if (const UClass* AssetClass = InAssetData.GetClass())
 		{
+			// Don't add temporary assets used for diffing
+			if (InAssetData.HasAnyPackageFlags(PKG_ForDiffing))
+			{
+				return;
+			}
 			const bool bIsRegisteredClass = IMetasoundUObjectRegistry::Get().IsRegisteredClass(*AssetClass);
 			if (bIsRegisteredClass)
 			{
