Fix for BuildToAsset fix-up collision with BuildAndOverwriteMetaSound

#rnx

diff --git a/Engine/Plugins/Runtime/Metasound/Source/MetasoundEditor/Private/MetasoundEditorSubsystem.cpp b/Engine/Plugins/Runtime/Metasound/Source/MetasoundEditor/Private/MetasoundEditorSubsystem.cpp
--- a/Engine/Plugins/Runtime/Metasound/Source/MetasoundEditor/Private/MetasoundEditorSubsystem.cpp
+++ b/Engine/Plugins/Runtime/Metasound/Source/MetasoundEditor/Private/MetasoundEditorSubsystem.cpp
@@ -63,7 +63,9 @@ TScriptInterface<IMetaSoundDocumentInterface> UMetaSoundEditorSubsystem::BuildTo
 			// Initialize and Build
 			{
 				constexpr bool bForceUniqueClassName = true;
-				InBuilder->BuildAndOverwriteMetaSound(NewMetaSound, bForceUniqueClassName);
+				constexpr bool bAddToRegistry = true;
+				const FMetaSoundBuilderOptions BuilderOptions { FName(*AssetName), bForceUniqueClassName, bAddToRegistry, NewMetaSound };
+				InBuilder->Build(BuilderOptions);
 			}
 
 			UMetaSoundBuilderBase& NewDocBuilder = FDocumentBuilderRegistry::GetChecked().FindOrBeginBuilding(*NewMetaSound);
diff --git a/Engine/Plugins/Runtime/Metasound/Source/MetasoundEngine/Private/MetasoundBuilderSubsystem.cpp b/Engine/Plugins/Runtime/Metasound/Source/MetasoundEngine/Private/MetasoundBuilderSubsystem.cpp
--- a/Engine/Plugins/Runtime/Metasound/Source/MetasoundEngine/Private/MetasoundBuilderSubsystem.cpp
+++ b/Engine/Plugins/Runtime/Metasound/Source/MetasoundEngine/Private/MetasoundBuilderSubsystem.cpp
@@ -96,16 +96,8 @@ namespace Metasound::Engine
 	} // namespace BuilderSubsystemPrivate
 } // namespace Metasound::Engine
 
-void UMetaSoundPatchBuilder::BuildAndOverwriteMetaSound(TScriptInterface<IMetaSoundDocumentInterface> ExistingMetaSound, bool bForceUniqueClassName)
+void UMetaSoundPatchBuilder::BuildAndOverwriteMetaSoundInternal(TScriptInterface<IMetaSoundDocumentInterface> ExistingMetaSound, bool bForceUniqueClassName) const
 {
-	if (!ExistingMetaSound || ExistingMetaSound.GetObject()->IsAsset())
-	{
-		UE_LOG(LogMetaSound, Error, TEXT("Failed to build and overwrite MetaSound: No existing MetaSound supplied or is serialized asset "
-			"(use 'BuildNewMetaSound' to create a new, transient MetaSound. Overwriting serialized asset is only supported at edit time via "
-			"UMetaSoundEditorSubsystem::BuildToAsset."));
-		return;
-	}
-
 	FMetaSoundBuilderOptions Options;
 	Options.ExistingMetaSound = ExistingMetaSound;
 	Options.bForceUniqueClassName = bForceUniqueClassName;
@@ -270,16 +262,8 @@ bool UMetaSoundSourceBuilder::ExecuteAuditionableTransaction(FAuditionableTransa
 	return false;
 }
 
-void UMetaSoundSourceBuilder::BuildAndOverwriteMetaSound(TScriptInterface<IMetaSoundDocumentInterface> ExistingMetaSound, bool bForceUniqueClassName)
+void UMetaSoundSourceBuilder::BuildAndOverwriteMetaSoundInternal(TScriptInterface<IMetaSoundDocumentInterface> ExistingMetaSound, bool bForceUniqueClassName) const
 {
-	if (!ExistingMetaSound || !ExistingMetaSound.GetObject()->IsAsset())
-	{
-		UE_LOG(LogMetaSound, Error, TEXT("Failed to build and overwrite MetaSound: No existing MetaSound supplied or is serialized asset "
-			"(use 'BuildNewMetaSound' to create a new, transient MetaSound. Overwriting serialized asset is only supported at edit time via "
-			"UMetaSoundEditorSubsystem::BuildToAsset."));
-		return;
-	}
-
 	FMetaSoundBuilderOptions Options;
 	Options.ExistingMetaSound = ExistingMetaSound;
 	Options.bForceUniqueClassName = bForceUniqueClassName;
diff --git a/Engine/Plugins/Runtime/Metasound/Source/MetasoundEngine/Public/MetasoundBuilderBase.cpp b/Engine/Plugins/Runtime/Metasound/Source/MetasoundEngine/Public/MetasoundBuilderBase.cpp
--- a/Engine/Plugins/Runtime/Metasound/Source/MetasoundEngine/Public/MetasoundBuilderBase.cpp
+++ b/Engine/Plugins/Runtime/Metasound/Source/MetasoundEngine/Public/MetasoundBuilderBase.cpp
@@ -248,6 +248,38 @@ FMetaSoundNodeHandle UMetaSoundBuilderBase::AddNodeByClassName(const FMetasoundF
 	return NewHandle;
 }
 
+#if WITH_EDITORONLY_DATA
+TScriptInterface<IMetaSoundDocumentInterface> UMetaSoundBuilderBase::Build(const FMetaSoundBuilderOptions& Options) const
+{
+	if (Options.ExistingMetaSound)
+	{
+		BuildAndOverwriteMetaSoundInternal(Options.ExistingMetaSound, Options.bForceUniqueClassName);
+		return Options.ExistingMetaSound;
+	}
+
+	return BuildNewMetaSound(Options.Name);
+}
+#endif // WITH_EDITORONLY_DATA
+
+void UMetaSoundBuilderBase::BuildAndOverwriteMetaSound(TScriptInterface<IMetaSoundDocumentInterface> ExistingMetaSound, bool bForceUniqueClassName)
+{
+	if (!ExistingMetaSound)
+	{
+		UE_LOG(LogMetaSound, Error, TEXT("Failed to build and overwrite MetaSound: No existing MetaSound supplied."));
+		return;
+	}
+
+	if (ExistingMetaSound.GetObject()->IsAsset())
+	{
+		UE_LOG(LogMetaSound, Error, TEXT("Failed to build and overwrite MetaSound: Cannot overwrite serialized asset "
+			"(use 'BuildNewMetaSound' to create a new, transient MetaSound. Overwriting serialized asset is only supported at edit time via "
+			"UMetaSoundEditorSubsystem::BuildToAsset."));
+		return;
+	}
+
+	BuildAndOverwriteMetaSoundInternal(ExistingMetaSound, bForceUniqueClassName);
+}
+
 void UMetaSoundBuilderBase::BuildInternal(TScriptInterface<IMetaSoundDocumentInterface> NewMetaSound, const FMetasoundFrontendClassName* InDocClassName) const
 {
 	using namespace Metasound::Engine;
diff --git a/Engine/Plugins/Runtime/Metasound/Source/MetasoundEngine/Public/MetasoundBuilderBase.h b/Engine/Plugins/Runtime/Metasound/Source/MetasoundEngine/Public/MetasoundBuilderBase.h
--- a/Engine/Plugins/Runtime/Metasound/Source/MetasoundEngine/Public/MetasoundBuilderBase.h
+++ b/Engine/Plugins/Runtime/Metasound/Source/MetasoundEngine/Public/MetasoundBuilderBase.h
@@ -436,15 +436,20 @@ public:
 	UE_DEPRECATED(5.5, "Call directly on Frontend builder using 'GetBuilder'")
 	void UpdateDependencyClassNames(const TMap<FMetasoundFrontendClassName, FMetasoundFrontendClassName>& OldToNewReferencedClassNames);
 
-	UFUNCTION(BlueprintCallable, Category = "Audio|MetaSound|Builder", meta = (WorldContext = "Parent"))
+	UFUNCTION(BlueprintCallable, Category = "Audio|MetaSound|Builder", meta = (WorldContext = "Parent", DisplayName = "Build (Soft Deprecated. Parent no longer supported and field is ignored.)"))
 	virtual UPARAM(DisplayName = "MetaSound") TScriptInterface<IMetaSoundDocumentInterface> Build(UObject* Parent, const FMetaSoundBuilderOptions& Options) const { return BuildNewMetaSound(Options.Name); }
 
+#if WITH_EDITORONLY_DATA
+	TScriptInterface<IMetaSoundDocumentInterface> Build(const FMetaSoundBuilderOptions& Options) const;
+#endif // WITH_EDITORONLY_DATA
+
 	// Copies a transient MetaSound with the provided builder options, copying the underlying MetaSound
 	// managed by this builder and registering it with the MetaSound Node Registry as a unique name.
 	// If 'Force Unique Class Name' is true, registers MetaSound as a new class in the registry, potentially
-	// invalidating existing references in other MetaSounds.
+	// invalidating existing references in other MetaSounds. Not permissible to overwrite MetaSound asset,
+	// only transient MetaSound (see EditorSubsystem for overwriting assets at edit time).
 	UFUNCTION(BlueprintCallable, Category = "Audio|MetaSound|Builder", meta = (DisplayName = "Build And Overwrite MetaSound", AdvancedDisplay = "1"))
-	virtual void BuildAndOverwriteMetaSound(UPARAM(DisplayName = "Existing MetaSound") TScriptInterface<IMetaSoundDocumentInterface> ExistingMetaSound, bool bForceUniqueClassName = false) PURE_VIRTUAL(UMetaSoundBuilderBase::BuildAndOverwriteMetaSound, ;);
+	void BuildAndOverwriteMetaSound(UPARAM(DisplayName = "Existing MetaSound") TScriptInterface<IMetaSoundDocumentInterface> ExistingMetaSound, bool bForceUniqueClassName = false);
 
 	// Builds a transient MetaSound with the provided builder options, copying the underlying MetaSound
 	// managed by this builder and registering it with the MetaSound Node Registry as a unique class. If
@@ -496,6 +501,7 @@ public:
 	void Initialize();
 
 protected:
+	virtual void BuildAndOverwriteMetaSoundInternal(TScriptInterface<IMetaSoundDocumentInterface> ExistingMetaSound, bool bForceUniqueClassName) const PURE_VIRTUAL(UMetaSoundBuilderBase::BuildAndOverwriteMetaSoundInternal, ;);
 	virtual void InitDelegates(Metasound::Frontend::FDocumentModifyDelegates& OutDocumentDelegates);
 
 	virtual void OnAssetReferenceAdded(TScriptInterface<IMetaSoundDocumentInterface> DocInterface) PURE_VIRTUAL(UMetaSoundBuilderBase::OnAssetReferenceAdded, );
diff --git a/Engine/Plugins/Runtime/Metasound/Source/MetasoundEngine/Public/MetasoundBuilderSubsystem.h b/Engine/Plugins/Runtime/Metasound/Source/MetasoundEngine/Public/MetasoundBuilderSubsystem.h
--- a/Engine/Plugins/Runtime/Metasound/Source/MetasoundEngine/Public/MetasoundBuilderSubsystem.h
+++ b/Engine/Plugins/Runtime/Metasound/Source/MetasoundEngine/Public/MetasoundBuilderSubsystem.h
@@ -53,11 +53,11 @@ class METASOUNDENGINE_API UMetaSoundPatchBuilder : public UMetaSoundBuilderBase
 	GENERATED_BODY()
 
 public:
-	virtual void BuildAndOverwriteMetaSound(TScriptInterface<IMetaSoundDocumentInterface> ExistingMetaSound, bool bForceUniqueClassName) override;
 	virtual TScriptInterface<IMetaSoundDocumentInterface> BuildNewMetaSound(FName NameBase) const override;
 	virtual const UClass& GetBaseMetaSoundUClass() const override;
 
 protected:
+	virtual void BuildAndOverwriteMetaSoundInternal(TScriptInterface<IMetaSoundDocumentInterface> ExistingMetaSound, bool bForceUniqueClassName) const override;
 	virtual void OnAssetReferenceAdded(TScriptInterface<IMetaSoundDocumentInterface> DocInterface) override;
 	virtual void OnRemovingAssetReference(TScriptInterface<IMetaSoundDocumentInterface> DocInterface) override;
 
@@ -92,7 +92,6 @@ public:
 
 	const Metasound::Engine::FOutputAudioFormatInfoPair* FindOutputAudioFormatInfo() const;
 
-	virtual void BuildAndOverwriteMetaSound(TScriptInterface<IMetaSoundDocumentInterface> ExistingMetaSound, bool bForceUniqueClassName) override;
 	virtual TScriptInterface<IMetaSoundDocumentInterface> BuildNewMetaSound(FName NameBase) const override;
 	virtual const UClass& GetBaseMetaSoundUClass() const override;
 
@@ -109,6 +108,7 @@ public:
 	void SetQuality(FName Quality);
 
 protected:
+	virtual void BuildAndOverwriteMetaSoundInternal(TScriptInterface<IMetaSoundDocumentInterface> ExistingMetaSound, bool bForceUniqueClassName) const override;
 	virtual void InitDelegates(Metasound::Frontend::FDocumentModifyDelegates& OutDocumentDelegates) override;
 	virtual void OnAssetReferenceAdded(TScriptInterface<IMetaSoundDocumentInterface> DocInterface) override;
 	virtual void OnRemovingAssetReference(TScriptInterface<IMetaSoundDocumentInterface> DocInterface) override;
