Make MetaSound float input default automatic range settings consistent

#rnx

diff --git a/Engine/Plugins/Runtime/Metasound/Source/MetasoundEditor/Private/MetasoundEditorGraphMemberDefaults.cpp b/Engine/Plugins/Runtime/Metasound/Source/MetasoundEditor/Private/MetasoundEditorGraphMemberDefaults.cpp
--- a/Engine/Plugins/Runtime/Metasound/Source/MetasoundEditor/Private/MetasoundEditorGraphMemberDefaults.cpp
+++ b/Engine/Plugins/Runtime/Metasound/Source/MetasoundEditor/Private/MetasoundEditorGraphMemberDefaults.cpp
@@ -672,7 +672,11 @@ void UMetasoundEditorGraphMemberDefaultFloat::PostEditChangeChainProperty(FPrope
 		}
 		else if (WidgetUnitValueType == EAudioUnitsValueType::Frequency)
 		{
-			SetRange(FVector2D(MIN_FILTER_FREQUENCY, MAX_FILTER_FREQUENCY));
+			// Set to a reasonable frequency range if range is set to the default
+			if (Range.Equals(FVector2D(0.0f, 1.0f)))
+			{
+				SetRange(FVector2D(MIN_FILTER_FREQUENCY, MAX_FILTER_FREQUENCY));
+			}
 		}
 		else
 		{
