Fix MetaSound editor graph comment issues
- Fix comment box display issues when duplicating MetaSound assets
- Fix node comments not saving when set via right click menu

#rnx

diff --git a/Engine/Plugins/Runtime/Metasound/Source/MetasoundEditor/Private/MetasoundEditorGraphBuilder.cpp b/Engine/Plugins/Runtime/Metasound/Source/MetasoundEditor/Private/MetasoundEditorGraphBuilder.cpp
--- a/Engine/Plugins/Runtime/Metasound/Source/MetasoundEditor/Private/MetasoundEditorGraphBuilder.cpp
+++ b/Engine/Plugins/Runtime/Metasound/Source/MetasoundEditor/Private/MetasoundEditorGraphBuilder.cpp
@@ -85,6 +85,8 @@ namespace Metasound
 					// Override finalized EdGraphNode ID if a location is found.
 					constexpr bool bUpdateEditorNodeID = true;
 					NewGraphNode->SyncLocationFromFrontendNode(bUpdateEditorNodeID);
+					NewGraphNode->SyncCommentFromFrontendNode();
+
 					return NewGraphNode;
 				}
 
diff --git a/Engine/Plugins/Runtime/Metasound/Source/MetasoundEditor/Private/MetasoundEditorGraphNode.cpp b/Engine/Plugins/Runtime/Metasound/Source/MetasoundEditor/Private/MetasoundEditorGraphNode.cpp
--- a/Engine/Plugins/Runtime/Metasound/Source/MetasoundEditor/Private/MetasoundEditorGraphNode.cpp
+++ b/Engine/Plugins/Runtime/Metasound/Source/MetasoundEditor/Private/MetasoundEditorGraphNode.cpp
@@ -671,6 +671,29 @@ void UMetasoundEditorGraphNode::PostEditImport()
 {
 }
 
+void UMetasoundEditorGraphNode::PostEditChangeProperty(struct FPropertyChangedEvent& InEvent)
+{
+	Super::PostEditChangeProperty(InEvent);
+	
+	if (InEvent.GetPropertyName() == GET_MEMBER_NAME_CHECKED(UEdGraphNode, NodeComment))
+	{
+		UMetaSoundBuilderBase& Builder = GetBuilderChecked();
+		if (const FMetasoundFrontendNode* Node = Builder.GetConstBuilder().FindNode(GetNodeID()))
+		{
+			if (!Node->Style.Display.Comment.Equals(NodeComment))
+			{
+				UObject& MetaSound = GetMetasoundChecked();
+				MetaSound.Modify();
+				EMetaSoundBuilderResult Result;
+				Builder.SetNodeComment(Node->GetID(), NodeComment, Result);
+				ensure(Result == EMetaSoundBuilderResult::Succeeded);
+				Builder.SetNodeCommentVisible(Node->GetID(), bCommentBubbleMakeVisible, Result);
+				ensure(Result == EMetaSoundBuilderResult::Succeeded);
+			}
+		}
+	}
+}
+
 void UMetasoundEditorGraphNode::PostEditUndo()
 {
 	using namespace Metasound::Editor;
diff --git a/Engine/Plugins/Runtime/Metasound/Source/MetasoundEditor/Private/MetasoundEditorGraphSchema.cpp b/Engine/Plugins/Runtime/Metasound/Source/MetasoundEditor/Private/MetasoundEditorGraphSchema.cpp
--- a/Engine/Plugins/Runtime/Metasound/Source/MetasoundEditor/Private/MetasoundEditorGraphSchema.cpp
+++ b/Engine/Plugins/Runtime/Metasound/Source/MetasoundEditor/Private/MetasoundEditorGraphSchema.cpp
@@ -1180,12 +1180,12 @@ UEdGraphNode* FMetasoundGraphSchemaAction_NewComment::PerformAction(UEdGraph* Pa
 			NewComment->NodePosX = Location.X;
 			NewComment->NodePosY = Location.Y;
 			NewComment->SnapToGrid(GetDefault<UEditorStyleSettings>()->GridSnapSize);
-
-			// Applies new snapped node data to frontend doc data
-			FMetaSoundFrontendGraphComment& FrontendComment = MetaSoundGraph->GetBuilderChecked().FindOrAddGraphComment(NewComment->GetCommentID());
-			UMetasoundEditorGraphCommentNode::ConvertToFrontendComment(*NewComment, FrontendComment);
 		}
 
+		// Applies new node data to frontend doc data
+		FMetaSoundFrontendGraphComment& FrontendComment = MetaSoundGraph->GetBuilderChecked().FindOrAddGraphComment(NewComment->GetCommentID());
+		UMetasoundEditorGraphCommentNode::ConvertToFrontendComment(*NewComment, FrontendComment);
+
 		return NewComment;
 	}
 
diff --git a/Engine/Plugins/Runtime/Metasound/Source/MetasoundEditor/Public/MetasoundEditorGraphNode.h b/Engine/Plugins/Runtime/Metasound/Source/MetasoundEditor/Public/MetasoundEditorGraphNode.h
--- a/Engine/Plugins/Runtime/Metasound/Source/MetasoundEditor/Public/MetasoundEditorGraphNode.h
+++ b/Engine/Plugins/Runtime/Metasound/Source/MetasoundEditor/Public/MetasoundEditorGraphNode.h
@@ -121,6 +121,7 @@ public:
 	virtual void PreSave(FObjectPreSaveContext InSaveContext) override;
 
 	virtual void PostLoad() override;
+	virtual void PostEditChangeProperty(struct FPropertyChangedEvent& InEvent) override;
 	virtual void PostEditUndo() override;
 
 	virtual void PostEditImport() override;
