[Metasounds] Warning flag does not appear on One-Shot interface node after it's been removed and re-added

#rnx

diff --git a/Engine/Plugins/Runtime/Metasound/Source/MetasoundEditor/Private/MetasoundEditorGraphBuilder.cpp b/Engine/Plugins/Runtime/Metasound/Source/MetasoundEditor/Private/MetasoundEditorGraphBuilder.cpp
--- a/Engine/Plugins/Runtime/Metasound/Source/MetasoundEditor/Private/MetasoundEditorGraphBuilder.cpp
+++ b/Engine/Plugins/Runtime/Metasound/Source/MetasoundEditor/Private/MetasoundEditorGraphBuilder.cpp
@@ -622,6 +622,7 @@ namespace Metasound
 							}
 						};
 
+						Builder.GetConstDocumentChecked().Metadata.ModifyContext.AddNodeIDModified(InNodeID);
 						return GraphBuilderPrivate::AddNode<UMetasoundEditorGraphOutputNode>(InMetaSound, InitNodeFunc, bInSelectNewNode);
 					}
 				}
