Fix for missing default when swapping variable data types

#rnx

diff --git a/Engine/Plugins/Runtime/Metasound/Source/MetasoundEditor/Private/MetasoundEditorGraph.cpp b/Engine/Plugins/Runtime/Metasound/Source/MetasoundEditor/Private/MetasoundEditorGraph.cpp
--- a/Engine/Plugins/Runtime/Metasound/Source/MetasoundEditor/Private/MetasoundEditorGraph.cpp
+++ b/Engine/Plugins/Runtime/Metasound/Source/MetasoundEditor/Private/MetasoundEditorGraph.cpp
@@ -1672,6 +1672,13 @@ void UMetasoundEditorGraphVariable::SetDataType(FName InNewType, bool bPostTrans
 		TypeName = InNewType;
 		InitializeLiteral();
 
+		{
+			FMetasoundFrontendLiteral DefaultLiteral;
+			DefaultLiteral.SetFromLiteral(IDataTypeRegistry::Get().CreateDefaultLiteral(GetDataType()));
+			check(Literal); // Should exist from prior InitializeLiteral() call
+			Literal->SetFromLiteral(DefaultLiteral);
+		}
+
 		// Add the nodes with the same identifier data but new datatype.
 		UObject& Metasound = Graph->GetMetasoundChecked();
 		AddVariableNodes(Metasound, FrontendGraph, CachedData.Locations);
