MetaSound Editor Output Node Location/Page Fixes
- Fix bug with moving initially placed editor output node moving same node on other pages in some cases
- Fix removing editor output nodes removing output nodes from all pages (and root output graph vertex)
- Disallow required interface input/output nodes from being removed from the ed graph
- Fix for occasional audition page desync

#rnx

diff --git a/Engine/Plugins/Runtime/Metasound/Source/MetasoundEditor/Private/MetasoundEditor.cpp b/Engine/Plugins/Runtime/Metasound/Source/MetasoundEditor/Private/MetasoundEditor.cpp
--- a/Engine/Plugins/Runtime/Metasound/Source/MetasoundEditor/Private/MetasoundEditor.cpp
+++ b/Engine/Plugins/Runtime/Metasound/Source/MetasoundEditor/Private/MetasoundEditor.cpp
@@ -319,7 +319,7 @@ namespace Metasound
 			DRAG_DROP_OPERATOR_TYPE(FMetaSoundDragDropMemberAction, FGraphSchemaActionDragDropAction)
 
 			virtual FReply DroppedOnPanel(const TSharedRef<SWidget>& InPanel, FVector2D InScreenPosition, FVector2D InGraphPosition, UEdGraph& InGraph) override
-			{			
+			{
 				if (!GraphMember.IsValid() || &InGraph != GraphMember->GetOwningGraph())
 				{
 					return FReply::Unhandled();
@@ -1503,6 +1503,8 @@ namespace Metasound
 			{
 				ToolMenus->RefreshAllWidgets();
 			}
+
+			SyncAuditionState();
 		}
 
 		void FEditor::NotifyAssetPrimeInProgress()
@@ -2714,12 +2716,14 @@ namespace Metasound
 
 			if (Builder.IsValid())
 			{
-				constexpr bool bOpenEditor = false; // Already Focused
-				constexpr bool bPostTransaction = false;
-				const FMetaSoundFrontendDocumentBuilder& DocBuilder = Builder->GetConstBuilder();
-				const FGuid BuildPageID = DocBuilder.GetBuildPageID();
-
-				UMetaSoundEditorSubsystem::GetChecked().SetFocusedPage(*Builder.Get(), BuildPageID, bOpenEditor, bPostTransaction);
+				if (bSetAuditionFocus)
+				{
+					constexpr bool bOpenEditor = false; // Already Focused
+					constexpr bool bPostTransaction = false;
+					const FMetaSoundFrontendDocumentBuilder& DocBuilder = Builder->GetConstBuilder();
+					const FGuid BuildPageID = DocBuilder.GetBuildPageID();
+					UMetaSoundEditorSubsystem::GetChecked().SetFocusedPage(*Builder.Get(), BuildPageID, bOpenEditor, bPostTransaction);
+				}
 
 				if (const UMetasoundEditorSettings* EdSettings = GetDefault<UMetasoundEditorSettings>())
 				{
diff --git a/Engine/Plugins/Runtime/Metasound/Source/MetasoundEditor/Private/MetasoundEditorGraphBuilder.cpp b/Engine/Plugins/Runtime/Metasound/Source/MetasoundEditor/Private/MetasoundEditorGraphBuilder.cpp
--- a/Engine/Plugins/Runtime/Metasound/Source/MetasoundEditor/Private/MetasoundEditorGraphBuilder.cpp
+++ b/Engine/Plugins/Runtime/Metasound/Source/MetasoundEditor/Private/MetasoundEditorGraphBuilder.cpp
@@ -2284,7 +2284,7 @@ namespace Metasound
 						if (const FMetasoundFrontendClassOutput* Output = ClassOutputPtr.Get())
 						{
 							const FGuid& InitialID = Node->GetNodeID();
-							if (Node->GetNodeHandle()->GetID() != Output->NodeID)
+							if (NodeHandle->GetID() != Output->NodeID)
 							{
 								Node->SetNodeID(Output->NodeID);
 								UE_LOG(LogMetasoundEditor, Verbose, TEXT("Editor Output Node '%s' interface versioned"), *Node->GetDisplayName().ToString());
@@ -2340,21 +2340,28 @@ namespace Metasound
 					UMetasoundEditorGraphNode* EditorNode = EditorNodes[j];
 					if (EditorNode->GetNodeID() == Node->GetID())
 					{
-						bFoundEditorNode = true;
-						FAssociatedNodes& AssociatedNodeData = AssociatedNodes.FindOrAdd(Node->GetID());
-						if (AssociatedNodeData.Node->IsValid())
+						// Editor node may have the same Frontend NodeID as another page,
+						// but may have been assigned a different Editor NodeID, so synchronize
+						// from frontend data here only if node location was able to sync.
+						constexpr bool bUpdateEditorNodeID = true;
+						const bool bLocationFound = EditorNode->SyncLocationFromFrontendNode(bUpdateEditorNodeID);
+						if (bLocationFound)
 						{
-							ensure(AssociatedNodeData.Node == Node);
-						}
-						else
-						{
-							AssociatedNodeData.Node = Node;
-						}
+							bFoundEditorNode = true;
+							FAssociatedNodes& AssociatedNodeData = AssociatedNodes.FindOrAdd(Node->GetID());
+							if (AssociatedNodeData.Node->IsValid())
+							{
+								ensure(AssociatedNodeData.Node == Node);
+							}
+							else
+							{
+								AssociatedNodeData.Node = Node;
+							}
 
-						EditorNode->SyncLocationFromFrontendNode();
-						EditorNode->SyncCommentFromFrontendNode();
-						AssociatedNodeData.EditorNodes.Add(EditorNode);
-						EditorNodes.RemoveAtSwap(j, EAllowShrinking::No);
+							EditorNode->SyncCommentFromFrontendNode();
+							AssociatedNodeData.EditorNodes.Add(EditorNode);
+							EditorNodes.RemoveAtSwap(j, EAllowShrinking::No);
+						}
 					}
 				}
 
diff --git a/Engine/Plugins/Runtime/Metasound/Source/MetasoundEditor/Private/MetasoundEditorGraphNode.cpp b/Engine/Plugins/Runtime/Metasound/Source/MetasoundEditor/Private/MetasoundEditorGraphNode.cpp
--- a/Engine/Plugins/Runtime/Metasound/Source/MetasoundEditor/Private/MetasoundEditorGraphNode.cpp
+++ b/Engine/Plugins/Runtime/Metasound/Source/MetasoundEditor/Private/MetasoundEditorGraphNode.cpp
@@ -70,10 +70,9 @@ void UMetasoundEditorGraphNode::UpdateFrontendNodeLocation(const FVector2D& InLo
 	
 	ensureMsgf(NodeGuid.IsValid(), TEXT("Cannot update frontend node location prior to node guid being finalized."));
 
-	FNodeHandle NodeHandle = GetNodeHandle();
-	FMetasoundFrontendNodeStyle Style = NodeHandle->GetNodeStyle();
-	Style.Display.Locations.FindOrAdd(NodeGuid) = InLocation;
-	NodeHandle->SetNodeStyle(Style);
+	const FGuid NodeID = GetNodeID();
+	UMetaSoundBuilderBase& Builder = GetBuilderChecked();
+	Builder.GetBuilder().SetNodeLocation(NodeID, InLocation, &NodeGuid);
 }
 
 bool UMetasoundEditorGraphNode::ShowNodeDebugData()
@@ -912,10 +911,18 @@ bool UMetasoundEditorGraphOutputNode::RemoveFromDocument() const
 {
 	if (Output)
 	{
+		// When removing ed graph output nodes, just remove the location as all frontend page graphs
+		// require the graph vertex node to exist and contain matching NodeIDs across all pages.
+		const FGuid& NodeID = GetNodeID();
 		UMetaSoundBuilderBase& Builder = GetBuilderChecked();
-		EMetaSoundBuilderResult Result = EMetaSoundBuilderResult::Failed;
-		Builder.RemoveGraphOutput(Output->GetMemberName(), Result);
-		return Result == EMetaSoundBuilderResult::Succeeded;
+		const int32 NumLocationsRemoved = Builder.GetBuilder().RemoveNodeLocation(NodeID);
+
+		if (FMetasoundAssetBase* MetaSoundAsset = Metasound::IMetasoundUObjectRegistry::Get().GetObjectAsAssetBase(&GetMetasoundChecked()))
+		{
+			MetaSoundAsset->GetModifyContext().AddNodeIDsModified({ NodeID });
+		}
+
+		return NumLocationsRemoved > 0;
 	}
 
 	return false;
@@ -1032,7 +1039,12 @@ bool UMetasoundEditorGraphOutputNode::CanUserDeleteNode() const
 {
 	if (const UMetasoundEditorGraphVertex* Vertex = Cast<UMetasoundEditorGraphVertex>(GetMember()))
 	{
-		return !Vertex->IsInterfaceMember();
+		FMetasoundFrontendInterface MemberInterface;
+		if (Vertex->IsInterfaceMember(&MemberInterface))
+		{
+			FText RequiredText;
+			return !MemberInterface.IsMemberOutputRequired(Vertex->GetMemberName(), RequiredText);
+		}
 	}
 	
 	return true;
diff --git a/Engine/Plugins/Runtime/Metasound/Source/MetasoundFrontend/Private/MetasoundFrontendDocumentBuilder.cpp b/Engine/Plugins/Runtime/Metasound/Source/MetasoundFrontend/Private/MetasoundFrontendDocumentBuilder.cpp
--- a/Engine/Plugins/Runtime/Metasound/Source/MetasoundFrontend/Private/MetasoundFrontendDocumentBuilder.cpp
+++ b/Engine/Plugins/Runtime/Metasound/Source/MetasoundFrontend/Private/MetasoundFrontendDocumentBuilder.cpp
@@ -8,6 +8,7 @@
 #include "Algo/NoneOf.h"
 #include "Algo/Sort.h"
 #include "Algo/Transform.h"
+#include "AudioParameter.h"
 #include "Interfaces/MetasoundFrontendInterfaceBindingRegistry.h"
 #include "Interfaces/MetasoundFrontendInterfaceRegistry.h"
 #include "MetasoundAssetBase.h"
@@ -90,6 +91,85 @@ namespace Metasound::Frontend
 			return false;
 		}
 
+		bool NameContainsInterfaceNamespace(FName VertexName, FMetasoundFrontendInterface* OutInterface)
+		{
+			using namespace Metasound::Frontend;
+
+			FName InterfaceNamespace;
+			FName ParamName;
+			Audio::FParameterPath::SplitName(VertexName, InterfaceNamespace, ParamName);
+
+			FMetasoundFrontendInterface FoundInterface;
+			if (!InterfaceNamespace.IsNone() && ISearchEngine::Get().FindInterfaceWithHighestVersion(InterfaceNamespace, FoundInterface))
+			{
+				if (OutInterface)
+				{
+					*OutInterface = MoveTemp(FoundInterface);
+				}
+				return true;
+			}
+
+			if (OutInterface)
+			{
+				*OutInterface = { };
+			}
+			return false;
+		}
+
+		bool IsInterfaceInput(FName InputName, FName TypeName, FMetasoundFrontendInterface* OutInterface)
+		{
+			FMetasoundFrontendInterface Interface;
+			if (NameContainsInterfaceNamespace(InputName, &Interface))
+			{
+				auto IsInput = [&InputName, &TypeName](const FMetasoundFrontendClassInput& InterfaceInput)
+				{
+					return InputName == InterfaceInput.Name && InterfaceInput.TypeName == TypeName;
+				};
+
+				if (Interface.Inputs.ContainsByPredicate(IsInput))
+				{
+					if (OutInterface)
+					{
+						*OutInterface = MoveTemp(Interface);
+					}
+					return true;
+				}
+			}
+
+			if (OutInterface)
+			{
+				*OutInterface = { };
+			}
+			return false;
+		}
+
+		bool IsInterfaceOutput(FName OutputName, FName TypeName, FMetasoundFrontendInterface* OutInterface)
+		{
+			FMetasoundFrontendInterface Interface;
+			if (NameContainsInterfaceNamespace(OutputName, &Interface))
+			{
+				auto IsOutput = [&OutputName, &TypeName](const FMetasoundFrontendClassInput& InterfaceOutput)
+				{
+					return OutputName == InterfaceOutput.Name && InterfaceOutput.TypeName == TypeName;
+				};
+
+				if (Interface.Outputs.ContainsByPredicate(IsOutput))
+				{
+					if (OutInterface)
+					{
+						*OutInterface = MoveTemp(Interface);
+					}
+					return true;
+				}
+			}
+
+			if (OutInterface)
+			{
+				*OutInterface = { };
+			}
+			return false;
+		}
+
 		bool TryGetInterfaceBoundEdges(
 			const FGuid& InFromNodeID,
 			const TSet<FMetasoundFrontendVersion>& InFromNodeInterfaces,
@@ -440,14 +520,24 @@ namespace Metasound::Frontend
 					}
 				}
 
-				// Iterate through sorted map in sequence, slotting in new locations after existing swapped nodes with predefined locations.
+				// Iterate through sorted map in sequence, slotting in new locations after
+				// existing swapped nodes with predefined locations relative to one another.
 				for (TPair<int32, FMetasoundFrontendNode*>& Pair : SortOrderToNode)
 				{
 					FMetasoundFrontendNode* Node = Pair.Value;
 					const FName NodeName = Node->Name;
 					if (AddedNames.Contains(NodeName))
 					{
-						Node->Style.Display.Locations.Add(FGuid::NewGuid(), NextLocation);
+						bool bAddedLocation = false;
+						for (TPair<FGuid, FVector2D>& LocationPair : Node->Style.Display.Locations)
+						{
+							bAddedLocation = true;
+							LocationPair.Value = NextLocation;
+						}
+						if (!bAddedLocation)
+						{
+							Node->Style.Display.Locations.Add(FGuid::NewGuid(), NextLocation);
+						}
 						NextLocation += DisplayStyle::NodeLayout::DefaultOffsetY;
 					}
 					else
@@ -1032,6 +1122,16 @@ const FMetasoundFrontendNode* FMetaSoundFrontendDocumentBuilder::AddGraphInput(c
 			DocumentBuilderPrivate::SetDefaultLiteralOnInputNode(InOutNode, InClassInput);
 		};
 
+#if WITH_EDITORONLY_DATA
+		bool bIsRequired = false;
+		FMetasoundFrontendInterface Interface;
+		if (DocumentBuilderPrivate::IsInterfaceInput(InClassInput.Name, InClassInput.TypeName, &Interface))
+		{
+			FText RequiredText;
+			bIsRequired = Interface.IsMemberInputRequired(InClassInput.Name, RequiredText);
+		}
+#endif // WITH_EDITORONLY_DATA
+
 		// Must add input node to all paged graphs to maintain API parity for all page implementations
 		FMetasoundFrontendNode* NewNode = nullptr;
 		RootGraph.IterateGraphPages([&](const FMetasoundFrontendGraph& Graph)
@@ -1043,6 +1143,17 @@ const FMetasoundFrontendNode* FMetaSoundFrontendDocumentBuilder::AddGraphInput(c
 				NewNode = NewPageNode;
 			}
 
+#if WITH_EDITORONLY_DATA
+			if (bIsRequired)
+			{
+				// LocationGuid corresponds with the assigned editor graph node guid when dynamically created.
+				// This is added if this is an interface member that is required to force page to create visual
+				// representation that can inform the user of its required state.
+				FGuid LocationGuid = FDocumentIDGenerator::Get().CreateVertexID(Document);
+				SetNodeLocation(NewNode->GetID(), FVector2D::ZeroVector, &LocationGuid, &Graph.PageID);
+			}
+#endif // WITH_EDITORONLY_DATA
+
 			// Remove the default literal on the node added during the "FinalizeNode" call. This matches how 
 			// nodes are serialized in editor. The default literals are only stored on the FMetasoundFrontendClassInputs.
 			NewPageNode->InputLiterals.Reset();
@@ -1106,6 +1217,16 @@ const FMetasoundFrontendNode* FMetaSoundFrontendDocumentBuilder::AddGraphOutput(
 			DocumentBuilderPrivate::SetNodeAndVertexNames(InOutNode, InClassOutput);
 		};
 
+#if WITH_EDITORONLY_DATA
+		bool bIsRequired = false;
+		FMetasoundFrontendInterface Interface;
+		if (DocumentBuilderPrivate::IsInterfaceOutput(InClassOutput.Name, InClassOutput.TypeName, &Interface))
+		{
+			FText RequiredText;
+			bIsRequired = Interface.IsMemberOutputRequired(InClassOutput.Name, RequiredText);
+		}
+#endif // WITH_EDITORONLY_DATA
+
 		bool bAddedNodes = true;
 		FMetasoundFrontendNode* NewNodeToReturn = nullptr;
 		FMetasoundFrontendDocument& Document = GetDocumentChecked();
@@ -1117,6 +1238,18 @@ const FMetasoundFrontendNode* FMetaSoundFrontendDocumentBuilder::AddGraphOutput(
 			{
 				NewNodeToReturn = NewNode;
 			}
+
+#if WITH_EDITORONLY_DATA
+			if (bIsRequired)
+			{
+				// LocationGuid corresponds with the assigned editor graph node guid when dynamically created.
+				// This is added if this is an interface member that is required to force page to create visual
+				// representation that can inform the user of its required state.
+				FGuid LocationGuid = FDocumentIDGenerator::Get().CreateVertexID(Document);
+				SetNodeLocation(NewNode->GetID(), FVector2D::ZeroVector, &LocationGuid, &Graph.PageID);
+			}
+#endif // WITH_EDITORONLY_DATA
+
 			bAddedNodes &= NewNode != nullptr;
 		});
 
