Fix static compilation unchecked nullref warning

#rnx

diff --git a/Engine/Plugins/Runtime/Metasound/Source/MetasoundFrontend/Private/MetasoundFrontendDocumentBuilder.cpp b/Engine/Plugins/Runtime/Metasound/Source/MetasoundFrontend/Private/MetasoundFrontendDocumentBuilder.cpp
--- a/Engine/Plugins/Runtime/Metasound/Source/MetasoundFrontend/Private/MetasoundFrontendDocumentBuilder.cpp
+++ b/Engine/Plugins/Runtime/Metasound/Source/MetasoundFrontend/Private/MetasoundFrontendDocumentBuilder.cpp
@@ -1150,7 +1150,7 @@ const FMetasoundFrontendNode* FMetaSoundFrontendDocumentBuilder::AddGraphInput(c
 				// This is added if this is an interface member that is required to force page to create visual
 				// representation that can inform the user of its required state.
 				FGuid LocationGuid = FDocumentIDGenerator::Get().CreateVertexID(Document);
-				SetNodeLocation(NewNode->GetID(), FVector2D::ZeroVector, &LocationGuid, &Graph.PageID);
+				SetNodeLocation(InClassInput.NodeID, FVector2D::ZeroVector, &LocationGuid, &Graph.PageID);
 			}
 #endif // WITH_EDITORONLY_DATA
 
@@ -1246,7 +1246,7 @@ const FMetasoundFrontendNode* FMetaSoundFrontendDocumentBuilder::AddGraphOutput(
 				// This is added if this is an interface member that is required to force page to create visual
 				// representation that can inform the user of its required state.
 				FGuid LocationGuid = FDocumentIDGenerator::Get().CreateVertexID(Document);
-				SetNodeLocation(NewNode->GetID(), FVector2D::ZeroVector, &LocationGuid, &Graph.PageID);
+				SetNodeLocation(InClassOutput.NodeID, FVector2D::ZeroVector, &LocationGuid, &Graph.PageID);
 			}
 #endif // WITH_EDITORONLY_DATA
 
