Fix missing modify call when swapping MetaSound page focus
#rnx

diff --git a/Engine/Plugins/Runtime/Metasound/Source/MetasoundEditor/Private/MetasoundEditor.cpp b/Engine/Plugins/Runtime/Metasound/Source/MetasoundEditor/Private/MetasoundEditor.cpp
--- a/Engine/Plugins/Runtime/Metasound/Source/MetasoundEditor/Private/MetasoundEditor.cpp
+++ b/Engine/Plugins/Runtime/Metasound/Source/MetasoundEditor/Private/MetasoundEditor.cpp
@@ -1505,6 +1505,7 @@ namespace Metasound
 			}
 
 			SyncAuditionState();
+			bRefreshGraph = true;
 		}
 
 		void FEditor::NotifyAssetPrimeInProgress()
diff --git a/Engine/Plugins/Runtime/Metasound/Source/MetasoundEditor/Private/MetasoundEditorSubsystem.cpp b/Engine/Plugins/Runtime/Metasound/Source/MetasoundEditor/Private/MetasoundEditorSubsystem.cpp
--- a/Engine/Plugins/Runtime/Metasound/Source/MetasoundEditor/Private/MetasoundEditorSubsystem.cpp
+++ b/Engine/Plugins/Runtime/Metasound/Source/MetasoundEditor/Private/MetasoundEditorSubsystem.cpp
@@ -446,6 +446,7 @@ bool UMetaSoundEditorSubsystem::SetFocusedPageInternal(FName PageName, const FGu
 {
 	using namespace Metasound::Frontend;
 
+	const FScopedTransaction Transaction(FText::Format(LOCTEXT("SetFocusedPageTransactionFormat", "Set Focused Page '{0}'"), FText::FromName(PageName)), bPostTransaction);
 	bool bAuditionPageSet = false;
 	// Must set audition target page before setting build page ID as listeners
 	// to build page ID changes need to reliably be able to adjust to newly assigned
@@ -456,6 +457,7 @@ bool UMetaSoundEditorSubsystem::SetFocusedPageInternal(FName PageName, const FGu
 	{
 		if (EditorSettings->AuditionPage != PageName)
 		{
+			EditorSettings->Modify();
 			EditorSettings->AuditionPage = PageName;
 			bAuditionPageSet = true;
 		}
@@ -464,7 +466,6 @@ bool UMetaSoundEditorSubsystem::SetFocusedPageInternal(FName PageName, const FGu
 	const FMetaSoundFrontendDocumentBuilder& DocBuilder = Builder.GetConstBuilder();
 	if (DocBuilder.GetBuildPageID() != InPageID)
 	{
-		const FScopedTransaction Transaction(FText::Format(LOCTEXT("SetFocusedPageTransactionFormat", "Set Focused Page '{0}'"), FText::FromName(PageName)), bPostTransaction);
 		Builder.Modify();
 		UObject& MetaSound = DocBuilder.CastDocumentObjectChecked<UObject>();
 		if (Builder.GetBuilder().SetBuildPageID(InPageID))
