Renaming a MetaSound patch results in duplicates/crashes if MetaSound Editor is Open

#rnx

diff --git a/Engine/Plugins/Runtime/Metasound/Source/MetasoundEditor/Private/MetasoundEditorGraphBuilder.cpp b/Engine/Plugins/Runtime/Metasound/Source/MetasoundEditor/Private/MetasoundEditorGraphBuilder.cpp
--- a/Engine/Plugins/Runtime/Metasound/Source/MetasoundEditor/Private/MetasoundEditorGraphBuilder.cpp
+++ b/Engine/Plugins/Runtime/Metasound/Source/MetasoundEditor/Private/MetasoundEditorGraphBuilder.cpp
@@ -1783,36 +1783,6 @@ namespace Metasound
 			}
 		}
 
-		void FGraphBuilder::UnregisterGraphWithFrontend(UObject& InMetaSound)
-		{
-			FMetasoundAssetBase* MetaSoundAsset = IMetasoundUObjectRegistry::Get().GetObjectAsAssetBase(&InMetaSound);
-			if (!ensure(MetaSoundAsset))
-			{
-				return;
-			}
-
-			if (GEditor)
-			{
-				TArray<UObject*> EditedAssets = GEditor->GetEditorSubsystem<UAssetEditorSubsystem>()->GetAllEditedAssets();
-				for (UObject* Asset : EditedAssets)
-				{
-					if (Asset != &InMetaSound)
-					{
-						if (FMetasoundAssetBase* EditedMetaSound = IMetasoundUObjectRegistry::Get().GetObjectAsAssetBase(Asset))
-						{
-							EditedMetaSound->RebuildReferencedAssetClasses();
-							if (EditedMetaSound->IsReferencedAsset(*MetaSoundAsset))
-							{
-								EditedMetaSound->GetModifyContext().SetDocumentModified();
-							}
-						}
-					}
-				}
-			}
-
-			MetaSoundAsset->UnregisterGraphWithFrontend();
-		}
-
 		bool FGraphBuilder::IsMatchingInputHandleAndPin(const Frontend::FConstInputHandle& InInputHandle, const UEdGraphPin& InEditorPin)
 		{
 			if (InEditorPin.Direction != EGPD_Input)
diff --git a/Engine/Plugins/Runtime/Metasound/Source/MetasoundEditor/Private/MetasoundEditorGraphBuilder.h b/Engine/Plugins/Runtime/Metasound/Source/MetasoundEditor/Private/MetasoundEditorGraphBuilder.h
--- a/Engine/Plugins/Runtime/Metasound/Source/MetasoundEditor/Private/MetasoundEditorGraphBuilder.h
+++ b/Engine/Plugins/Runtime/Metasound/Source/MetasoundEditor/Private/MetasoundEditorGraphBuilder.h
@@ -81,9 +81,6 @@ namespace Metasound
 			// referencing graphs open in editors)
 			static void RegisterGraphWithFrontend(UObject& InMetaSound, bool bInForceViewSynchronization = false);
 
-			// Wraps RegisterGraphWithFrontend logic in Frontend with any additional logic required to refresh editor & respective editor object state.
-			static void UnregisterGraphWithFrontend(UObject& InMetaSound);
-
 			// Returns whether pin category is a custom MetaSound DataType
 			static bool IsPinCategoryMetaSoundCustomDataType(FName InPinCategoryName);
 
diff --git a/Engine/Plugins/Runtime/Metasound/Source/MetasoundEditor/Private/MetasoundEditorModule.cpp b/Engine/Plugins/Runtime/Metasound/Source/MetasoundEditor/Private/MetasoundEditorModule.cpp
--- a/Engine/Plugins/Runtime/Metasound/Source/MetasoundEditor/Private/MetasoundEditorModule.cpp
+++ b/Engine/Plugins/Runtime/Metasound/Source/MetasoundEditor/Private/MetasoundEditorModule.cpp
@@ -782,13 +782,46 @@ namespace Metasound
 				IMetasoundEngineModule& MetaSoundEngineModule = FModuleManager::LoadModuleChecked<IMetasoundEngineModule>("MetasoundEngine");
 				
 				// Bind delegates for MetaSound registration in the asset registry
-				MetaSoundEngineModule.GetOnGraphRegisteredDelegate().BindLambda([](UObject& InMetaSound, bool bForceViewSynchronization)
+				MetaSoundEngineModule.GetOnGraphRegisteredDelegate().BindLambda([](UObject& InMetaSound, Engine::ERegistrationAssetContext AssetContext)
 				{
+					using namespace Engine;
+
+					// Use the editor version of RegisterWithFrontend so it refreshes any open MetaSound editors.
+					const bool bForceViewSynchronization = AssetContext == ERegistrationAssetContext::Renaming;
 					FGraphBuilder::RegisterGraphWithFrontend(InMetaSound, bForceViewSynchronization);
 				});
-				MetaSoundEngineModule.GetOnGraphUnregisteredDelegate().BindLambda([](UObject& InMetaSound)
+				MetaSoundEngineModule.GetOnGraphUnregisteredDelegate().BindLambda([](UObject& InMetaSound, Engine::ERegistrationAssetContext AssetContext)
 				{
-					FGraphBuilder::UnregisterGraphWithFrontend(InMetaSound);
+					using namespace Engine;
+
+					switch(AssetContext)
+					{
+						case ERegistrationAssetContext::Reloading:
+						case ERegistrationAssetContext::Removing:
+						case ERegistrationAssetContext::Renaming:
+						{
+							UAssetEditorSubsystem* AssetEditorSubsystem = GEditor->GetEditorSubsystem<UAssetEditorSubsystem>();
+							if (AssetEditorSubsystem)
+							{
+								// Close the editors so the internal reference to the builder doesn't suddenly 
+								// upon GC enter an invalid state (pointing to a null MetaSound asset)
+								AssetEditorSubsystem->CloseAllEditorsForAsset(&InMetaSound);
+							}
+							break;
+						}
+
+						case ERegistrationAssetContext::None:
+						default:
+						{
+							break;
+						}
+					}
+
+					Metasound::IMetasoundUObjectRegistry& UObjectRegistry = Metasound::IMetasoundUObjectRegistry::Get();
+					if (FMetasoundAssetBase* AssetBase = UObjectRegistry.GetObjectAsAssetBase(&InMetaSound))
+					{
+						AssetBase->UnregisterGraphWithFrontend();
+					}
 				});
 
 				// Required to ensure logic to order nodes for presets exclusive to
diff --git a/Engine/Plugins/Runtime/Metasound/Source/MetasoundEngine/Private/MetasoundEngineModule.cpp b/Engine/Plugins/Runtime/Metasound/Source/MetasoundEngine/Private/MetasoundEngineModule.cpp
--- a/Engine/Plugins/Runtime/Metasound/Source/MetasoundEngine/Private/MetasoundEngineModule.cpp
+++ b/Engine/Plugins/Runtime/Metasound/Source/MetasoundEngine/Private/MetasoundEngineModule.cpp
@@ -191,7 +191,7 @@ namespace Metasound::Engine
 		{
 			if (UObject* AssetObject = InAssetData.GetAsset())
 			{
-				OnGraphRegister.ExecuteIfBound(*AssetObject, /*bInForceViewSynchronization=*/false);
+				OnGraphRegister.ExecuteIfBound(*AssetObject, ERegistrationAssetContext::None);
 			}
 		}
 		else
@@ -218,7 +218,7 @@ namespace Metasound::Engine
 					check(MetaSoundAsset);
 					if (!MetaSoundAsset->IsRegistered())
 					{
-						OnGraphRegister.ExecuteIfBound(*MetaSoundObj, /*bInForceViewSynchronization=*/false);
+						OnGraphRegister.ExecuteIfBound(*MetaSoundObj, ERegistrationAssetContext::None);
 					}
 				}
 
@@ -323,10 +323,8 @@ namespace Metasound::Engine
 			{
 				if (IsAssetMetaSound(Obj))
 				{
+					OnGraphUnregister.ExecuteIfBound(*Obj, ERegistrationAssetContext::Reloading);
 					IMetaSoundAssetManager::GetChecked().RemoveAsset(*Pair.Key);
-
-					// Use the editor version of UnregisterWithFrontend so it refreshes any open MetaSound editors
-					OnGraphUnregister.ExecuteIfBound(*Obj);
 				}
 			}
 
@@ -335,9 +333,7 @@ namespace Metasound::Engine
 				if (IsAssetMetaSound(Obj))
 				{
 					IMetaSoundAssetManager::GetChecked().AddOrUpdateAsset(*Pair.Value);
-
-					// Use the editor version of RegisterWithFrontend so it refreshes any open MetaSound editors
-					OnGraphRegister.ExecuteIfBound(*Obj, /*bInForceViewSynchronization=*/false);
+					OnGraphRegister.ExecuteIfBound(*Obj, ERegistrationAssetContext::Reloading);
 				}
 			}
 		}
@@ -382,7 +378,7 @@ namespace Metasound::Engine
 				// occurred on object destroy.
 				if (UObject* AssetObject = InAssetData.GetSoftObjectPath().ResolveObject())
 				{
-					OnGraphUnregister.ExecuteIfBound(*AssetObject);
+					OnGraphUnregister.ExecuteIfBound(*AssetObject, ERegistrationAssetContext::Removing);
 				}
 
 				IMetaSoundAssetManager::GetChecked().RemoveAsset(InAssetData);
@@ -412,15 +408,14 @@ namespace Metasound::Engine
 				bool bIsRegistered = AssetBase->IsRegistered();
 				if (bIsRegistered)
 				{
-					OnGraphUnregister.ExecuteIfBound(*AssetObject);
+					OnGraphUnregister.ExecuteIfBound(*AssetObject, ERegistrationAssetContext::Renaming);
 				}
 
 				IMetaSoundAssetManager::GetChecked().RenameAsset(InAssetData, InOldObjectPath);
 
 				if (bIsRegistered)
 				{
-					constexpr bool bForceViewSynchronization = true;
-					OnGraphRegister.ExecuteIfBound(*AssetObject, bForceViewSynchronization);
+					OnGraphRegister.ExecuteIfBound(*AssetObject, ERegistrationAssetContext::Renaming);
 				}
 			}
 		}
diff --git a/Engine/Plugins/Runtime/Metasound/Source/MetasoundEngine/Public/IMetasoundEngineModule.h b/Engine/Plugins/Runtime/Metasound/Source/MetasoundEngine/Public/IMetasoundEngineModule.h
--- a/Engine/Plugins/Runtime/Metasound/Source/MetasoundEngine/Public/IMetasoundEngineModule.h
+++ b/Engine/Plugins/Runtime/Metasound/Source/MetasoundEngine/Public/IMetasoundEngineModule.h
@@ -29,8 +29,16 @@ namespace Metasound::Engine
 		Canceled = 4
 	};
 
-	DECLARE_DELEGATE_TwoParams(FOnMetasoundGraphRegister, UObject&, bool)
-	DECLARE_DELEGATE_OneParam(FOnMetasoundGraphUnregister, UObject&)
+	enum class ERegistrationAssetContext
+	{
+		None, // No special asset context associated with this graph registration action
+		Removing, // Graph registration during asset removal
+		Renaming, // Graph registration during asset rename
+		Reloading, // Graph registration during asset reload
+	};
+
+	DECLARE_DELEGATE_TwoParams(FOnMetasoundGraphRegister, UObject&, ERegistrationAssetContext)
+	DECLARE_DELEGATE_TwoParams(FOnMetasoundGraphUnregister, UObject&, ERegistrationAssetContext)
 #endif // WITH_EDITOR
 
 	class IMetasoundEngineModule : public IModuleInterface
