[MetaSounds] - When deleting MetaSound inputs using namespaces, the next highlighted entry is seemingly random

#rnx

diff --git a/Engine/Plugins/Runtime/Metasound/Source/MetasoundEditor/Private/MetasoundEditorGraph.cpp b/Engine/Plugins/Runtime/Metasound/Source/MetasoundEditor/Private/MetasoundEditorGraph.cpp
--- a/Engine/Plugins/Runtime/Metasound/Source/MetasoundEditor/Private/MetasoundEditorGraph.cpp
+++ b/Engine/Plugins/Runtime/Metasound/Source/MetasoundEditor/Private/MetasoundEditorGraph.cpp
@@ -2261,46 +2261,65 @@ UMetasoundEditorGraphMember* UMetasoundEditorGraph::FindMember(FGuid InMemberID)
 
 UMetasoundEditorGraphMember* UMetasoundEditorGraph::FindAdjacentMember(const UMetasoundEditorGraphMember& InMember)
 {
-	int32 IndexInArray = Inputs.IndexOfByPredicate([&](const TObjectPtr<UMetasoundEditorGraphInput>& InputMember)
+	auto CheckPredicate = [&InMember](const TObjectPtr<UMetasoundEditorGraphMember>& InGraphMember)
 	{
-		return &InMember == ToRawPtr(InputMember);
-	});
+		return &InMember == ToRawPtr(InGraphMember);
+	};
+
+	// Input
+	int32 IndexInArray = Inputs.IndexOfByPredicate(CheckPredicate);
 
 	if (INDEX_NONE != IndexInArray)
 	{
-		if (IndexInArray < (Inputs.Num() - 1))
+		if (TObjectPtr<UMetasoundEditorGraphInput> Input = FindAdjacentMemberFromSorted<UMetasoundEditorGraphInput>(Inputs, CheckPredicate))
 		{
-			return Inputs[IndexInArray + 1];
+			return Input;
 		}
-		else if (IndexInArray > 0)
+		else if (Outputs.Num() > 0)
 		{
-			return Inputs[IndexInArray - 1];
+			return Outputs[0];
 		}
-		else
+		else if (Variables.Num() > 0)
 		{
-			if (Outputs.Num() > 0)
-			{
-				return Outputs[0];
-			}
+			return Variables[0];
 		}
 
 		return nullptr;
 	}
 
-	IndexInArray = Outputs.IndexOfByPredicate([&](const TObjectPtr<UMetasoundEditorGraphOutput>& OutputMember)
+	// Output
+	IndexInArray = Outputs.IndexOfByPredicate(CheckPredicate);
+
+	if (INDEX_NONE != IndexInArray)
 	{
-		return &InMember == ToRawPtr(OutputMember);
-	});
+		if (TObjectPtr<UMetasoundEditorGraphOutput> Output = FindAdjacentMemberFromSorted<UMetasoundEditorGraphOutput>(Outputs, CheckPredicate))
+		{
+			return Output;
+		}
+		else if (Inputs.Num() > 0)
+		{
+			return Inputs.Last();
+		}
+		else if (Variables.Num() > 0)
+		{
+			return Variables[0];
+		}
+
+		return nullptr;
+	}
+
+	// Variable
+	IndexInArray = Variables.IndexOfByPredicate(CheckPredicate);
 
 	if (INDEX_NONE != IndexInArray)
 	{
-		if (IndexInArray < (Outputs.Num() - 1))
+		if (TObjectPtr<UMetasoundEditorGraphVariable> Variable = FindAdjacentMemberFromSorted<UMetasoundEditorGraphVariable>(Variables, CheckPredicate))
 		{
-			return Outputs[IndexInArray + 1];
+			return Variable;
 		}
-		else if (IndexInArray > 0)
+		else if (Outputs.Num() > 0)
 		{
-			return Outputs[IndexInArray - 1];
+			return Outputs.Last();
 		}
 		else if (Inputs.Num() > 0)
 		{
diff --git a/Engine/Plugins/Runtime/Metasound/Source/MetasoundEditor/Public/MetasoundEditorGraph.h b/Engine/Plugins/Runtime/Metasound/Source/MetasoundEditor/Public/MetasoundEditorGraph.h
--- a/Engine/Plugins/Runtime/Metasound/Source/MetasoundEditor/Public/MetasoundEditorGraph.h
+++ b/Engine/Plugins/Runtime/Metasound/Source/MetasoundEditor/Public/MetasoundEditorGraph.h
@@ -592,6 +592,39 @@ public:
 	UMetasoundEditorGraphMember* FindMember(FGuid InMemberID) const;
 	UMetasoundEditorGraphMember* FindAdjacentMember(const UMetasoundEditorGraphMember& InMember);
 
+private:
+	/**
+	* Sorts the incoming members by their MemberName and returns next element after sorting. Will not affect the order of the original list.
+	* @param InMembers - Members to sort
+	* @param InPredicate - Member to search for
+	* @return the adjacent member from the sorted members
+	*/
+	template<typename T>
+	TObjectPtr<T> FindAdjacentMemberFromSorted(const TArray<TObjectPtr<T>>& InMembers, TFunction<bool(const TObjectPtr<UMetasoundEditorGraphMember>&)> InPredicate)
+	{
+		auto SortMembers = [](const UMetasoundEditorGraphMember& InFirstMember, const UMetasoundEditorGraphMember& InSecondMember)
+			{
+				return InFirstMember.GetMemberName().ToString() < InSecondMember.GetMemberName().ToString();
+			};
+
+		TArray<TObjectPtr<T>> MembersCopy = InMembers;
+		MembersCopy.Sort(SortMembers);
+
+		int32 IndexInArray = MembersCopy.IndexOfByPredicate(InPredicate);
+
+		if (IndexInArray < (InMembers.Num() - 1))
+		{
+			return MembersCopy[IndexInArray + 1];
+		}
+		else if (IndexInArray > 0)
+		{
+			return MembersCopy[IndexInArray - 1];
+		}
+
+		return nullptr;
+	}
+
+public:
 	UE_DEPRECATED(5.5, "Use the associated Frontend builder and synchronize the graph remove a given member.")
 	bool RemoveMember(UMetasoundEditorGraphMember& InGraphMember) { return false; }
 
