[MetaSounds][WaveTable] Add tooltip explaining float interpolation on "Get WaveTable From Bank" Node

#rnx

diff --git a/Engine/Plugins/Runtime/Metasound/Source/MetasoundEngine/Private/MetasoundWaveTableGetNode.cpp b/Engine/Plugins/Runtime/Metasound/Source/MetasoundEngine/Private/MetasoundWaveTableGetNode.cpp
--- a/Engine/Plugins/Runtime/Metasound/Source/MetasoundEngine/Private/MetasoundWaveTableGetNode.cpp
+++ b/Engine/Plugins/Runtime/Metasound/Source/MetasoundEngine/Private/MetasoundWaveTableGetNode.cpp
@@ -28,7 +28,7 @@ namespace Metasound
 			static const FVertexInterface DefaultInterface(
 				FInputVertexInterface(
 					TInputDataVertex<FWaveTableBankAsset>("WaveTableBank", FDataVertexMetadata { LOCTEXT("MetasoundWaveTableGetNode_InputWaveTableBank", "WaveTableBank") }),
-					TInputDataVertex<float>("TableIndex", FDataVertexMetadata { LOCTEXT("MetasoundWaveTableGetNode_InputTableIndex", "Index") }, 0.0f)
+					TInputDataVertex<float>("TableIndex", FDataVertexMetadata { LOCTEXT("MetasoundWaveTableGetNode_InputTableIndex", "Index - supports Float input to allow index interpolation") }, 0.0f)
 				),
 				FOutputVertexInterface(
 					TOutputDataVertex<FWaveTable>("Out", FDataVertexMetadata { LOCTEXT("MetasoundWaveTableGetNode_OutputWaveTable", "Out") })
