Fix for init variable on trigger.
#rnx

diff --git a/Engine/Plugins/Runtime/Metasound/Source/MetasoundGraphCore/Private/MetasoundDynamicOperator.cpp b/Engine/Plugins/Runtime/Metasound/Source/MetasoundGraphCore/Private/MetasoundDynamicOperator.cpp
--- a/Engine/Plugins/Runtime/Metasound/Source/MetasoundGraphCore/Private/MetasoundDynamicOperator.cpp
+++ b/Engine/Plugins/Runtime/Metasound/Source/MetasoundGraphCore/Private/MetasoundDynamicOperator.cpp
@@ -10,6 +10,7 @@
 #include "MetasoundOperatorInterface.h"
 #include "MetasoundVertexData.h"
 #include "MetasoundTrace.h"
+#include "Misc/ReverseIterate.h"
 #include "Templates/SharedPointer.h"
 #include "Templates/UniquePtr.h"
 
@@ -220,7 +221,7 @@ namespace Metasound
 
 		void FDynamicOperator::PostExecute()
 		{
-			for (FPostExecuteEntry& Entry : DynamicOperatorData.PostExecuteTable)
+			for (FPostExecuteEntry& Entry : ReverseIterate(DynamicOperatorData.PostExecuteTable))
 			{
 				Entry.PostExecute();
 			}
diff --git a/Engine/Plugins/Runtime/Metasound/Source/MetasoundGraphCore/Private/MetasoundGraphOperator.cpp b/Engine/Plugins/Runtime/Metasound/Source/MetasoundGraphCore/Private/MetasoundGraphOperator.cpp
--- a/Engine/Plugins/Runtime/Metasound/Source/MetasoundGraphCore/Private/MetasoundGraphOperator.cpp
+++ b/Engine/Plugins/Runtime/Metasound/Source/MetasoundGraphCore/Private/MetasoundGraphOperator.cpp
@@ -9,6 +9,7 @@
 #include "MetasoundExecutableOperator.h"
 #include "MetasoundGraphAlgoPrivate.h"
 #include "MetasoundOperatorInterface.h"
+#include "Misc/ReverseIterate.h"
 
 namespace Metasound
 {
@@ -163,31 +164,27 @@ namespace Metasound
 
 	void FGraphOperator::Execute()
 	{
-		FExecuteEntry* StackPtr = ExecuteStack.GetData();
-		const int32 Num = ExecuteStack.Num();
-		for (int32 i = 0; i < Num; i++)
+		for (FExecuteEntry& Entry : ExecuteStack)
 		{
-			StackPtr[i].Execute();
+			Entry.Execute();
 		}
 	}
 
 	void FGraphOperator::PostExecute()
 	{
-		FPostExecuteEntry* StackPtr = PostExecuteStack.GetData();
-		const int32 Num = PostExecuteStack.Num();
-		for (int32 i = 0; i < Num; i++)
+		// Reverse iterate over post execute so that inputs to operators do
+		// not change from last execute
+		for (FPostExecuteEntry& Entry : ReverseIterate(PostExecuteStack))
 		{
-			StackPtr[i].PostExecute();
+			Entry.PostExecute();
 		}
 	}
 
 	void FGraphOperator::Reset(const FGraphOperator::FResetParams& InParams)
 	{
-		FResetEntry* StackPtr = ResetStack.GetData();
-		const int32 Num = ResetStack.Num();
-		for (int32 i = 0; i < Num; i++)
+		for (FResetEntry& Entry : ResetStack)
 		{
-			StackPtr[i].Reset(InParams);
+			Entry.Reset(InParams);
 		}
 	}
 
diff --git a/Engine/Plugins/Runtime/Metasound/Source/MetasoundGraphCore/Private/MetasoundRebindableGraphOperator.cpp b/Engine/Plugins/Runtime/Metasound/Source/MetasoundGraphCore/Private/MetasoundRebindableGraphOperator.cpp
--- a/Engine/Plugins/Runtime/Metasound/Source/MetasoundGraphCore/Private/MetasoundRebindableGraphOperator.cpp
+++ b/Engine/Plugins/Runtime/Metasound/Source/MetasoundGraphCore/Private/MetasoundRebindableGraphOperator.cpp
@@ -10,6 +10,7 @@
 #include "MetasoundVertex.h"
 #include "MetasoundVertexData.h"
 #include "Misc/Guid.h"
+#include "Misc/ReverseIterate.h"
 #include "Templates/UniquePtr.h"
 
 namespace Metasound
@@ -48,7 +49,9 @@ namespace Metasound
 	{
 		using namespace DynamicGraph;
 
-		for (FPostExecuteEntry& Entry : GraphOperatorData.PostExecuteTable)
+		// Reverse iterate over post execute to keep inputs to operators unchanged
+		// between calls to Execute() and PostExecute()
+		for (FPostExecuteEntry& Entry : ReverseIterate(GraphOperatorData.PostExecuteTable))
 		{
 			Entry.PostExecute();
 		}
