Fix for MetaSound detail customization crash

#rnx

diff --git a/Engine/Plugins/Runtime/Metasound/Source/MetasoundEditor/Private/MetasoundDetailCustomization.cpp b/Engine/Plugins/Runtime/Metasound/Source/MetasoundEditor/Private/MetasoundDetailCustomization.cpp
--- a/Engine/Plugins/Runtime/Metasound/Source/MetasoundEditor/Private/MetasoundDetailCustomization.cpp
+++ b/Engine/Plugins/Runtime/Metasound/Source/MetasoundEditor/Private/MetasoundDetailCustomization.cpp
@@ -563,7 +563,7 @@ namespace Metasound::Editor
 			// Page Playing Icon
 			{
 				TAttribute<FText> ToolTip = LOCTEXT("MetaSound_ExecutingPageGraphTooltip", "Currently executing graph.");
-				TAttribute<EVisibility> Visibility = TAttribute<EVisibility>::CreateLambda([this, PageID = InEntryPageID]()
+				TAttribute<EVisibility> Visibility = TAttribute<EVisibility>::CreateSPLambda(AsShared(), [this, PageID = InEntryPageID]()
 				{
 					if (Builder.IsValid())
 					{
diff --git a/Engine/Plugins/Runtime/Metasound/Source/MetasoundEditor/Private/MetasoundEditor.cpp b/Engine/Plugins/Runtime/Metasound/Source/MetasoundEditor/Private/MetasoundEditor.cpp
--- a/Engine/Plugins/Runtime/Metasound/Source/MetasoundEditor/Private/MetasoundEditor.cpp
+++ b/Engine/Plugins/Runtime/Metasound/Source/MetasoundEditor/Private/MetasoundEditor.cpp
@@ -5166,7 +5166,7 @@ namespace Metasound
 		{
 			if (PageStatsWidget.IsValid())
 			{
-				TAttribute<EVisibility> ExecVisibility = TAttribute<EVisibility>::CreateLambda([this, InPageID]()
+				TAttribute<EVisibility> ExecVisibility = TAttribute<EVisibility>::CreateSPLambda(AsShared(), [this, InPageID]()
 				{
 					using namespace Engine;
 					if (Builder.IsValid() && ShowPageGraphDetails())
diff --git a/Engine/Plugins/Runtime/Metasound/Source/MetasoundEditor/Private/MetasoundNodeDetailCustomization.cpp b/Engine/Plugins/Runtime/Metasound/Source/MetasoundEditor/Private/MetasoundNodeDetailCustomization.cpp
--- a/Engine/Plugins/Runtime/Metasound/Source/MetasoundEditor/Private/MetasoundNodeDetailCustomization.cpp
+++ b/Engine/Plugins/Runtime/Metasound/Source/MetasoundEditor/Private/MetasoundNodeDetailCustomization.cpp
@@ -1025,7 +1025,7 @@ PRAGMA_ENABLE_DEPRECATION_WARNINGS
 				IMetasoundEditorModule& EditorModule = FModuleManager::GetModuleChecked<IMetasoundEditorModule>("MetaSoundEditor");
 				LiteralCustomization = EditorModule.CreateMemberDefaultLiteralCustomization(*MemberClass, DefaultCategoryBuilder);
 
-				TAttribute<EVisibility> Visibility = TAttribute<EVisibility>::CreateLambda([this]()
+				TAttribute<EVisibility> Visibility = TAttribute<EVisibility>::CreateSPLambda(AsShared(), [this]()
 				{
 					return GetDefaultVisibility();
 				});
@@ -1648,16 +1648,16 @@ PRAGMA_ENABLE_DEPRECATION_WARNINGS
 					const UMetasoundEditorGraphInput* Input = Cast<UMetasoundEditorGraphInput>(MemberDefaultLiteral->FindMember());
 					if (Input)
 					{
-						Enabled = TAttribute<bool>::CreateLambda([this] { return !GetInputInheritsDefault(); });
+						Enabled = TAttribute<bool>::CreateSPLambda(AsShared(), [this] { return !GetInputInheritsDefault(); });
 						ResetOverride = FResetToDefaultOverride::Create(
-							FIsResetToDefaultVisible::CreateLambda([this](TSharedPtr<IPropertyHandle> /* PropertyHandle */) { return !GetInputInheritsDefault(); }),
-							FResetToDefaultHandler::CreateLambda([this](TSharedPtr<IPropertyHandle> /* PropertyHandle */) { SetInputInheritsDefault(); }));
+							FIsResetToDefaultVisible::CreateSPLambda(AsShared(), [this](TSharedPtr<IPropertyHandle> /* PropertyHandle */) { return !GetInputInheritsDefault(); }),
+							FResetToDefaultHandler::CreateSPLambda(AsShared(), [this](TSharedPtr<IPropertyHandle> /* PropertyHandle */) { SetInputInheritsDefault(); }));
 					}
 				}
 			}
 			else
 			{
-				Enabled = TAttribute<bool>::CreateLambda([this]
+				Enabled = TAttribute<bool>::CreateSPLambda(AsShared(), [this]
 				{
 					if (GraphMember.IsValid())
 					{
diff --git a/Engine/Plugins/Runtime/Metasound/Source/MetasoundEditor/Private/SMetasoundGraphNode.cpp b/Engine/Plugins/Runtime/Metasound/Source/MetasoundEditor/Private/SMetasoundGraphNode.cpp
--- a/Engine/Plugins/Runtime/Metasound/Source/MetasoundEditor/Private/SMetasoundGraphNode.cpp
+++ b/Engine/Plugins/Runtime/Metasound/Source/MetasoundEditor/Private/SMetasoundGraphNode.cpp
@@ -164,7 +164,7 @@ namespace Metasound
 
 		TAttribute<EVisibility> SMetaSoundGraphNode::GetSimulationVisibilityAttribute() const
 		{
-			return TAttribute<EVisibility>::CreateLambda([this]()
+			return TAttribute<EVisibility>::CreateSPLambda(AsShared(), [this]()
 			{
 				using namespace Frontend;
 
