Fix for multiple referenced nodes of the same asset class not all auto-updating their interface when placed on multiple pages of the same MetaSound asset.

#rnx

diff --git a/Engine/Plugins/Runtime/Metasound/Source/MetasoundFrontend/Private/MetasoundAssetBase.cpp b/Engine/Plugins/Runtime/Metasound/Source/MetasoundFrontend/Private/MetasoundAssetBase.cpp
--- a/Engine/Plugins/Runtime/Metasound/Source/MetasoundFrontend/Private/MetasoundAssetBase.cpp
+++ b/Engine/Plugins/Runtime/Metasound/Source/MetasoundFrontend/Private/MetasoundAssetBase.cpp
@@ -1049,13 +1049,14 @@ bool FMetasoundAssetBase::VersionDependencies(FMetaSoundFrontendDocumentBuilder&
 	Metasound::Frontend::FDocumentHandle DocHandle = GetDocumentHandle();
 	const FMetasoundFrontendGraphClass& RootGraph = Builder.GetConstDocumentChecked().RootGraph;
 	constexpr bool bBroadcastPageIDDelegate = false;
+
+	FAutoUpdateRootGraph AutoUpdateTransform(GetOwningAssetName(), bInLogWarningsOnDroppedConnection);
 	RootGraph.IterateGraphPages([&](const FMetasoundFrontendGraph& Graph)
 	{
 		// Set the build page ID to this graph as a hack to apply dependency versioning logic using
 		// the controller/handle API until auto-update is renamed & moved to use document builder API.
 		Builder.SetBuildPageID(Graph.PageID, bBroadcastPageIDDelegate);
-		FString OwningAssetName = GetOwningAssetName();
-		bDocumentModified |= FAutoUpdateRootGraph(MoveTemp(OwningAssetName), bInLogWarningsOnDroppedConnection).Transform(DocHandle);
+		bDocumentModified |= AutoUpdateTransform.Transform(DocHandle);
 	});
 	Builder.SetBuildPageID(InitBuildPageID, bBroadcastPageIDDelegate);
 #endif // WITH_EDITORONLY_DATA
diff --git a/Engine/Plugins/Runtime/Metasound/Source/MetasoundFrontend/Private/MetasoundFrontendTransform.cpp b/Engine/Plugins/Runtime/Metasound/Source/MetasoundFrontend/Private/MetasoundFrontendTransform.cpp
--- a/Engine/Plugins/Runtime/Metasound/Source/MetasoundFrontend/Private/MetasoundFrontendTransform.cpp
+++ b/Engine/Plugins/Runtime/Metasound/Source/MetasoundFrontend/Private/MetasoundFrontendTransform.cpp
@@ -509,7 +509,7 @@ namespace Metasound
 		}
 #endif // WITH_EDITOR
 
-		bool FAutoUpdateRootGraph::Transform(FDocumentHandle InDocument) const
+		bool FAutoUpdateRootGraph::Transform(FDocumentHandle InDocument)
 		{
 			METASOUND_TRACE_CPUPROFILER_EVENT_SCOPE(FAutoUpdateRootGraph::Transform);
 			bool bDidEdit = false;
@@ -543,11 +543,15 @@ namespace Metasound
 				}
 
 				FClassInterfaceUpdates InterfaceUpdates;
-				if (!NodeHandle->CanAutoUpdate(InterfaceUpdates))
+				const FGuid& ClassID = NodeHandle->GetClassID();
+				const bool bHasUpdated = UpdatedClasses.Contains(ClassID);
+				if (!bHasUpdated && !NodeHandle->CanAutoUpdate(InterfaceUpdates))
 				{
 					return;
 				}
 
+				UpdatedClasses.Add(ClassID);
+
 				// Check if a updated minor version exists.
 				FMetasoundFrontendClass ClassWithHighestMinorVersion;
 				const bool bFoundClassInSearchEngine = Frontend::ISearchEngine::Get().FindClassWithHighestMinorVersion(ClassMetadata.GetClassName(), ClassMetadata.GetVersion().Major, ClassWithHighestMinorVersion);
diff --git a/Engine/Plugins/Runtime/Metasound/Source/MetasoundFrontend/Public/MetasoundFrontendTransform.h b/Engine/Plugins/Runtime/Metasound/Source/MetasoundFrontend/Public/MetasoundFrontendTransform.h
--- a/Engine/Plugins/Runtime/Metasound/Source/MetasoundFrontend/Public/MetasoundFrontendTransform.h
+++ b/Engine/Plugins/Runtime/Metasound/Source/MetasoundFrontend/Public/MetasoundFrontendTransform.h
@@ -188,7 +188,7 @@ namespace Metasound::Frontend
 	/** Automatically updates all nodes and respective dependencies in graph where
 		* newer versions exist in the loaded MetaSound Class Node Registry.
 		*/
-	class METASOUNDFRONTEND_API FAutoUpdateRootGraph : public IDocumentTransform
+	class METASOUNDFRONTEND_API FAutoUpdateRootGraph 
 	{
 	public:
 		/** Construct an AutoUpdate transform
@@ -202,9 +202,14 @@ namespace Metasound::Frontend
 		{
 		}
 
-		bool Transform(FDocumentHandle InDocument) const override;
+		bool Transform(FDocumentHandle InDocument);
 
 	private:
+		// Keeps track of classes already updated so node check can be avoided.
+		// Hack to avoid issue where earlier auto-update passes on pages can
+		// clear out internal change state of a class in the registry causing
+		// nodes to get ignored on later page auto-update passes.
+		TSet<FGuid> UpdatedClasses;
 		const FString DebugAssetPath;
 		bool bLogWarningOnDroppedConnection;
 	};
