Fix for external node pin names missing namespace

#rnx

diff --git a/Engine/Plugins/Runtime/Metasound/Source/MetasoundEditor/Private/MetasoundEditorGraphSchema.cpp b/Engine/Plugins/Runtime/Metasound/Source/MetasoundEditor/Private/MetasoundEditorGraphSchema.cpp
--- a/Engine/Plugins/Runtime/Metasound/Source/MetasoundEditor/Private/MetasoundEditorGraphSchema.cpp
+++ b/Engine/Plugins/Runtime/Metasound/Source/MetasoundEditor/Private/MetasoundEditorGraphSchema.cpp
@@ -1901,10 +1901,21 @@ FText UMetasoundEditorGraphSchema::GetPinDisplayName(const UEdGraphPin* Pin) con
 			{
 				FName Namespace, ParamName;
 				ClassVertex->SplitName(Namespace, ParamName);
-				FText DisplayName = ClassVertex->Metadata.GetDisplayName();
+				const FText DisplayName = ClassVertex->Metadata.GetDisplayName();
 				if (DisplayName.IsEmptyOrWhitespace())
 				{
-					DisplayName = FText::FromName(ParamName);
+					if (Namespace.IsNone())
+					{
+						return FText::FromName(ParamName);
+					}
+					else
+					{
+						return FText::Format(
+							LOCTEXT("ClassMetadataDisplayNameWithNamespaceFormat", "{0} ({1})"),
+							FText::FromName(ParamName),
+							FText::FromName(Namespace)
+						);
+					}
 				}
 
 				return DisplayName;
