Fix MetaSound widget value changes not being synced between all input nodes

#rnx

diff --git a/Engine/Plugins/Runtime/Metasound/Source/MetasoundEditor/Private/SMetasoundGraphNode.cpp b/Engine/Plugins/Runtime/Metasound/Source/MetasoundEditor/Private/SMetasoundGraphNode.cpp
--- a/Engine/Plugins/Runtime/Metasound/Source/MetasoundEditor/Private/SMetasoundGraphNode.cpp
+++ b/Engine/Plugins/Runtime/Metasound/Source/MetasoundEditor/Private/SMetasoundGraphNode.cpp
@@ -827,6 +827,7 @@ namespace Metasound
 
 						const FGuid BuildPageID = Builder.GetBuildPageID();
 						DefaultFloat->SetFromLiteral(OutputLiteral, BuildPageID);
+						DefaultFloat->OnDefaultValueChanged.Broadcast(BuildPageID, Value);
 
 						constexpr bool bPostTransaction = true;
 						GraphMember->UpdateFrontendDefaultLiteral(bPostTransaction, &BuildPageID);
