More logging around resolving page id/proxy generation to hunt down page packaging mystery of expected sound not playing
- Fix minor bug where removal of graph input would return false despite removal succeeding if input was on preset but not initially set to inherit
- Disable MetaSound live audition transactor in cook to avoid potential warnings caused by resolution on cook page mutations
#rnx

diff --git a/Engine/Plugins/Runtime/Metasound/Source/MetasoundEngine/Private/MetasoundBuilderSubsystem.cpp b/Engine/Plugins/Runtime/Metasound/Source/MetasoundEngine/Private/MetasoundBuilderSubsystem.cpp
--- a/Engine/Plugins/Runtime/Metasound/Source/MetasoundEngine/Private/MetasoundBuilderSubsystem.cpp
+++ b/Engine/Plugins/Runtime/Metasound/Source/MetasoundEngine/Private/MetasoundBuilderSubsystem.cpp
@@ -5,6 +5,7 @@
 #include "Algo/Transform.h"
 #include "AudioDevice.h"
 #include "Components/AudioComponent.h"
+#include "CoreGlobals.h"
 #include "Engine/Engine.h"
 #include "HAL/IConsoleManager.h"
 #include "Interfaces/MetasoundOutputFormatInterfaces.h"
@@ -253,6 +254,11 @@ bool UMetaSoundSourceBuilder::ExecuteAuditionableTransaction(FAuditionableTransa
 
 	METASOUND_TRACE_CPUPROFILER_EVENT_SCOPE(UMetaSoundSourceBuilder::ExecuteAuditionableTransaction);
 
+	checkf(!IsRunningCookCommandlet(),
+		TEXT("ExecuteAuditionableTransaction cannot be called while running cook. ")
+		TEXT("Requires resolved graph which should not be relied on while cooking as it can mutate depending on cook's ")
+		TEXT("target platform, which is independent of targetable page(s)"));
+
 	TSharedPtr<FDynamicOperatorTransactor> Transactor = GetMetaSoundSource().GetDynamicGeneratorTransactor();
 	if (Transactor.IsValid())
 	{
@@ -326,15 +332,18 @@ void UMetaSoundSourceBuilder::InitDelegates(Metasound::Frontend::FDocumentModify
 {
 	Super::InitDelegates(OutDocumentDelegates);
 
-	OutDocumentDelegates.PageDelegates.OnPageAdded.AddUObject(this, &UMetaSoundSourceBuilder::OnPageAdded);
-	OutDocumentDelegates.PageDelegates.OnRemovingPage.AddUObject(this, &UMetaSoundSourceBuilder::OnRemovingPage);
+	if (!IsRunningCookCommandlet())
+	{
+		OutDocumentDelegates.PageDelegates.OnPageAdded.AddUObject(this, &UMetaSoundSourceBuilder::OnPageAdded);
+		OutDocumentDelegates.PageDelegates.OnRemovingPage.AddUObject(this, &UMetaSoundSourceBuilder::OnRemovingPage);
 
-	OutDocumentDelegates.InterfaceDelegates.OnInputAdded.AddUObject(this, &UMetaSoundSourceBuilder::OnInputAdded);
-	OutDocumentDelegates.InterfaceDelegates.OnOutputAdded.AddUObject(this, &UMetaSoundSourceBuilder::OnOutputAdded);
-	OutDocumentDelegates.InterfaceDelegates.OnRemovingInput.AddUObject(this, &UMetaSoundSourceBuilder::OnRemovingInput);
-	OutDocumentDelegates.InterfaceDelegates.OnRemovingOutput.AddUObject(this, &UMetaSoundSourceBuilder::OnRemovingOutput);
+		OutDocumentDelegates.InterfaceDelegates.OnInputAdded.AddUObject(this, &UMetaSoundSourceBuilder::OnInputAdded);
+		OutDocumentDelegates.InterfaceDelegates.OnOutputAdded.AddUObject(this, &UMetaSoundSourceBuilder::OnOutputAdded);
+		OutDocumentDelegates.InterfaceDelegates.OnRemovingInput.AddUObject(this, &UMetaSoundSourceBuilder::OnRemovingInput);
+		OutDocumentDelegates.InterfaceDelegates.OnRemovingOutput.AddUObject(this, &UMetaSoundSourceBuilder::OnRemovingOutput);
 
-	InitTargetPageDelegates(OutDocumentDelegates);
+		InitTargetPageDelegates(OutDocumentDelegates);
+	}
 }
 
 void UMetaSoundSourceBuilder::InitTargetPageDelegates(Metasound::Frontend::FDocumentModifyDelegates& OutDocumentDelegates)
diff --git a/Engine/Plugins/Runtime/Metasound/Source/MetasoundEngine/Private/MetasoundDocumentBuilderRegistry.cpp b/Engine/Plugins/Runtime/Metasound/Source/MetasoundEngine/Private/MetasoundDocumentBuilderRegistry.cpp
--- a/Engine/Plugins/Runtime/Metasound/Source/MetasoundEngine/Private/MetasoundDocumentBuilderRegistry.cpp
+++ b/Engine/Plugins/Runtime/Metasound/Source/MetasoundEngine/Private/MetasoundDocumentBuilderRegistry.cpp
@@ -109,7 +109,14 @@ namespace Metasound::Engine
 		const TArray<FGuid> PlatformTargetPageIDs = Settings->GetCookedTargetPageIDs(PlatformName);
 		checkf(!PlatformTargetPageIDs.IsEmpty(), TEXT("Must have at least one targeted page ID to cook MetaSound."));
 
-		auto StripPageEntries = [&](TArray<FGuid>& PageIdsToResolve, TSet<FGuid>& ResolveTargetScratch, TFunctionRef<bool(const FGuid&)> RemovePageItem)
+		const FString DebugName = Builder.GetDebugName();
+
+		auto StripPageEntries = [&](
+			TArray<FGuid>& PageIdsToResolve,
+			TSet<FGuid>& ResolveTargetScratch,
+			TFunctionRef<bool(const FGuid&)> RemovePageItem,
+			const FName& ItemName,
+			const FString& ItemType)
 		{
 			ResolveTargetScratch.Reset();
 			for (const FGuid& TargetPage : PlatformTargetPageIDs)
@@ -124,7 +131,16 @@ namespace Metasound::Engine
 
 			for (const FGuid& PageID : PageIdsToResolve)
 			{
-				bModified |= RemovePageItem(PageID);
+				const bool bRemovedPageItem = RemovePageItem(PageID);
+				if (bRemovedPageItem)
+				{
+					UE_LOG(LogMetaSound, Display, TEXT("%s: Removed %s %s w/PageID '%s'"),
+						*DebugName,
+						ItemName.IsNone() ? TEXT("paged") : *ItemName.ToString(),
+						*ItemType,
+						*PageID.ToString());
+				}
+				bModified |= bRemovedPageItem;
 			}
 		};
 
@@ -142,7 +158,8 @@ namespace Metasound::Engine
 			};
 
 			const int32 NumInitGraphs = Document.RootGraph.GetConstGraphPages().Num();
-			StripPageEntries(ResolvePageIDs, ResolvedTargetScratchIDs, RemoveGraphPage);
+			StripPageEntries(ResolvePageIDs, ResolvedTargetScratchIDs, RemoveGraphPage, FName(), TEXT("graph"));
+
 			const int32 NumRemainingGraphs = Document.RootGraph.GetConstGraphPages().Num();
 
 			checkf(NumRemainingGraphs > 0,
@@ -166,11 +183,13 @@ namespace Metasound::Engine
 
 				auto RemoveDefault = [&Builder, &GraphInput](const FGuid& InPageID)
 				{
-					return Builder.RemoveGraphInputDefault(GraphInput.Name, InPageID);
+					const bool bClearInheritsDefault = false;
+					return Builder.RemoveGraphInputDefault(GraphInput.Name, InPageID, bClearInheritsDefault);
 				};
 
 				const int32 NumInitDefaults = GraphInput.GetDefaults().Num();
-				StripPageEntries(ResolvePageIDs, ResolvedTargetScratchIDs, RemoveDefault);
+				StripPageEntries(ResolvePageIDs, ResolvedTargetScratchIDs, RemoveDefault, GraphInput.Name, TEXT("input default"));
+
 				const int32 NumRemainingDefaults = GraphInput.GetDefaults().Num();
 
 				checkf(NumRemainingDefaults > 0,
diff --git a/Engine/Plugins/Runtime/Metasound/Source/MetasoundFrontend/Private/MetasoundFrontendDataTypeRegistry.cpp b/Engine/Plugins/Runtime/Metasound/Source/MetasoundFrontend/Private/MetasoundFrontendDataTypeRegistry.cpp
--- a/Engine/Plugins/Runtime/Metasound/Source/MetasoundFrontend/Private/MetasoundFrontendDataTypeRegistry.cpp
+++ b/Engine/Plugins/Runtime/Metasound/Source/MetasoundFrontend/Private/MetasoundFrontendDataTypeRegistry.cpp
@@ -862,9 +862,16 @@ namespace Metasound
 				if (const IDataTypeRegistryEntry* Entry = FindDataTypeEntry(InDataType))
 				{
 					ProxyPtr = Entry->CreateProxy(InObject);
-					if (!ProxyPtr && InObject)
+					if (ProxyPtr)
 					{
-						UE_LOG(LogMetaSound, Error, TEXT("Failed to create a valid proxy from UObject '%s'."), *InObject->GetName());
+						UE_LOG(LogMetaSound, VeryVerbose, TEXT("Created UObject proxy for '%s'."), *InObject->GetName());
+					}
+					else
+					{
+						if (InObject)
+						{
+							UE_LOG(LogMetaSound, Error, TEXT("Failed to create a valid proxy from UObject '%s'."), *InObject->GetName());
+						}
 					}
 				}
 
diff --git a/Engine/Plugins/Runtime/Metasound/Source/MetasoundFrontend/Private/MetasoundFrontendDocumentBuilder.cpp b/Engine/Plugins/Runtime/Metasound/Source/MetasoundFrontend/Private/MetasoundFrontendDocumentBuilder.cpp
--- a/Engine/Plugins/Runtime/Metasound/Source/MetasoundFrontend/Private/MetasoundFrontendDocumentBuilder.cpp
+++ b/Engine/Plugins/Runtime/Metasound/Source/MetasoundFrontend/Private/MetasoundFrontendDocumentBuilder.cpp
@@ -3229,7 +3229,7 @@ void FMetaSoundFrontendDocumentBuilder::Reload(TSharedPtr<Metasound::Frontend::F
 }
 
 #if WITH_EDITORONLY_DATA
-bool FMetaSoundFrontendDocumentBuilder::RemoveGraphInputDefault(FName InputName, const FGuid& InPageID)
+bool FMetaSoundFrontendDocumentBuilder::RemoveGraphInputDefault(FName InputName, const FGuid& InPageID, bool bClearInheritsDefault)
 {
 	using namespace Metasound;
 
@@ -3246,11 +3246,12 @@ bool FMetaSoundFrontendDocumentBuilder::RemoveGraphInputDefault(FName InputName,
 		{
 			DocumentDelegates->InterfaceDelegates.OnInputDefaultChanged.Broadcast(Index);
 
-			// Set the input as no longer inheriting default for presets
-			if (IsPreset())
+			if (bClearInheritsDefault)
 			{
+				// Set the input as no longer inheriting default for presets
+				// (No-ops if MetaSound isn't preset or isn't set to inherit default).
 				constexpr bool bInputInheritsDefault = false;
-				return SetGraphInputInheritsDefault(InputName, bInputInheritsDefault);
+				SetGraphInputInheritsDefault(InputName, bInputInheritsDefault);
 			}
 
 			return true;
@@ -3772,11 +3773,9 @@ bool FMetaSoundFrontendDocumentBuilder::ResetGraphInputDefault(FName InputName)
 		DocumentDelegates->InterfaceDelegates.OnInputDefaultChanged.Broadcast(Index);
 
 		// Set the input as inheriting default for presets
-		if (IsPreset())
-		{
-			constexpr bool bInputInheritsDefault = true;
-			return SetGraphInputInheritsDefault(InputName, bInputInheritsDefault);
-		}
+		// (No-ops if MetaSound isn't preset or is already set to inherit default).
+		constexpr bool bInputInheritsDefault = true;
+		SetGraphInputInheritsDefault(InputName, bInputInheritsDefault);
 
 		Document.Metadata.ModifyContext.AddMemberIDModified(Input.NodeID);
 		return true;
@@ -4033,12 +4032,10 @@ bool FMetaSoundFrontendDocumentBuilder::SetGraphInputDefault(FName InputName, FM
 			}
 			DocumentDelegates->InterfaceDelegates.OnInputDefaultChanged.Broadcast(Index);
 
-			// Set the input as no longer inheriting default for presets
-			if (IsPreset())
-			{
-				constexpr bool bInputInheritsDefault = false;
-				return SetGraphInputInheritsDefault(InputName, bInputInheritsDefault);
-			}
+			// Set the input as inheriting default for presets
+			// (No-ops if MetaSound isn't preset or is already set to inherit default).
+			constexpr bool bInputInheritsDefault = false;
+			SetGraphInputInheritsDefault(InputName, bInputInheritsDefault);
 
 			return true;
 		}
@@ -4071,12 +4068,9 @@ bool FMetaSoundFrontendDocumentBuilder::SetGraphInputDefaults(FName InputName, T
 			DocumentDelegates->InterfaceDelegates.OnInputDefaultChanged.Broadcast(Index);
 
 			// Set the input as no longer inheriting default for presets
-			if (IsPreset())
-			{
-				constexpr bool bInputInheritsDefault = false;
-				return SetGraphInputInheritsDefault(InputName, bInputInheritsDefault);
-			}
-
+			// (No-ops if MetaSound isn't preset or isn't set to inherit default).
+			constexpr bool bInputInheritsDefault = false;
+			SetGraphInputInheritsDefault(InputName, bInputInheritsDefault);
 			return true;
 		}
 		UE_LOG(LogMetaSound, Error, TEXT("Attempting to set graph input of type '%s' with unsupported literal type(s)"), *Input.TypeName.ToString());
diff --git a/Engine/Plugins/Runtime/Metasound/Source/MetasoundFrontend/Private/MetasoundFrontendRegistryContainerImpl.cpp b/Engine/Plugins/Runtime/Metasound/Source/MetasoundFrontend/Private/MetasoundFrontendRegistryContainerImpl.cpp
--- a/Engine/Plugins/Runtime/Metasound/Source/MetasoundFrontend/Private/MetasoundFrontendRegistryContainerImpl.cpp
+++ b/Engine/Plugins/Runtime/Metasound/Source/MetasoundFrontend/Private/MetasoundFrontendRegistryContainerImpl.cpp
@@ -489,6 +489,28 @@ namespace Metasound::Frontend
 		FProxyDataCache ProxyDataCache;
 		ProxyDataCache.CreateAndCacheProxies(Document);
 
+#if !NO_LOGGING
+		if (UE_LOG_ACTIVE(LogMetaSound, Verbose))
+		{
+			const FGuid PageID = IDocumentBuilderRegistry::GetChecked().ResolveTargetPageID(Document.RootGraph);
+			const bool bContainsMultipleGraphs = Document.RootGraph.GetConstGraphPages().Num() > 1;
+			if (bContainsMultipleGraphs || PageID != Metasound::Frontend::DefaultPageID)
+			{
+				UE_LOG(LogMetaSound, Verbose, TEXT("Registered MetaSound '%s' Graph Page with PageID '%s'."),
+					*AssetPath.GetAssetName().ToString(),
+					*PageID.ToString());
+				if (bContainsMultipleGraphs)
+				{
+					UE_LOG(LogMetaSound, Verbose, TEXT("Graphs found with following PageIDs Implemented:"));
+					Document.RootGraph.IterateGraphPages([](const FMetasoundFrontendGraph& Graph)
+					{
+						UE_LOG(LogMetaSound, Verbose, TEXT("    - %s'"), *Graph.PageID.ToString());
+					});
+				}
+			}
+		}
+#endif // !NO_LOGGING
+
 		// Store update to newly registered node in history so nodes
 		// can be queried by transaction ID
 		FNodeClassInfo NodeClassInfo(Document.RootGraph, AssetPath);
diff --git a/Engine/Plugins/Runtime/Metasound/Source/MetasoundFrontend/Public/MetasoundFrontendDocumentBuilder.h b/Engine/Plugins/Runtime/Metasound/Source/MetasoundFrontend/Public/MetasoundFrontendDocumentBuilder.h
--- a/Engine/Plugins/Runtime/Metasound/Source/MetasoundFrontend/Public/MetasoundFrontendDocumentBuilder.h
+++ b/Engine/Plugins/Runtime/Metasound/Source/MetasoundFrontend/Public/MetasoundFrontendDocumentBuilder.h
@@ -389,7 +389,7 @@ public:
 	void Reload(TSharedPtr<Metasound::Frontend::FDocumentModifyDelegates> Delegates = {}, bool bPrimeCache = false);
 
 #if WITH_EDITORONLY_DATA
-	bool RemoveGraphInputDefault(FName InputName, const FGuid& InPageID);
+	bool RemoveGraphInputDefault(FName InputName, const FGuid& InPageID, bool bClearInheritsDefault = true);
 #endif // WITH_EDITORONLY_DATA
 
 	bool RemoveNodeInputDefault(const FGuid& InNodeID, const FGuid& InVertexID, const FGuid* InPageID = nullptr);
