Fix array drag drop functionality in MetaSounds

#rnx

diff --git a/Engine/Plugins/Runtime/Metasound/Source/MetasoundEditor/Private/MetasoundDefaultLiteralCustomization.cpp b/Engine/Plugins/Runtime/Metasound/Source/MetasoundEditor/Private/MetasoundDefaultLiteralCustomization.cpp
--- a/Engine/Plugins/Runtime/Metasound/Source/MetasoundEditor/Private/MetasoundDefaultLiteralCustomization.cpp
+++ b/Engine/Plugins/Runtime/Metasound/Source/MetasoundEditor/Private/MetasoundDefaultLiteralCustomization.cpp
@@ -383,27 +383,46 @@ namespace Metasound::Editor
 					DefaultProperties.Add(ValueProperty);
 
 					IDetailPropertyRow& ValueRow = InDetailLayout.AddPropertyToCategory(ValueProperty);
+					ValueRow.CustomWidget(true /*bShowChilden */);
+
+					// Name widget
 					if (bIsPagedDefault && bShowPageModifiers)
 					{
-						ValueRow.CustomWidget(true /*bShowChilden */);
-
 						(*ValueRow.CustomNameWidget())
 						[
 							BuildPageDefaultNameWidget(InLiteral, ElementProperty)
 						];
-
-						ValueRow.ShowPropertyButtons(false);
-						(*ValueRow.CustomValueWidget())
+					}
+					else
+					{
+						(*ValueRow.CustomNameWidget())
 						[
-							ValueProperty->CreatePropertyValueWidget()
+							ElementProperty->CreatePropertyNameWidget()
 						];
 					}
+					ValueRow.ShowPropertyButtons(false);
+
+					// Value widget
+					BuildDefaultValueWidget(ValueRow, ValueProperty);
 					ValueRow.IsEnabled(GetEnabled());
 				}
 			}
 		}
 	}
 
+	void FMetasoundDefaultLiteralCustomizationBase::BuildDefaultValueWidget(IDetailPropertyRow& ValueRow, TSharedPtr<IPropertyHandle> ValueProperty)
+	{
+		if (!ValueProperty.IsValid())
+		{
+			return;
+		}
+
+		(*ValueRow.CustomValueWidget())
+		[
+			ValueProperty->CreatePropertyValueWidget()
+		];
+	}
+
 	TArray<IDetailPropertyRow*> FMetasoundDefaultLiteralCustomizationBase::CustomizeLiteral(UMetasoundEditorGraphMemberDefaultLiteral& InLiteral, IDetailLayoutBuilder& InDetailLayout)
 	{
 		CustomizeDefaults(InLiteral, InDetailLayout);
diff --git a/Engine/Plugins/Runtime/Metasound/Source/MetasoundEditor/Private/MetasoundEditorGraph.cpp b/Engine/Plugins/Runtime/Metasound/Source/MetasoundEditor/Private/MetasoundEditorGraph.cpp
--- a/Engine/Plugins/Runtime/Metasound/Source/MetasoundEditor/Private/MetasoundEditorGraph.cpp
+++ b/Engine/Plugins/Runtime/Metasound/Source/MetasoundEditor/Private/MetasoundEditorGraph.cpp
@@ -1677,6 +1677,7 @@ void UMetasoundEditorGraphVariable::SetDataType(FName InNewType, bool bPostTrans
 			DefaultLiteral.SetFromLiteral(IDataTypeRegistry::Get().CreateDefaultLiteral(GetDataType()));
 			check(Literal); // Should exist from prior InitializeLiteral() call
 			Literal->SetFromLiteral(DefaultLiteral);
+			Literal->MemberID = VariableID;
 		}
 
 		// Add the nodes with the same identifier data but new datatype.
diff --git a/Engine/Plugins/Runtime/Metasound/Source/MetasoundEditor/Private/MetasoundNodeDetailCustomization.cpp b/Engine/Plugins/Runtime/Metasound/Source/MetasoundEditor/Private/MetasoundNodeDetailCustomization.cpp
--- a/Engine/Plugins/Runtime/Metasound/Source/MetasoundEditor/Private/MetasoundNodeDetailCustomization.cpp
+++ b/Engine/Plugins/Runtime/Metasound/Source/MetasoundEditor/Private/MetasoundNodeDetailCustomization.cpp
@@ -459,6 +459,65 @@ namespace Metasound
 			FMetasoundDefaultLiteralCustomizationBase::CustomizePageDefaultRows(InLiteral, InDetailLayout);
 		}
 
+		void FMetasoundObjectArrayLiteralCustomization::BuildDefaultValueWidget(IDetailPropertyRow& ValueRow, TSharedPtr<IPropertyHandle> ValueProperty)
+		{
+			if (!ValueProperty.IsValid())
+			{
+				return;
+			}
+
+			(*ValueRow.CustomValueWidget())
+			[
+				SNew(SAssetDropTarget)
+					.bSupportsMultiDrop(true)
+					.OnAreAssetsAcceptableForDropWithReason_Lambda([this, ValueProperty](TArrayView<FAssetData> InAssets, FText& OutReason)
+					{
+						Frontend::FDataTypeRegistryInfo DataTypeInfo;
+						const bool bMemberFound = MemberCustomizationPrivate::GetDataTypeFromElementPropertyHandle(ValueProperty, DataTypeInfo);
+						bool bCanDrop = bMemberFound;
+						if (UClass* ProxyGenClass = DataTypeInfo.ProxyGeneratorClass; bCanDrop && bMemberFound)
+						{
+							bCanDrop = true;
+							for (const FAssetData& AssetData : InAssets)
+							{
+								if (UClass* Class = AssetData.GetClass())
+								{
+									if (DataTypeInfo.bIsExplicit)
+									{
+										bCanDrop &= Class == DataTypeInfo.ProxyGeneratorClass;
+									}
+									else
+									{
+										bCanDrop &= Class->IsChildOf(DataTypeInfo.ProxyGeneratorClass);
+									}
+								}
+							}
+						}
+						return bCanDrop;
+					})
+					.OnAssetsDropped_Lambda([this, ValueProperty](const FDragDropEvent& DragDropEvent, TArrayView<FAssetData> InAssets)
+					{
+						TSharedPtr<IPropertyHandleArray> ArrayProperty = ValueProperty->AsArray();
+						if (ArrayProperty.IsValid())
+						{
+							FScopedTransaction Transaction(LOCTEXT("DragDropInputAssets", "Drop Asset(s) on MetaSound Input"));
+							for (const FAssetData& AssetData : InAssets)
+							{
+								uint32 AddIndex = INDEX_NONE;
+								ArrayProperty->GetNumElements(AddIndex);
+								ArrayProperty->AddItem();
+								TSharedPtr<IPropertyHandle> ElementHandle = ArrayProperty->GetElement(static_cast<int32>(AddIndex));
+								TSharedPtr<IPropertyHandle> ObjectHandle = ElementHandle->GetChildHandle(GET_MEMBER_NAME_CHECKED(FMetasoundEditorGraphMemberDefaultObjectRef, Object));
+								ObjectHandle->SetValue(AssetData.GetAsset());
+							}
+						}
+					})
+					[
+						ValueProperty->CreatePropertyValueWidget()
+					]
+			];
+		}
+
 		FText FMetasoundMemberDefaultBoolDetailCustomization::GetPropertyNameOverride() const
 		{
 			using namespace MemberCustomizationPrivate;
diff --git a/Engine/Plugins/Runtime/Metasound/Source/MetasoundEditor/Private/MetasoundNodeDetailCustomization.h b/Engine/Plugins/Runtime/Metasound/Source/MetasoundEditor/Private/MetasoundNodeDetailCustomization.h
--- a/Engine/Plugins/Runtime/Metasound/Source/MetasoundEditor/Private/MetasoundNodeDetailCustomization.h
+++ b/Engine/Plugins/Runtime/Metasound/Source/MetasoundEditor/Private/MetasoundNodeDetailCustomization.h
@@ -214,6 +214,7 @@ namespace Metasound
 
 			virtual ~FMetasoundObjectArrayLiteralCustomization() = default;
 
+			virtual void BuildDefaultValueWidget(IDetailPropertyRow& ValueRow, TSharedPtr<IPropertyHandle> ValueProperty) override;
 			virtual void CustomizeDefaults(UMetasoundEditorGraphMemberDefaultLiteral& InLiteral, IDetailLayoutBuilder& InDetailLayout) override;
 		};
 
diff --git a/Engine/Plugins/Runtime/Metasound/Source/MetasoundEditor/Public/MetasoundDefaultLiteralCustomization.h b/Engine/Plugins/Runtime/Metasound/Source/MetasoundEditor/Public/MetasoundDefaultLiteralCustomization.h
--- a/Engine/Plugins/Runtime/Metasound/Source/MetasoundEditor/Public/MetasoundDefaultLiteralCustomization.h
+++ b/Engine/Plugins/Runtime/Metasound/Source/MetasoundEditor/Public/MetasoundDefaultLiteralCustomization.h
@@ -47,6 +47,7 @@ namespace Metasound::Editor
 
 	protected:
 		void CustomizePageDefaultRows(UMetasoundEditorGraphMemberDefaultLiteral& InLiteral, IDetailLayoutBuilder& InDetailLayout);
+		virtual void BuildDefaultValueWidget(IDetailPropertyRow& ValueRow, TSharedPtr<IPropertyHandle> ValueProperty);
 
 		TArray<TSharedPtr<IPropertyHandle>> DefaultProperties;
 
