Fix for cook stripping targeting default instead of desired page due to ini platform name/platform name difference

#tests horde cooks with lots of logging
#rnx

diff --git a/Engine/Plugins/Runtime/Metasound/Source/MetasoundEngine/Private/MetasoundEngineAsset.h b/Engine/Plugins/Runtime/Metasound/Source/MetasoundEngine/Private/MetasoundEngineAsset.h
--- a/Engine/Plugins/Runtime/Metasound/Source/MetasoundEngine/Private/MetasoundEngineAsset.h
+++ b/Engine/Plugins/Runtime/Metasound/Source/MetasoundEngine/Private/MetasoundEngineAsset.h
@@ -252,7 +252,7 @@ namespace Metasound::Engine
 				FName PlatformName;
 				if (const ITargetPlatform* TargetPlatform = InSaveContext.GetTargetPlatform())
 				{
-					PlatformName = *TargetPlatform->PlatformName();
+					PlatformName = *TargetPlatform->IniPlatformName();
 				}
 				const bool bIsDeterministic = SerializationRequiresDeterminism(bIsCooking);
 				FDocumentIDGenerator::FScopeDeterminism DeterminismScope(bIsDeterministic);
