Fix warning when synchronizing MetaSound literal types

#rnx

diff --git a/Engine/Plugins/Runtime/Metasound/Source/MetasoundEditor/Private/MetasoundEditorGraphMemberDefaults.cpp b/Engine/Plugins/Runtime/Metasound/Source/MetasoundEditor/Private/MetasoundEditorGraphMemberDefaults.cpp
--- a/Engine/Plugins/Runtime/Metasound/Source/MetasoundEditor/Private/MetasoundEditorGraphMemberDefaults.cpp
+++ b/Engine/Plugins/Runtime/Metasound/Source/MetasoundEditor/Private/MetasoundEditorGraphMemberDefaults.cpp
@@ -13,6 +13,7 @@
 #include "MetasoundEditorGraph.h"
 #include "MetasoundEditorGraphBuilder.h"
 #include "MetasoundEditorGraphNode.h"
+#include "MetasoundEditorModule.h"
 #include "MetasoundEditorSettings.h"
 #include "MetasoundFrontend.h"
 #include "MetasoundFrontendController.h"
@@ -290,9 +291,9 @@ namespace Metasound::Editor
 									LiteralType NewValue;
 									if (!TestLiteral.TryGet(NewValue))
 									{
-										UE_LOG(LogMetaSound, Warning, TEXT("Synchronizing Page Default of member '%s' (type '%s') was supplied literal value '%s' with mismatched type"),
+										UE_LOG(LogMetasoundEditor, Verbose, TEXT("Synchronizing Page Default: Setting member '%s' (type '%s') to literal value '%s'. Type has changed or literal could not be set and will be set to type's default constructed value."),
 											*Member->GetMemberName().ToString(),
-											*Member->GetLiteral()->GetClass()->GetName(),
+											*Member->GetDataType().ToString(),
 											*TestLiteral.ToString());
 									}
 									Default.Value = LiteralToMember(NewValue);
