[MetaSound][Editor] Move audiowidgets style override logic to be inside the Metasound::Editor::FSlateStyle class. Putting this logic behind an ISlateStyle interface will allow an ISlateStyle to be passed to the analyzer rack, allowing the overridden styles to still be applied in the analyzer rack.
This implementation also fixes the risk of a crash if a material style asset is deleted after a material widget was created, due to audio material widgets storing a pointer to the widget style.

diff --git a/Engine/Plugins/Runtime/AudioWidgets/Source/AudioWidgets/Private/AudioWidgetsStyle.cpp b/Engine/Plugins/Runtime/AudioWidgets/Source/AudioWidgets/Private/AudioWidgetsStyle.cpp
--- a/Engine/Plugins/Runtime/AudioWidgets/Source/AudioWidgets/Private/AudioWidgetsStyle.cpp
+++ b/Engine/Plugins/Runtime/AudioWidgets/Source/AudioWidgets/Private/AudioWidgetsStyle.cpp
@@ -13,6 +13,7 @@
 #include "Brushes/SlateRoundedBoxBrush.h"
 #include "Engine/Texture2D.h"
 #include "Misc/Paths.h"
+#include "Styling/AppStyle.h"
 #include "Styling/StyleColors.h"
 #include "Styling/SlateStyleRegistry.h"
 
@@ -26,7 +27,7 @@ namespace AudioWidgetsStylePrivate
 FAudioWidgetsStyle::FAudioWidgetsStyle()
 	: FSlateStyleSet(StyleName)
 {
-	SetParentStyleName("CoreStyle");
+	SetParentStyleName(FAppStyle::GetAppStyleSetName());
 	SetContentRoot(FPaths::EnginePluginsDir() / TEXT("Runtime/AudioWidgets/Content"));
 	SetResources();
 
diff --git a/Engine/Plugins/Runtime/Metasound/Source/MetasoundEditor/Private/MetasoundEditor.cpp b/Engine/Plugins/Runtime/Metasound/Source/MetasoundEditor/Private/MetasoundEditor.cpp
--- a/Engine/Plugins/Runtime/Metasound/Source/MetasoundEditor/Private/MetasoundEditor.cpp
+++ b/Engine/Plugins/Runtime/Metasound/Source/MetasoundEditor/Private/MetasoundEditor.cpp
@@ -1715,28 +1715,24 @@ namespace Metasound
 
 		void FEditor::CreateAnalyzers(UMetaSoundSource& MetaSoundSource)
 		{
-			if (ensure(GEditor))
+			const ISlateStyle* StyleSet = FSlateStyleRegistry::FindSlateStyle("MetaSoundStyle");
+			if (ensure(GEditor) && ensure(StyleSet))
 			{
 				const Audio::FDeviceId AudioDeviceId = GEditor->GetMainAudioDeviceID();
 				constexpr UAudioBus* DefaultBus = nullptr;
 
 				if (!OutputMeter.IsValid())
 				{
-					const UMetasoundEditorSettings* EditorSettings = GetDefault<UMetasoundEditorSettings>();
-					check(EditorSettings)
-					const bool bUseAudioMaterialWidgets = EditorSettings->bUseAudioMaterialWidgets;
-					if (bUseAudioMaterialWidgets)
+					const FName MaterialMeterStyleName("AudioMaterialMeter.Style");
+					if (StyleSet->HasWidgetStyle<FAudioMaterialMeterStyle>(MaterialMeterStyleName))
 					{
-						const FAudioMaterialMeterStyle* MeterStyle = EditorSettings->GetMeterStyle();
-						if (ensureMsgf(MeterStyle, TEXT("Failed to find MaterialMeterStyle when attempting to build MetaSound Editor output meter. Falling back to default non-material meter.")))
-						{
-							OutputMeter = MakeShared<AudioWidgets::FAudioMeter>(MetaSoundSource.NumChannels, AudioDeviceId, *MeterStyle, DefaultBus);
-						}
+						const FAudioMaterialMeterStyle& AudioMaterialMeterStyle = StyleSet->GetWidgetStyle<FAudioMaterialMeterStyle>(MaterialMeterStyleName);
+						OutputMeter = MakeShared<AudioWidgets::FAudioMeter>(MetaSoundSource.NumChannels, AudioDeviceId, AudioMaterialMeterStyle, DefaultBus);
 					}
-
-					if (!OutputMeter.IsValid())
+					else
 					{
-						OutputMeter = MakeShared<AudioWidgets::FAudioMeter>(MetaSoundSource.NumChannels, AudioDeviceId, DefaultBus, &Style::GetMeterDefaultColorStyle());
+						const FAudioMeterDefaultColorStyle& MeterDefaultColorStyle = StyleSet->GetWidgetStyle<FAudioMeterDefaultColorStyle>("AudioMeter.DefaultColorStyle");
+						OutputMeter = MakeShared<AudioWidgets::FAudioMeter>(MetaSoundSource.NumChannels, AudioDeviceId, DefaultBus, &MeterDefaultColorStyle);
 					}
 				}
 				else if (OutputMeter->GetAudioBus()->GetNumChannels() != MetaSoundSource.NumChannels)
@@ -1760,13 +1756,13 @@ namespace Metasound
 						OscilloscopeMaxTimeWindowMs,
 						OscilloscopeAnalysisPeriodMs,
 						OscilloscopePanelLayoutType,
-						&Style::GetOscilloscopeStyle());
+						&StyleSet->GetWidgetStyle<FAudioOscilloscopePanelStyle>("AudioOscilloscope.PanelStyle"));
 				}
 				else if (OutputOscilloscope->GetAudioBus()->GetNumChannels() != MetaSoundSource.NumChannels)
 				{
 					OutputOscilloscope->CreateAudioBus(MetaSoundNumChannels);
 					OutputOscilloscope->CreateDataProvider(AudioDeviceId, OscilloscopeTimeWindowMs, OscilloscopeMaxTimeWindowMs, OscilloscopeAnalysisPeriodMs, OscilloscopePanelLayoutType);
-					OutputOscilloscope->CreateOscilloscopeWidget(MetaSoundNumChannels, OscilloscopePanelLayoutType, &Style::GetOscilloscopeStyle());
+					OutputOscilloscope->CreateOscilloscopeWidget(MetaSoundNumChannels, OscilloscopePanelLayoutType, &StyleSet->GetWidgetStyle<FAudioOscilloscopePanelStyle>("AudioOscilloscope.PanelStyle"));
 				}
 
 				// Init Vectorscope
@@ -1783,13 +1779,13 @@ namespace Metasound
 						VectorscopeMaxTimeWindowMs,
 						VectorscopeAnalysisPeriodMs,
 						VectorscopePanelLayoutType,
-						&Style::GetVectorscopeStyle());
+						&StyleSet->GetWidgetStyle<FAudioVectorscopePanelStyle>("AudioVectorscope.PanelStyle"));
 				}
 				else if (OutputVectorscope->GetAudioBus()->GetNumChannels() != MetaSoundSource.NumChannels)
 				{
 					OutputVectorscope->CreateAudioBus(MetaSoundNumChannels);
 					OutputVectorscope->CreateDataProvider(AudioDeviceId, VectorscopeTimeWindowMs, VectorscopeMaxTimeWindowMs, VectorscopeAnalysisPeriodMs);
-					OutputVectorscope->CreateVectorscopeWidget(VectorscopePanelLayoutType, &Style::GetVectorscopeStyle());
+					OutputVectorscope->CreateVectorscopeWidget(VectorscopePanelLayoutType, &StyleSet->GetWidgetStyle<FAudioVectorscopePanelStyle>("AudioVectorscope.PanelStyle"));
 				}
 
 				if (!OutputSpectrumAnalyzer.IsValid())
@@ -1883,7 +1879,7 @@ namespace Metasound
 							SpectrumAnalyzerSettings.bDisplaySoundLevelAxisLabels = !SpectrumAnalyzerSettings.bDisplaySoundLevelAxisLabels;
 							GetMutableDefault<UMetasoundEditorSettings>()->SaveConfig();
 						});
-					Params.PlotStyle = &Style::GetSpectrumPlotStyle();
+					Params.PlotStyle = &StyleSet->GetWidgetStyle<FAudioSpectrumPlotStyle>("AudioSpectrumPlot.Style");
 					OutputSpectrumAnalyzer = MakeShared<AudioWidgets::FAudioSpectrumAnalyzer>(Params);
 				}
 				else if (OutputSpectrumAnalyzer->GetAudioBus()->GetNumChannels() != MetaSoundSource.NumChannels)
diff --git a/Engine/Plugins/Runtime/Metasound/Source/MetasoundEditor/Private/MetasoundEditorModule.cpp b/Engine/Plugins/Runtime/Metasound/Source/MetasoundEditor/Private/MetasoundEditorModule.cpp
--- a/Engine/Plugins/Runtime/Metasound/Source/MetasoundEditor/Private/MetasoundEditorModule.cpp
+++ b/Engine/Plugins/Runtime/Metasound/Source/MetasoundEditor/Private/MetasoundEditorModule.cpp
@@ -2,6 +2,7 @@
 #include "MetasoundEditorModule.h"
 
 #include "AssetTypeActions_Base.h"
+#include "AudioMaterialSlate/AudioMaterialSlateTypes.h"
 #include "AudioMeter.h"
 #include "AudioOscilloscopePanelStyle.h"
 #include "AudioSpectrumPlotStyle.h"
@@ -86,7 +87,7 @@ namespace Metasound
 			FSlateStyle()
 				: FSlateStyleSet("MetaSoundStyle")
 			{
-				SetParentStyleName(FAppStyle::GetAppStyleSetName());
+				SetParentStyleName(FAudioWidgetsStyle::Get().GetStyleSetName());
 
 				SetContentRoot(FPaths::EnginePluginsDir() / TEXT("Runtime/Metasound/Content/Editor/Slate"));
 				SetCoreContentRoot(FPaths::EngineContentDir() / TEXT("Slate"));
@@ -198,30 +199,119 @@ namespace Metasound
 					Set("MetasoundEditor.MetasoundSource.Preset.Thumbnail", new IMAGE_BRUSH_SVG(TEXT("Icons/metasoundsourcepreset_thumbnail"), Icon20x20));
 				}
 
+				// Audio Widgets
+				{
+					const FLinearColor& AnalyzerForegroundColor = GetColor("MetasoundEditor.Analyzers.ForegroundColor");
+
+					// Add static overrides for these widget styles:
+					FAudioMeterDefaultColorStyle MeterStyle;
+					MeterStyle.MeterValueColor = AnalyzerForegroundColor;
+					Set("AudioMeter.DefaultColorStyle", MeterStyle);
+
+					Set("AudioOscilloscope.PanelStyle", FAudioOscilloscopePanelStyle()
+						.SetWaveViewerStyle(FSampledSequenceViewerStyle()
+							.SetSequenceColor(AnalyzerForegroundColor)));
+
+					Set("AudioSpectrumPlot.Style", FAudioSpectrumPlotStyle()
+						.SetCrosshairColor(FSlateColor(AnalyzerForegroundColor).UseSubduedForeground())
+						.SetSpectrumColor(AnalyzerForegroundColor));
+
+					Set("AudioVectorscope.PanelStyle", FAudioVectorscopePanelStyle()
+						.SetVectorViewerStyle(FSampledSequenceVectorViewerStyle()
+							.SetLineColor(AnalyzerForegroundColor)));
+
+					// Add dynamic overrides for the widget style types with settings overrides:
+					AddDynamicLoadedWidgetStyle<FAudioMaterialKnobStyle>([]() { return GetDefault<UMetasoundEditorSettings>()->KnobStyleOverride; });
+					AddDynamicLoadedWidgetStyle<FAudioMaterialSliderStyle>([]() { return GetDefault<UMetasoundEditorSettings>()->SliderStyleOverride; });
+					AddDynamicLoadedWidgetStyle<FAudioMaterialButtonStyle>([]() { return GetDefault<UMetasoundEditorSettings>()->ButtonStyleOverride; });
+					AddDynamicLoadedWidgetStyle<FAudioMaterialMeterStyle>([]() { return GetDefault<UMetasoundEditorSettings>()->MeterStyleOverride; });
+				}
+
 				FSlateStyleRegistry::RegisterSlateStyle(*this);
 			}
-		};
 
-		namespace Style
-		{
-			FSlateIcon CreateSlateIcon(FName InName)
+		private:
+			virtual const FSlateWidgetStyle* GetWidgetStyleInternal(const FName DesiredTypeName, const FName StyleName, const FSlateWidgetStyle* DefaultStyle, bool bWarnIfNotFound) const override
 			{
-				return { "MetaSoundStyle", InName};
+				if (const auto* WidgetStyleLoader = AudioMaterialWidgetStyleLoaders.Find(DesiredTypeName))
+				{
+					const UMetasoundEditorSettings* EditorSettings = GetDefault<UMetasoundEditorSettings>();
+					if (EditorSettings == nullptr || !EditorSettings->bUseAudioMaterialWidgets)
+					{
+						ensure(!bWarnIfNotFound);
+
+						// AudioMaterialWidgets are disabled, so explicitly return null rather than fallback to the base style.
+						return nullptr;
+					}
+					
+					if (const FSlateWidgetStyle* CustomWidgetStyle = (*WidgetStyleLoader)->LoadWidgetStyle())
+					{
+						return CustomWidgetStyle;
+					}
+				}
+				
+				const FSlateWidgetStyle* WidgetStyle = FSlateStyleSet::GetWidgetStyleInternal(DesiredTypeName, StyleName, DefaultStyle, bWarnIfNotFound);
+				ensure(!bWarnIfNotFound || WidgetStyle != nullptr);
+				return WidgetStyle;
 			}
 
-			const FSlateColor& GetDefaultAnalyzerColor()
+			using FGetAssetPathFunc = TFunction<FSoftObjectPath()>;
+
+			template<std::derived_from<FSlateWidgetStyle> WidgetStyleType>
+			void AddDynamicLoadedWidgetStyle(FGetAssetPathFunc GetAssetPath)
 			{
-				auto MakeColor = []() -> FSlateColor
+				static_assert(std::is_base_of_v<FAudioMaterialWidgetStyle, WidgetStyleType>);
+				AudioMaterialWidgetStyleLoaders.Add(WidgetStyleType::TypeName, MakeShared<TWidgetStyleLoader<WidgetStyleType>>(GetAssetPath));
+			}
+
+			class IWidgetStyleLoader
+			{
+			public:
+				virtual ~IWidgetStyleLoader() {};
+				virtual const FSlateWidgetStyle* LoadWidgetStyle() = 0;
+			};
+
+			template<std::derived_from<FSlateWidgetStyle> WidgetStyleType>
+			class TWidgetStyleLoader : public IWidgetStyleLoader
+			{
+			public:
+				TWidgetStyleLoader(FGetAssetPathFunc InGetAssetPath)
+					: GetAssetPath(InGetAssetPath)
 				{
-					if (const ISlateStyle* MetaSoundStyle = FSlateStyleRegistry::FindSlateStyle("MetaSoundStyle"))
+					//
+				}
+
+				virtual const FSlateWidgetStyle* LoadWidgetStyle() override
+				{
+					TSoftObjectPtr<USlateWidgetStyleAsset> SoftObjectPtr(GetAssetPath());
+					if (const USlateWidgetStyleAsset* WidgetStyleAsset = SoftObjectPtr.LoadSynchronous())
 					{
-						return MetaSoundStyle->GetColor("MetasoundEditor.Analyzers.ForegroundColor");
+						if (const WidgetStyleType* WidgetStyle = WidgetStyleAsset->GetStyle<WidgetStyleType>())
+						{
+							// Copy style:
+							CachedStyle = *WidgetStyle;
+
+							// Return pointer to our copy:
+							return &CachedStyle;
+						}
 					}
 
-					return FStyleColors::AccentWhite;
-				};
-				static const FSlateColor AnalyzerColor = MakeColor();
-				return AnalyzerColor;
+					return nullptr;
+				}
+
+			private:
+				FGetAssetPathFunc GetAssetPath;
+				WidgetStyleType CachedStyle;
+			};
+
+			TMap<FName, TSharedRef<IWidgetStyleLoader>> AudioMaterialWidgetStyleLoaders;
+		};
+
+		namespace Style
+		{
+			FSlateIcon CreateSlateIcon(FName InName)
+			{
+				return { "MetaSoundStyle", InName};
 			}
 
 			const FSlateBrush& GetSlateBrushSafe(FName InName)
@@ -245,33 +335,6 @@ namespace Metasound
 				return NullBrush;
 			}
 
-			const FAudioMeterDefaultColorStyle& GetMeterDefaultColorStyle()
-			{
-				auto MakeStyle = []()
-				{
-					FAudioMeterDefaultColorStyle MeterStyle;
-					MeterStyle.MeterValueColor = GetDefaultAnalyzerColor().GetSpecifiedColor();
-					return MeterStyle;
-				};
-				static const FAudioMeterDefaultColorStyle ThisStyle = MakeStyle();
-				return ThisStyle;
-			}
-
-			const FAudioOscilloscopePanelStyle& GetOscilloscopeStyle()
-			{
-				auto MakeStyle = []()
-				{
-					FAudioOscilloscopePanelStyle OscStyle;
-					FSampledSequenceViewerStyle SampleView;
-					SampleView.SequenceColor = GetDefaultAnalyzerColor();
-					OscStyle.SetWaveViewerStyle(SampleView);
-
-					return OscStyle;
-				};
-				static FAudioOscilloscopePanelStyle ThisStyle = MakeStyle();
-				return ThisStyle;
-			}
-
 			const FSlateColor& GetPageExecutingColor()
 			{
 				auto MakeColor = []() -> FSlateColor
@@ -286,34 +349,6 @@ namespace Metasound
 				static const FSlateColor AnalyzerColor = MakeColor();
 				return AnalyzerColor;
 			}
-
-			const FAudioSpectrumPlotStyle& GetSpectrumPlotStyle()
-			{
-				auto MakeStyle = []()
-				{
-					FAudioSpectrumPlotStyle PlotStyle;
-					const FSlateColor AnalyzerColor = GetDefaultAnalyzerColor();
-					PlotStyle.CrosshairColor = AnalyzerColor.UseSubduedForeground();
-					PlotStyle.SpectrumColor = AnalyzerColor;
-					return PlotStyle;
-				};
-				static FAudioSpectrumPlotStyle ThisStyle = MakeStyle();
-				return ThisStyle;
-			}
-
-			const FAudioVectorscopePanelStyle& GetVectorscopeStyle()
-			{
-				auto MakeStyle = []()
-				{
-					FAudioVectorscopePanelStyle PanelStyle;
-					FSampledSequenceVectorViewerStyle VectorViewStyle;
-					VectorViewStyle.LineColor = GetDefaultAnalyzerColor().GetSpecifiedColor();
-					PanelStyle.SetVectorViewerStyle(VectorViewStyle);
-					return PanelStyle;
-				};
-				static FAudioVectorscopePanelStyle ThisStyle = MakeStyle();
-				return ThisStyle;
-			}
 		} // namespace Style
 
 		// A structure that contains information about registered custom pin types. 
diff --git a/Engine/Plugins/Runtime/Metasound/Source/MetasoundEditor/Private/MetasoundEditorSettings.cpp b/Engine/Plugins/Runtime/Metasound/Source/MetasoundEditor/Private/MetasoundEditorSettings.cpp
--- a/Engine/Plugins/Runtime/Metasound/Source/MetasoundEditor/Private/MetasoundEditorSettings.cpp
+++ b/Engine/Plugins/Runtime/Metasound/Source/MetasoundEditor/Private/MetasoundEditorSettings.cpp
@@ -61,58 +61,6 @@ void UMetasoundEditorSettings::PostEditChangeProperty(FPropertyChangedEvent& Pro
 }
 #endif // WITH_EDITOR
 
-const FAudioMaterialKnobStyle* UMetasoundEditorSettings::GetKnobStyle() const
-{
-	if (const UObject* Style = KnobStyleOverride.TryLoad())
-	{
-		if (const USlateWidgetStyleAsset* SlateWidgetStyleAsset = CastChecked<USlateWidgetStyleAsset>(Style))
-		{
-			return SlateWidgetStyleAsset->GetStyle<FAudioMaterialKnobStyle>();
-		}
-	}
-
-	return &FAudioWidgetsStyle::Get().GetWidgetStyle<FAudioMaterialKnobStyle>("AudioMaterialKnob.Style");
-}
-
-const FAudioMaterialSliderStyle* UMetasoundEditorSettings::GetSliderStyle() const
-{
-	if (const UObject* Style = SliderStyleOverride.TryLoad())
-	{
-		if (const USlateWidgetStyleAsset* SlateWidgetStyleAsset = CastChecked<USlateWidgetStyleAsset>(Style))
-		{
-			return SlateWidgetStyleAsset->GetStyle<FAudioMaterialSliderStyle>();
-		}
-	}
-
-	return &FAudioWidgetsStyle::Get().GetWidgetStyle<FAudioMaterialSliderStyle>("AudioMaterialSlider.Style");
-}
-
-const FAudioMaterialButtonStyle* UMetasoundEditorSettings::GetButtonStyle() const
-{
-	if (const UObject* Style = ButtonStyleOverride.TryLoad())
-	{
-		if (const USlateWidgetStyleAsset* SlateWidgetStyleAsset = CastChecked<USlateWidgetStyleAsset>(Style))
-		{
-			return SlateWidgetStyleAsset->GetStyle<FAudioMaterialButtonStyle>();
-		}
-	}
-
-	return &FAudioWidgetsStyle::Get().GetWidgetStyle<FAudioMaterialButtonStyle>("AudioMaterialButton.Style");
-}
-
-const FAudioMaterialMeterStyle* UMetasoundEditorSettings::GetMeterStyle() const
-{
-	if (const UObject* Style = MeterStyleOverride.TryLoad())
-	{
-		if (const USlateWidgetStyleAsset* SlateWidgetStyleAsset = CastChecked<USlateWidgetStyleAsset>(Style))
-		{
-			return SlateWidgetStyleAsset->GetStyle<FAudioMaterialMeterStyle>();
-		}
-	}
-
-	return &FAudioWidgetsStyle::Get().GetWidgetStyle<FAudioMaterialMeterStyle>("AudioMaterialMeter.Style");
-}
-
 Metasound::Engine::FPageResolutionEditorResults UMetasoundEditorSettings::ResolveAuditionPage(const TArray<FGuid>& InPageIDs) const
 {
 	using namespace Metasound;
diff --git a/Engine/Plugins/Runtime/Metasound/Source/MetasoundEditor/Private/MetasoundEditorSettings.h b/Engine/Plugins/Runtime/Metasound/Source/MetasoundEditor/Private/MetasoundEditorSettings.h
--- a/Engine/Plugins/Runtime/Metasound/Source/MetasoundEditor/Private/MetasoundEditorSettings.h
+++ b/Engine/Plugins/Runtime/Metasound/Source/MetasoundEditor/Private/MetasoundEditorSettings.h
@@ -295,19 +295,7 @@ public:
 	virtual void PostEditChangeProperty(FPropertyChangedEvent& PropertyChangedEvent) override;
 #endif // WITH_EDITOR
 	//~UObject
-		
-	/** Get the AudioMaterialKnob Style. If KnobStyleOverride is not set, returns default style.*/
-	const FAudioMaterialKnobStyle* GetKnobStyle() const;
-
-	/** Get the AudioMaterialSlider Style. If SliderStyleOverride is not set, returns default style.*/
-	const FAudioMaterialSliderStyle* GetSliderStyle() const;
-	
-	/** Get the AudioMaterialButton Style. If ButtonStyleOverride is not set, returns default style.*/
-	const FAudioMaterialButtonStyle* GetButtonStyle() const;
 	
-	/** Get the AudioMaterialMeter Style. If MeterStyleOverride is not set, returns default style.*/
-	const FAudioMaterialMeterStyle* GetMeterStyle() const;
-
 	Metasound::Engine::FPageResolutionEditorResults ResolveAuditionPage(const TArray<FGuid>& InPageIDs) const;
 	
 	/** Given the provided AuditionPageID, returns the resolved PageID from the provided array of values based on fallback logic. */
diff --git a/Engine/Plugins/Runtime/Metasound/Source/MetasoundEditor/Private/MetasoundPinAudioInspector.cpp b/Engine/Plugins/Runtime/Metasound/Source/MetasoundEditor/Private/MetasoundPinAudioInspector.cpp
--- a/Engine/Plugins/Runtime/Metasound/Source/MetasoundEditor/Private/MetasoundPinAudioInspector.cpp
+++ b/Engine/Plugins/Runtime/Metasound/Source/MetasoundEditor/Private/MetasoundPinAudioInspector.cpp
@@ -43,13 +43,19 @@ namespace Metasound::Editor
 
 		constexpr int32 NumChannels = 1; // Audio wires are currently mono signals
 
+		const FAudioOscilloscopePanelStyle* OscilloscopePanelStyle = nullptr;
+		if (const ISlateStyle* StyleSet = FSlateStyleRegistry::FindSlateStyle("MetaSoundStyle"))
+		{
+			OscilloscopePanelStyle = &StyleSet->GetWidgetStyle<FAudioOscilloscopePanelStyle>("AudioOscilloscope.PanelStyle");
+		}
+
 		Oscilloscope = MakeShared<AudioWidgets::FAudioOscilloscope>(AudioDeviceId,
 			NumChannels,
 			/*InTimeWindowMs*/     10.0f,
 			/*InMaxTimeWindowMs*/  10.0f,
 			/*InAnalysisPeriodMs*/ 10.0f,
 			/*InPanelLayoutType*/  EAudioPanelLayoutType::Basic,
-			/*InOscilloscopePanelStyle*/ &Style::GetOscilloscopeStyle()
+			/*InOscilloscopePanelStyle*/ OscilloscopePanelStyle
 		);
 
 		PinAudioInspectorWidget = SNew(SMetasoundPinAudioInspector)
diff --git a/Engine/Plugins/Runtime/Metasound/Source/MetasoundEditor/Private/SMetasoundFilterFrequencyResponsePlots.cpp b/Engine/Plugins/Runtime/Metasound/Source/MetasoundEditor/Private/SMetasoundFilterFrequencyResponsePlots.cpp
--- a/Engine/Plugins/Runtime/Metasound/Source/MetasoundEditor/Private/SMetasoundFilterFrequencyResponsePlots.cpp
+++ b/Engine/Plugins/Runtime/Metasound/Source/MetasoundEditor/Private/SMetasoundFilterFrequencyResponsePlots.cpp
@@ -107,7 +107,7 @@ namespace Metasound::Editor
 					.TiltExponent_Lambda([] { return 0.0f; }) // Binding this property has the effect of hiding its context menu entry (Tilting the spectrum is not desired here).
 					.FrequencyAxisPixelBucketMode_Lambda([]() { return EAudioSpectrumPlotFrequencyAxisPixelBucketMode::Sample; }) // Binding this property has the effect of hiding its context menu entry (PixelBucketMode is not much use here).
 					.OnGetAudioSpectrumData(this, &SMetaSoundFilterFrequencyResponsePlot::GetAudioSpectrumData)
-					.Style(&Editor::Style::GetSpectrumPlotStyle())
+					.Style(FSlateStyleRegistry::FindSlateStyle("MetaSoundStyle"), "AudioSpectrumPlot.Style")
 			];
 
 		ContextMenuExtension = FrequencyResponsePlot->AddContextMenuExtension(EExtensionHook::Before, nullptr, FMenuExtensionDelegate::CreateSP(this, &SMetaSoundFilterFrequencyResponsePlot::ExtendSpectrumPlotContextMenu));
diff --git a/Engine/Plugins/Runtime/Metasound/Source/MetasoundEditor/Private/SMetasoundGraphNode.cpp b/Engine/Plugins/Runtime/Metasound/Source/MetasoundEditor/Private/SMetasoundGraphNode.cpp
--- a/Engine/Plugins/Runtime/Metasound/Source/MetasoundEditor/Private/SMetasoundGraphNode.cpp
+++ b/Engine/Plugins/Runtime/Metasound/Source/MetasoundEditor/Private/SMetasoundGraphNode.cpp
@@ -871,7 +871,7 @@ namespace Metasound
 					{
 						SAssignNew(FloatInputWidget, SAudioMaterialLabeledSlider)
 							.Owner(GraphMember->GetOwningGraph())
-							.Style(EditorSettings->GetSliderStyle())
+							.Style(FSlateStyleRegistry::FindSlateStyle("MetaSoundStyle"), "AudioMaterialSlider.Style")
 							.AudioUnitsValueType(DefaultFloat->WidgetUnitValueType)
 							.bUseLinearOutput(DefaultFloat->VolumeWidgetUseLinearOutput)
 							.OnValueChanged_Lambda(OnValueChangedLambda)
@@ -1004,7 +1004,7 @@ namespace Metasound
 					{
 						SAssignNew(FloatInputWidget, SAudioMaterialLabeledKnob)
 							.Owner(GraphMember->GetOwningGraph())
-							.Style(EditorSettings->GetKnobStyle())
+							.Style(FSlateStyleRegistry::FindSlateStyle("MetaSoundStyle"), "AudioMaterialKnob.Style")
 							.OnValueChanged_Lambda(OnValueChangedLambda)
 							.AudioUnitsValueType(DefaultFloat->WidgetUnitValueType)
 							.bUseLinearOutput(DefaultFloat->VolumeWidgetUseLinearOutput)
@@ -1163,7 +1163,7 @@ namespace Metasound
 							};
 
 							SAssignNew(MaterialButtonWidget, SAudioMaterialButton)
-								.AudioMaterialButtonStyle(EditorSettings->GetButtonStyle())
+								.AudioMaterialButtonStyle(FSlateStyleRegistry::FindSlateStyle("MetaSoundStyle"), "AudioMaterialButton.Style")
 								.OnBooleanValueChanged_Lambda(OnBoolValueChangedLambda)
 								.OnMouseCaptureEnd_Lambda(OnBoolMouseCaptureEndLambda)
 								.bIsPressedAttribute(DefaultBool->GetDefaultAs<bool>());
diff --git a/Engine/Plugins/Runtime/Metasound/Source/MetasoundEditor/Private/SMetasoundSpectrumAnalyzerGraphNode.cpp b/Engine/Plugins/Runtime/Metasound/Source/MetasoundEditor/Private/SMetasoundSpectrumAnalyzerGraphNode.cpp
--- a/Engine/Plugins/Runtime/Metasound/Source/MetasoundEditor/Private/SMetasoundSpectrumAnalyzerGraphNode.cpp
+++ b/Engine/Plugins/Runtime/Metasound/Source/MetasoundEditor/Private/SMetasoundSpectrumAnalyzerGraphNode.cpp
@@ -66,7 +66,10 @@ namespace Metasound::Editor
 				FAudioSpectrumAnalyzerParams Params;
 				Params.NumChannels = 1;
 				Params.AudioDeviceId = GEditor->GetMainAudioDeviceID();
-				Params.PlotStyle = &Style::GetSpectrumPlotStyle();
+				if (const ISlateStyle* StyleSet = FSlateStyleRegistry::FindSlateStyle("MetaSoundStyle"))
+				{
+					Params.PlotStyle = &StyleSet->GetWidgetStyle<FAudioSpectrumPlotStyle>("AudioSpectrumPlot.Style");
+				}
 				SpectrumAnalyzer = MakeShared<FAudioSpectrumAnalyzer>(Params);
 			}
 
diff --git a/Engine/Plugins/Runtime/Metasound/Source/MetasoundEditor/Public/MetasoundEditorModule.h b/Engine/Plugins/Runtime/Metasound/Source/MetasoundEditor/Public/MetasoundEditorModule.h
--- a/Engine/Plugins/Runtime/Metasound/Source/MetasoundEditor/Public/MetasoundEditorModule.h
+++ b/Engine/Plugins/Runtime/Metasound/Source/MetasoundEditor/Public/MetasoundEditorModule.h
@@ -47,13 +47,7 @@ namespace Metasound
 			METASOUNDEDITOR_API FSlateIcon CreateSlateIcon(FName InName);
 			METASOUNDEDITOR_API const FSlateBrush& GetSlateBrushSafe(FName InName);
 
-			const FSlateColor& GetDefaultAnalyzerColor();
 			const FSlateColor& GetPageExecutingColor();
-
-			const FAudioMeterDefaultColorStyle& GetMeterDefaultColorStyle();
-			const FAudioOscilloscopePanelStyle& GetOscilloscopeStyle();
-			const FAudioSpectrumPlotStyle& GetSpectrumPlotStyle();
-			const FAudioVectorscopePanelStyle& GetVectorscopeStyle();
 		} // namespace Style
 
 		using EAssetScanStatus = Metasound::Engine::EAssetScanStatus;
