Fix for displaying the wrong value or crashing when promoting to input from output pin on node on page with no default value overwritten.

#rnx

diff --git a/Engine/Plugins/Runtime/Metasound/Source/MetasoundEditor/Private/MetasoundEditorGraphBuilder.cpp b/Engine/Plugins/Runtime/Metasound/Source/MetasoundEditor/Private/MetasoundEditorGraphBuilder.cpp
--- a/Engine/Plugins/Runtime/Metasound/Source/MetasoundEditor/Private/MetasoundEditorGraphBuilder.cpp
+++ b/Engine/Plugins/Runtime/Metasound/Source/MetasoundEditor/Private/MetasoundEditorGraphBuilder.cpp
@@ -2644,7 +2644,12 @@ namespace Metasound
 			const TArray<FMetasoundFrontendClassInputDefault>* ClassDefaults = InBuilder.FindNodeClassInputDefaults(InputHandle.NodeID, InputVertex->Name);
 			if (ClassDefaults)
 			{
-				const FGuid PageID = Engine::FDocumentBuilderRegistry::GetChecked().ResolveTargetPageID(*ClassDefaults);
+				const UMetasoundEditorSettings* EditorSettings = GetDefault<UMetasoundEditorSettings>();
+				check(EditorSettings);
+
+				TArray<FGuid> PageIDs;
+				Algo::Transform(*ClassDefaults, PageIDs, [](const FMetasoundFrontendClassInputDefault& InputDefault) { return InputDefault.PageID; });
+				const FGuid PageID = EditorSettings->ResolveAuditionPage(PageIDs, InBuilder.GetBuildPageID());
 				auto MatchesPageID = [&PageID](const FMetasoundFrontendClassInputDefault& InputDefault) { return InputDefault.PageID == PageID; };
 				if (const FMetasoundFrontendClassInputDefault* ClassDefault = ClassDefaults->FindByPredicate(MatchesPageID))
 				{
diff --git a/Engine/Plugins/Runtime/Metasound/Source/MetasoundEditor/Private/MetasoundEditorGraphSchema.cpp b/Engine/Plugins/Runtime/Metasound/Source/MetasoundEditor/Private/MetasoundEditorGraphSchema.cpp
--- a/Engine/Plugins/Runtime/Metasound/Source/MetasoundEditor/Private/MetasoundEditorGraphSchema.cpp
+++ b/Engine/Plugins/Runtime/Metasound/Source/MetasoundEditor/Private/MetasoundEditorGraphSchema.cpp
@@ -426,18 +426,48 @@ PRAGMA_ENABLE_DEPRECATION_WARNINGS
 			ParentMetasound.Modify();
 			MetasoundGraph->Modify();
 
+			// Default literal must ALWAYS have value for default pageID, so even though this may get called
+			// from a page "higher" in the page stack, always add the literal value for the default page ID.
+			auto InitDefaultLiterals = [](FMetasoundFrontendLiteral NewLiteral)
+			{
+				TArray<FMetasoundFrontendClassInputDefault> InitValues;
+				InitValues.Add_GetRef(Metasound::Frontend::DefaultPageID).Literal = NewLiteral;
+				return InitValues;
+			};
+
 			TArray<FMetasoundFrontendClassInputDefault> DefaultLiterals;
 			if (const FMetasoundFrontendVertexLiteral* VertexLiteral = Builder.FindNodeInputDefault(InputVertexHandle.NodeID, InputVertexHandle.VertexID))
 			{
-				DefaultLiterals.Add_GetRef(Builder.GetBuildPageID()).Literal = VertexLiteral->Value;
+				// Since a default page ID requires an associated value and no other nodes on any page would be connected to this new input, use
+				// the default page ID. If the user wants different behavior, when they connect the newly created input on a lower-indexed graph,
+				// they will assign a proper page default value therein.  This in practice should cut down on duplicate page input default data.
+				DefaultLiterals = InitDefaultLiterals(VertexLiteral->Value);
 			}
-			if (DefaultLiterals.IsEmpty())
+			else if (const TArray<FMetasoundFrontendClassInputDefault>* ClassDefaults = Builder.FindNodeClassInputDefaults(InputVertexHandle.NodeID, InputVertex->Name))
 			{
-				if (const TArray<FMetasoundFrontendClassInputDefault>* ClassDefaults = Builder.FindNodeClassInputDefaults(InputVertexHandle.NodeID, InputVertex->Name))
+				if (!ensure(!ClassDefaults->IsEmpty()))
+				{
+					return nullptr;
+				}
+
+				DefaultLiterals = *ClassDefaults;
+				const FMetasoundFrontendClassInputDefault* DefaultPageValueLiteral = DefaultLiterals.FindByPredicate([](const FMetasoundFrontendClassInputDefault& InputDefault)
 				{
-					DefaultLiterals = *ClassDefaults;
+					return InputDefault.PageID == Metasound::Frontend::DefaultPageID;
+				});
+
+				// Code OR asset-defined classes should ALWAYS include input default value associated with default page ID by this point
+				if (!ensure(DefaultPageValueLiteral))
+				{
+					DefaultLiterals.Add_GetRef(Metasound::Frontend::DefaultPageID).Literal = ClassDefaults->Last().Literal;
 				}
 			}
+			else
+			{
+				FMetasoundFrontendLiteral DefaultValue;
+				DefaultValue.SetFromLiteral(IDataTypeRegistry::Get().CreateDefaultLiteral(InputVertex->TypeName));
+				DefaultLiterals = InitDefaultLiterals(MoveTemp(DefaultValue));
+			}
 
 			const FCreateNodeVertexParams VertexParams =
 			{
