Fix MetaSound Editor performance regression

#rnx

diff --git a/Engine/Plugins/Runtime/Metasound/Source/MetasoundEditor/Private/MetasoundEditor.cpp b/Engine/Plugins/Runtime/Metasound/Source/MetasoundEditor/Private/MetasoundEditor.cpp
--- a/Engine/Plugins/Runtime/Metasound/Source/MetasoundEditor/Private/MetasoundEditor.cpp
+++ b/Engine/Plugins/Runtime/Metasound/Source/MetasoundEditor/Private/MetasoundEditor.cpp
@@ -1690,11 +1690,11 @@ namespace Metasound
 				PagesDetails->SetObjects(PagesViewObj);
 				PagesDetails->HideFilterArea(true);
 
-				TAttribute<bool> EnabledAttr = TAttribute<bool>::Create([this]()
+				TAttribute<bool> EnabledAttr = TAttribute<bool>::CreateSPLambda(AsShared(), [this]()
 				{
 					return ShowPageGraphDetails();
 				});
-				TAttribute<EVisibility> VisibilityAttr = TAttribute<EVisibility>::Create([this]()
+				TAttribute<EVisibility> VisibilityAttr = TAttribute<EVisibility>::CreateSPLambda(AsShared(), [this]()
 				{
 					return ShowPageGraphDetails() ? EVisibility::Visible : EVisibility::Hidden;
 				});
@@ -3620,6 +3620,16 @@ namespace Metasound
 			return nullptr;
 		}
 
+		FMetaSoundFrontendDocumentBuilder* FEditor::GetFrontendBuilder() const
+		{
+			if (Builder.IsValid())
+			{
+				return &Builder->GetBuilder();
+			}
+
+			return nullptr;
+		}
+
 		void FEditor::RenameSelectedNode()
 		{
 			const FGraphPanelSelectionSet& SelectedNodes = MetasoundGraphEditor->GetSelectedNodes();
diff --git a/Engine/Plugins/Runtime/Metasound/Source/MetasoundEditor/Private/MetasoundEditor.h b/Engine/Plugins/Runtime/Metasound/Source/MetasoundEditor/Private/MetasoundEditor.h
--- a/Engine/Plugins/Runtime/Metasound/Source/MetasoundEditor/Private/MetasoundEditor.h
+++ b/Engine/Plugins/Runtime/Metasound/Source/MetasoundEditor/Private/MetasoundEditor.h
@@ -142,6 +142,8 @@ namespace Metasound::Editor
 
 		UAudioComponent* GetAudioComponent() const;
 
+		FMetaSoundFrontendDocumentBuilder* GetFrontendBuilder() const;
+
 		/** IMetasoundEditor interface */
 		virtual UObject* GetMetasoundObject() const override;
 		virtual void SetSelection(const TArray<UObject*>& SelectedObjects) override;
diff --git a/Engine/Plugins/Runtime/Metasound/Source/MetasoundEditor/Private/MetasoundEditorGraphBuilder.cpp b/Engine/Plugins/Runtime/Metasound/Source/MetasoundEditor/Private/MetasoundEditorGraphBuilder.cpp
--- a/Engine/Plugins/Runtime/Metasound/Source/MetasoundEditor/Private/MetasoundEditorGraphBuilder.cpp
+++ b/Engine/Plugins/Runtime/Metasound/Source/MetasoundEditor/Private/MetasoundEditorGraphBuilder.cpp
@@ -934,10 +934,15 @@ namespace Metasound
 			return VertexHandle;
 		}
 
-		const FMetasoundFrontendVertex* FGraphBuilder::GetPinVertex(const FMetaSoundFrontendDocumentBuilder& InBuilder, const UEdGraphPin* InPin)
+		const FMetasoundFrontendVertex* FGraphBuilder::GetPinVertex(const FMetaSoundFrontendDocumentBuilder& InBuilder, const UEdGraphPin* InPin, const FMetasoundFrontendNode** Node)
 		{
 			using namespace VariableNames;
 
+			if (Node)
+			{
+				*Node = nullptr;
+			}
+
 			if (!InPin)
 			{
 				return nullptr;
@@ -946,13 +951,18 @@ namespace Metasound
 			const UMetasoundEditorGraphNode* OwningNode = CastChecked<UMetasoundEditorGraphNode>(InPin->GetOwningNode());
 
 			const FGuid NodeID = OwningNode->GetNodeID();
-			const FMetasoundFrontendNode* Node = InBuilder.FindNode(NodeID);
-			if (!Node)
+			const FMetasoundFrontendNode* FoundNode = InBuilder.FindNode(NodeID);
+			if (!FoundNode)
 			{
 				return nullptr;
 			}
 
-			const FMetasoundFrontendClass* Class = InBuilder.FindDependency(Node->ClassID);
+			if (Node)
+			{
+				*Node = FoundNode;
+			}
+
+			const FMetasoundFrontendClass* Class = InBuilder.FindDependency(FoundNode->ClassID);
 			if (!Class)
 			{
 				return nullptr;
@@ -975,11 +985,11 @@ namespace Metasound
 				}
 				case EMetasoundFrontendClassType::Input:
 				{
-					return &Node->Interface.Inputs.Last();
+					return &FoundNode->Interface.Inputs.Last();
 				}
 				case EMetasoundFrontendClassType::Output:
 				{
-					return &Node->Interface.Outputs.Last();
+					return &FoundNode->Interface.Outputs.Last();
 				}
 
 				default:
@@ -1015,33 +1025,6 @@ namespace Metasound
 			return IOutputController::GetInvalidHandle();
 		}
 
-		const FMetasoundFrontendEdgeStyle* FGraphBuilder::GetOutputEdgeStyle(Frontend::FConstOutputHandle InOutputHandle)
-		{
-			using namespace Frontend;
-
-			if (InOutputHandle->IsValid())
-			{
-				FConstNodeHandle NodeHandle = InOutputHandle->GetOwningNode();
-				FConstGraphHandle GraphHandle = NodeHandle->GetOwningGraph();
-
-				const TArray<FMetasoundFrontendEdgeStyle>& EdgeStyles = GraphHandle->GetGraphStyle().EdgeStyles;
-				return EdgeStyles.FindByPredicate([&InOutputHandle](const FMetasoundFrontendEdgeStyle& InCandidate)
-				{
-					return InCandidate.NodeID == InOutputHandle->GetOwningNodeID() && InCandidate.OutputName == InOutputHandle->GetName();
-				});
-			}
-
-			return nullptr;
-		}
-
-		const FMetasoundFrontendEdgeStyle* FGraphBuilder::GetOutputEdgeStyle(const UEdGraphPin* InGraphPin)
-		{
-			using namespace Frontend;
-
-			FConstOutputHandle OutputHandle = FindReroutedConstOutputHandleFromPin(InGraphPin);
-			return GetOutputEdgeStyle(OutputHandle);
-		}
-
 		Frontend::FConstOutputHandle FGraphBuilder::GetConstOutputHandleFromPin(const UEdGraphPin* InPin)
 		{
 			using namespace Frontend;
diff --git a/Engine/Plugins/Runtime/Metasound/Source/MetasoundEditor/Private/MetasoundEditorGraphBuilder.h b/Engine/Plugins/Runtime/Metasound/Source/MetasoundEditor/Private/MetasoundEditorGraphBuilder.h
--- a/Engine/Plugins/Runtime/Metasound/Source/MetasoundEditor/Private/MetasoundEditorGraphBuilder.h
+++ b/Engine/Plugins/Runtime/Metasound/Source/MetasoundEditor/Private/MetasoundEditorGraphBuilder.h
@@ -205,10 +205,7 @@ namespace Metasound
 
 			static FName GetPinDataType(const UEdGraphPin* InPin);
 			static FMetasoundFrontendVertexHandle GetPinVertexHandle(const FMetaSoundFrontendDocumentBuilder& InBuilder, const UEdGraphPin* InPin);
-			static const FMetasoundFrontendVertex* GetPinVertex(const FMetaSoundFrontendDocumentBuilder& InBuilder, const UEdGraphPin* InPin);
-
-			static const FMetasoundFrontendEdgeStyle* GetOutputEdgeStyle(Frontend::FConstOutputHandle InOutputHandle);
-			static const FMetasoundFrontendEdgeStyle* GetOutputEdgeStyle(const UEdGraphPin* InPin);
+			static const FMetasoundFrontendVertex* GetPinVertex(const FMetaSoundFrontendDocumentBuilder& InBuilder, const UEdGraphPin* InPin, const FMetasoundFrontendNode** Node = nullptr);
 
 			// Get the output handle from an output pin.  Ensures pin is an output pin.
 			// TODO: use IDs to connect rather than names. Likely need an UMetasoundEditorGraphPin
diff --git a/Engine/Plugins/Runtime/Metasound/Source/MetasoundEditor/Private/MetasoundEditorGraphConnectionDrawingPolicy.cpp b/Engine/Plugins/Runtime/Metasound/Source/MetasoundEditor/Private/MetasoundEditorGraphConnectionDrawingPolicy.cpp
--- a/Engine/Plugins/Runtime/Metasound/Source/MetasoundEditor/Private/MetasoundEditorGraphConnectionDrawingPolicy.cpp
+++ b/Engine/Plugins/Runtime/Metasound/Source/MetasoundEditor/Private/MetasoundEditorGraphConnectionDrawingPolicy.cpp
@@ -79,68 +79,36 @@ namespace Metasound
 #endif // METASOUND_EDITOR_DEBUG_CONNECTIONS
 
 			template<typename TAnalyzerType>
-			bool DetermineWiringStyle_Envelope(TSharedPtr<FEditor> MetasoundEditor, UEdGraphPin* OutputPin, UEdGraphPin* InputPin, FConnectionParams& OutParams)
+			void DetermineWiringStyle_Envelope(bool bIsPlaying, FGraphConnectionManager& ConnectionManager, const FGuid& NodeID, FName OutputName, UEdGraphPin* InputPin, FConnectionParams& OutParams)
 			{
 				using namespace Frontend;
 
 				OutParams.bDrawBubbles = false;
 
-				if (MetasoundEditor.IsValid())
+				if (bIsPlaying)
 				{
-					FConstOutputHandle OutputHandle = FGraphBuilder::FindReroutedConstOutputHandleFromPin(OutputPin);
-					const FGuid NodeId = OutputHandle->GetOwningNodeID();
-					FName OutputName = OutputHandle->GetName();
-
-					OutParams.bDrawBubbles = MetasoundEditor->GetConnectionManager().IsTracked(NodeId, OutputName, TAnalyzerType::GetAnalyzerName());
+					OutParams.bDrawBubbles = ConnectionManager.IsTracked(NodeID, OutputName, TAnalyzerType::GetAnalyzerName());
 					if (OutParams.bDrawBubbles)
 					{
-						float WindowValue = MetasoundEditor->GetConnectionManager().UpdateValueWindow<TAnalyzerType>(NodeId, OutputName);
+						const float WindowValue = ConnectionManager.UpdateValueWindow<TAnalyzerType>(NodeID, OutputName);
 						OutParams.WireThickness = FMath::Clamp(WindowValue, 0.0f, 1.0f);
 						FLinearColor HSV = OutParams.WireColor.LinearRGBToHSV();
 						HSV.G = FMath::Lerp(HSV.G, 1.0f, WindowValue);
 						OutParams.WireColor = HSV.HSVToLinearRGB();
 					}
-
-					return true;
 				}
-
-				return false;
 			}
 
 			template<typename TNumericType>
-			void DrawConnectionSpline_Numeric(FGraphConnectionDrawingPolicy& InDrawingPolicy, int32 InLayerId, const FDrawConnectionData& InData, TNumericType InDefaultValue)
+			void DrawConnectionSpline_Numeric(FGraphConnectionDrawingPolicy& InDrawingPolicy, const FEditor& InEditor, int32 InLayerId, const FDrawConnectionData& InData, TNumericType InDefaultValue)
 			{
-				if (!InData.OutputPin)
-				{
-					return;
-				}
-
-				const UMetasoundEditorGraphNode* Node = Cast<UMetasoundEditorGraphNode>(InData.OutputPin->GetOwningNode());
-				if (!ensureMsgf(Node, TEXT("Expected MetaSound pin to be member of MetaSound node")))
-				{
-					return;
-				}
-
-				const UMetasoundEditorGraph* Graph = Cast<UMetasoundEditorGraph>(Node->GetGraph());
-				if (!ensureMsgf(Graph, TEXT("Expected MetaSound node to be member of MetaSound graph")))
+				if (!InData.OutputVertex || !InData.OutputNode)
 				{
 					return;
 				}
 
-				const TSharedPtr<const FEditor> Editor = static_cast<const FGraphConnectionDrawingPolicy&>(InDrawingPolicy).GetEditor();
-				if (!Editor)
-				{
-					return;
-				}
-
-				PRAGMA_DISABLE_DEPRECATION_WARNINGS
-				Frontend::FConstOutputHandle OutputHandle = Frontend::FindReroutedOutput(InData.OutputHandle);
-				PRAGMA_ENABLE_DEPRECATION_WARNINGS
-
-				const FGuid NodeID = OutputHandle->GetOwningNodeID();
 				TNumericType Value = InDefaultValue;
-				FName OutputName = OutputHandle->GetName();
-				Editor->GetConnectionManager().GetValue(NodeID, OutputName, Value);
+				InEditor.GetConnectionManager().GetValue(InData.OutputNode->GetID(), InData.OutputVertex->Name, Value);
 
 				FLinearColor InnerColor = Frontend::DisplayStyle::EdgeAnimation::DefaultColor;
 
@@ -191,18 +159,16 @@ namespace Metasound
 			}
 		} // namespace DrawingPolicyPrivate
 
-		FDrawConnectionData::FDrawConnectionData(const FVector2D& InStart, const FVector2D& InEnd, const FVector2D& InSplineTangent, const FConnectionParams& InParams, const UGraphEditorSettings& InSettings, const FVector2D& InMousePosition)
+		FDrawConnectionData::FDrawConnectionData(const FVector2D& InStart, const FVector2D& InEnd, const FEditor& Editor, const FVector2D& InSplineTangent, const FConnectionParams& InParams, const UGraphEditorSettings& InSettings, const FVector2D& InMousePosition)
 			: P0(InStart)
 			, P0Tangent(InParams.StartDirection == EGPD_Output ? InSplineTangent : -InSplineTangent)
 			, P1(InEnd)
 			, P1Tangent(InParams.EndDirection == EGPD_Input ? InSplineTangent : -InSplineTangent)
 			, Params(InParams)
-			, OutputHandle(Frontend::IOutputController::GetInvalidHandle())
 		{
 			if (Params.AssociatedPin1)
 			{
 				OutputPin = Params.AssociatedPin1->Direction == EGPD_Output ? Params.AssociatedPin1 : Params.AssociatedPin2;
-				OutputHandle = FGraphBuilder::FindReroutedConstOutputHandleFromPin(OutputPin);
 			}
 
 			// The curve will include the endpoints but can extend out of a tight bounds because of the tangents
@@ -248,7 +214,14 @@ namespace Metasound
 				}
 			}
 
-			EdgeStyle = FGraphBuilder::GetOutputEdgeStyle(OutputHandle);
+			if (const FMetaSoundFrontendDocumentBuilder* Builder = Editor.GetFrontendBuilder())
+			{
+				if (const UEdGraphPin* ReroutePin = FGraphBuilder::FindReroutedOutputPin(OutputPin))
+				{
+					OutputVertex = FGraphBuilder::GetPinVertex(*Builder, ReroutePin, &OutputNode);
+				}
+				EdgeStyle = Builder->FindConstEdgeStyle(OutputNode->GetID(), OutputVertex->Name);
+			}
 		}
 
 		void FDrawConnectionData::UpdateSplineOverlap(FGraphSplineOverlapResult& OutResult) const
@@ -339,6 +312,8 @@ namespace Metasound
 		{
 			using namespace Frontend;
 
+			METASOUND_TRACE_CPUPROFILER_EVENT_SCOPE(Metasound::Editor::FGraphConnectionDrawingPolicy::DetermineWiringStyle);
+
 			TSharedPtr<FEditor> EditorPtr = MetasoundEditor.Pin();
 			if (!EditorPtr)
 			{
@@ -382,37 +357,49 @@ namespace Metasound
 				OutParams.WireColor = FGraphBuilder::GetPinCategoryColor(AnyPin->PinType);
 			}
 
+			bool bThicknessSet = false;
 			if (AnyPin && AnalyzerSettings.bAnimateConnections)
 			{
-				if (AnyPin->PinType.PinCategory == FGraphBuilder::PinCategoryTrigger)
-				{
-					DrawingPolicyPrivate::DetermineWiringStyle_Envelope<FVertexAnalyzerTriggerDensity>(EditorPtr, OutputPin, InputPin, OutParams);
-					ScaleEnvelopeWireThickness(OutParams);
-				}
-				else if (AnyPin->PinType.PinCategory == FGraphBuilder::PinCategoryAudio)
+				if (const UEdGraphPin* ReroutePin = FGraphBuilder::FindReroutedOutputPin(OutputPin))
 				{
-					DrawingPolicyPrivate::DetermineWiringStyle_Envelope<FVertexAnalyzerEnvelopeFollower>(EditorPtr, OutputPin, InputPin, OutParams);
-					ScaleEnvelopeWireThickness(OutParams);
-				}
-				else if (FGraphBuilder::CanInspectPin(AnyPin))
-				{
-					bool bEdgeStyleValid = false;
-					if (const FMetasoundFrontendEdgeStyle* Style = FGraphBuilder::GetOutputEdgeStyle(OutputPin))
-					{
-						bEdgeStyleValid = !Style->LiteralColorPairs.IsEmpty();
-					}
-
-					if (bEdgeStyleValid && EditorPtr->IsPlaying())
-					{
-						OutParams.WireThickness = AnalyzerSettings.NumericWireThickness;
-					}
-					else
+					if (const FMetaSoundFrontendDocumentBuilder* Builder = EditorPtr->GetFrontendBuilder())
 					{
-						OutParams.WireThickness = Settings->TraceReleaseWireThickness;
+						const FMetasoundFrontendNode* OutputNode = nullptr;
+						if (const FMetasoundFrontendVertex* OutputVertex = FGraphBuilder::GetPinVertex(*Builder, ReroutePin, &OutputNode))
+						{
+							const bool bIsPlaying = EditorPtr->IsPlaying();
+							const FGuid& NodeID = OutputNode->GetID();
+							if (AnyPin->PinType.PinCategory == FGraphBuilder::PinCategoryTrigger)
+							{
+								FGraphConnectionManager& ConnectionManager = EditorPtr->GetConnectionManager();
+								DrawingPolicyPrivate::DetermineWiringStyle_Envelope<FVertexAnalyzerTriggerDensity>(bIsPlaying, ConnectionManager, NodeID, OutputVertex->Name, InputPin, OutParams);
+								ScaleEnvelopeWireThickness(OutParams);
+								bThicknessSet = true;
+							}
+							else if (AnyPin->PinType.PinCategory == FGraphBuilder::PinCategoryAudio)
+							{
+								FGraphConnectionManager& ConnectionManager = EditorPtr->GetConnectionManager();
+								DrawingPolicyPrivate::DetermineWiringStyle_Envelope<FVertexAnalyzerEnvelopeFollower>(bIsPlaying, ConnectionManager, NodeID, OutputVertex->Name, InputPin, OutParams);
+								ScaleEnvelopeWireThickness(OutParams);
+								bThicknessSet = true;
+							}
+							else if (FGraphBuilder::CanInspectPin(AnyPin))
+							{
+								if (const FMetasoundFrontendEdgeStyle* EdgeStyle = Builder->FindConstEdgeStyle(NodeID, OutputVertex->Name))
+								{
+									if (!EdgeStyle->LiteralColorPairs.IsEmpty())
+									{
+										OutParams.WireThickness = AnalyzerSettings.NumericWireThickness;
+										bThicknessSet = true;
+									}
+								}
+							}
+						}
 					}
 				}
 			}
-			else
+
+			if (!bThicknessSet)
 			{
 				OutParams.WireThickness = Settings->TraceReleaseWireThickness;
 			}
@@ -420,7 +407,7 @@ namespace Metasound
 			const bool bDeemphasizeUnhoveredPins = HoveredPins.Num() > 0;
 			if (bDeemphasizeUnhoveredPins)
 			{
-				const bool bIsPlaying = GetEditor()->IsPlaying();
+				const bool bIsPlaying = EditorPtr->IsPlaying();
 				ApplyHoverDeemphasisMetaSound(OutputPin, InputPin, bIsPlaying, /*inout*/ OutParams.WireThickness, /*inout*/ OutParams.WireColor);
 			}
 		}
@@ -455,7 +442,7 @@ namespace Metasound
 
 		void FGraphConnectionDrawingPolicy::DrawConnection(int32 LayerId, const FVector2D& Start, const FVector2D& End, const FConnectionParams& Params)
 		{
-			METASOUND_TRACE_CPUPROFILER_EVENT_SCOPE(FGraphConnectionDrawingPolicy::DrawConnection);
+			METASOUND_TRACE_CPUPROFILER_EVENT_SCOPE(Metasound::Editor::FGraphConnectionDrawingPolicy::DrawConnection);
 
 			using namespace Frontend;
 
@@ -472,7 +459,7 @@ namespace Metasound
 				return;
 			}
 
-			const FDrawConnectionData ConnectionData(Start, End, FGraphConnectionDrawingPolicy::ComputeSplineTangent(Start, End), Params, *Settings, LocalMousePosition);
+			const FDrawConnectionData ConnectionData(Start, End, *EditorPtr, FGraphConnectionDrawingPolicy::ComputeSplineTangent(Start, End), Params, *Settings, LocalMousePosition);
 			if (Settings->bTreatSplinesLikePins)
 			{
 #if METASOUND_EDITOR_DEBUG_CONNECTIONS
@@ -494,7 +481,8 @@ namespace Metasound
 				{
 					const float AppTime = FPlatformTime::Seconds() - GStartTime;
 					FDrawConnectionSignalData SignalParams = { AppTime, LayerId, SplineReparamTable, SplineLength, ConnectionData };
-					const FName DataType = SignalParams.ConnectionData.OutputHandle->GetDataType();
+
+					const FName DataType = SignalParams.ConnectionData.OutputVertex->TypeName;
 					if (DataType == GetMetasoundDataTypeName<FAudioBuffer>())
 					{
 						SignalParams.SpacingFactor = DrawingPolicyPrivate::EnvelopeConnectionSpacingCVar;
@@ -523,20 +511,20 @@ namespace Metasound
 			}
 
 			// Draw other types
-			if (ConnectionData.OutputHandle->IsValid() && MetasoundEditor.IsValid())
+			if (ConnectionData.OutputVertex && MetasoundEditor.IsValid())
 			{
-				const FName DataType = ConnectionData.OutputHandle->GetDataType();
+				const FName DataType = ConnectionData.OutputVertex->TypeName;
 				if (DataType == GetMetasoundDataTypeName<float>())
 				{
-					DrawingPolicyPrivate::DrawConnectionSpline_Numeric<float>(*this, LayerId, ConnectionData, 0.0f);
+					DrawingPolicyPrivate::DrawConnectionSpline_Numeric<float>(*this, *EditorPtr, LayerId, ConnectionData, 0.0f);
 				}
 				else if (DataType == GetMetasoundDataTypeName<int32>())
 				{
-					DrawingPolicyPrivate::DrawConnectionSpline_Numeric<int32>(*this, LayerId, ConnectionData, 0);
+					DrawingPolicyPrivate::DrawConnectionSpline_Numeric<int32>(*this, *EditorPtr, LayerId, ConnectionData, 0);
 				}
 				else if (DataType == GetMetasoundDataTypeName<bool>())
 				{
-					DrawingPolicyPrivate::DrawConnectionSpline_Numeric<bool>(*this, LayerId, ConnectionData, 0);
+					DrawingPolicyPrivate::DrawConnectionSpline_Numeric<bool>(*this, *EditorPtr, LayerId, ConnectionData, 0);
 				}
 				else 
 				{
diff --git a/Engine/Plugins/Runtime/Metasound/Source/MetasoundEditor/Private/MetasoundEditorGraphConnectionDrawingPolicy.h b/Engine/Plugins/Runtime/Metasound/Source/MetasoundEditor/Private/MetasoundEditorGraphConnectionDrawingPolicy.h
--- a/Engine/Plugins/Runtime/Metasound/Source/MetasoundEditor/Private/MetasoundEditorGraphConnectionDrawingPolicy.h
+++ b/Engine/Plugins/Runtime/Metasound/Source/MetasoundEditor/Private/MetasoundEditorGraphConnectionDrawingPolicy.h
@@ -41,7 +41,8 @@ namespace Metasound
 			FBox2D Bounds;
 
 			const UEdGraphPin* OutputPin = nullptr;
-			Frontend::FConstOutputHandle OutputHandle;
+			const FMetasoundFrontendVertex* OutputVertex = nullptr;
+			const FMetasoundFrontendNode* OutputNode = nullptr;
 
 			bool bCloseToSpline = false;
 			float ClosestDistanceSquared = TNumericLimits<float>::Max();
@@ -57,7 +58,7 @@ namespace Metasound
 
 			const FMetasoundFrontendEdgeStyle* EdgeStyle = nullptr;
 
-			FDrawConnectionData(const FVector2D& InStart, const FVector2D& InEnd, const FVector2D& InSplineTangent, const FConnectionParams& InParams, const UGraphEditorSettings& InSettings, const FVector2D& InMousePosition);
+			FDrawConnectionData(const FVector2D& InStart, const FVector2D& InEnd, const FEditor& Editor, const FVector2D& InSplineTangent, const FConnectionParams& InParams, const UGraphEditorSettings& InSettings, const FVector2D& InMousePosition);
 
 			void UpdateSplineOverlap(FGraphSplineOverlapResult& OutResult) const;
 		};
diff --git a/Engine/Plugins/Runtime/Metasound/Source/MetasoundEditor/Private/SMetasoundGraphPin.h b/Engine/Plugins/Runtime/Metasound/Source/MetasoundEditor/Private/SMetasoundGraphPin.h
--- a/Engine/Plugins/Runtime/Metasound/Source/MetasoundEditor/Private/SMetasoundGraphPin.h
+++ b/Engine/Plugins/Runtime/Metasound/Source/MetasoundEditor/Private/SMetasoundGraphPin.h
@@ -42,6 +42,9 @@ namespace Metasound
 
 			TSharedPtr<FMetasoundPinAudioInspector> PinAudioInspector;
 
+			// Cached builder pointer for fast access
+			mutable TWeakObjectPtr<UMetaSoundBuilderBase> BuilderPtr;
+
 		protected:
 			virtual bool CanInspectPin(const UEdGraphPin* InPin)
 			{
@@ -195,15 +198,17 @@ namespace Metasound
 						}
 						else if (const UMetasoundEditorGraphExternalNode* ExternalNode = Cast<UMetasoundEditorGraphExternalNode>(Node))
 						{
-							if (Pin->Direction == EGPD_Input)
-							{
-								Frontend::FConstInputHandle InputHandle = FGraphBuilder::GetConstInputHandleFromPin(Pin);
-								AccessType = InputHandle->GetVertexAccessType();
-							}
-							else if (Pin->Direction == EGPD_Output)
+							if (UMetaSoundBuilderBase* Builder = GetBuilder())
 							{
-								Frontend::FConstOutputHandle OutputHandle = FGraphBuilder::GetConstOutputHandleFromPin(Pin);
-								AccessType = OutputHandle->GetVertexAccessType();
+								const FMetasoundFrontendVertexHandle VertexHandle = FGraphBuilder::GetPinVertexHandle(Builder->GetConstBuilder(), Pin);
+								if (Pin->Direction == EGPD_Input)
+								{
+									AccessType = Builder->GetConstBuilder().GetNodeInputAccessType(VertexHandle.NodeID, VertexHandle.VertexID);
+								}
+								else if (Pin->Direction == EGPD_Output)
+								{
+									AccessType = Builder->GetConstBuilder().GetNodeOutputAccessType(VertexHandle.NodeID, VertexHandle.VertexID);
+								}
 							}
 						}
 					}
@@ -279,14 +284,27 @@ namespace Metasound
 			}
 
 			UMetaSoundBuilderBase& GetBuilderChecked() const
+			{
+				UMetaSoundBuilderBase* Builder = GetBuilder();
+				check(Builder);
+				return *Builder;
+			}
+
+			UMetaSoundBuilderBase* GetBuilder() const
 			{
 				using namespace Metasound::Engine;
 
+				if (UMetaSoundBuilderBase* Builder = BuilderPtr.Get())
+				{
+					return Builder;
+				}
+
 				const UMetasoundEditorGraphNode* Node = GetOwningMetaSoundNode();
 				check(Node);
 				UObject* Outermost = Node->GetOutermostObject();
 				check(Outermost);
-				return FDocumentBuilderRegistry::GetChecked().FindOrBeginBuilding(*Outermost);
+				BuilderPtr = &FDocumentBuilderRegistry::GetChecked().FindOrBeginBuilding(*Outermost);
+				return BuilderPtr.Get();
 			}
 
 			const FMetasoundFrontendNode* GetFrontendNode() const
@@ -297,9 +315,26 @@ namespace Metasound
 				{
 					if (UObject* MetaSound = Node->GetMetasound())
 					{
-						const FGuid NodeID = Node->GetNodeID();
-						const UMetaSoundBuilderBase& Builder = FDocumentBuilderRegistry::GetChecked().FindOrBeginBuilding(*MetaSound);
-						return Builder.GetConstBuilder().FindNode(NodeID);
+						if (UMetaSoundBuilderBase* Builder = GetBuilder())
+						{
+							const FGuid NodeID = Node->GetNodeID();
+							return Builder.GetConstBuilder().FindNode(NodeID);
+						}
+					}
+				}
+
+				return nullptr;
+			}
+
+			const FMetasoundFrontendVertex* GetFrontendVertex() const
+			{
+				using namespace Metasound::Engine;
+
+				if (UEdGraphPin* Pin = ParentPinType::GetPinObj())
+				{
+					if (UMetaSoundBuilderBase* Builder = GetBuilder())
+					{
+						return FGraphBuilder::GetPinVertex(*Builder, *Pin);
 					}
 				}
 
@@ -489,32 +524,34 @@ namespace Metasound
 				if (UEdGraphPin* GraphPin = ParentPinType::GetPinObj())
 				{
 					const bool bIsHoveringPin = ParentPinType::IsHovered();
-					UpdatePinInspector(*GraphPin, bIsHoveringPin, PinInspector, ParentPinType::ValueInspectorTooltip,
+
+					// General Value Inspector update
+					if (bIsHoveringPin || PinInspector.IsValid())
+					{
+						UpdatePinInspector(*GraphPin, bIsHoveringPin, PinInspector, ParentPinType::ValueInspectorTooltip,
 						[this](FVector2D& OutTooltipLocation)
 						{
 							ParentPinType::GetInteractiveTooltipLocation(OutTooltipLocation);
 						});
-
-					// Audio Pin Inspector update
-					FName DataType;
-					if (GraphPin->Direction == EGPD_Input)
-					{
-						const Metasound::Frontend::FConstInputHandle InputHandle = FGraphBuilder::GetConstInputHandleFromPin(GraphPin);
-						DataType = InputHandle->GetDataType();
-					}
-					else
-					{
-						const Metasound::Frontend::FConstOutputHandle OutputHandle = FGraphBuilder::GetConstOutputHandleFromPin(GraphPin);
-						DataType = OutputHandle->GetDataType();
 					}
 
-					if (DataType == GetMetasoundDataTypeName<FAudioBuffer>())
+					// Audio Pin Inspector update
+					if (bIsHoveringPin || PinAudioInspector.IsValid())
 					{
-						UpdatePinAudioInspector(*GraphPin, bIsHoveringPin, PinAudioInspector, ParentPinType::ValueInspectorTooltip,
-							[this](FVector2D& OutTooltipLocation)
+						if (UMetaSoundBuilderBase* Builder = GetBuilder())
+						{
+							if (const FMetasoundFrontendVertex* Vertex = FGraphBuilder::GetPinVertex(Builder->GetConstBuilder(), GraphPin))
 							{
-								ParentPinType::GetInteractiveTooltipLocation(OutTooltipLocation);
-							});
+								if (Vertex->TypeName == GetMetasoundDataTypeName<FAudioBuffer>())
+								{
+									UpdatePinAudioInspector(*GraphPin, bIsHoveringPin, PinAudioInspector, ParentPinType::ValueInspectorTooltip,
+									[this](FVector2D& OutTooltipLocation)
+									{
+										ParentPinType::GetInteractiveTooltipLocation(OutTooltipLocation);
+									});
+								}
+							}
+						}
 					}
 				}
 			}
