Fix build due to unused function with bad syntax

#rnx

diff --git a/Engine/Plugins/Runtime/Metasound/Source/MetasoundEditor/Private/MetasoundEditorGraphConnectionDrawingPolicy.cpp b/Engine/Plugins/Runtime/Metasound/Source/MetasoundEditor/Private/MetasoundEditorGraphConnectionDrawingPolicy.cpp
--- a/Engine/Plugins/Runtime/Metasound/Source/MetasoundEditor/Private/MetasoundEditorGraphConnectionDrawingPolicy.cpp
+++ b/Engine/Plugins/Runtime/Metasound/Source/MetasoundEditor/Private/MetasoundEditorGraphConnectionDrawingPolicy.cpp
@@ -219,8 +219,9 @@ namespace Metasound
 				if (const UEdGraphPin* ReroutePin = FGraphBuilder::FindReroutedOutputPin(OutputPin))
 				{
 					OutputVertex = FGraphBuilder::GetPinVertex(*Builder, ReroutePin, &OutputNode);
+					check(OutputNode);
+					EdgeStyle = Builder->FindConstEdgeStyle(OutputNode->GetID(), OutputVertex->Name);
 				}
-				EdgeStyle = Builder->FindConstEdgeStyle(OutputNode->GetID(), OutputVertex->Name);
 			}
 		}
 
@@ -367,6 +368,7 @@ namespace Metasound
 						const FMetasoundFrontendNode* OutputNode = nullptr;
 						if (const FMetasoundFrontendVertex* OutputVertex = FGraphBuilder::GetPinVertex(*Builder, ReroutePin, &OutputNode))
 						{
+							check(OutputNode);
 							const bool bIsPlaying = EditorPtr->IsPlaying();
 							const FGuid& NodeID = OutputNode->GetID();
 							if (AnyPin->PinType.PinCategory == FGraphBuilder::PinCategoryTrigger)
diff --git a/Engine/Plugins/Runtime/Metasound/Source/MetasoundEditor/Private/SMetasoundGraphPin.h b/Engine/Plugins/Runtime/Metasound/Source/MetasoundEditor/Private/SMetasoundGraphPin.h
--- a/Engine/Plugins/Runtime/Metasound/Source/MetasoundEditor/Private/SMetasoundGraphPin.h
+++ b/Engine/Plugins/Runtime/Metasound/Source/MetasoundEditor/Private/SMetasoundGraphPin.h
@@ -315,10 +315,10 @@ namespace Metasound
 				{
 					if (UObject* MetaSound = Node->GetMetasound())
 					{
-						if (UMetaSoundBuilderBase* Builder = GetBuilder())
+						if (const UMetaSoundBuilderBase* Builder = GetBuilder())
 						{
 							const FGuid NodeID = Node->GetNodeID();
-							return Builder.GetConstBuilder().FindNode(NodeID);
+							return Builder->GetConstBuilder().FindNode(NodeID);
 						}
 					}
 				}
