Another static analysis build fix
#rnx

diff --git a/Engine/Plugins/Runtime/Metasound/Source/MetasoundEditor/Private/SMetasoundGraphPin.h b/Engine/Plugins/Runtime/Metasound/Source/MetasoundEditor/Private/SMetasoundGraphPin.h
--- a/Engine/Plugins/Runtime/Metasound/Source/MetasoundEditor/Private/SMetasoundGraphPin.h
+++ b/Engine/Plugins/Runtime/Metasound/Source/MetasoundEditor/Private/SMetasoundGraphPin.h
@@ -313,13 +313,10 @@ namespace Metasound
 
 				if (const UMetasoundEditorGraphNode* Node = GetOwningMetaSoundNode())
 				{
-					if (UObject* MetaSound = Node->GetMetasound())
+					if (const UMetaSoundBuilderBase* Builder = GetBuilder())
 					{
-						if (const UMetaSoundBuilderBase* Builder = GetBuilder())
-						{
-							const FGuid NodeID = Node->GetNodeID();
-							return Builder->GetConstBuilder().FindNode(NodeID);
-						}
+						const FGuid NodeID = Node->GetNodeID();
+						return Builder->GetConstBuilder().FindNode(NodeID);
 					}
 				}
 
@@ -334,7 +331,7 @@ namespace Metasound
 				{
 					if (UMetaSoundBuilderBase* Builder = GetBuilder())
 					{
-						return FGraphBuilder::GetPinVertex(*Builder, *Pin);
+						return FGraphBuilder::GetPinVertex(Builder->GetConstBuilder(), Pin);
 					}
 				}
 
