Fix ensure on copy/pasting required MetaSound interface outputs

#rnx

diff --git a/Engine/Plugins/Runtime/Metasound/Source/MetasoundEditor/Private/MetasoundEditorDocumentClipboardUtils.cpp b/Engine/Plugins/Runtime/Metasound/Source/MetasoundEditor/Private/MetasoundEditorDocumentClipboardUtils.cpp
--- a/Engine/Plugins/Runtime/Metasound/Source/MetasoundEditor/Private/MetasoundEditorDocumentClipboardUtils.cpp
+++ b/Engine/Plugins/Runtime/Metasound/Source/MetasoundEditor/Private/MetasoundEditorDocumentClipboardUtils.cpp
@@ -195,6 +195,16 @@ namespace Metasound::Editor
 									MetaSoundEditorSubsystem.BindMemberMetadata(Builder, *Output, LiteralClass, DefaultLiteral);
 								}
 							}
+
+							// Remove default node location from output node. 
+							// Correct node location from the ed graph node will be set subsequently in ProcessPastedNodePositions
+							TArray<FGuid> NodeLocationGuids;
+							NewNode->Style.Display.Locations.GetKeys(NodeLocationGuids);
+							if (!NodeLocationGuids.IsEmpty())
+							{
+								Builder.RemoveNodeLocation(NewNode->GetID());
+							}
+
 							OutputNode->Output = Output;
 						}
 						else
diff --git a/Engine/Plugins/Runtime/Metasound/Source/MetasoundEditor/Private/MetasoundEditorGraph.cpp b/Engine/Plugins/Runtime/Metasound/Source/MetasoundEditor/Private/MetasoundEditorGraph.cpp
--- a/Engine/Plugins/Runtime/Metasound/Source/MetasoundEditor/Private/MetasoundEditorGraph.cpp
+++ b/Engine/Plugins/Runtime/Metasound/Source/MetasoundEditor/Private/MetasoundEditorGraph.cpp
@@ -880,22 +880,28 @@ bool UMetasoundEditorGraphInput::IsInterfaceMember(FMetasoundFrontendInterface*
 	FMetasoundFrontendInterface Interface;
 	if (NameContainsInterfaceNamespace(&Interface))
 	{
-		// Check if Input is a member of the found interface
-		if (const FMetasoundFrontendNode* InputNode = GetFrontendNode())
+		// Is interface declared on this MetaSound 
+		const UObject& MetaSoundObject = GetOwningGraph()->GetMetasoundChecked();
+		const FMetasoundAssetBase* MetaSoundAsset = Metasound::IMetasoundUObjectRegistry::Get().GetObjectAsAssetBase(&MetaSoundObject);
+		if (MetaSoundAsset && MetaSoundAsset->IsInterfaceDeclared(Interface.Version))
 		{
-			const FMetasoundFrontendVertex& Input = InputNode->Interface.Inputs.Last();
-			auto IsInput = [&Input](const FMetasoundFrontendClassInput& InterfaceInput)
+			// Check if Input is a member of the found interface
+			if (const FMetasoundFrontendNode* InputNode = GetFrontendNode())
 			{
-				return FMetasoundFrontendVertex::IsFunctionalEquivalent(Input, InterfaceInput);
-			};
+				const FMetasoundFrontendVertex& Input = InputNode->Interface.Inputs.Last();
+				auto IsInput = [&Input](const FMetasoundFrontendClassInput& InterfaceInput)
+				{
+					return FMetasoundFrontendVertex::IsFunctionalEquivalent(Input, InterfaceInput);
+				};
 
-			if (Interface.Inputs.ContainsByPredicate(IsInput))
-			{
-				if (OutInterface)
+				if (Interface.Inputs.ContainsByPredicate(IsInput))
 				{
-					*OutInterface = MoveTemp(Interface);
+					if (OutInterface)
+					{
+						*OutInterface = MoveTemp(Interface);
+					}
+					return true;
 				}
-				return true;
 			}
 		}
 	}
@@ -1162,22 +1168,28 @@ bool UMetasoundEditorGraphOutput::IsInterfaceMember(FMetasoundFrontendInterface*
 	FMetasoundFrontendInterface Interface;
 	if (NameContainsInterfaceNamespace(&Interface))
 	{
-		// Check if Output is a member of the found interface
-		if (const FMetasoundFrontendNode* OutputNode = GetFrontendNode())
+		// Is interface declared on this MetaSound 
+		const UObject& MetaSoundObject = GetOwningGraph()->GetMetasoundChecked();
+		const FMetasoundAssetBase* MetaSoundAsset = Metasound::IMetasoundUObjectRegistry::Get().GetObjectAsAssetBase(&MetaSoundObject);
+		if (MetaSoundAsset && MetaSoundAsset->IsInterfaceDeclared(Interface.Version))
 		{
-			const FMetasoundFrontendVertex& Output = OutputNode->Interface.Outputs.Last();
-			auto IsOutput = [&Output](const FMetasoundFrontendClassOutput& InterfaceOutput)
+			// Check if Output is a member of the found interface
+			if (const FMetasoundFrontendNode* OutputNode = GetFrontendNode())
 			{
-				return FMetasoundFrontendVertex::IsFunctionalEquivalent(Output, InterfaceOutput);
-			};
+				const FMetasoundFrontendVertex& Output = OutputNode->Interface.Outputs.Last();
+				auto IsOutput = [&Output](const FMetasoundFrontendClassOutput& InterfaceOutput)
+				{
+					return FMetasoundFrontendVertex::IsFunctionalEquivalent(Output, InterfaceOutput);
+				};
 
-			if (Interface.Outputs.ContainsByPredicate(IsOutput))
-			{
-				if (OutInterface)
+				if (Interface.Outputs.ContainsByPredicate(IsOutput))
 				{
-					*OutInterface = MoveTemp(Interface);
+					if (OutInterface)
+					{
+						*OutInterface = MoveTemp(Interface);
+					}
+					return true;
 				}
-				return true;
 			}
 		}
 	}
diff --git a/Engine/Plugins/Runtime/Metasound/Source/MetasoundFrontend/Private/MetasoundFrontendDocumentBuilder.cpp b/Engine/Plugins/Runtime/Metasound/Source/MetasoundFrontend/Private/MetasoundFrontendDocumentBuilder.cpp
--- a/Engine/Plugins/Runtime/Metasound/Source/MetasoundFrontend/Private/MetasoundFrontendDocumentBuilder.cpp
+++ b/Engine/Plugins/Runtime/Metasound/Source/MetasoundFrontend/Private/MetasoundFrontendDocumentBuilder.cpp
@@ -1127,8 +1127,11 @@ const FMetasoundFrontendNode* FMetaSoundFrontendDocumentBuilder::AddGraphInput(c
 		FMetasoundFrontendInterface Interface;
 		if (DocumentBuilderPrivate::IsInterfaceInput(InClassInput.Name, InClassInput.TypeName, &Interface))
 		{
-			FText RequiredText;
-			bIsRequired = Interface.IsMemberInputRequired(InClassInput.Name, RequiredText);
+			if (Document.Interfaces.Contains(Interface.Version))
+			{
+				FText RequiredText;
+				bIsRequired = Interface.IsMemberInputRequired(InClassInput.Name, RequiredText);
+			}
 		}
 #endif // WITH_EDITORONLY_DATA
 
