Fix for input default range widget/setter fighting & transaction stack issues
#rnx

diff --git a/Engine/Plugins/Runtime/Metasound/Source/MetasoundEditor/Private/SMetasoundGraphNode.cpp b/Engine/Plugins/Runtime/Metasound/Source/MetasoundEditor/Private/SMetasoundGraphNode.cpp
--- a/Engine/Plugins/Runtime/Metasound/Source/MetasoundEditor/Private/SMetasoundGraphNode.cpp
+++ b/Engine/Plugins/Runtime/Metasound/Source/MetasoundEditor/Private/SMetasoundGraphNode.cpp
@@ -810,59 +810,41 @@ namespace Metasound
 
 				auto OnValueChangedLambda = [DefaultFloat, GraphMember, this](float Value)
 				{
-					if (DefaultFloat.IsValid() && FloatInputWidget.IsValid())
+					BeginOrUpdateValueTransaction(GraphMember, [this, DefaultFloat, Value](const FGuid& BuildPageID, UMetasoundEditorGraphMember& Member)
 					{
-						if (!bIsInputWidgetTransacting)
+						if (DefaultFloat.IsValid() && FloatInputWidget.IsValid())
 						{
-							GEditor->BeginTransaction(LOCTEXT("MetasoundGraphNode_MetasoundSetInputDefault", "Set MetaSound Input Default"));
-							bIsInputWidgetTransacting = true;
-						}
-
-						FMetaSoundFrontendDocumentBuilder& Builder = GraphMember->GetFrontendBuilderChecked();
-						Builder.CastDocumentObjectChecked<UObject>().Modify();
-						DefaultFloat->Modify();
-
-						FMetasoundFrontendLiteral OutputLiteral;
-						OutputLiteral.Set(FloatInputWidget->GetOutputValue(Value));
+							DefaultFloat->Modify();
 
-						const FGuid BuildPageID = Builder.GetBuildPageID();
-						DefaultFloat->SetFromLiteral(OutputLiteral, BuildPageID);
-						DefaultFloat->OnDefaultValueChanged.Broadcast(BuildPageID, Value);
+							{
+								FMetasoundFrontendLiteral OutputLiteral;
+								OutputLiteral.Set(FloatInputWidget->GetOutputValue(Value));
+								DefaultFloat->SetFromLiteral(OutputLiteral, BuildPageID);
+							}
 
-						constexpr bool bPostTransaction = true;
-						GraphMember->UpdateFrontendDefaultLiteral(bPostTransaction, &BuildPageID);
-					}
+							constexpr bool bPostTransaction = false;
+							Member.UpdateFrontendDefaultLiteral(bPostTransaction, &BuildPageID);
+						}
+					});
 				};
 
-				auto OnValueCommittedLambda = [DefaultFloat, GraphMember, this](float Value)
+				auto OnValueCommittedLambda = [this, GraphMember, DefaultFloat](float Value)
 				{
-					if (DefaultFloat.IsValid() && FloatInputWidget.IsValid())
+					FinalizeValueTransaction(GraphMember, [&](const FGuid& BuildPageID, UMetasoundEditorGraphMember& Member, bool bPostTransaction)
 					{
-						bool bPostTransaction = false;
-						FMetasoundFrontendLiteral OutputLiteral;
-						OutputLiteral.Set(FloatInputWidget->GetOutputValue(Value));
-						const FGuid BuildPageID = GraphMember->GetFrontendBuilderChecked().GetBuildPageID();
-						DefaultFloat->SetFromLiteral(OutputLiteral, BuildPageID);
-						DefaultFloat->OnDefaultValueChanged.Broadcast(BuildPageID, Value);
-
-						if (bIsInputWidgetTransacting)
+						if (DefaultFloat.IsValid() && FloatInputWidget.IsValid())
 						{
-							GEditor->EndTransaction();
-							bIsInputWidgetTransacting = false;
-						}
-						else
-						{
-							bPostTransaction = true;
-							UE_LOG(LogMetaSound, Warning, TEXT("Unmatched MetaSound editor widget transaction."));
-						}
-
-						GraphMember->UpdateFrontendDefaultLiteral(bPostTransaction);
+							DefaultFloat->Modify();
+							{
+								FMetasoundFrontendLiteral OutputLiteral;
+								OutputLiteral.Set(FloatInputWidget->GetOutputValue(Value));
+								DefaultFloat->SetFromLiteral(OutputLiteral, BuildPageID);
+							}
 
-						if (UMetasoundEditorGraph* Graph = GraphMember->GetOwningGraph())
-						{
-							FGraphBuilder::GetOutermostMetaSoundChecked(*Graph).GetModifyContext().AddMemberIDsModified({ GraphMember->GetMemberID() });
+							Member.UpdateFrontendDefaultLiteral(bPostTransaction);
+							DefaultFloat->OnDefaultValueChanged.Broadcast(BuildPageID, Value);
 						}
-					}
+					});
 				};
 
 				if (DefaultFloat->WidgetType == EMetasoundMemberDefaultWidget::Slider)
@@ -978,26 +960,30 @@ namespace Metasound
 				}
 				else if (DefaultFloat->WidgetType == EMetasoundMemberDefaultWidget::RadialSlider)
 				{
-					auto OnRadialSliderMouseCaptureBeginLambda = [this]()
+					auto OnRadialSliderMouseCaptureBeginLambda = [this, GraphMember]()
 					{
-						if (!bIsInputWidgetTransacting)
-						{
-							GEditor->BeginTransaction(LOCTEXT("MetasoundSetRadialSliderInputDefault", "Set MetaSound Input Default"));
-							bIsInputWidgetTransacting = true;
-						}
+						BeginOrUpdateValueTransaction(GraphMember, [](const FGuid&, UMetasoundEditorGraphMember&) { });
 					};
 
-					auto OnRadialSliderMouseCaptureEndLambda = [this]()
+					auto OnRadialSliderMouseCaptureEndLambda = [this, GraphMember, DefaultFloat]()
 					{
-						if (bIsInputWidgetTransacting)
-						{
-							GEditor->EndTransaction();
-							bIsInputWidgetTransacting = false;
-						}
-						else
+						FinalizeValueTransaction(GraphMember, [this, DefaultFloat](const FGuid& BuildPageID, UMetasoundEditorGraphMember& Member, bool bPostTransaction)
 						{
-							UE_LOG(LogMetaSound, Warning, TEXT("Unmatched MetaSound editor widget transaction."));
-						}
+							if (DefaultFloat.IsValid())
+							{
+								DefaultFloat->Modify();
+								const float FinalValue = DefaultFloat->GetDefaultAs<float>(BuildPageID);
+
+								{
+									FMetasoundFrontendLiteral OutputLiteral;
+									OutputLiteral.Set(FinalValue);
+									DefaultFloat->SetFromLiteral(OutputLiteral, BuildPageID);
+								}
+
+								Member.UpdateFrontendDefaultLiteral(bPostTransaction);
+								DefaultFloat->OnDefaultValueChanged.Broadcast(BuildPageID, FinalValue);
+							}
+						});
 					};
 
 					if (bUseAudioMaterialWidgets)
@@ -1123,43 +1109,45 @@ namespace Metasound
 							constexpr float WidgetPadding = 3.0f;
 							static const FVector2D ButtonDesiredSize = FVector2D(56.0f, 87.0f);
 
-							auto OnBoolValueChangedLambda = [DefaultBool, GraphMember, this](bool Value)
+							auto OnBoolValueChangedLambda = [this, GraphMember, DefaultBool](bool Value)
 							{
-								if (DefaultBool.IsValid() && MaterialButtonWidget.IsValid())
+								BeginOrUpdateValueTransaction(GraphMember, [this, DefaultBool, Value](const FGuid& BuildPageID, UMetasoundEditorGraphMember& Member)
 								{
-									if (!bIsInputWidgetTransacting)
+									if (DefaultBool.IsValid())
 									{
-										GEditor->BeginTransaction(FText::Format(LOCTEXT("MetasoundGraphNode_MetasoundSetBoolInput", "Set MetaSound Input '{0}' Boolean Default"), GraphMember->GetDisplayName()));
-										bIsInputWidgetTransacting = true;
-									}
+										DefaultBool->Modify();
 
-									FMetaSoundFrontendDocumentBuilder& Builder = GraphMember->GetFrontendBuilderChecked();
-									Builder.CastDocumentObjectChecked<UObject>().Modify();
-									DefaultBool->Modify();
+										{
+											FMetasoundFrontendLiteral OutputLiteral;
+											OutputLiteral.Set(Value);
+											DefaultBool->SetFromLiteral(OutputLiteral, BuildPageID);
+										}
 
-									FMetasoundFrontendLiteral OutputLiteral;
-									OutputLiteral.Set(Value);
-
-									const FGuid BuildPageID = Builder.GetBuildPageID();
-									DefaultBool->SetFromLiteral(OutputLiteral, BuildPageID);
-									DefaultBool->OnDefaultStateChanged.Broadcast(Value, BuildPageID);
-
-									constexpr bool bPostTransaction = true;
-									GraphMember->UpdateFrontendDefaultLiteral(bPostTransaction, &BuildPageID);
-								}
+										constexpr bool bPostTransaction = false;
+										Member.UpdateFrontendDefaultLiteral(bPostTransaction, &BuildPageID);
+									}
+								});
 							};
 
-							auto OnBoolMouseCaptureEndLambda = [this]()
+							auto OnBoolMouseCaptureEndLambda = [this, GraphMember, DefaultBool]()
 							{
-								if (bIsInputWidgetTransacting)
+								FinalizeValueTransaction(GraphMember, [this, DefaultBool](const FGuid& BuildPageID, UMetasoundEditorGraphMember& Member, bool bPostTransaction)
 								{
-									GEditor->EndTransaction();
-									bIsInputWidgetTransacting = false;
-								}
-								else
-								{
-									UE_LOG(LogMetaSound, Warning, TEXT("Unmatched MetaSound editor widget transaction."));
-								}
+									if (DefaultBool.IsValid())
+									{
+										DefaultBool->Modify();
+										const bool FinalValue = DefaultBool->GetDefaultAs<bool>(BuildPageID);
+
+										{
+											FMetasoundFrontendLiteral OutputLiteral;
+											OutputLiteral.Set(FinalValue);
+											DefaultBool->SetFromLiteral(OutputLiteral, BuildPageID);
+										}
+
+										Member.UpdateFrontendDefaultLiteral(bPostTransaction);
+										DefaultBool->OnDefaultStateChanged.Broadcast(FinalValue, BuildPageID);
+									}
+								});
 							};
 
 							SAssignNew(MaterialButtonWidget, SAudioMaterialButton)
@@ -1277,6 +1265,79 @@ namespace Metasound
 				];
 		}
 
+		void SMetaSoundGraphNode::BeginOrUpdateValueTransaction(TWeakObjectPtr<UMetasoundEditorGraphMember> GraphMemberPtr, TFunctionRef<void(const FGuid&, UMetasoundEditorGraphMember&)> SetValue)
+		{
+			UMetasoundEditorGraphMember* GraphMember = GraphMemberPtr.Pin().Get();
+			if (!GraphMember)
+			{
+				return;
+			}
+
+			FMetaSoundFrontendDocumentBuilder& Builder = GraphMember->GetFrontendBuilderChecked();
+			Builder.CastDocumentObjectChecked<UObject>().Modify();
+			const FGuid BuildPageID = Builder.GetBuildPageID();
+
+			if (!bIsInputWidgetTransacting)
+			{
+				const UMetaSoundSettings* MetaSoundSettings = GetDefault<UMetaSoundSettings>();
+				check(MetaSoundSettings);
+				const FMetaSoundPageSettings* PageSettings = MetaSoundSettings->FindPageSettings(BuildPageID);
+
+				GEditor->BeginTransaction(FText::Format(LOCTEXT("MetaSoundGraphNode_SetMemberDefault", "Set MetaSound {0} '{1}' Default (Page: {2})"),
+					GraphMember->GetGraphMemberLabel(),
+					GraphMember->GetDisplayName(),
+					PageSettings ? FText::FromName(PageSettings->Name) : LOCTEXT("MetaSoundGraphPage_Unknown", "Unknown")
+				));
+				bIsInputWidgetTransacting = true;
+			}
+
+			SetValue(BuildPageID, *GraphMember);
+		}
+
+		void SMetaSoundGraphNode::FinalizeValueTransaction(TWeakObjectPtr<UMetasoundEditorGraphMember> GraphMemberPtr, TFunctionRef<void(const FGuid&, UMetasoundEditorGraphMember&, bool)> SetValue)
+		{
+			UMetasoundEditorGraphMember* GraphMember = GraphMemberPtr.Pin().Get();
+			if (!GraphMember)
+			{
+				return;
+			}
+
+			FMetaSoundFrontendDocumentBuilder& Builder = GraphMember->GetFrontendBuilderChecked();
+			const FGuid& BuildPageID = Builder.GetBuildPageID();
+			UObject& MetaSound = Builder.CastDocumentObjectChecked<UObject>();
+			MetaSound.Modify();
+
+			bool bPostTransaction = false;
+			if (!bIsInputWidgetTransacting)
+			{
+				bPostTransaction = true;
+				UE_LOG(LogMetaSound, Warning, TEXT("Unmatched MetaSound editor widget transaction."));
+			}
+
+			SetValue(BuildPageID, *GraphMember, bPostTransaction);
+
+			if (bIsInputWidgetTransacting)
+			{
+				GEditor->EndTransaction();
+				bIsInputWidgetTransacting = false;
+			}
+
+			if (UMetasoundEditorGraph* Graph = GraphMember->GetOwningGraph())
+			{
+				FMetasoundFrontendDocumentModifyContext& ModifyContext = FGraphBuilder::GetOutermostMetaSoundChecked(*Graph).GetModifyContext();
+				ModifyContext.AddMemberIDsModified({ GraphMember->GetMemberID() });
+				ModifyContext.AddNodeIDModified(GetMetaSoundNode().GetNodeID());
+
+				// Only inputs require registration as changes to default values on other types (i.e. variables, outputs) are not external
+				// graph API changes that may have an effect on other open MetaSound asset(s) visible state(s) (ex. presets, referenced node
+				// defaults, etc.)
+				if (GraphMember->IsA<UMetasoundEditorGraphInput>())
+				{
+					FGraphBuilder::RegisterGraphWithFrontend(Builder.CastDocumentObjectChecked<UObject>());
+				}
+			}
+		}
+
 		TSharedPtr<SGraphPin> SMetaSoundGraphNodeKnot::CreatePinWidget(UEdGraphPin* Pin) const
 		{
 			return SNew(SMetaSoundGraphPinKnot, Pin);
diff --git a/Engine/Plugins/Runtime/Metasound/Source/MetasoundEditor/Private/SMetasoundGraphNode.h b/Engine/Plugins/Runtime/Metasound/Source/MetasoundEditor/Private/SMetasoundGraphNode.h
--- a/Engine/Plugins/Runtime/Metasound/Source/MetasoundEditor/Private/SMetasoundGraphNode.h
+++ b/Engine/Plugins/Runtime/Metasound/Source/MetasoundEditor/Private/SMetasoundGraphNode.h
@@ -71,6 +71,9 @@ namespace Metasound::Editor
 		static TSharedRef<SWidget> CreateTriggerSimulationWidget(UMetasoundEditorGraphMemberDefaultLiteral& Literal, TAttribute<EVisibility>&& InVisibility, TAttribute<bool>&& InEnablement, const FText* InToolTip = nullptr);
 
 	private:
+		void BeginOrUpdateValueTransaction(TWeakObjectPtr<UMetasoundEditorGraphMember> GraphMemberPtr, TFunctionRef<void(const FGuid& /*BuildPageID*/, UMetasoundEditorGraphMember& /*GraphMember*/)> SetValue);
+		void FinalizeValueTransaction(TWeakObjectPtr<UMetasoundEditorGraphMember> GraphMemberPtr, TFunctionRef<void(const FGuid& /*BuildPageID*/, UMetasoundEditorGraphMember& /*GraphMember*/, bool /*bPostTransaction*/)> SetValue);
+
 		// Returns attribute that returns whether or not the input widget should be enabled.
 		TAttribute<bool> GetInputWidgetEnabled() const;
 
