Fix crash when adding new MetaSound variable nodes after node removal

#rnx

diff --git a/Engine/Plugins/Runtime/Metasound/Source/MetasoundFrontend/Private/MetasoundFrontendDocumentBuilder.cpp b/Engine/Plugins/Runtime/Metasound/Source/MetasoundFrontend/Private/MetasoundFrontendDocumentBuilder.cpp
--- a/Engine/Plugins/Runtime/Metasound/Source/MetasoundFrontend/Private/MetasoundFrontendDocumentBuilder.cpp
+++ b/Engine/Plugins/Runtime/Metasound/Source/MetasoundFrontend/Private/MetasoundFrontendDocumentBuilder.cpp
@@ -4752,14 +4752,16 @@ bool FMetaSoundFrontendDocumentBuilder::UnlinkVariableNode(const FGuid& InNodeID
 			return true;
 		}
 
-		const bool bRemovedDeferredNode = Variable.DeferredAccessorNodeIDs.RemoveAllSwap(IsNodeID, EAllowShrinking::No) > 0;
+		// Removal must maintain array order to preserve head/tail positions in stack
+		const bool bRemovedDeferredNode = Variable.DeferredAccessorNodeIDs.RemoveAll(IsNodeID) > 0;
 		if (bRemovedDeferredNode)
 		{
 			SpliceVariableNodeFromStack(InNodeID, InPageID);
 			return true;
 		}
-
-		const bool bRemovedAccessorNode = Variable.AccessorNodeIDs.RemoveAllSwap(IsNodeID, EAllowShrinking::No) > 0;
+		
+		// Removal must maintain array order to preserve head/tail positions in stack
+		const bool bRemovedAccessorNode = Variable.AccessorNodeIDs.RemoveAll(IsNodeID) > 0;
 		if (bRemovedAccessorNode)
 		{
 			SpliceVariableNodeFromStack(InNodeID, InPageID);
