Fix for certain MetaSound Editor Asset actions causing focus stealing when not expected for end-user
#rnx

diff --git a/Engine/Plugins/Runtime/Metasound/Source/MetasoundEditor/Private/MetasoundEditor.cpp b/Engine/Plugins/Runtime/Metasound/Source/MetasoundEditor/Private/MetasoundEditor.cpp
--- a/Engine/Plugins/Runtime/Metasound/Source/MetasoundEditor/Private/MetasoundEditor.cpp
+++ b/Engine/Plugins/Runtime/Metasound/Source/MetasoundEditor/Private/MetasoundEditor.cpp
@@ -1335,7 +1335,7 @@ namespace Metasound
 			return nullptr;
 		}
 
-		void FEditor::SetSelection(const TArray<UObject*>& SelectedObjects)
+		void FEditor::SetSelection(const TArray<UObject*>& SelectedObjects, bool bInvokeTabOnSelectionSet)
 		{
 			if (GraphMembersMenu.IsValid())
 			{
@@ -1354,7 +1354,7 @@ namespace Metasound
 			{
 				if (SelectedObjects.IsEmpty())
 				{
-					if (TabManager.IsValid())
+					if (bInvokeTabOnSelectionSet && TabManager.IsValid())
 					{
 						if (ShowPageGraphDetails())
 						{
@@ -1366,7 +1366,7 @@ namespace Metasound
 				{
 					MetasoundDetails->SetObjects(SelectedObjects);
 					MetasoundDetails->HideFilterArea(false);
-					if (TabManager.IsValid())
+					if (bInvokeTabOnSelectionSet && TabManager.IsValid())
 					{
 						TabManager->TryInvokeTab(TabNamesPrivate::Details);
 					}
@@ -1668,7 +1668,11 @@ namespace Metasound
 			MetasoundDetails->SetExtensionHandler(MakeShared<FMetaSoundNodeExtensionHandler>());
 
 			// Set details selection to the MetaSound's source settings 
-			SetSelection({ &MetaSound });
+				// Don't invoke tab as this can be called in response
+				// to opening multiple assets, and the higher level
+				// request handles tab invocation/focus
+				constexpr bool bInvokeTabOnSelectionSet = false;
+			SetSelection({ &MetaSound }, false);
 			InterfacesDetails = PropertyModule.CreateDetailView(Args);
 			if (InterfacesDetails.IsValid())
 			{
@@ -4793,7 +4797,11 @@ namespace Metasound
 
 				if (!Selection.IsEmpty())
 				{
-					SetSelection(Selection);
+					// Don't invoke tab as this can be called in response
+					// to another focused, referenced graph mutating (ex.
+					// interface changing).
+					constexpr bool bInvokeTabOnSelectionSet = false;
+					SetSelection(Selection, bInvokeTabOnSelectionSet);
 				}
 
 				// Avoids details panel displaying
diff --git a/Engine/Plugins/Runtime/Metasound/Source/MetasoundEditor/Private/MetasoundEditor.h b/Engine/Plugins/Runtime/Metasound/Source/MetasoundEditor/Private/MetasoundEditor.h
--- a/Engine/Plugins/Runtime/Metasound/Source/MetasoundEditor/Private/MetasoundEditor.h
+++ b/Engine/Plugins/Runtime/Metasound/Source/MetasoundEditor/Private/MetasoundEditor.h
@@ -146,7 +146,7 @@ namespace Metasound::Editor
 
 		/** IMetasoundEditor interface */
 		virtual UObject* GetMetasoundObject() const override;
-		virtual void SetSelection(const TArray<UObject*>& SelectedObjects) override;
+		virtual void SetSelection(const TArray<UObject*>& SelectedObjects, bool bInvokeTabOnSelectionSet = true) override;
 		virtual bool GetBoundsForSelectedNodes(FSlateRect& Rect, float Padding) override;
 
 		/** IToolkit interface */
diff --git a/Engine/Plugins/Runtime/Metasound/Source/MetasoundEditor/Public/IMetasoundEditor.h b/Engine/Plugins/Runtime/Metasound/Source/MetasoundEditor/Public/IMetasoundEditor.h
--- a/Engine/Plugins/Runtime/Metasound/Source/MetasoundEditor/Public/IMetasoundEditor.h
+++ b/Engine/Plugins/Runtime/Metasound/Source/MetasoundEditor/Public/IMetasoundEditor.h
@@ -15,7 +15,7 @@ public:
 	virtual UObject* GetMetasoundObject() const = 0;
 
 	// Sets the selected object(s).
-	virtual void SetSelection(const TArray<UObject*>& SelectedObjects) = 0;
+	virtual void SetSelection(const TArray<UObject*>& SelectedObjects, bool bInvokeTabOnSelectionSet = true) = 0;
 
 	// Returns the bounds for the selected node(s).
 	virtual bool GetBoundsForSelectedNodes(FSlateRect& OutRect, float Padding) = 0;
