Fix crash when undo-ing blutility that calls BuildToAsset
#rnx

diff --git a/Engine/Plugins/Runtime/Metasound/Source/MetasoundEngine/Private/MetasoundEngineAsset.h b/Engine/Plugins/Runtime/Metasound/Source/MetasoundEngine/Private/MetasoundEngineAsset.h
--- a/Engine/Plugins/Runtime/Metasound/Source/MetasoundEngine/Private/MetasoundEngineAsset.h
+++ b/Engine/Plugins/Runtime/Metasound/Source/MetasoundEngine/Private/MetasoundEngineAsset.h
@@ -285,10 +285,11 @@ namespace Metasound::Engine
 
 			if (InArchive.IsLoading())
 			{
+				const bool bIsTransacting = InArchive.IsTransacting();
 				TStrongObjectPtr<UMetaSoundBuilderBase> Builder;
 				{
 					FGCScopeGuard ScopeGuard;
-					Builder.Reset(&FDocumentBuilderRegistry::GetChecked().FindOrBeginBuilding(InMetaSound));
+					Builder.Reset(&FDocumentBuilderRegistry::GetChecked().FindOrBeginBuilding(InMetaSound, bIsTransacting));
 				}
 
 				{
diff --git a/Engine/Plugins/Runtime/Metasound/Source/MetasoundEngine/Public/MetasoundDocumentBuilderRegistry.h b/Engine/Plugins/Runtime/Metasound/Source/MetasoundEngine/Public/MetasoundDocumentBuilderRegistry.h
--- a/Engine/Plugins/Runtime/Metasound/Source/MetasoundEngine/Public/MetasoundDocumentBuilderRegistry.h
+++ b/Engine/Plugins/Runtime/Metasound/Source/MetasoundEngine/Public/MetasoundDocumentBuilderRegistry.h
@@ -70,10 +70,19 @@ namespace Metasound::Engine
 		}
 
 #if WITH_EDITORONLY_DATA
+		// Find or begin building a MetaSound asset.  Optionally, provide
+		// whether or not the builder is being accessed during a transaction.
+		// If false, enforces MetaSound being built is an asset.  If true, does
+		// not enforce (transactions may result in assets being moved and becoming
+		// transient wherein the builder can and should be valid to act on the
+		// transient UObject in these rare cases).
 		template <typename BuilderClass = UMetaSoundBuilderBase>
-		BuilderClass& FindOrBeginBuilding(UObject& InMetaSoundObject) const
+		BuilderClass& FindOrBeginBuilding(UObject& InMetaSoundObject, bool bIsTransacting = false) const
 		{
-			check(InMetaSoundObject.IsAsset());
+			if (!bIsTransacting)
+			{
+				check(InMetaSoundObject.IsAsset());
+			}
 
 			TScriptInterface<IMetaSoundDocumentInterface> DocInterface = &InMetaSoundObject;
 			check(DocInterface.GetObject());
