Fix triggers showing page options, empty default row on variables/outputs, and placing defaults under array entry when page data is not supported (i.e. triggers, all variables & outputs)

#rnx

diff --git a/Engine/Plugins/Runtime/Metasound/Source/MetasoundEditor/Private/MetasoundDefaultLiteralCustomization.cpp b/Engine/Plugins/Runtime/Metasound/Source/MetasoundEditor/Private/MetasoundDefaultLiteralCustomization.cpp
--- a/Engine/Plugins/Runtime/Metasound/Source/MetasoundEditor/Private/MetasoundDefaultLiteralCustomization.cpp
+++ b/Engine/Plugins/Runtime/Metasound/Source/MetasoundEditor/Private/MetasoundDefaultLiteralCustomization.cpp
@@ -395,9 +395,15 @@ namespace Metasound::Editor
 					}
 					else
 					{
+						TSharedRef<IPropertyHandle> PagedDefaultProperty = ElementProperty;
+						uint32 NumChildren = 0;
+						if (ElementProperty->GetNumChildren(NumChildren) == FPropertyAccess::Result::Success && NumChildren > 0)
+						{
+							PagedDefaultProperty = ElementProperty->GetChildHandle(0).ToSharedRef();
+						}
 						(*ValueRow.CustomNameWidget())
 						[
-							ElementProperty->CreatePropertyNameWidget()
+							PagedDefaultProperty->CreatePropertyNameWidget()
 						];
 					}
 					ValueRow.ShowPropertyButtons(false);
diff --git a/Engine/Plugins/Runtime/Metasound/Source/MetasoundEditor/Private/MetasoundEditorGraph.cpp b/Engine/Plugins/Runtime/Metasound/Source/MetasoundEditor/Private/MetasoundEditorGraph.cpp
--- a/Engine/Plugins/Runtime/Metasound/Source/MetasoundEditor/Private/MetasoundEditorGraph.cpp
+++ b/Engine/Plugins/Runtime/Metasound/Source/MetasoundEditor/Private/MetasoundEditorGraph.cpp
@@ -872,7 +872,10 @@ TArray<UMetasoundEditorGraphMemberNode*> UMetasoundEditorGraphInput::GetNodes()
 
 bool UMetasoundEditorGraphInput::IsDefaultPaged() const
 {
-	return true;
+	// Triggers are special and do not show their default value, but are visible
+	// to allow for interact button when auditioning.  Therefore, default paging
+	// is unnecessary.
+	return TypeName != Metasound::GetMetasoundDataTypeName<Metasound::FTrigger>();
 }
 
 bool UMetasoundEditorGraphInput::IsInterfaceMember(FMetasoundFrontendInterface* OutInterface) const
diff --git a/Engine/Plugins/Runtime/Metasound/Source/MetasoundEditor/Private/MetasoundNodeDetailCustomization.cpp b/Engine/Plugins/Runtime/Metasound/Source/MetasoundEditor/Private/MetasoundNodeDetailCustomization.cpp
--- a/Engine/Plugins/Runtime/Metasound/Source/MetasoundEditor/Private/MetasoundNodeDetailCustomization.cpp
+++ b/Engine/Plugins/Runtime/Metasound/Source/MetasoundEditor/Private/MetasoundNodeDetailCustomization.cpp
@@ -409,6 +409,20 @@ namespace Metasound
 				return;
 			}
 
+			bool bIsTrigger = false;
+			const UMetasoundEditorGraphMember* Member = BoolLiteral->FindMember();
+			if (!Member)
+			{
+				return;
+			}
+
+			// Non-input members don't show any options for inputs, so early out if trigger.
+			bIsTrigger = Member->GetDataType() == Metasound::GetMetasoundDataTypeName<Metasound::FTrigger>();
+			if (bIsTrigger && !Member->IsA<UMetasoundEditorGraphInput>())
+			{
+				return;
+			}
+
 			FMetasoundDefaultLiteralCustomizationBase::CustomizeDefaults(InLiteral, InDetailLayout);
 
 			TAttribute<EVisibility> DefaultVisibility = GetDefaultVisibility();
@@ -431,25 +445,41 @@ namespace Metasound
 
 			if (EditorSettings->bUseAudioMaterialWidgets)
 			{
-				if (const UMetasoundEditorGraphInput* Member = Cast<UMetasoundEditorGraphInput>(InLiteral.FindMember()))
+				if (!bIsTrigger)
 				{
-					if (Member->GetDataType() != GetMetasoundDataTypeName<FTrigger>())
+					if (const UMetasoundEditorGraph* OwningGraph = Member->GetOwningGraph())
 					{
-						if (const UMetasoundEditorGraph* OwningGraph = Member->GetOwningGraph())
-						{
-							bShowWidgetOptions = OwningGraph->IsEditable();
-						}
+						bShowWidgetOptions = OwningGraph->IsEditable();
 					}
 				}
 			}
 
 			if (bShowWidgetOptions)
 			{
-				if (InLiteral.GetDataType() != Metasound::GetMetasoundDataTypeName<Metasound::FTrigger>())
+				AddOptionPropRow(GET_MEMBER_NAME_CHECKED(UMetasoundEditorGraphMemberDefaultBool, WidgetType));
+			}
+		}
+
+		TAttribute<EVisibility> FMetasoundBoolLiteralCustomization::GetDefaultVisibility() const
+		{
+			if (UMetasoundEditorGraphMemberDefaultBool* DefaultBool = BoolLiteral.Pin().Get())
+			{
+				UMetasoundEditorGraphMember* Member = DefaultBool->FindMember();
+				if (Member->IsA<UMetasoundEditorGraphInput>())
+				{
+					return FMetasoundDefaultLiteralCustomizationBase::GetDefaultVisibility();
+				}
+				else
 				{
-					AddOptionPropRow(GET_MEMBER_NAME_CHECKED(UMetasoundEditorGraphMemberDefaultBool, WidgetType));
+					const bool bIsTrigger = Member->GetDataType() == Metasound::GetMetasoundDataTypeName<Metasound::FTrigger>();
+					if (bIsTrigger)
+					{
+						return EVisibility::Collapsed;
+					}
 				}
 			}
+
+			return FMetasoundDefaultLiteralCustomizationBase::GetDefaultVisibility();
 		}
 
 		void FMetasoundObjectArrayLiteralCustomization::CustomizeDefaults(UMetasoundEditorGraphMemberDefaultLiteral& InLiteral, IDetailLayoutBuilder& InDetailLayout)
diff --git a/Engine/Plugins/Runtime/Metasound/Source/MetasoundEditor/Private/MetasoundNodeDetailCustomization.h b/Engine/Plugins/Runtime/Metasound/Source/MetasoundEditor/Private/MetasoundNodeDetailCustomization.h
--- a/Engine/Plugins/Runtime/Metasound/Source/MetasoundEditor/Private/MetasoundNodeDetailCustomization.h
+++ b/Engine/Plugins/Runtime/Metasound/Source/MetasoundEditor/Private/MetasoundNodeDetailCustomization.h
@@ -200,6 +200,7 @@ namespace Metasound
 			virtual ~FMetasoundBoolLiteralCustomization();
 
 			virtual void CustomizeDefaults(UMetasoundEditorGraphMemberDefaultLiteral& InLiteral, IDetailLayoutBuilder& InDetailLayout) override;
+			virtual TAttribute<EVisibility> GetDefaultVisibility() const override;
 		};
 
 		// Customization to support drag-and-drop of Proxy UObject types on underlying members that are structs.
