Fix for stale MetaSound graph name showing when swapping page graphs on undo/redo

#rnx

diff --git a/Engine/Plugins/Runtime/Metasound/Source/MetasoundEditor/Private/MetasoundEditor.cpp b/Engine/Plugins/Runtime/Metasound/Source/MetasoundEditor/Private/MetasoundEditor.cpp
--- a/Engine/Plugins/Runtime/Metasound/Source/MetasoundEditor/Private/MetasoundEditor.cpp
+++ b/Engine/Plugins/Runtime/Metasound/Source/MetasoundEditor/Private/MetasoundEditor.cpp
@@ -1504,7 +1504,22 @@ namespace Metasound
 				ToolMenus->RefreshAllWidgets();
 			}
 
+			// Playback must be stopped if undoing a page change transaction
+			bool bStopPlayback = !Builder.IsValid() || !PageStatsWidget.IsValid();
+			if (!bStopPlayback)
+			{
+				const FMetaSoundFrontendDocumentBuilder& DocBuilder = Builder->GetConstBuilder();
+				bStopPlayback = DocBuilder.GetBuildPageID() != PageStatsWidget->GetDisplayedPageID();
+			}
+
 			SyncAuditionState();
+
+			if (bStopPlayback)
+			{
+				Stop();
+			}
+
+			UpdatePageInfo(IsPlaying());
 			bRefreshGraph = true;
 		}
 
@@ -1668,11 +1683,11 @@ namespace Metasound
 			MetasoundDetails->SetExtensionHandler(MakeShared<FMetaSoundNodeExtensionHandler>());
 
 			// Set details selection to the MetaSound's source settings 
-				// Don't invoke tab as this can be called in response
-				// to opening multiple assets, and the higher level
-				// request handles tab invocation/focus
-				constexpr bool bInvokeTabOnSelectionSet = false;
-			SetSelection({ &MetaSound }, false);
+			// Don't invoke tab as this can be called in response
+			// to opening multiple assets, and the higher level
+			// request handles tab invocation/focus
+			constexpr bool bInvokeTabOnSelectionSet = false;
+			SetSelection({ &MetaSound }, bInvokeTabOnSelectionSet);
 			InterfacesDetails = PropertyModule.CreateDetailView(Args);
 			if (InterfacesDetails.IsValid())
 			{
