Missing files from prior CL 37274219
#rnx

diff --git a/Engine/Plugins/Runtime/Metasound/Source/MetasoundEditor/Private/SMetasoundStats.cpp b/Engine/Plugins/Runtime/Metasound/Source/MetasoundEditor/Private/SMetasoundStats.cpp
--- a/Engine/Plugins/Runtime/Metasound/Source/MetasoundEditor/Private/SMetasoundStats.cpp
+++ b/Engine/Plugins/Runtime/Metasound/Source/MetasoundEditor/Private/SMetasoundStats.cpp
@@ -32,6 +32,9 @@ namespace Metasound::Editor
 	{
 		SVerticalBox::Construct(SVerticalBox::FArguments());
 
+		DisplayedPageID = Metasound::Frontend::DefaultPageID;
+		DisplayedPageName = Metasound::Frontend::DefaultPageName;
+
 		AddSlot().HAlign(HAlign_Left)
 		[
 			SAssignNew(AuditionPageTextWidget, STextBlock)
@@ -103,6 +106,19 @@ namespace Metasound::Editor
 			GraphPageTextWidget->SetText(PageInfo);
 			GraphPageTextWidget->SetColorAndOpacity(ActiveColor ? *ActiveColor : StatsPrivate::BaseTextColor);
 		}
+
+		DisplayedPageID = GraphPageSettings ? GraphPageSettings->UniqueId : Metasound::Frontend::DefaultPageID;
+		DisplayedPageName = GraphPageSettings ? GraphPageSettings->Name : Metasound::Frontend::DefaultPageName;
+	}
+
+	const FGuid& SPageStats::GetDisplayedPageID() const
+	{
+		return DisplayedPageID;
+	}
+
+	FName SPageStats::GetDisplayedPageName() const
+	{
+		return DisplayedPageName;
 	}
 
 	void SRenderStats::Construct(const FArguments& InArgs)
@@ -129,7 +145,6 @@ namespace Metasound::Editor
 		AddTextWidgetSlot(AuditionPlatformWidget);
 	}
 
-
 	void SRenderStats::Update(bool bIsPlaying, double InDeltaTime, const UMetaSoundSource* InSource)
 	{
 		using namespace Metasound;
diff --git a/Engine/Plugins/Runtime/Metasound/Source/MetasoundEditor/Private/SMetasoundStats.h b/Engine/Plugins/Runtime/Metasound/Source/MetasoundEditor/Private/SMetasoundStats.h
--- a/Engine/Plugins/Runtime/Metasound/Source/MetasoundEditor/Private/SMetasoundStats.h
+++ b/Engine/Plugins/Runtime/Metasound/Source/MetasoundEditor/Private/SMetasoundStats.h
@@ -30,10 +30,16 @@ namespace Metasound::Editor
 		void SetExecVisibility(TAttribute<EVisibility> InVisibility);
 		void Update(const FMetaSoundPageSettings* AuditionPageSettings, const FMetaSoundPageSettings* GraphPageSettings, const FSlateColor* ActiveColor);
 
+		const FGuid& GetDisplayedPageID() const;
+		FName GetDisplayedPageName() const;
+
 	private:
 		TSharedPtr<SImage> ExecImageWidget;
 		TSharedPtr<STextBlock> GraphPageTextWidget;
 		TSharedPtr<STextBlock> AuditionPageTextWidget;
+
+		FGuid DisplayedPageID;
+		FName DisplayedPageName;
 	};
 
 
