Variables (and potentially other values when using the BuilderAPI) on different pages but of the same name cannot have unique default values (Using Frontend::DefaultPageID as fallback in these instances where they should be using BuildPageID)

#rnx

diff --git a/Engine/Plugins/Runtime/Metasound/Source/MetasoundEditor/Private/MetasoundEditorGraph.cpp b/Engine/Plugins/Runtime/Metasound/Source/MetasoundEditor/Private/MetasoundEditorGraph.cpp
--- a/Engine/Plugins/Runtime/Metasound/Source/MetasoundEditor/Private/MetasoundEditorGraph.cpp
+++ b/Engine/Plugins/Runtime/Metasound/Source/MetasoundEditor/Private/MetasoundEditorGraph.cpp
@@ -1321,6 +1321,7 @@ void UMetasoundEditorGraphOutput::UpdateFrontendDefaultLiteral(bool bPostTransac
 		return;
 	}
 
+	// Use the default page ID here as output defaults do *not* support paged defaults (they exist per paged graph on the singleton output node)
 	FMetasoundFrontendLiteral DefaultLiteral;
 	if (ensure(Literal->TryFindDefault(DefaultLiteral)))
 	{
@@ -1856,20 +1857,26 @@ bool UMetasoundEditorGraphVariable::Synchronize()
 	using namespace Metasound;
 
 	bool bModified = Super::Synchronize();
-
 	FMetaSoundFrontendDocumentBuilder& Builder = GetFrontendBuilderChecked();
 	const FMetasoundFrontendGraph& Graph = Builder.FindConstBuildGraphChecked();
 	const FName MemberName = GetMemberName();
-	auto IsVariable = [&MemberName] (const FMetasoundFrontendVariable& Var) { return Var.Name == MemberName; };
-	if (const FMetasoundFrontendVariable* Variable = Graph.Variables.FindByPredicate(IsVariable))
+	if (const FMetasoundFrontendVariable* Variable = Builder.FindGraphVariable(MemberName))
 	{
+		if (TypeName != Variable->TypeName)
+		{
+			bModified = true;
+			TypeName = Variable->TypeName;
+
+			InitializeLiteral();
+		}
+
 		if (ensure(Literal))
 		{
 			TOptional<FMetasoundFrontendLiteral> NewDefault;
-			if (const FMetasoundFrontendNode* MutatorNode = Builder.FindNode(Variable->MutatorNodeID); ensure(MutatorNode))
+			FMetasoundFrontendLiteral DefaultLiteral;
+			Literal->TryFindDefault(DefaultLiteral);
+			if (const FMetasoundFrontendNode* MutatorNode = Builder.FindNode(Variable->MutatorNodeID))
 			{
-				FMetasoundFrontendLiteral DefaultLiteral;
-				Literal->TryFindDefault(DefaultLiteral);
 				if (!MutatorNode->InputLiterals.IsEmpty())
 				{
 					const FMetasoundFrontendVertexLiteral& VertexLiteral = MutatorNode->InputLiterals.Last();
@@ -1888,6 +1895,10 @@ bool UMetasoundEditorGraphVariable::Synchronize()
 					}
 				}
 			}
+			else if (!Variable->Literal.IsEqual(DefaultLiteral))
+			{
+				NewDefault = Variable->Literal;
+			}
 
 			if (NewDefault.IsSet())
 			{
@@ -1923,13 +1934,16 @@ void UMetasoundEditorGraphVariable::UpdateFrontendDefaultLiteral(bool bPostTrans
 	const FScopedTransaction Transaction(LOCTEXT("Set Variable Default", "Set MetaSound Variable Default"), bPostTransaction);
 	Builder.CastDocumentObjectChecked<UObject>().Modify();
 
+	// Use the default page ID here as variables do *not* support paged defaults
+	// (they, as well as their mutator node which has a matching default, exist in each paged graph).
 	FMetasoundFrontendLiteral DefaultLiteral;
-	Literal->TryFindDefault(DefaultLiteral, &Builder.GetBuildPageID());
+	Literal->TryFindDefault(DefaultLiteral);
 
-	// Swap once variable support is added to document builder API
-	Builder.SetGraphVariableDefault(GetMemberName(), DefaultLiteral);
+	// Page ID is passed along to the builder from here because the builder needs the current BuildPageID to access the appropriate in-graph variable
+	// (variables can have the same IDs/names in different paged graphs).
+	Builder.SetGraphVariableDefault(GetMemberName(), DefaultLiteral, InPageID);
 
-	if (const FMetasoundFrontendNode* MutatorNode = Builder.FindNode(Variable->MutatorNodeID))
+	if (const FMetasoundFrontendNode* MutatorNode = Builder.FindNode(Variable->MutatorNodeID, InPageID))
 	{
 		const FMetasoundFrontendVertex* Input = MutatorNode->Interface.Inputs.FindByPredicate([](const FMetasoundFrontendVertex& Vertex) { return Vertex.Name == METASOUND_GET_PARAM_NAME(InputData); });
 		if (ensure(Input))
diff --git a/Engine/Plugins/Runtime/Metasound/Source/MetasoundEditor/Private/MetasoundEditorGraphBuilder.cpp b/Engine/Plugins/Runtime/Metasound/Source/MetasoundEditor/Private/MetasoundEditorGraphBuilder.cpp
--- a/Engine/Plugins/Runtime/Metasound/Source/MetasoundEditor/Private/MetasoundEditorGraphBuilder.cpp
+++ b/Engine/Plugins/Runtime/Metasound/Source/MetasoundEditor/Private/MetasoundEditorGraphBuilder.cpp
@@ -2764,6 +2764,11 @@ namespace Metasound
 				}
 			}
 
+			for (TObjectPtr<UMetasoundEditorGraphVariable> Variable : OutGraph.Variables)
+			{
+				bEditorGraphModified |= Variable->Synchronize();
+			}
+
 			// Remove empty entries
 			bEditorGraphModified |= OutGraph.Inputs.RemoveAllSwap([](const TObjectPtr<UMetasoundEditorGraphInput>& Input) { return !Input; }) > 0;
 			bEditorGraphModified |= OutGraph.Outputs.RemoveAllSwap([](const TObjectPtr<UMetasoundEditorGraphOutput>& Output) { return !Output; }) > 0;
diff --git a/Engine/Plugins/Runtime/Metasound/Source/MetasoundFrontend/Private/MetasoundFrontendDocumentBuilder.cpp b/Engine/Plugins/Runtime/Metasound/Source/MetasoundFrontend/Private/MetasoundFrontendDocumentBuilder.cpp
--- a/Engine/Plugins/Runtime/Metasound/Source/MetasoundFrontend/Private/MetasoundFrontendDocumentBuilder.cpp
+++ b/Engine/Plugins/Runtime/Metasound/Source/MetasoundFrontend/Private/MetasoundFrontendDocumentBuilder.cpp
@@ -2122,7 +2122,7 @@ const TArray<FMetasoundFrontendClassInputDefault>* FMetaSoundFrontendDocumentBui
 					const Frontend::FNodeRegistryKey Key = Frontend::FNodeRegistryKey(Class->Metadata);
 					const Frontend::INodeTemplate* Template = Frontend::INodeTemplateRegistry::Get().FindTemplate(Key);
 					check(Template);
-					const FGuid PageID = InPageID ? *InPageID : Frontend::DefaultPageID;
+					const FGuid PageID = InPageID ? *InPageID : BuildPageID;
 					return Template->FindNodeClassInputDefaults(*this, PageID, InNodeID, InVertexName);
 				}
 				break;
@@ -4019,7 +4019,7 @@ bool FMetaSoundFrontendDocumentBuilder::SetGraphInputDefault(FName InputName, FM
 		FMetasoundFrontendClassInput& Input = Inputs[Index];
 		if (Frontend::IDataTypeRegistry::Get().IsLiteralTypeSupported(Input.TypeName, InDefaultLiteral.GetType()))
 		{
-			const FGuid PageID = InPageID ? *InPageID : Frontend::DefaultPageID;
+			const FGuid PageID = InPageID ? *InPageID : BuildPageID;
 			bool bFound = false;
 			Input.IterateDefaults([&bFound, &PageID, &InDefaultLiteral](const FGuid& InputPageID, FMetasoundFrontendLiteral& InputLiteral)
 			{
@@ -4339,7 +4339,7 @@ bool FMetaSoundFrontendDocumentBuilder::SetGraphVariableDefault(FName VariableNa
 {
 	using namespace Metasound;
 
-	const FGuid PageID = InPageID ? *InPageID : Frontend::DefaultPageID;
+	const FGuid PageID = InPageID ? *InPageID : BuildPageID;
 	FMetasoundFrontendDocument& Document = GetDocumentChecked();
 	FMetasoundFrontendGraph& Graph = Document.RootGraph.FindGraphChecked(PageID);
 
