Null deref fix

#rnx

diff --git a/Engine/Plugins/Runtime/Metasound/Source/MetasoundEditor/Private/MetasoundEditorGraphConnectionDrawingPolicy.cpp b/Engine/Plugins/Runtime/Metasound/Source/MetasoundEditor/Private/MetasoundEditorGraphConnectionDrawingPolicy.cpp
--- a/Engine/Plugins/Runtime/Metasound/Source/MetasoundEditor/Private/MetasoundEditorGraphConnectionDrawingPolicy.cpp
+++ b/Engine/Plugins/Runtime/Metasound/Source/MetasoundEditor/Private/MetasoundEditorGraphConnectionDrawingPolicy.cpp
@@ -484,7 +484,12 @@ namespace Metasound
 					const float AppTime = FPlatformTime::Seconds() - GStartTime;
 					FDrawConnectionSignalData SignalParams = { AppTime, LayerId, SplineReparamTable, SplineLength, ConnectionData };
 
-					const FName DataType = SignalParams.ConnectionData.OutputVertex->TypeName;
+					FName DataType;
+					if (const FMetasoundFrontendVertex* OutputVertex = SignalParams.ConnectionData.OutputVertex)
+					{
+						DataType = OutputVertex->TypeName;
+					}
+						
 					if (DataType == GetMetasoundDataTypeName<FAudioBuffer>())
 					{
 						SignalParams.SpacingFactor = DrawingPolicyPrivate::EnvelopeConnectionSpacingCVar;
