Disallow page editing in MetaSound restricted mode

#rnx

diff --git a/Engine/Plugins/Runtime/Metasound/Source/MetasoundEditor/Private/MetasoundEditor.cpp b/Engine/Plugins/Runtime/Metasound/Source/MetasoundEditor/Private/MetasoundEditor.cpp
--- a/Engine/Plugins/Runtime/Metasound/Source/MetasoundEditor/Private/MetasoundEditor.cpp
+++ b/Engine/Plugins/Runtime/Metasound/Source/MetasoundEditor/Private/MetasoundEditor.cpp
@@ -155,6 +155,12 @@ namespace Metasound
 
 		bool PageEditorEnabled(const FMetaSoundFrontendDocumentBuilder& Builder, bool bHasProjectPageValues, bool bPresetCanEditPageValues)
 		{
+			IMetasoundEditorModule& MetaSoundEditorModule = FModuleManager::GetModuleChecked<IMetasoundEditorModule>("MetaSoundEditor");
+			if (MetaSoundEditorModule.IsRestrictedMode())
+			{
+				return false;
+			}
+
 			const UMetaSoundSettings* Settings = GetDefault<UMetaSoundSettings>();
 			if (!Settings)
 			{
