Disable MetaSound versioning logging in restricted mode

#rnx

diff --git a/Engine/Plugins/Runtime/Metasound/Source/MetasoundEditor/Private/MetasoundEditorModule.cpp b/Engine/Plugins/Runtime/Metasound/Source/MetasoundEditor/Private/MetasoundEditorModule.cpp
--- a/Engine/Plugins/Runtime/Metasound/Source/MetasoundEditor/Private/MetasoundEditorModule.cpp
+++ b/Engine/Plugins/Runtime/Metasound/Source/MetasoundEditor/Private/MetasoundEditorModule.cpp
@@ -544,6 +544,8 @@ namespace Metasound
 			virtual void SetRestrictedMode(bool bInRestrictedMode) override
 			{
 				bIsRestrictedMode = bInRestrictedMode;
+				const bool bEnableLogging = !bIsRestrictedMode;
+				Frontend::DocumentTransform::SetVersioningLoggingEnabled(bEnableLogging);
 			}
 
 			void RegisterSettingsDelegates()
diff --git a/Engine/Plugins/Runtime/Metasound/Source/MetasoundFrontend/Private/MetasoundFrontendTransform.cpp b/Engine/Plugins/Runtime/Metasound/Source/MetasoundFrontend/Private/MetasoundFrontendTransform.cpp
--- a/Engine/Plugins/Runtime/Metasound/Source/MetasoundFrontend/Private/MetasoundFrontendTransform.cpp
+++ b/Engine/Plugins/Runtime/Metasound/Source/MetasoundFrontend/Private/MetasoundFrontendTransform.cpp
@@ -21,8 +21,12 @@ namespace Metasound
 {
 	namespace Frontend
 	{
+#define	METASOUND_VERSIONING_LOG(Verbosity, Format, ...) if (DocumentTransform::bVersioningLoggingEnabled) { UE_LOG(LogMetaSound, Verbosity, Format, ##__VA_ARGS__); }
+
 		namespace DocumentTransform
 		{
+			bool bVersioningLoggingEnabled = true;
+
 			void LogAutoUpdateWarning(const FString& LogMessage)
 			{
 				// These should eventually move back to warning on cook 
@@ -31,17 +35,22 @@ namespace Metasound
 				// blocking code checkins 
 				if (IsRunningCookCommandlet())
 				{
-					UE_LOG(LogMetaSound, Display, TEXT("%s"), *LogMessage);
+					METASOUND_VERSIONING_LOG(Display, TEXT("%s"), *LogMessage);
 				}
 				else
 				{
-					UE_LOG(LogMetaSound, Warning, TEXT("%s"), *LogMessage);
+					METASOUND_VERSIONING_LOG(Warning, TEXT("%s"), *LogMessage);
 				}
 			}
 
 #if WITH_EDITOR
 			FGetNodeDisplayNameProjection NodeDisplayNameProjection;
 
+			void SetVersioningLoggingEnabled(bool bEnabled)
+			{
+				bVersioningLoggingEnabled = bEnabled;
+			}
+
 			void RegisterNodeDisplayNameProjection(FGetNodeDisplayNameProjection&& InNameProjection)
 			{
 				NodeDisplayNameProjection = MoveTemp(InNameProjection);
@@ -85,7 +94,7 @@ namespace Metasound
 				const bool bFromInterfaceFound = IInterfaceRegistry::Get().FindInterface(GetInterfaceRegistryKey(Version), Interface);
 				if (!ensureAlways(bFromInterfaceFound))
 				{
-					UE_LOG(LogMetaSound, Error, TEXT("Failed to find interface '%s' to remove"), *Version.ToString());
+					METASOUND_VERSIONING_LOG(Error, TEXT("Failed to find interface '%s' to remove"), *Version.ToString());
 				}
 				return Interface;
 			});
@@ -96,7 +105,7 @@ namespace Metasound
 				const bool bToInterfaceFound = IInterfaceRegistry::Get().FindInterface(GetInterfaceRegistryKey(Version), Interface);
 				if (!ensureAlways(bToInterfaceFound))
 				{
-					UE_LOG(LogMetaSound, Error, TEXT("Failed to find interface '%s' to add"), *Version.ToString());
+					METASOUND_VERSIONING_LOG(Error, TEXT("Failed to find interface '%s' to add"), *Version.ToString());
 				}
 				return Interface;
 			});
@@ -537,7 +546,7 @@ namespace Metasound
 				{
 					if (bIsPreset)
 					{
-						UE_LOG(LogMetaSound, Error, TEXT("Auto-Updating preset '%s' failed: Referenced class '%s' missing."), *DebugAssetPath, *ClassMetadata.GetClassName().ToString());
+						METASOUND_VERSIONING_LOG(Error, TEXT("Auto-Updating preset '%s' failed: Referenced class '%s' missing."), *DebugAssetPath, *ClassMetadata.GetClassName().ToString());
 						return;
 					}
 				}
@@ -560,14 +569,14 @@ namespace Metasound
 				{
 					const FMetasoundFrontendVersionNumber UpdateVersion = ClassWithHighestMinorVersion.Metadata.GetVersion();
 
-					UE_LOG(LogMetaSound, Display, TEXT("Auto-Updating '%s' node class '%s': Newer version '%s' found."), *DebugAssetPath, *ClassMetadata.GetClassName().ToString(), *UpdateVersion.ToString());
+					METASOUND_VERSIONING_LOG(Display, TEXT("Auto-Updating '%s' node class '%s': Newer version '%s' found."), *DebugAssetPath, *ClassMetadata.GetClassName().ToString(), *UpdateVersion.ToString());
 
 					NodesToUpdate.Add(TPair<FNodeHandle, FMetasoundFrontendVersionNumber>(NodeHandle, UpdateVersion));
 				}
 				else if (InterfaceUpdates.ContainsChanges())
 				{
 					const FMetasoundFrontendVersionNumber UpdateVersion = ClassMetadata.GetVersion();
-					UE_LOG(LogMetaSound, Display, TEXT("Auto-Updating '%s' node class '%s (%s)': Interface change detected."), *DebugAssetPath, *ClassMetadata.GetClassName().ToString(), *UpdateVersion.ToString());
+					METASOUND_VERSIONING_LOG(Display, TEXT("Auto-Updating '%s' node class '%s (%s)': Interface change detected."), *DebugAssetPath, *ClassMetadata.GetClassName().ToString(), *UpdateVersion.ToString());
 
 					NodesToUpdate.Add(TPair<FNodeHandle, FMetasoundFrontendVersionNumber>(NodeHandle, UpdateVersion));
 				}
@@ -1004,5 +1013,6 @@ namespace Metasound
 		{
 			return false;
 		}
+#undef METASOUND_VERSIONING_LOG
 	} // namespace Frontend
 } // namespace Metasound
diff --git a/Engine/Plugins/Runtime/Metasound/Source/MetasoundFrontend/Public/MetasoundFrontendTransform.h b/Engine/Plugins/Runtime/Metasound/Source/MetasoundFrontend/Public/MetasoundFrontendTransform.h
--- a/Engine/Plugins/Runtime/Metasound/Source/MetasoundFrontend/Public/MetasoundFrontendTransform.h
+++ b/Engine/Plugins/Runtime/Metasound/Source/MetasoundFrontend/Public/MetasoundFrontendTransform.h
@@ -22,6 +22,7 @@ namespace Metasound::Frontend
 		using FGetNodeDisplayNameProjection = TFunction<FText(const FNodeHandle&)>;
 		using FGetNodeDisplayNameProjectionRef = TFunctionRef<FText(const FNodeHandle&)>;
 
+		METASOUNDFRONTEND_API void SetVersioningLoggingEnabled(bool bIsEnabled);
 		METASOUNDFRONTEND_API void RegisterNodeDisplayNameProjection(FGetNodeDisplayNameProjection&& InNameProjection);
 		METASOUNDFRONTEND_API FGetNodeDisplayNameProjectionRef GetNodeDisplayNameProjection();
 #endif // WITH_EDITOR
