More document versioing log supression in restricted mode
#rnx

diff --git a/Engine/Plugins/Runtime/Metasound/Source/MetasoundFrontend/Private/MetasoundFrontendDocumentVersioning.cpp b/Engine/Plugins/Runtime/Metasound/Source/MetasoundFrontend/Private/MetasoundFrontendDocumentVersioning.cpp
--- a/Engine/Plugins/Runtime/Metasound/Source/MetasoundFrontend/Private/MetasoundFrontendDocumentVersioning.cpp
+++ b/Engine/Plugins/Runtime/Metasound/Source/MetasoundFrontend/Private/MetasoundFrontendDocumentVersioning.cpp
@@ -12,6 +12,7 @@
 #include "MetasoundFrontendDocumentController.h"
 #include "MetasoundFrontendRegistries.h"
 #include "MetasoundFrontendSearchEngine.h"
+#include "MetasoundFrontendTransform.h"
 #include "MetasoundLog.h"
 #include "MetasoundTrace.h"
 #include "Misc/App.h"
@@ -19,6 +20,8 @@
 
 namespace Metasound::Frontend
 {
+#define METASOUND_VERSIONING_LOG(Verbosity, Format, ...) if (DocumentTransform::GetVersioningLoggingEnabled()) { UE_LOG(LogMetaSound, Verbosity, Format, ##__VA_ARGS__); }
+
 	namespace VersioningPrivate
 	{
 		class FMigratePagePropertiesTransform : public FMetaSoundFrontendDocumentBuilder::IPropertyVersionTransform
@@ -222,7 +225,7 @@ PRAGMA_ENABLE_DEPRECATION_WARNINGS
 					}
 				}
 	#else
-				UE_LOG(LogMetaSound, Error, TEXT("Asset '%s' at '%s' must be saved with editor enabled in order to version document to target version '%s'."), *Name.ToString(), *Path, *GetTargetVersion().ToString());
+				METASOUND_VERSIONING_LOG(Error, TEXT("Asset '%s' at '%s' must be saved with editor enabled in order to version document to target version '%s'."), *Name.ToString(), *Path, *GetTargetVersion().ToString());
 	#endif // !WITH_EDITOR
 			}
 		};
@@ -257,7 +260,7 @@ PRAGMA_ENABLE_DEPRECATION_WARNINGS
 				Metadata.SetDisplayName(FText::FromString(Name.ToString()));
 				InDocument->GetRootGraph()->SetGraphMetadata(Metadata);
 	#else
-				UE_LOG(LogMetaSound, Error, TEXT("Asset '%s' at '%s' must be saved with editor enabled in order to version document to target version '%s'."), *Name.ToString(), *Path, *GetTargetVersion().ToString());
+				METASOUND_VERSIONING_LOG(Error, TEXT("Asset '%s' at '%s' must be saved with editor enabled in order to version document to target version '%s'."), *Name.ToString(), *Path, *GetTargetVersion().ToString());
 	#endif // !WITH_EDITOR
 			}
 		};
@@ -336,14 +339,14 @@ PRAGMA_ENABLE_DEPRECATION_WARNINGS
 
 					if (const FMetasoundFrontendInterface* Interface = FindMostSimilarInterfaceSupportingEnvironment(RootGraph, Dependencies, Subgraphs, CandidateInterfaces))
 					{
-						UE_LOG(LogMetaSound, Display, TEXT("Assigned interface [InterfaceVersion:%s] to document [RootGraphClassName:%s]"),
+						METASOUND_VERSIONING_LOG(Display, TEXT("Assigned interface [InterfaceVersion:%s] to document [RootGraphClassName:%s]"),
 							*Interface->Version.ToString(), *RootGraph.Metadata.GetClassName().ToString());
 
 						InDocument->AddInterfaceVersion(Interface->Version);
 					}
 					else
 					{
-						UE_LOG(LogMetaSound, Warning, TEXT("Failed to find interface for document [RootGraphClassName:%s]"),
+						METASOUND_VERSIONING_LOG(Warning, TEXT("Failed to find interface for document [RootGraphClassName:%s]"),
 							*RootGraph.Metadata.GetClassName().ToString());
 					}
 				}
@@ -381,7 +384,7 @@ PRAGMA_ENABLE_DEPRECATION_WARNINGS
 					InDocument->GetRootGraph()->SetGraphMetadata(NewMetadata);
 				}
 	#else
-				UE_LOG(LogMetaSound, Error, TEXT("Asset '%s' at '%s' must be saved with editor enabled in order to version document to target version '%s'."), *Name.ToString(), *Path, *GetTargetVersion().ToString());
+				METASOUND_VERSIONING_LOG(Error, TEXT("Asset '%s' at '%s' must be saved with editor enabled in order to version document to target version '%s'."), *Name.ToString(), *Path, *GetTargetVersion().ToString());
 	#endif // !WITH_EDITOR
 			}
 		};
@@ -455,7 +458,7 @@ PRAGMA_ENABLE_DEPRECATION_WARNINGS
 				InDocument->GetRootGraph()->IterateNodes(RenameTransform, EMetasoundFrontendClassType::Input);
 				InDocument->GetRootGraph()->IterateNodes(RenameTransform, EMetasoundFrontendClassType::Output);
 	#else
-				UE_LOG(LogMetaSound, Error, TEXT("Asset '%s' at '%s' must be saved with editor enabled in order to version document to target version '%s'."), *Name.ToString(), *Path, *GetTargetVersion().ToString());
+				METASOUND_VERSIONING_LOG(Error, TEXT("Asset '%s' at '%s' must be saved with editor enabled in order to version document to target version '%s'."), *Name.ToString(), *Path, *GetTargetVersion().ToString());
 	#endif // !WITH_EDITOR
 			}
 		};
@@ -542,7 +545,7 @@ PRAGMA_ENABLE_DEPRECATION_WARNINGS
 
 				InDocument->SetRootGraphClass(MoveTemp(RootGraphClass));
 	#else
-			UE_LOG(LogMetaSound, Error, TEXT("Asset '%s' at '%s' must be saved with editor enabled in order to version document to target version '%s'."), *Name.ToString(), *Path, *GetTargetVersion().ToString());
+			METASOUND_VERSIONING_LOG(Error, TEXT("Asset '%s' at '%s' must be saved with editor enabled in order to version document to target version '%s'."), *Name.ToString(), *Path, *GetTargetVersion().ToString());
 	#endif // !WITH_EDITOR
 			}
 		};
@@ -576,7 +579,7 @@ PRAGMA_ENABLE_DEPRECATION_WARNINGS
 				RootGraphClass.Metadata.SetDisplayName(FText());
 				InDocument->SetRootGraphClass(MoveTemp(RootGraphClass));
 	#else
-				UE_LOG(LogMetaSound, Error, TEXT("Asset '%s' at '%s' must be saved with editor enabled in order to version document to target version '%s'."), *Name.ToString(), *Path, *GetTargetVersion().ToString());
+				METASOUND_VERSIONING_LOG(Error, TEXT("Asset '%s' at '%s' must be saved with editor enabled in order to version document to target version '%s'."), *Name.ToString(), *Path, *GetTargetVersion().ToString());
 	#endif // !WITH_EDITOR
 			}
 		};
@@ -652,13 +655,16 @@ PRAGMA_ENABLE_DEPRECATION_WARNINGS
 
 				if (IsRunningCookCommandlet())
 				{
-					UE_LOG(LogMetaSound, Display, TEXT("Resave recommended: Asset '%s' at '%s' skipped migrated editor data/creation of input template nodes during cook to target document version '%s'."), *Name.ToString(), *Path->ToString(), *GetTargetVersion().ToString());
+					if (DocumentTransform::GetVersioningLoggingEnabled())
+					{
+						METASOUND_VERSIONING_LOG(Display, TEXT("Resave recommended: Asset '%s' at '%s' skipped migrated editor data/creation of input template nodes during cook to target document version '%s'."), *Name.ToString(), *Path->ToString(), *GetTargetVersion().ToString());
+					}
 				}
 				else
 				{
 					FMigratePagePropertiesTransform().Transform(OutBuilder);
 					OutBuilder.GetMetasoundAsset().MigrateEditorGraph(OutBuilder);
-					UE_LOG(LogMetaSound, Display, TEXT("Resave recommended: Asset '%s' at '%s' successfully migrated editor data in target document version '%s'."), *Name.ToString(), *Path->ToString(), *GetTargetVersion().ToString());
+					METASOUND_VERSIONING_LOG(Display, TEXT("Resave recommended: Asset '%s' at '%s' successfully migrated editor data in target document version '%s'."), *Name.ToString(), *Path->ToString(), *GetTargetVersion().ToString());
 				}
 			}
 		};
@@ -790,11 +796,12 @@ PRAGMA_ENABLE_DEPRECATION_WARNINGS
 			if (bWasUpdated)
 			{
 				const FMetasoundFrontendVersionNumber& NewVersionNumber = Builder.GetConstDocumentChecked().Metadata.Version.Number;
-				UE_LOG(LogMetaSound, Verbose, TEXT("MetaSound at '%s' Document Versioned: '%s' --> '%s'"), *Path, *InitVersionNumber.ToString(), *NewVersionNumber.ToString());
+				METASOUND_VERSIONING_LOG(Verbose, TEXT("MetaSound at '%s' Document Versioned: '%s' --> '%s'"), *Path, *InitVersionNumber.ToString(), *NewVersionNumber.ToString());
 			}
 		}
 
 		return bWasUpdated;
 	}
+#undef METASOUND_VERSIONING_LOG
 } // namespace Metasound::Frontend
 #endif // WITH_EDITORONLY_DATA
\ No newline at end of file
diff --git a/Engine/Plugins/Runtime/Metasound/Source/MetasoundFrontend/Private/MetasoundFrontendTransform.cpp b/Engine/Plugins/Runtime/Metasound/Source/MetasoundFrontend/Private/MetasoundFrontendTransform.cpp
--- a/Engine/Plugins/Runtime/Metasound/Source/MetasoundFrontend/Private/MetasoundFrontendTransform.cpp
+++ b/Engine/Plugins/Runtime/Metasound/Source/MetasoundFrontend/Private/MetasoundFrontendTransform.cpp
@@ -46,6 +46,11 @@ namespace Metasound
 #if WITH_EDITOR
 			FGetNodeDisplayNameProjection NodeDisplayNameProjection;
 
+			bool GetVersioningLoggingEnabled()
+			{
+				return bVersioningLoggingEnabled;
+			}
+
 			void SetVersioningLoggingEnabled(bool bEnabled)
 			{
 				bVersioningLoggingEnabled = bEnabled;
diff --git a/Engine/Plugins/Runtime/Metasound/Source/MetasoundFrontend/Public/MetasoundFrontendTransform.h b/Engine/Plugins/Runtime/Metasound/Source/MetasoundFrontend/Public/MetasoundFrontendTransform.h
--- a/Engine/Plugins/Runtime/Metasound/Source/MetasoundFrontend/Public/MetasoundFrontendTransform.h
+++ b/Engine/Plugins/Runtime/Metasound/Source/MetasoundFrontend/Public/MetasoundFrontendTransform.h
@@ -22,6 +22,7 @@ namespace Metasound::Frontend
 		using FGetNodeDisplayNameProjection = TFunction<FText(const FNodeHandle&)>;
 		using FGetNodeDisplayNameProjectionRef = TFunctionRef<FText(const FNodeHandle&)>;
 
+		METASOUNDFRONTEND_API bool GetVersioningLoggingEnabled();
 		METASOUNDFRONTEND_API void SetVersioningLoggingEnabled(bool bIsEnabled);
 		METASOUNDFRONTEND_API void RegisterNodeDisplayNameProjection(FGetNodeDisplayNameProjection&& InNameProjection);
 		METASOUNDFRONTEND_API FGetNodeDisplayNameProjectionRef GetNodeDisplayNameProjection();
