Do not prime search engine when running a commandlet
#rnx

diff --git a/Engine/Plugins/Runtime/Metasound/Source/MetasoundFrontend/Private/MetasoundFrontendRegistryContainerImpl.cpp b/Engine/Plugins/Runtime/Metasound/Source/MetasoundFrontend/Private/MetasoundFrontendRegistryContainerImpl.cpp
--- a/Engine/Plugins/Runtime/Metasound/Source/MetasoundFrontend/Private/MetasoundFrontendRegistryContainerImpl.cpp
+++ b/Engine/Plugins/Runtime/Metasound/Source/MetasoundFrontend/Private/MetasoundFrontendRegistryContainerImpl.cpp
@@ -342,8 +342,11 @@ namespace Metasound::Frontend
 			LazyInitCommands.Empty();
 		}
 
-		// Prime search engine after bulk registration.
-		ISearchEngine::Get().Prime();
+		if (!IsRunningCommandlet())
+		{
+			// Prime search engine after bulk registration.
+			ISearchEngine::Get().Prime();
+		}
 	}
 
 	bool FRegistryContainerImpl::EnqueueInitCommand(TUniqueFunction<void()>&& InFunc)
