Do not prime the MetaSound Node Class Registry with asset-defined graph nodes in Restricted Editor mode
#rnx
#tests Force restricted mode locally, witness prime not getting kicked on load of MetaSound Asset Editor

diff --git a/Engine/Plugins/Runtime/Metasound/Source/MetasoundEditor/Private/MetasoundEditor.cpp b/Engine/Plugins/Runtime/Metasound/Source/MetasoundEditor/Private/MetasoundEditor.cpp
--- a/Engine/Plugins/Runtime/Metasound/Source/MetasoundEditor/Private/MetasoundEditor.cpp
+++ b/Engine/Plugins/Runtime/Metasound/Source/MetasoundEditor/Private/MetasoundEditor.cpp
@@ -1145,11 +1145,15 @@ namespace Metasound
 			check(ObjectToEdit);
 			checkf(IMetasoundUObjectRegistry::Get().IsRegisteredClass(ObjectToEdit), TEXT("Object passed in was not registered as a valid metasound interface!"));
 
-			IMetasoundEngineModule& MetaSoundEngineModule = FModuleManager::GetModuleChecked<IMetasoundEngineModule>("MetaSoundEngine");
-			bPrimingRegistry = MetaSoundEngineModule.GetNodeClassRegistryPrimeStatus() <= ENodeClassRegistryPrimeStatus::InProgress;
-			if (MetaSoundEngineModule.GetNodeClassRegistryPrimeStatus() < ENodeClassRegistryPrimeStatus::InProgress)
+			IMetasoundEditorModule& MetaSoundEditorModule = FModuleManager::GetModuleChecked<IMetasoundEditorModule>("MetaSoundEditor");
+			if (!MetaSoundEditorModule.IsRestrictedMode())
 			{
-				MetaSoundEngineModule.PrimeAssetRegistryAsync();
+				IMetasoundEngineModule& MetaSoundEngineModule = FModuleManager::GetModuleChecked<IMetasoundEngineModule>("MetaSoundEngine");
+				if (MetaSoundEngineModule.GetNodeClassRegistryPrimeStatus() < ENodeClassRegistryPrimeStatus::InProgress)
+				{
+					MetaSoundEngineModule.PrimeAssetRegistryAsync();
+				}
+				bPrimingRegistry = MetaSoundEngineModule.GetNodeClassRegistryPrimeStatus() <= ENodeClassRegistryPrimeStatus::InProgress;
 			}
 
 			// Support undo/redo
