Change MetaSound Create Preset asset action to use a create asset dialog
- Hide Convert From Preset in restricted mode

diff --git a/Engine/Plugins/Runtime/Metasound/Source/MetasoundEditor/Private/MetasoundAssetDefinitions.cpp b/Engine/Plugins/Runtime/Metasound/Source/MetasoundEditor/Private/MetasoundAssetDefinitions.cpp
--- a/Engine/Plugins/Runtime/Metasound/Source/MetasoundEditor/Private/MetasoundAssetDefinitions.cpp
+++ b/Engine/Plugins/Runtime/Metasound/Source/MetasoundEditor/Private/MetasoundAssetDefinitions.cpp
@@ -397,7 +397,7 @@ void UAssetDefinition_MetaSoundSource::GetAssetActionButtonExtensions(const FAss
 
 namespace MenuExtension_MetaSoundSourceTemplate
 {
-	template <typename TClass>
+	template <typename TClass, typename TFactoryClass>
 	void ExecuteCreateMetaSoundPreset(const FToolMenuContext& MenuContext)
 	{
 		if (const UContentBrowserAssetContextMenuContext* Context = UContentBrowserAssetContextMenuContext::FindContextWithAssets(MenuContext))
@@ -411,27 +411,18 @@ namespace MenuExtension_MetaSoundSourceTemplate
 
 				IAssetTools::Get().CreateUniqueAssetName(ReferencedMetaSound->GetOutermost()->GetName(), TEXT("_Preset"), PackagePath, AssetName);
 
-				EMetaSoundBuilderResult BuilderResult;
-				UMetaSoundBuilderBase& Builder = UMetaSoundBuilderSubsystem::GetChecked().CreatePresetBuilder(FName(AssetName), ReferencedMetaSound, BuilderResult);
+				UMetaSoundBaseFactory* Factory = NewObject<TFactoryClass>();
+				Factory->ReferencedMetaSoundObject = ReferencedMetaSound;
 
-				if (BuilderResult != EMetaSoundBuilderResult::Succeeded)
-				{
-					UE_LOG(LogMetaSound, Error, TEXT("Error creating a preset builder for MetaSound '%s'"), *AssetName);
-					return;
-				}
+				UObject* NewMetaSound = IAssetTools::Get().CreateAssetWithDialog(AssetName, FPackageName::GetLongPackagePath(PackagePath), Factory->GetSupportedClass(), Factory);
 
-				UMetaSoundEditorSubsystem& MetaSoundEditorSubsystem = UMetaSoundEditorSubsystem::GetChecked();
-				TScriptInterface<IMetaSoundDocumentInterface> NewMetaSound = MetaSoundEditorSubsystem.BuildToAsset(&Builder, MetaSoundEditorSubsystem.GetDefaultAuthor(), AssetName, FPackageName::GetLongPackagePath(PackagePath), BuilderResult);
-				if (BuilderResult == EMetaSoundBuilderResult::Succeeded)
+				if (ensure(NewMetaSound))
 				{
-					if (ensure(NewMetaSound))
-					{
-						FGraphBuilder::RegisterGraphWithFrontend(*NewMetaSound.GetObject());
-					}
+					FGraphBuilder::RegisterGraphWithFrontend(*NewMetaSound);
 				}
 				else
 				{
-					UE_LOG(LogMetaSound, Error, TEXT("Error building to asset when creating preset '%s'"), *AssetName);
+					UE_LOG(LogMetaSound, Error, TEXT("Error creating new asset with factory '%s' when creating preset '%s'"), *Factory->GetName(), *AssetName);
 				}
 			}
 		}
@@ -498,7 +489,7 @@ namespace MenuExtension_MetaSoundSourceTemplate
 							const FSlateIcon Icon = Metasound::Editor::Style::CreateSlateIcon("ClassIcon.MetasoundSource");
 
 							FToolUIAction UIAction;
-							UIAction.ExecuteAction = FToolMenuExecuteAction::CreateStatic(&ExecuteCreateMetaSoundPreset<UMetaSoundSource>);
+							UIAction.ExecuteAction = FToolMenuExecuteAction::CreateStatic(&ExecuteCreateMetaSoundPreset<UMetaSoundSource, UMetaSoundSourceFactory>);
 							InSection.AddMenuEntry("MetaSoundSource_CreatePreset", Label, ToolTip, Icon, UIAction);
 						}
 					}));
@@ -517,7 +508,7 @@ namespace MenuExtension_MetaSoundSourceTemplate
 					const FSlateIcon Icon = Metasound::Editor::Style::CreateSlateIcon("ClassIcon.MetasoundPatch");
 
 					FToolUIAction UIAction;
-					UIAction.ExecuteAction = FToolMenuExecuteAction::CreateStatic(&ExecuteCreateMetaSoundPreset<UMetaSoundPatch>);
+					UIAction.ExecuteAction = FToolMenuExecuteAction::CreateStatic(&ExecuteCreateMetaSoundPreset<UMetaSoundPatch, UMetaSoundFactory>);
 					InSection.AddMenuEntry("MetaSoundPatch_CreatePreset", Label, ToolTip, Icon, UIAction);
 				}));
 			}
diff --git a/Engine/Plugins/Runtime/Metasound/Source/MetasoundEditor/Private/MetasoundEditor.cpp b/Engine/Plugins/Runtime/Metasound/Source/MetasoundEditor/Private/MetasoundEditor.cpp
--- a/Engine/Plugins/Runtime/Metasound/Source/MetasoundEditor/Private/MetasoundEditor.cpp
+++ b/Engine/Plugins/Runtime/Metasound/Source/MetasoundEditor/Private/MetasoundEditor.cpp
@@ -2204,7 +2204,8 @@ namespace Metasound
 // 							"ExportMetasound"
 // 						);
 
-						if (!IsGraphEditable())
+						const IMetasoundEditorModule& MetaSoundEditorModule = FModuleManager::GetModuleChecked<IMetasoundEditorModule>("MetaSoundEditor");
+						if (!IsGraphEditable() && !MetaSoundEditorModule.IsRestrictedMode())
 						{
 							ToolbarBuilder.AddToolBarButton
  							(
diff --git a/Engine/Plugins/Runtime/Metasound/Source/MetasoundFrontend/Private/MetasoundFrontendDocumentBuilder.cpp b/Engine/Plugins/Runtime/Metasound/Source/MetasoundFrontend/Private/MetasoundFrontendDocumentBuilder.cpp
--- a/Engine/Plugins/Runtime/Metasound/Source/MetasoundFrontend/Private/MetasoundFrontendDocumentBuilder.cpp
+++ b/Engine/Plugins/Runtime/Metasound/Source/MetasoundFrontend/Private/MetasoundFrontendDocumentBuilder.cpp
@@ -1681,12 +1681,13 @@ bool FMetaSoundFrontendDocumentBuilder::ConvertFromPreset()
 	return false;
 }
 
-bool FMetaSoundFrontendDocumentBuilder::ConvertToPreset(const FMetasoundFrontendDocument& InReferencedDocument, TSharedRef<Metasound::Frontend::FDocumentModifyDelegates> ModifyDelegates)
+bool FMetaSoundFrontendDocumentBuilder::ConvertToPreset(const FMetasoundFrontendDocument& InReferencedDocument, TSharedPtr<Metasound::Frontend::FDocumentModifyDelegates> ModifyDelegates)
 {
 	using namespace Metasound;
 	using namespace Metasound::Frontend;
 
-	ClearDocument(ModifyDelegates);
+	TSharedRef<FDocumentModifyDelegates> ModifyDelegatesRef = ModifyDelegates.IsValid() ? ModifyDelegates->AsShared() : MakeShared<FDocumentModifyDelegates>(InReferencedDocument);
+	ClearDocument(ModifyDelegatesRef);
 
 	FMetasoundFrontendGraphClass& PresetAssetRootGraph = GetDocumentChecked().RootGraph;
 	PresetAssetRootGraph.IterateGraphPages([](FMetasoundFrontendGraph& PresetAssetGraph)
diff --git a/Engine/Plugins/Runtime/Metasound/Source/MetasoundFrontend/Public/MetasoundFrontendDocumentBuilder.h b/Engine/Plugins/Runtime/Metasound/Source/MetasoundFrontend/Public/MetasoundFrontendDocumentBuilder.h
--- a/Engine/Plugins/Runtime/Metasound/Source/MetasoundFrontend/Public/MetasoundFrontendDocumentBuilder.h
+++ b/Engine/Plugins/Runtime/Metasound/Source/MetasoundFrontend/Public/MetasoundFrontendDocumentBuilder.h
@@ -205,7 +205,7 @@ public:
 	bool ContainsNode(const FGuid& InNodeID, const FGuid* InPageID = nullptr) const;
 
 	bool ConvertFromPreset();
-	bool ConvertToPreset(const FMetasoundFrontendDocument& InReferencedDocument, TSharedRef<Metasound::Frontend::FDocumentModifyDelegates> ModifyDelegates = { });
+	bool ConvertToPreset(const FMetasoundFrontendDocument& InReferencedDocument, TSharedPtr<Metasound::Frontend::FDocumentModifyDelegates> ModifyDelegates = { });
 
 	const FMetasoundFrontendNode* DuplicateGraphInput(const FMetasoundFrontendClassInput& InClassInput, const FName InName, const FGuid* InPageID = nullptr);
 	const FMetasoundFrontendNode* DuplicateGraphOutput(const FMetasoundFrontendClassOutput& InClassOutput, const FName InName, const FGuid* InPageID = nullptr);
