Speculative fix for cook crash due to builder not being initialized for referenced graph that has had its default graph cooked away.
- Guard against initializing a builder with an invalid build page ID set at cook/runtime

#rnx

diff --git a/Engine/Plugins/Runtime/Metasound/Source/MetasoundFrontend/Private/MetasoundFrontendDocumentBuilder.cpp b/Engine/Plugins/Runtime/Metasound/Source/MetasoundFrontend/Private/MetasoundFrontendDocumentBuilder.cpp
--- a/Engine/Plugins/Runtime/Metasound/Source/MetasoundFrontend/Private/MetasoundFrontendDocumentBuilder.cpp
+++ b/Engine/Plugins/Runtime/Metasound/Source/MetasoundFrontend/Private/MetasoundFrontendDocumentBuilder.cpp
@@ -784,6 +784,7 @@ void UMetaSoundBuilderDocument::OnFinishActiveBuilder()
 
 FMetaSoundFrontendDocumentBuilder::FMetaSoundFrontendDocumentBuilder(TScriptInterface<IMetaSoundDocumentInterface> InDocumentInterface, TSharedPtr<Metasound::Frontend::FDocumentModifyDelegates> InDocumentDelegates, bool bPrimeCache)
 	: DocumentInterface(InDocumentInterface)
+	, BuildPageID(Metasound::Frontend::DefaultPageID)
 {
 	BeginBuilding(InDocumentDelegates, bPrimeCache);
 }
@@ -3013,6 +3014,38 @@ void FMetaSoundFrontendDocumentBuilder::BeginBuilding(TSharedPtr<Metasound::Fron
 {
 	using namespace Metasound::Frontend;
 
+	if (DocumentInterface)
+	{
+		// Potentially at cook and runtime, the default graph may have been cooked away,
+		// so initialize build page to a valid page ID if possible. On initial construction,
+		// it may be possible the default graph has yet to be initialized, so don't error if
+		// default page graph has yet to be created (BuildPageID is then left as ctor Default
+		// PageID).
+		const FMetasoundFrontendDocument& Document = DocumentInterface->GetConstDocument();
+		bool bPageIDSet = false;
+		Document.RootGraph.IterateGraphPages([&bPageIDSet, this](const FMetasoundFrontendGraph& Graph)
+		{
+			if (Graph.PageID == Metasound::Frontend::DefaultPageID)
+			{
+				bPageIDSet = true;
+				BuildPageID = Graph.PageID;
+			}
+			else if (!bPageIDSet)
+			{
+				BuildPageID = Graph.PageID;
+			}
+		});
+
+		if (UE_LOG_ACTIVE(LogMetaSound, VeryVerbose))
+		{
+			FTopLevelAssetPath DebugPath;
+			if (DebugPath.TrySetPath(DocumentInterface.GetObject()))
+			{
+				UE_LOG(LogMetaSound, VeryVerbose, TEXT("MetaSoundFrontendDocumentBuilder::BeginBuilding for asset '%s': BuildPageID initialized to '%s'"), *DebugPath.ToString(), *BuildPageID.ToString());
+			}
+		}
+	}
+
 	if (Delegates.IsValid())
 	{
 		DocumentDelegates = Delegates;
diff --git a/Engine/Plugins/Runtime/Metasound/Source/MetasoundFrontend/Private/MetasoundFrontendGraphController.cpp b/Engine/Plugins/Runtime/Metasound/Source/MetasoundFrontend/Private/MetasoundFrontendGraphController.cpp
--- a/Engine/Plugins/Runtime/Metasound/Source/MetasoundFrontend/Private/MetasoundFrontendGraphController.cpp
+++ b/Engine/Plugins/Runtime/Metasound/Source/MetasoundFrontend/Private/MetasoundFrontendGraphController.cpp
@@ -1262,7 +1262,8 @@ namespace Metasound
 		{
 			if (const FMetasoundFrontendClassInput* Desc = FindInputDescriptionWithVertexID(InVertexID))
 			{
-				const FMetasoundFrontendLiteral& DefaultLiteral = Desc->FindConstDefaultChecked(Frontend::DefaultPageID);
+				const FGuid TargetPageID = IDocumentBuilderRegistry::GetChecked().ResolveTargetPageID(*Desc);
+				const FMetasoundFrontendLiteral& DefaultLiteral = Desc->FindConstDefaultChecked(TargetPageID);
 				return DefaultLiteral;
 			}
 			return FMetasoundFrontendLiteral{};
@@ -1274,8 +1275,8 @@ namespace Metasound
 			{
 				if (ensure(IDataTypeRegistry::Get().IsLiteralTypeSupported(Desc->TypeName, InLiteral.GetType())))
 				{
-					FMetasoundFrontendLiteral& DefaultLiteral = Desc->FindDefaultChecked(Frontend::DefaultPageID);
-					DefaultLiteral = InLiteral;
+					const FGuid TargetPageID = IDocumentBuilderRegistry::GetChecked().ResolveTargetPageID(*Desc);
+					const FMetasoundFrontendLiteral& DefaultLiteral = Desc->FindConstDefaultChecked(TargetPageID);
 					return true;
 				}
 				else
diff --git a/Engine/Plugins/Runtime/Metasound/Source/MetasoundFrontend/Private/MetasoundFrontendTransform.cpp b/Engine/Plugins/Runtime/Metasound/Source/MetasoundFrontend/Private/MetasoundFrontendTransform.cpp
--- a/Engine/Plugins/Runtime/Metasound/Source/MetasoundFrontend/Private/MetasoundFrontendTransform.cpp
+++ b/Engine/Plugins/Runtime/Metasound/Source/MetasoundFrontend/Private/MetasoundFrontendTransform.cpp
@@ -7,6 +7,7 @@
 #include "NodeTemplates/MetasoundFrontendNodeTemplateInput.h"
 #include "MetasoundAccessPtr.h"
 #include "MetasoundAssetBase.h"
+#include "MetasoundDocumentInterface.h"
 #include "MetasoundFrontendDocument.h"
 #include "MetasoundFrontendDocumentBuilder.h"
 #include "MetasoundFrontendDocumentController.h"
@@ -600,6 +601,15 @@ namespace Metasound
 			{
 				if (bIsPreset)
 				{
+#if WITH_EDITORONLY_DATA
+					// This call to FindOrBeginBuilding is a bit of a hack to guarantee instantiation of the referenced graph's builder while
+					// using the deprecated controllers to rebuild the preset root graph.  This is needed in case the referenced graph implements
+					// page data as the builder is required to resolve the correct graph data used to rebuild the preset.  Once the rebuild transform
+					// is migrated to the Builder API, this will no longer be necessary.
+					TScriptInterface<IMetaSoundDocumentInterface> RefDocInterface = PresetReferencedMetaSoundAsset->GetOwningAsset();
+					FMetaSoundFrontendDocumentBuilder& ReferenceBuilder = IDocumentBuilderRegistry::GetChecked().FindOrBeginBuilding(RefDocInterface);
+#endif // WITH_EDITORONLY_DATA
+
 					bDidEdit |= FRebuildPresetRootGraph(PresetReferencedMetaSoundAsset->GetDocumentHandle()).Transform(InDocument);
 					if (bDidEdit)
 					{
