Sync content browser focus on MetaSound preset asset creation

#rnx

diff --git a/Engine/Plugins/Runtime/Metasound/Source/MetasoundEditor/Private/MetasoundAssetDefinitions.cpp b/Engine/Plugins/Runtime/Metasound/Source/MetasoundEditor/Private/MetasoundAssetDefinitions.cpp
--- a/Engine/Plugins/Runtime/Metasound/Source/MetasoundEditor/Private/MetasoundAssetDefinitions.cpp
+++ b/Engine/Plugins/Runtime/Metasound/Source/MetasoundEditor/Private/MetasoundAssetDefinitions.cpp
@@ -403,6 +403,7 @@ namespace MenuExtension_MetaSoundSourceTemplate
 		if (const UContentBrowserAssetContextMenuContext* Context = UContentBrowserAssetContextMenuContext::FindContextWithAssets(MenuContext))
 		{
 			using namespace Metasound::Editor;
+			TArray<UObject*> ObjectsToSync;
 
 			for (TClass* ReferencedMetaSound : Context->LoadSelectedObjects<TClass>())
 			{
@@ -419,12 +420,20 @@ namespace MenuExtension_MetaSoundSourceTemplate
 				if (ensure(NewMetaSound))
 				{
 					FGraphBuilder::RegisterGraphWithFrontend(*NewMetaSound);
+					ObjectsToSync.Add(NewMetaSound);
 				}
 				else
 				{
 					UE_LOG(LogMetaSound, Error, TEXT("Error creating new asset with factory '%s' when creating preset '%s'"), *Factory->GetName(), *AssetName);
 				}
 			}
+
+			// Sync content browser to newly created assets
+			if (ObjectsToSync.Num() > 0)
+			{
+				FContentBrowserModule& ContentBrowserModule = FModuleManager::LoadModuleChecked<FContentBrowserModule>("ContentBrowser");
+				ContentBrowserModule.Get().SyncBrowserToAssets(ObjectsToSync);
+			}
 		}
 	}
 
