Do not validate class names during asset validation when running a cook commandlet.

#rnx

diff --git a/Engine/Plugins/Runtime/Metasound/Source/MetasoundEngine/Private/MetasoundEngineAsset.h b/Engine/Plugins/Runtime/Metasound/Source/MetasoundEngine/Private/MetasoundEngineAsset.h
--- a/Engine/Plugins/Runtime/Metasound/Source/MetasoundEngine/Private/MetasoundEngineAsset.h
+++ b/Engine/Plugins/Runtime/Metasound/Source/MetasoundEngine/Private/MetasoundEngineAsset.h
@@ -160,7 +160,13 @@ namespace Metasound::Engine
 			EDataValidationResult Result = EDataValidationResult::Valid;
 			if (MetasoundEngineModulePrivate::EnableMetaSoundEditorAssetValidation)
 			{
-				Result = IsClassNameUnique(Document, InOutContext);
+				// We cannot rely on the asset registry scan being complete during the call
+				// to IsDataValid(...) while running a cook commandlet. The IMetasoundAssetManager
+				// will still log errors on duplicate assets which will fail cook. 
+				if (!IsRunningCookCommandlet())
+				{
+					Result = IsClassNameUnique(Document, InOutContext);
+				}
 			}
 
 			const UMetaSoundSettings* Settings = GetDefault<UMetaSoundSettings>();
