[AudioWidgets] Added Audio Analyzer Rack - Uses tab spawners so that rack units can be shown, hidden, and reordered by the user. Includes extensibility so that third party plugins can register custom analyzer rack units.
[Metasound Editor] Using AudioWidgets Audio Analyzer Rack instead of custom code in the editor for managing display of analyzer widgets. Spectrum analyzer settings struct moved to be defined in AudioWidgets instead of Metasound Editor.

diff --git a/Engine/Plugins/Runtime/AudioWidgets/Source/AudioWidgets/Private/AudioAnalyzerRack.cpp b/Engine/Plugins/Runtime/AudioWidgets/Source/AudioWidgets/Private/AudioAnalyzerRack.cpp
new file mode 100644
--- /dev/null
+++ b/Engine/Plugins/Runtime/AudioWidgets/Source/AudioWidgets/Private/AudioAnalyzerRack.cpp
@@ -0,0 +1,285 @@
+// Copyright Epic Games, Inc. All Rights Reserved.
+
+#include "AudioAnalyzerRack.h"
+
+#include "AudioAnalyzerRackUnitRegistry.h"
+#include "AudioMeter.h"
+#include "AudioOscilloscope.h"
+#include "AudioSpectrumAnalyzer.h"
+#include "AudioWidgetsStyle.h"
+#include "Framework/Docking/LayoutService.h"
+#include "Framework/MultiBox/MultiBoxBuilder.h"
+#include "Widgets/Colors/SColorBlock.h"
+#include "Widgets/Docking/SDockTab.h"
+
+#define LOCTEXT_NAMESPACE "FAudioAnalyzerRack"
+
+namespace AudioWidgets
+{
+	FAudioAnalyzerRack::FAudioAnalyzerRack(const FRackConstructParams& Params)
+		: TabManagerLayoutName(Params.TabManagerLayoutName)
+	{
+		if (TabManagerLayoutName.IsNone())
+		{
+			TabManagerLayoutName = TEXT("AudioWidgets_FAudioAnalyzerRack_v0");
+		}
+
+		RackUnitConstructParams.StyleSet = (Params.StyleSet != nullptr) ? Params.StyleSet : &FAudioWidgetsStyle::Get();
+		RackUnitConstructParams.EditorSettingsClass = Params.EditorSettingsClass;
+	}
+
+	void FAudioAnalyzerRack::Init(int32 InNumChannels, Audio::FDeviceId InAudioDeviceId)
+	{
+		const EAudioBusChannels AudioBusChannels = AudioBusUtils::ConvertIntToEAudioBusChannels(InNumChannels);
+		if (RackUnitConstructParams.AudioBusInfo.AudioDeviceId != InAudioDeviceId || !AudioBus.IsValid() || AudioBus->AudioBusChannels != AudioBusChannels)
+		{
+			// Create a UAudioBus with the required num channels:
+			AudioBus = TStrongObjectPtr(NewObject<UAudioBus>());
+			AudioBus->AudioBusChannels = AudioBusChannels;
+
+			// Update cached AudioBusInfo:
+			RackUnitConstructParams.AudioBusInfo = { .AudioDeviceId = InAudioDeviceId, .AudioBus = AudioBus.Get() };
+
+			// Reinit any existing rack units:
+			for (auto&& KVPair : RackUnits)
+			{
+				KVPair.Value->SetAudioBusInfo(RackUnitConstructParams.AudioBusInfo);
+			}
+		}
+	}
+
+	void FAudioAnalyzerRack::DestroyAnalyzers()
+	{
+		RackUnits.Reset();
+	}
+
+	TSharedRef<SWidget> FAudioAnalyzerRack::CreateWidget(TSharedRef<SDockTab> DockTab)
+	{
+		// Create a TabManager:
+		TabManager = FGlobalTabmanager::Get()->NewTabManager(DockTab);
+		TabManager->SetOnPersistLayout(FTabManager::FOnPersistLayout::CreateStatic(&FAudioAnalyzerRack::SaveTabLayout));
+
+		// Register TabSpawners for all registered rack unit types:
+		const TSharedRef<FWorkspaceItem> AppMenuGroup = TabManager->AddLocalWorkspaceMenuCategory(LOCTEXT("AnalyzerRackGroupName", "Analyzers"));
+
+		TArray<const FAudioAnalyzerRackUnitTypeInfo*> RackUnitTypes;
+		FAudioAnalyzerRackUnitRegistry::Get().GetRegisteredRackUnitTypes(RackUnitTypes);
+		for (const FAudioAnalyzerRackUnitTypeInfo* RackUnitType : RackUnitTypes)
+		{
+			TabManager->RegisterTabSpawner(RackUnitType->TypeName, FOnSpawnTab::CreateLambda([this](const FSpawnTabArgs& Args)
+				{
+					const FName RackUnitTypeName = Args.GetTabId().TabType;
+					if (TSharedRef<IAudioAnalyzerRackUnit>* ExistingRackUnit = RackUnits.Find(RackUnitTypeName))
+					{
+						return (*ExistingRackUnit)->SpawnTab(Args);
+					}
+					else
+					{
+						const TSharedRef<IAudioAnalyzerRackUnit> RackUnit = RackUnits.Add(RackUnitTypeName, MakeRackUnit(RackUnitTypeName));
+						return RackUnit->SpawnTab(Args);
+					}
+				}))
+				.SetGroup(AppMenuGroup)
+				.SetDisplayName(RackUnitType->DisplayName)
+				.SetIcon(RackUnitType->Icon)
+				.SetCanSidebarTab(false);
+		}
+
+		// Create a ToolBar that can toggle visible analyzers:
+		const TSharedRef<FUICommandList> CommandList = MakeShared<FUICommandList>();
+		constexpr bool bForceSmallIcons = true;
+		FSlimHorizontalToolBarBuilder ToolBarBuilder(CommandList, FMultiBoxCustomization::None, nullptr, bForceSmallIcons);
+
+		const TSharedRef<FWorkspaceItem> LocalWorkspaceMenuRoot = TabManager->GetLocalWorkspaceMenuRoot();
+		for (const TSharedRef<FWorkspaceItem>& WorkspaceGroup : LocalWorkspaceMenuRoot->GetChildItems())
+		{
+			ToolBarBuilder.AddComboButton(
+				FUIAction(),
+				FOnGetContent::CreateStatic(&FAudioAnalyzerRack::MakeVisibleAnalyzersMenu, CommandList, WorkspaceGroup.ToWeakPtr(), TabManager.ToWeakPtr()),
+				WorkspaceGroup->GetDisplayName(),
+				WorkspaceGroup->GetTooltipText(),
+				WorkspaceGroup->GetIcon(),
+				bForceSmallIcons);
+		}
+
+		// Load saved tab layout, or create the default layout:
+		const TSharedRef<FTabManager::FLayout> TabLayout = LoadTabLayout();
+
+		// Create the actual SWidget, containing the toolbar and the rack unit docking tab layout:
+		const FLinearColor BackgroundColor = RackUnitConstructParams.StyleSet->GetColor("AudioAnalyzerRack.BackgroundColor");
+		const TSharedRef<SWidget> Widget = SNew(SVerticalBox)
+			+ SVerticalBox::Slot()
+			.AutoHeight()
+			[
+				SNew(SHorizontalBox)
+					+ SHorizontalBox::Slot()
+					.HAlign(EHorizontalAlignment::HAlign_Right)
+					[
+						ToolBarBuilder.MakeWidget()
+					]
+			]
+			+ SVerticalBox::Slot()
+			.FillHeight(1.0f)
+			[
+				SNew(SOverlay)
+					+ SOverlay::Slot()
+					[
+						SNew(SColorBlock)
+							.Color(BackgroundColor)
+					]
+					+ SOverlay::Slot()
+					[
+						TabManager->RestoreFrom(TabLayout, DockTab->GetParentWindow()).ToSharedRef()
+					]
+			];
+
+		return Widget;
+	}
+
+	UAudioBus* FAudioAnalyzerRack::GetAudioBus() const
+	{
+		return AudioBus.Get();
+	}
+
+	void FAudioAnalyzerRack::StartProcessing()
+	{
+		for (auto&& KVPair : RackUnits)
+		{
+			KVPair.Value->StartProcessing();
+		}
+
+		bIsProcessingStarted = true;
+	}
+
+	void FAudioAnalyzerRack::StopProcessing()
+	{
+		for (auto&& KVPair : RackUnits)
+		{
+			KVPair.Value->StopProcessing();
+		}
+
+		bIsProcessingStarted = false;
+	}
+
+	TSharedRef<SWidget> FAudioAnalyzerRack::MakeVisibleAnalyzersMenu(TSharedRef<FUICommandList> InCommandList, TWeakPtr<FWorkspaceItem> InWorkspaceGroup, TWeakPtr<FTabManager> InTabManager)
+	{
+		// Create a menu, with toggle buttons for each rack unit type.
+		FMenuBuilder MenuBuilder(true, InCommandList);
+
+		if (TSharedPtr<FWorkspaceItem> WorkspaceGroup = InWorkspaceGroup.Pin())
+		{
+			// Find all rack unit types (implemented here generically using the registered tab spawners):
+			for (const TSharedRef<FWorkspaceItem>& ChildItem : WorkspaceGroup->GetChildItems())
+			{
+				if (TSharedPtr<FTabSpawnerEntry> TabSpawnerEntry = ChildItem->AsSpawnerEntry())
+				{
+					if (TabSpawnerEntry->IsHidden())
+					{
+						continue;
+					}
+
+					const FName TabId = TabSpawnerEntry->GetFName();
+
+					MenuBuilder.AddMenuEntry(
+						TabSpawnerEntry->GetDisplayName(),
+						TabSpawnerEntry->GetTooltipText(),
+						TabSpawnerEntry->GetIcon(),
+						FUIAction(
+							FExecuteAction::CreateLambda([InTabManager, TabId]()
+								{
+									if (TSharedPtr<FTabManager> TabManager = InTabManager.Pin())
+									{
+										TSharedPtr<SDockTab> LiveTab = TabManager->FindExistingLiveTab(TabId);
+										if (!LiveTab.IsValid())
+										{
+											TabManager->TryInvokeTab(TabId);
+										}
+										else
+										{
+											LiveTab->RequestCloseTab();
+										}
+										TabManager->SavePersistentLayout();
+									}
+								}),
+							FCanExecuteAction(),
+							FIsActionChecked::CreateLambda([InTabManager, TabId]()
+								{
+									TSharedPtr<FTabManager> TabManager = InTabManager.Pin();
+									return TabManager.IsValid() && TabManager->FindExistingLiveTab(TabId).IsValid();
+								})
+						),
+						NAME_None,
+						EUserInterfaceActionType::ToggleButton
+					);
+				}
+			}
+		}
+
+		return MenuBuilder.MakeWidget();
+	}
+
+	void FAudioAnalyzerRack::SaveTabLayout(const TSharedRef<FTabManager::FLayout>& InLayout)
+	{
+		FLayoutSaveRestore::SaveToConfig(GEditorLayoutIni, InLayout);
+	}
+
+	TSharedRef<FTabManager::FLayout> FAudioAnalyzerRack::LoadTabLayout()
+	{
+		return FLayoutSaveRestore::LoadFromConfig(GEditorLayoutIni, GetDefaultTabLayout());
+	}
+
+	TSharedRef<FTabManager::FLayout> FAudioAnalyzerRack::GetDefaultTabLayout()
+	{
+		TArray<const FAudioAnalyzerRackUnitTypeInfo*> RackUnitTypes;
+		FAudioAnalyzerRackUnitRegistry::Get().GetRegisteredRackUnitTypes(RackUnitTypes);
+
+		// Create the default layout primary area (can be overriden by derived class):
+		TSharedRef<FTabManager::FArea> PrimaryArea = CreatePrimaryArea(RackUnitTypes);
+		return FTabManager::NewLayout(TabManagerLayoutName)->AddArea(PrimaryArea);
+	}
+
+	TSharedRef<FTabManager::FArea> FAudioAnalyzerRack::CreatePrimaryArea(const TArray<const FAudioAnalyzerRackUnitTypeInfo*>& RackUnitTypes)
+	{
+		TSharedRef<FTabManager::FArea> PrimaryArea = FTabManager::NewPrimaryArea()
+			->SetOrientation(EOrientation::Orient_Vertical);
+
+		// Add entries for all registered rack units:
+		for (const FAudioAnalyzerRackUnitTypeInfo* RackUnitType : RackUnitTypes)
+		{
+			// Set three known rack units visible by default:
+			const bool bOpenTab =
+				RackUnitType->TypeName == FAudioMeter::RackUnitTypeInfo.TypeName ||
+				RackUnitType->TypeName == FAudioOscilloscope::RackUnitTypeInfo.TypeName ||
+				RackUnitType->TypeName == FAudioSpectrumAnalyzer::RackUnitTypeInfo.TypeName;
+
+			const ETabState::Type TabState = (bOpenTab) ? ETabState::OpenedTab : ETabState::ClosedTab;
+
+			PrimaryArea->Split
+			(
+				FTabManager::NewStack()
+				->SetSizeCoefficient(RackUnitType->VerticalSizeCoefficient)
+				->SetHideTabWell(true)
+				->AddTab(RackUnitType->TypeName, TabState)
+			);
+		}
+
+		return PrimaryArea;
+	}
+
+	TSharedRef<IAudioAnalyzerRackUnit> FAudioAnalyzerRack::MakeRackUnit(FName RackUnitTypeName)
+	{
+		// FAudioAnalyzerRack::Init(...) is required to have been called first:
+		ensure(AudioBus.IsValid());
+
+		TSharedRef<IAudioAnalyzerRackUnit> RackUnit = FAudioAnalyzerRackUnitRegistry::Get().MakeRackUnit(RackUnitTypeName, RackUnitConstructParams);
+
+		if (bIsProcessingStarted)
+		{
+			RackUnit->StartProcessing();
+		}
+
+		return RackUnit;
+	}
+}
+
+#undef LOCTEXT_NAMESPACE
diff --git a/Engine/Plugins/Runtime/AudioWidgets/Source/AudioWidgets/Private/AudioAnalyzerRackUnitRegistry.cpp b/Engine/Plugins/Runtime/AudioWidgets/Source/AudioWidgets/Private/AudioAnalyzerRackUnitRegistry.cpp
new file mode 100644
--- /dev/null
+++ b/Engine/Plugins/Runtime/AudioWidgets/Source/AudioWidgets/Private/AudioAnalyzerRackUnitRegistry.cpp
@@ -0,0 +1,34 @@
+// Copyright Epic Games, Inc. All Rights Reserved.
+
+#include "AudioAnalyzerRackUnitRegistry.h"
+
+#include "Misc/LazySingleton.h"
+
+namespace AudioWidgets
+{
+	FAudioAnalyzerRackUnitRegistry& FAudioAnalyzerRackUnitRegistry::Get()
+	{
+		return TLazySingleton<FAudioAnalyzerRackUnitRegistry>::Get();
+	}
+
+	void FAudioAnalyzerRackUnitRegistry::TearDown()
+	{
+		TLazySingleton<FAudioAnalyzerRackUnitRegistry>::TearDown();
+	}
+
+	void FAudioAnalyzerRackUnitRegistry::RegisterRackUnitType(const FAudioAnalyzerRackUnitTypeInfo* RackUnitTypeInfo)
+	{
+		RegisteredRackUnitTypes.Add(RackUnitTypeInfo->TypeName, RackUnitTypeInfo);
+	}
+
+	TSharedRef<IAudioAnalyzerRackUnit> FAudioAnalyzerRackUnitRegistry::MakeRackUnit(FName RackUnitTypeName, const FAudioAnalyzerRackUnitConstructParams& Params)
+	{
+		const FAudioAnalyzerRackUnitTypeInfo* RegisteredRackUnitTypeInfo = RegisteredRackUnitTypes.FindChecked(RackUnitTypeName);
+		return RegisteredRackUnitTypeInfo->OnMakeAudioAnalyzerRackUnit.Execute(Params);
+	}
+
+	void FAudioAnalyzerRackUnitRegistry::GetRegisteredRackUnitTypes(TArray<const FAudioAnalyzerRackUnitTypeInfo*>& OutArray) const
+	{
+		RegisteredRackUnitTypes.GenerateValueArray(OutArray);
+	}
+}
diff --git a/Engine/Plugins/Runtime/AudioWidgets/Source/AudioWidgets/Private/AudioAnalyzerRackUnitRegistry.h b/Engine/Plugins/Runtime/AudioWidgets/Source/AudioWidgets/Private/AudioAnalyzerRackUnitRegistry.h
new file mode 100644
--- /dev/null
+++ b/Engine/Plugins/Runtime/AudioWidgets/Source/AudioWidgets/Private/AudioAnalyzerRackUnitRegistry.h
@@ -0,0 +1,23 @@
+// Copyright Epic Games, Inc. All Rights Reserved.
+#pragma once
+
+#include "AudioAnalyzerRack.h"
+
+namespace AudioWidgets
+{
+	class FAudioAnalyzerRackUnitRegistry
+	{
+	public:
+		static FAudioAnalyzerRackUnitRegistry& Get();
+		static void TearDown();
+
+		void RegisterRackUnitType(const FAudioAnalyzerRackUnitTypeInfo* RackUnitTypeInfo);
+
+		TSharedRef<IAudioAnalyzerRackUnit> MakeRackUnit(FName RackUnitTypeName, const FAudioAnalyzerRackUnitConstructParams& Params);
+
+		void GetRegisteredRackUnitTypes(TArray<const FAudioAnalyzerRackUnitTypeInfo*>& OutArray) const;
+
+	private:
+		TMap<FName, const FAudioAnalyzerRackUnitTypeInfo*> RegisteredRackUnitTypes;
+	};
+}
diff --git a/Engine/Plugins/Runtime/AudioWidgets/Source/AudioWidgets/Private/AudioMeter.cpp b/Engine/Plugins/Runtime/AudioWidgets/Source/AudioWidgets/Private/AudioMeter.cpp
--- a/Engine/Plugins/Runtime/AudioWidgets/Source/AudioWidgets/Private/AudioMeter.cpp
+++ b/Engine/Plugins/Runtime/AudioWidgets/Source/AudioWidgets/Private/AudioMeter.cpp
@@ -8,6 +8,7 @@
 #include "AudioWidgetsStyle.h"
 #include "SAudioMeter.h"
 #include "UObject/UObjectGlobals.h"
+#include "Widgets/Docking/SDockTab.h"
 
 #include UE_INLINE_GENERATED_CPP_BY_NAME(AudioMeter)
 
@@ -170,6 +171,14 @@ const FAudioMeterDefaultColorStyle& FAudioMeterDefaultColorStyle::GetDefault()
 
 namespace AudioWidgets
 {
+	const FAudioAnalyzerRackUnitTypeInfo FAudioMeter::RackUnitTypeInfo
+	{
+		.TypeName = TEXT("FAudioMeter"),
+		.DisplayName = LOCTEXT("AudioMeterDisplayName", "Meter"),
+		.OnMakeAudioAnalyzerRackUnit = FOnMakeAudioAnalyzerRackUnit::CreateStatic(&MakeRackUnit),
+		.VerticalSizeCoefficient = 0.5f,
+	};
+
 	FAudioMeter::FAudioMeter(int32 InNumChannels, UWorld& InWorld, TObjectPtr<UAudioBus> InExternalAudioBus)
 		: FAudioMeter(InNumChannels, InWorld.GetAudioDevice().GetDeviceID(), InExternalAudioBus)
 	{
@@ -300,5 +309,37 @@ namespace AudioWidgets
 
 		bUseExternalAudioBus = false;
 	}
+
+	TSharedRef<SDockTab> FAudioMeter::SpawnTab(const FSpawnTabArgs& Args) const
+	{
+		return SNew(SDockTab)
+			.Clipping(EWidgetClipping::ClipToBounds)
+			.Label(RackUnitTypeInfo.DisplayName)
+			[
+				SNew(SVerticalBox)
+					+ SVerticalBox::Slot()
+					.HAlign(HAlign_Center)
+					.VAlign(VAlign_Fill)
+					[
+						GetWidget()
+					]
+			];
+	}
+
+	TSharedRef<IAudioAnalyzerRackUnit> FAudioMeter::MakeRackUnit(const FAudioAnalyzerRackUnitConstructParams& Params)
+	{
+		// If we have an AudioMaterialMeterStyle create as an AudioMaterialMeter, otherwise create as a standard Slate meter:
+		const FName MaterialMeterStyleName("AudioMaterialMeter.Style");
+		if (Params.StyleSet->HasWidgetStyle<FAudioMaterialMeterStyle>(MaterialMeterStyleName))
+		{
+			const FAudioMaterialMeterStyle& AudioMaterialMeterStyle = Params.StyleSet->GetWidgetStyle<FAudioMaterialMeterStyle>(MaterialMeterStyleName);
+			return MakeShared<FAudioMeter>(Params.AudioBusInfo.AudioBus->GetNumChannels(), Params.AudioBusInfo.AudioDeviceId, AudioMaterialMeterStyle, Params.AudioBusInfo.AudioBus);
+		}
+		else
+		{
+			const FAudioMeterDefaultColorStyle& MeterDefaultColorStyle = Params.StyleSet->GetWidgetStyle<FAudioMeterDefaultColorStyle>("AudioMeter.DefaultColorStyle");
+			return MakeShared<FAudioMeter>(Params.AudioBusInfo.AudioBus->GetNumChannels(), Params.AudioBusInfo.AudioDeviceId, Params.AudioBusInfo.AudioBus, &MeterDefaultColorStyle);
+		}
+	}
 } // namespace AudioWidgets
 #undef LOCTEXT_NAMESPACE
diff --git a/Engine/Plugins/Runtime/AudioWidgets/Source/AudioWidgets/Private/AudioOscilloscope.cpp b/Engine/Plugins/Runtime/AudioWidgets/Source/AudioWidgets/Private/AudioOscilloscope.cpp
--- a/Engine/Plugins/Runtime/AudioWidgets/Source/AudioWidgets/Private/AudioOscilloscope.cpp
+++ b/Engine/Plugins/Runtime/AudioWidgets/Source/AudioWidgets/Private/AudioOscilloscope.cpp
@@ -2,18 +2,41 @@
 
 #include "AudioOscilloscope.h"
 
+#include "Widgets/Docking/SDockTab.h"
+
+#define LOCTEXT_NAMESPACE "FAudioOscilloscope"
+
 namespace AudioWidgets
 {
+	const FAudioAnalyzerRackUnitTypeInfo FAudioOscilloscope::RackUnitTypeInfo
+	{
+		.TypeName = TEXT("FAudioOscilloscope"),
+		.DisplayName = LOCTEXT("AudioOscilloscopeDisplayName", "Oscilloscope"),
+		.OnMakeAudioAnalyzerRackUnit = FOnMakeAudioAnalyzerRackUnit::CreateStatic(&MakeRackUnit),
+		.VerticalSizeCoefficient = 0.25f,
+	};
+
 	FAudioOscilloscope::FAudioOscilloscope(Audio::FDeviceId InAudioDeviceId,
 		const uint32 InNumChannels, 
 		const float InTimeWindowMs, 
 		const float InMaxTimeWindowMs, 
 		const float InAnalysisPeriodMs, 
 		const EAudioPanelLayoutType InPanelLayoutType,
-		const FAudioOscilloscopePanelStyle* InOscilloscopePanelStyle)
+		const FAudioOscilloscopePanelStyle* InOscilloscopePanelStyle,
+		TObjectPtr<UAudioBus> InExternalAudioBus)
 	{
 		OscilloscopePanelStyle = InOscilloscopePanelStyle ? *InOscilloscopePanelStyle : FAudioWidgetsStyle::Get().GetWidgetStyle<FAudioOscilloscopePanelStyle>("AudioOscilloscope.PanelStyle");
-		CreateAudioBus(InNumChannels);
+		
+		if (InExternalAudioBus != nullptr)
+		{
+			ensure(InExternalAudioBus->GetNumChannels() == InNumChannels);
+			AudioBus = TStrongObjectPtr(InExternalAudioBus.Get());
+		}
+		else
+		{
+			CreateAudioBus(InNumChannels);
+		}
+
 		CreateDataProvider(InAudioDeviceId, InTimeWindowMs, InMaxTimeWindowMs, InAnalysisPeriodMs, InPanelLayoutType);
 		CreateOscilloscopeWidget(InNumChannels, InPanelLayoutType);
 	}
@@ -90,4 +113,36 @@ namespace AudioWidgets
 	{
 		return OscilloscopePanelWidget.ToSharedRef();
 	}
+
+	void FAudioOscilloscope::SetAudioBusInfo(const FAudioBusInfo& AudioBusInfo)
+	{
+		AudioBus = TStrongObjectPtr(AudioBusInfo.AudioBus.Get());
+		CreateDataProvider(AudioBusInfo.AudioDeviceId, RackUnitTimeWindowMs, RackUnitMaxTimeWindowMs, RackUnitAnalysisPeriodMs, RackUnitPanelLayoutType);
+		CreateOscilloscopeWidget(AudioBusInfo.AudioBus->GetNumChannels(), RackUnitPanelLayoutType);
+	}
+
+	TSharedRef<SDockTab> FAudioOscilloscope::SpawnTab(const FSpawnTabArgs& Args) const
+	{
+		return SNew(SDockTab)
+			.Clipping(EWidgetClipping::ClipToBounds)
+			.Label(RackUnitTypeInfo.DisplayName)
+			[
+				GetPanelWidget()
+			];
+	}
+
+	TSharedRef<IAudioAnalyzerRackUnit> FAudioOscilloscope::MakeRackUnit(const FAudioAnalyzerRackUnitConstructParams& Params)
+	{
+		return MakeShared<FAudioOscilloscope>(
+			Params.AudioBusInfo.AudioDeviceId,
+			Params.AudioBusInfo.AudioBus->GetNumChannels(),
+			RackUnitTimeWindowMs,
+			RackUnitMaxTimeWindowMs,
+			RackUnitAnalysisPeriodMs,
+			RackUnitPanelLayoutType,
+			&Params.StyleSet->GetWidgetStyle<FAudioOscilloscopePanelStyle>("AudioOscilloscope.PanelStyle"),
+			Params.AudioBusInfo.AudioBus);
+	}
 }
+
+#undef LOCTEXT_NAMESPACE
diff --git a/Engine/Plugins/Runtime/AudioWidgets/Source/AudioWidgets/Private/AudioSpectrumAnalyzer.cpp b/Engine/Plugins/Runtime/AudioWidgets/Source/AudioWidgets/Private/AudioSpectrumAnalyzer.cpp
--- a/Engine/Plugins/Runtime/AudioWidgets/Source/AudioWidgets/Private/AudioSpectrumAnalyzer.cpp
+++ b/Engine/Plugins/Runtime/AudioWidgets/Source/AudioWidgets/Private/AudioSpectrumAnalyzer.cpp
@@ -8,6 +8,7 @@
 #include "DSP/EnvelopeFollower.h"
 #include "Framework/MultiBox/MultiBoxBuilder.h"
 #include "SynesthesiaSpectrumAnalysisFactory.h"
+#include "Widgets/Docking/SDockTab.h"
 
 #define LOCTEXT_NAMESPACE "FAudioSpectrumAnalyzer"
 
@@ -41,8 +42,48 @@ namespace AudioWidgets
 			constexpr float DCMeanSquare = 1.0f;
 			return DCMeanSquare / WindowedDCMeanSquare;
 		}
+
+		/**
+		 * Light wrapper for accessing settings for the analyzer rack unit. Can be passed by value.
+		 */
+		class FRackUnitSettingsHelper
+		{
+		public:
+			FRackUnitSettingsHelper(const FProperty& InSettingsProperty)
+				: SettingsProperty(InSettingsProperty)
+			{
+				//
+			}
+
+			FSpectrumAnalyzerRackUnitSettings* GetRackUnitSettings() const
+			{
+				UObject* EditorSettingsObject = GetEditorSettingsObject();
+				return SettingsProperty.ContainerPtrToValuePtr<FSpectrumAnalyzerRackUnitSettings>(EditorSettingsObject);
+			}
+
+			void SaveConfig() const
+			{
+				GetEditorSettingsObject()->SaveConfig();
+			}
+
+		private:
+			UObject* GetEditorSettingsObject() const
+			{
+				return SettingsProperty.GetOwnerClass()->GetDefaultObject();
+			}
+
+			const FProperty& SettingsProperty;
+		};
 	}
 
+	const FAudioAnalyzerRackUnitTypeInfo FAudioSpectrumAnalyzer::RackUnitTypeInfo
+	{
+		.TypeName = TEXT("FAudioSpectrumAnalyzer"),
+		.DisplayName = LOCTEXT("AudioSpectrumAnalyzerDisplayName", "Spectrum Analyzer"),
+		.OnMakeAudioAnalyzerRackUnit = FOnMakeAudioAnalyzerRackUnit::CreateStatic(&MakeRackUnit),
+		.VerticalSizeCoefficient = 0.25f,
+	};
+
 	FAudioSpectrumAnalyzer::FAudioSpectrumAnalyzer(const FAudioSpectrumAnalyzerParams& Params)
 		: SpectrumAnalysisSettings(NewObject<USynesthesiaSpectrumAnalysisSettings>())
 		, ConstantQSettings(NewObject<UConstantQSettings>())
@@ -548,6 +589,125 @@ namespace AudioWidgets
 		ConstantQAnalyzer.Reset();
 	}
 
+	TSharedRef<SDockTab> FAudioSpectrumAnalyzer::SpawnTab(const FSpawnTabArgs& Args) const
+	{
+		return SNew(SDockTab)
+			.Clipping(EWidgetClipping::ClipToBounds)
+			.Label(RackUnitTypeInfo.DisplayName)
+			[
+				GetWidget()
+			];
+	}
+
+	TSharedRef<IAudioAnalyzerRackUnit> FAudioSpectrumAnalyzer::MakeRackUnit(const FAudioAnalyzerRackUnitConstructParams& Params)
+	{
+		using namespace AudioSpectrumAnalyzerPrivate;
+
+		FAudioSpectrumAnalyzerParams AnalyzerParams;
+		AnalyzerParams.NumChannels = Params.AudioBusInfo.AudioBus->GetNumChannels();
+		AnalyzerParams.AudioDeviceId = Params.AudioBusInfo.AudioDeviceId;
+		AnalyzerParams.ExternalAudioBus = Params.AudioBusInfo.AudioBus;
+
+		if (Params.EditorSettingsClass != nullptr)
+		{
+			// If we have been given a valid editor settings class, bind analyzer options to the settings:
+			const FProperty* SpectrumAnalyzerSettingsProperty = Params.EditorSettingsClass->FindPropertyByName("SpectrumAnalyzerSettings");
+			if (SpectrumAnalyzerSettingsProperty != nullptr)
+			{
+				const FRackUnitSettingsHelper SettingsHelper(*SpectrumAnalyzerSettingsProperty);
+
+				AnalyzerParams.Ballistics.BindLambda([=]()
+					{
+						return SettingsHelper.GetRackUnitSettings()->Ballistics;
+					});
+				AnalyzerParams.AnalyzerType.BindLambda([=]()
+					{
+						return SettingsHelper.GetRackUnitSettings()->AnalyzerType;
+					});
+				AnalyzerParams.FFTAnalyzerFFTSize.BindLambda([=]()
+					{
+						return SettingsHelper.GetRackUnitSettings()->FFTAnalyzerFFTSize;
+					});
+				AnalyzerParams.CQTAnalyzerFFTSize.BindLambda([=]()
+					{
+						return SettingsHelper.GetRackUnitSettings()->CQTAnalyzerFFTSize;
+					});
+				AnalyzerParams.TiltExponent.BindLambda([=]()
+					{
+						const EAudioSpectrumPlotTilt TiltSpectrum = SettingsHelper.GetRackUnitSettings()->TiltSpectrum;
+						return SAudioSpectrumPlot::GetTiltExponentValue(TiltSpectrum);
+					});
+				AnalyzerParams.FrequencyAxisPixelBucketMode.BindLambda([=]()
+					{
+						return SettingsHelper.GetRackUnitSettings()->PixelPlotMode;
+					});
+				AnalyzerParams.FrequencyAxisScale.BindLambda([=]()
+					{
+						return SettingsHelper.GetRackUnitSettings()->FrequencyScale;
+					});
+				AnalyzerParams.bDisplayFrequencyAxisLabels.BindLambda([=]()
+					{
+						return SettingsHelper.GetRackUnitSettings()->bDisplayFrequencyAxisLabels;
+					});
+				AnalyzerParams.bDisplaySoundLevelAxisLabels.BindLambda([=]()
+					{
+						return SettingsHelper.GetRackUnitSettings()->bDisplaySoundLevelAxisLabels;
+					});
+
+				AnalyzerParams.OnBallisticsMenuEntryClicked.BindLambda([=](EAudioSpectrumAnalyzerBallistics SelectedValue)
+					{
+						SettingsHelper.GetRackUnitSettings()->Ballistics = SelectedValue;
+						SettingsHelper.SaveConfig();
+					});
+				AnalyzerParams.OnAnalyzerTypeMenuEntryClicked.BindLambda([=](EAudioSpectrumAnalyzerType SelectedValue)
+					{
+						SettingsHelper.GetRackUnitSettings()->AnalyzerType = SelectedValue;
+						SettingsHelper.SaveConfig();
+					});
+				AnalyzerParams.OnFFTAnalyzerFFTSizeMenuEntryClicked.BindLambda([=](EFFTSize SelectedValue)
+					{
+						SettingsHelper.GetRackUnitSettings()->FFTAnalyzerFFTSize = SelectedValue;
+						SettingsHelper.SaveConfig();
+					});
+				AnalyzerParams.OnCQTAnalyzerFFTSizeMenuEntryClicked.BindLambda([=](EConstantQFFTSizeEnum SelectedValue)
+					{
+						SettingsHelper.GetRackUnitSettings()->CQTAnalyzerFFTSize = SelectedValue;
+						SettingsHelper.SaveConfig();
+					});
+				AnalyzerParams.OnTiltSpectrumMenuEntryClicked.BindLambda([=](EAudioSpectrumPlotTilt SelectedValue)
+					{
+						SettingsHelper.GetRackUnitSettings()->TiltSpectrum = SelectedValue;
+						SettingsHelper.SaveConfig();
+					});
+				AnalyzerParams.OnFrequencyAxisPixelBucketModeMenuEntryClicked.BindLambda([=](EAudioSpectrumPlotFrequencyAxisPixelBucketMode SelectedValue)
+					{
+						SettingsHelper.GetRackUnitSettings()->PixelPlotMode = SelectedValue;
+						SettingsHelper.SaveConfig();
+					});
+				AnalyzerParams.OnFrequencyAxisScaleMenuEntryClicked.BindLambda([=](EAudioSpectrumPlotFrequencyAxisScale SelectedValue)
+					{
+						SettingsHelper.GetRackUnitSettings()->FrequencyScale = SelectedValue;
+						SettingsHelper.SaveConfig();
+					});
+				AnalyzerParams.OnDisplayFrequencyAxisLabelsButtonToggled.BindLambda([=]()
+					{
+						FSpectrumAnalyzerRackUnitSettings* SpectrumAnalyzerSettings = SettingsHelper.GetRackUnitSettings();
+						SpectrumAnalyzerSettings->bDisplayFrequencyAxisLabels = !SpectrumAnalyzerSettings->bDisplayFrequencyAxisLabels;
+						SettingsHelper.SaveConfig();
+					});
+				AnalyzerParams.OnDisplaySoundLevelAxisLabelsButtonToggled.BindLambda([=]()
+					{
+						FSpectrumAnalyzerRackUnitSettings* SpectrumAnalyzerSettings = SettingsHelper.GetRackUnitSettings();
+						SpectrumAnalyzerSettings->bDisplaySoundLevelAxisLabels = !SpectrumAnalyzerSettings->bDisplaySoundLevelAxisLabels;
+						SettingsHelper.SaveConfig();
+					});
+			}
+		}
+
+		AnalyzerParams.PlotStyle = &Params.StyleSet->GetWidgetStyle<FAudioSpectrumPlotStyle>("AudioSpectrumPlot.Style");
+
+		return MakeShared<FAudioSpectrumAnalyzer>(AnalyzerParams);
+	}
 } // namespace AudioWidgets
 
 #undef LOCTEXT_NAMESPACE
diff --git a/Engine/Plugins/Runtime/AudioWidgets/Source/AudioWidgets/Private/AudioVectorscope.cpp b/Engine/Plugins/Runtime/AudioWidgets/Source/AudioWidgets/Private/AudioVectorscope.cpp
--- a/Engine/Plugins/Runtime/AudioWidgets/Source/AudioWidgets/Private/AudioVectorscope.cpp
+++ b/Engine/Plugins/Runtime/AudioWidgets/Source/AudioWidgets/Private/AudioVectorscope.cpp
@@ -3,19 +3,41 @@
 #include "AudioVectorscope.h"
 
 #include "SAudioVectorscopePanelWidget.h"
+#include "Widgets/Docking/SDockTab.h"
+
+#define LOCTEXT_NAMESPACE "FAudioVectorscope"
 
 namespace AudioWidgets
 {
+	const FAudioAnalyzerRackUnitTypeInfo FAudioVectorscope::RackUnitTypeInfo
+	{
+		.TypeName = TEXT("FAudioVectorscope"),
+		.DisplayName = LOCTEXT("AudioVectorscopeDisplayName", "Vectorscope"),
+		.OnMakeAudioAnalyzerRackUnit = FOnMakeAudioAnalyzerRackUnit::CreateStatic(&MakeRackUnit),
+		.VerticalSizeCoefficient = 0.25f,
+	};
+
 	FAudioVectorscope::FAudioVectorscope(Audio::FDeviceId InAudioDeviceId,
 		const uint32 InNumChannels,
 		const float InTimeWindowMs,
 		const float InMaxTimeWindowMs,
 		const float InAnalysisPeriodMs,
 		const EAudioPanelLayoutType InPanelLayoutType,
-		const FAudioVectorscopePanelStyle* InPanelStyle)
+		const FAudioVectorscopePanelStyle* InPanelStyle,
+		TObjectPtr<UAudioBus> InExternalAudioBus)
 	{
 		VectorscopePanelStyle = InPanelStyle ? *InPanelStyle : FAudioWidgetsStyle::Get().GetWidgetStyle<FAudioVectorscopePanelStyle>("AudioVectorscope.PanelStyle");
-		CreateAudioBus(InNumChannels);
+
+		if (InExternalAudioBus != nullptr)
+		{
+			ensure(InExternalAudioBus->GetNumChannels() == InNumChannels);
+			AudioBus = TStrongObjectPtr(InExternalAudioBus.Get());
+		}
+		else
+		{
+			CreateAudioBus(InNumChannels);
+		}
+		
 		CreateDataProvider(InAudioDeviceId, InTimeWindowMs, InMaxTimeWindowMs, InAnalysisPeriodMs);
 		CreateVectorscopeWidget(InPanelLayoutType);
 	}
@@ -83,4 +105,36 @@ namespace AudioWidgets
 	{
 		return VectorscopePanelWidget.ToSharedRef();
 	}
+
+	void FAudioVectorscope::SetAudioBusInfo(const FAudioBusInfo& AudioBusInfo)
+	{
+		AudioBus = TStrongObjectPtr(AudioBusInfo.AudioBus.Get());
+		CreateDataProvider(AudioBusInfo.AudioDeviceId, RackUnitTimeWindowMs, RackUnitMaxTimeWindowMs, RackUnitAnalysisPeriodMs);
+		CreateVectorscopeWidget(RackUnitPanelLayoutType);
+	}
+
+	TSharedRef<SDockTab> FAudioVectorscope::SpawnTab(const FSpawnTabArgs& Args) const
+	{
+		return SNew(SDockTab)
+			.Clipping(EWidgetClipping::ClipToBounds)
+			.Label(RackUnitTypeInfo.DisplayName)
+			[
+				GetPanelWidget()
+			];
+	}
+
+	TSharedRef<IAudioAnalyzerRackUnit> FAudioVectorscope::MakeRackUnit(const FAudioAnalyzerRackUnitConstructParams& Params)
+	{
+		return MakeShared<FAudioVectorscope>(
+			Params.AudioBusInfo.AudioDeviceId,
+			Params.AudioBusInfo.AudioBus->GetNumChannels(),
+			RackUnitTimeWindowMs,
+			RackUnitMaxTimeWindowMs,
+			RackUnitAnalysisPeriodMs,
+			RackUnitPanelLayoutType,
+			&Params.StyleSet->GetWidgetStyle<FAudioVectorscopePanelStyle>("AudioVectorscope.PanelStyle"),
+			Params.AudioBusInfo.AudioBus);
+	}
 } // namespace AudioWidgets
+
+#undef LOCTEXT_NAMESPACE
diff --git a/Engine/Plugins/Runtime/AudioWidgets/Source/AudioWidgets/Private/AudioWidgets.cpp b/Engine/Plugins/Runtime/AudioWidgets/Source/AudioWidgets/Private/AudioWidgets.cpp
--- a/Engine/Plugins/Runtime/AudioWidgets/Source/AudioWidgets/Private/AudioWidgets.cpp
+++ b/Engine/Plugins/Runtime/AudioWidgets/Source/AudioWidgets/Private/AudioWidgets.cpp
@@ -2,6 +2,11 @@
 
 #include "AudioWidgets.h"
 #include "AdvancedWidgetsModule.h"
+#include "AudioAnalyzerRackUnitRegistry.h"
+#include "AudioMeter.h"
+#include "AudioOscilloscope.h"
+#include "AudioSpectrumAnalyzer.h"
+#include "AudioVectorscope.h"
 #include "AudioWidgetsStyle.h"
 #include "Modules/ModuleManager.h"
 
@@ -9,15 +14,29 @@
 
 void FAudioWidgetsModule::StartupModule()
 {
+	using namespace AudioWidgets;
+
 	// Required to load so the AudioWidget plugin content can reference widgets
 	// defined in AdvancedWidgets (ex. RadialSlider for UMG-defined knobs)
 	FModuleManager::Get().LoadModuleChecked<FAdvancedWidgetsModule>("AdvancedWidgets");
 	// Initialize static style instance 
 	FAudioWidgetsStyle::Get();
+
+	// Register standard analyzer rack units:
+	RegisterAudioAnalyzerRackUnitType(&FAudioMeter::RackUnitTypeInfo);
+	RegisterAudioAnalyzerRackUnitType(&FAudioOscilloscope::RackUnitTypeInfo);
+	RegisterAudioAnalyzerRackUnitType(&FAudioVectorscope::RackUnitTypeInfo);
+	RegisterAudioAnalyzerRackUnitType(&FAudioSpectrumAnalyzer::RackUnitTypeInfo);
 }
 
 void FAudioWidgetsModule::ShutdownModule()
 {
+	AudioWidgets::FAudioAnalyzerRackUnitRegistry::TearDown();
+}
+
+void FAudioWidgetsModule::RegisterAudioAnalyzerRackUnitType(const AudioWidgets::FAudioAnalyzerRackUnitTypeInfo* RackUnitTypeInfo)
+{
+	AudioWidgets::FAudioAnalyzerRackUnitRegistry::Get().RegisterRackUnitType(RackUnitTypeInfo);
 }
 
 #undef LOCTEXT_NAMESPACE
diff --git a/Engine/Plugins/Runtime/AudioWidgets/Source/AudioWidgets/Private/AudioWidgetsStyle.cpp b/Engine/Plugins/Runtime/AudioWidgets/Source/AudioWidgets/Private/AudioWidgetsStyle.cpp
--- a/Engine/Plugins/Runtime/AudioWidgets/Source/AudioWidgets/Private/AudioWidgetsStyle.cpp
+++ b/Engine/Plugins/Runtime/AudioWidgets/Source/AudioWidgets/Private/AudioWidgetsStyle.cpp
@@ -171,6 +171,8 @@ FAudioWidgetsStyle::FAudioWidgetsStyle()
 	Set("SampledSequenceVectorViewer.Style", FSampledSequenceVectorViewerStyle());
 	Set("TriggerThresholdLine.Style", FTriggerThresholdLineStyle());
 
+	Set("AudioAnalyzerRack.BackgroundColor", FLinearColor::Transparent);
+
 	FSlateStyleRegistry::RegisterSlateStyle(*this);
 }
 void FAudioWidgetsStyle::SetResources()
diff --git a/Engine/Plugins/Runtime/AudioWidgets/Source/AudioWidgets/Public/AudioAnalyzerRack.h b/Engine/Plugins/Runtime/AudioWidgets/Source/AudioWidgets/Public/AudioAnalyzerRack.h
new file mode 100644
--- /dev/null
+++ b/Engine/Plugins/Runtime/AudioWidgets/Source/AudioWidgets/Public/AudioAnalyzerRack.h
@@ -0,0 +1,119 @@
+// Copyright Epic Games, Inc. All Rights Reserved.
+#pragma once
+
+#include "AudioDefines.h"
+#include "Framework/Docking/TabManager.h"
+#include "Sound/AudioBus.h"
+
+class FUICommandList;
+
+namespace AudioWidgets
+{
+	/**
+	 * FAudioBusInfo encapsulates the required info that describes the audio bus that is to be analyzed.
+	 */
+	struct FAudioBusInfo
+	{
+		Audio::FDeviceId AudioDeviceId = static_cast<Audio::FDeviceId>(INDEX_NONE);
+		TObjectPtr<UAudioBus> AudioBus = nullptr;
+	};
+
+	/**
+	 * Interface for something that can be used in an analyzer rack.
+	 */
+	class IAudioAnalyzerRackUnit
+	{
+	public:
+		virtual ~IAudioAnalyzerRackUnit() {};
+
+		/** If the Audio Bus to analyze changes (due to channel count change), handle this here. */
+		virtual void SetAudioBusInfo(const FAudioBusInfo& AudioBusInfo) = 0;
+
+		/** Spawn the actual analyzer Widget in a DockTab. */
+		virtual TSharedRef<SDockTab> SpawnTab(const FSpawnTabArgs& Args) const = 0;
+
+		/** Be notified of request to start processing. */
+		virtual void StartProcessing() {};
+
+		/** Be notified of request to stop processing. */
+		virtual void StopProcessing() {};
+	};
+
+	/**
+	 * Parameters for the FOnMakeAudioAnalyzerRackUnit delegate.
+	 */
+	struct FAudioAnalyzerRackUnitConstructParams
+	{
+		FAudioBusInfo AudioBusInfo;
+		const ISlateStyle* StyleSet = nullptr;
+		const UClass* EditorSettingsClass = nullptr;
+	};
+
+	DECLARE_DELEGATE_RetVal_OneParam(TSharedRef<IAudioAnalyzerRackUnit>, FOnMakeAudioAnalyzerRackUnit, const FAudioAnalyzerRackUnitConstructParams&);
+
+	/**
+	 * Static type descriptor data for a rack unit type.
+	 */
+	struct FAudioAnalyzerRackUnitTypeInfo
+	{
+		FName TypeName;
+		FText DisplayName;
+		FSlateIcon Icon;
+		FOnMakeAudioAnalyzerRackUnit OnMakeAudioAnalyzerRackUnit;
+		float VerticalSizeCoefficient;
+	};
+
+	/**
+	 * Manages display of audio analyzer rack units. Rack units can be shown, hidden, and reordered by the user.
+	 */
+	class AUDIOWIDGETS_API FAudioAnalyzerRack : public TSharedFromThis<FAudioAnalyzerRack>
+	{
+	public:
+		struct FRackConstructParams
+		{
+			/** The rack layout can be saved using the given name. */
+			FName TabManagerLayoutName;
+
+			/** An ISlateStyle can be provided to override AudioWidgetsStyle. */
+			const ISlateStyle* StyleSet = nullptr;
+
+			/** An Editor Settings class can be provided for rack units that require one for saving settings. */
+			const UClass* EditorSettingsClass = nullptr;
+		};
+
+		FAudioAnalyzerRack(const FRackConstructParams& Params);
+		virtual ~FAudioAnalyzerRack() = default;
+
+		void Init(int32 InNumChannels, Audio::FDeviceId InAudioDeviceId);
+		void DestroyAnalyzers();
+
+		TSharedRef<SWidget> CreateWidget(TSharedRef<SDockTab> DockTab);
+
+		UAudioBus* GetAudioBus() const;
+
+		void StartProcessing();
+		void StopProcessing();
+
+	protected:
+		virtual TSharedRef<FTabManager::FArea> CreatePrimaryArea(const TArray<const FAudioAnalyzerRackUnitTypeInfo*>& RackUnitTypes);
+
+	private:
+		static TSharedRef<SWidget> MakeVisibleAnalyzersMenu(TSharedRef<FUICommandList> InCommandList, TWeakPtr<FWorkspaceItem> InWorkspaceGroup, TWeakPtr<FTabManager> InTabManager);
+
+		static void SaveTabLayout(const TSharedRef<FTabManager::FLayout>& InLayout);
+		TSharedRef<FTabManager::FLayout> LoadTabLayout();
+		TSharedRef<FTabManager::FLayout> GetDefaultTabLayout();
+
+		TSharedRef<IAudioAnalyzerRackUnit> MakeRackUnit(FName RackUnitTypeName);
+
+		FName TabManagerLayoutName;
+
+		TStrongObjectPtr<UAudioBus> AudioBus;
+		FAudioAnalyzerRackUnitConstructParams RackUnitConstructParams;
+
+		TMap<FName, TSharedRef<IAudioAnalyzerRackUnit>> RackUnits;
+		TSharedPtr<FTabManager> TabManager;
+		
+		bool bIsProcessingStarted = false;
+	};
+}
diff --git a/Engine/Plugins/Runtime/AudioWidgets/Source/AudioWidgets/Public/AudioMeter.h b/Engine/Plugins/Runtime/AudioWidgets/Source/AudioWidgets/Public/AudioMeter.h
--- a/Engine/Plugins/Runtime/AudioWidgets/Source/AudioWidgets/Public/AudioMeter.h
+++ b/Engine/Plugins/Runtime/AudioWidgets/Source/AudioWidgets/Public/AudioMeter.h
@@ -2,6 +2,7 @@
 #pragma once
 
 #include "AudioAnalyzer.h"
+#include "AudioAnalyzerRack.h"
 #include "AudioDefines.h"
 #include "AudioMeterStyle.h"
 #include "AudioMeterTypes.h"
@@ -187,9 +188,11 @@ protected:
 
 namespace AudioWidgets
 {
-	class AUDIOWIDGETS_API FAudioMeter
+	class AUDIOWIDGETS_API FAudioMeter : public IAudioAnalyzerRackUnit
 	{
 	public:
+		static const FAudioAnalyzerRackUnitTypeInfo RackUnitTypeInfo;
+
 		UE_DEPRECATED(5.4, "Use the FAudioMeter constructor that uses Audio::FDeviceId.")
 		FAudioMeter(int32 InNumChannels, UWorld& InWorld, TObjectPtr<UAudioBus> InExternalAudioBus = nullptr); 
 		
@@ -216,10 +219,21 @@ namespace AudioWidgets
 
 		void Init(const int32 InNumChannels, const Audio::FDeviceId InAudioDeviceId, TObjectPtr<UAudioBus> InExternalAudioBus);
 
+		// Begin IAudioAnalyzerRackUnit overrides.
+		virtual void SetAudioBusInfo(const FAudioBusInfo& AudioBusInfo) override
+		{
+			Init(AudioBusInfo.AudioBus->GetNumChannels(), AudioBusInfo.AudioDeviceId, AudioBusInfo.AudioBus);
+		}
+
+		virtual TSharedRef<SDockTab> SpawnTab(const FSpawnTabArgs& Args) const override;
+		// End IAudioAnalyzerRackUnit overrides.
+
 	protected:
 		void OnMeterOutput(UMeterAnalyzer* InMeterAnalyzer, int32 ChannelIndex, const FMeterResults& InMeterResults);
 
 	private:
+		static TSharedRef<IAudioAnalyzerRackUnit> MakeRackUnit(const FAudioAnalyzerRackUnitConstructParams& Params);
+
 		void Teardown();
 
 		/** Metasound analyzer object. */
diff --git a/Engine/Plugins/Runtime/AudioWidgets/Source/AudioWidgets/Public/AudioOscilloscope.h b/Engine/Plugins/Runtime/AudioWidgets/Source/AudioWidgets/Public/AudioOscilloscope.h
--- a/Engine/Plugins/Runtime/AudioWidgets/Source/AudioWidgets/Public/AudioOscilloscope.h
+++ b/Engine/Plugins/Runtime/AudioWidgets/Source/AudioWidgets/Public/AudioOscilloscope.h
@@ -2,6 +2,7 @@
 
 #pragma once
 
+#include "AudioAnalyzerRack.h"
 #include "AudioDefines.h"
 #include "AudioOscilloscopePanelStyle.h"
 #include "AudioWidgetsEnums.h"
@@ -16,16 +17,19 @@ namespace AudioWidgets
 {
 	class FWaveformAudioSamplesDataProvider;
 
-	class AUDIOWIDGETS_API FAudioOscilloscope
+	class AUDIOWIDGETS_API FAudioOscilloscope : public IAudioAnalyzerRackUnit
 	{
 	public:
+		static const FAudioAnalyzerRackUnitTypeInfo RackUnitTypeInfo;
+
 		FAudioOscilloscope(Audio::FDeviceId InAudioDeviceId,
 			const uint32 InNumChannels,
 			const float InTimeWindowMs,
 			const float InMaxTimeWindowMs,
 			const float InAnalysisPeriodMs,
 			const EAudioPanelLayoutType InPanelLayoutType,
-			const FAudioOscilloscopePanelStyle* InOscilloscopePanelStyle = nullptr);
+			const FAudioOscilloscopePanelStyle* InOscilloscopePanelStyle = nullptr,
+			TObjectPtr<UAudioBus> InExternalAudioBus = nullptr);
 
 		void CreateAudioBus(const uint32 InNumChannels);
 
@@ -37,13 +41,24 @@ namespace AudioWidgets
 
 		void CreateOscilloscopeWidget(const uint32 InNumChannels, const EAudioPanelLayoutType InPanelLayoutType, const FAudioOscilloscopePanelStyle* InOscilloscopePanelStyle = nullptr);
 
-		void StartProcessing();
-		void StopProcessing();
+		virtual void StartProcessing() override;
+		virtual void StopProcessing() override;
 
 		UAudioBus* GetAudioBus() const;
 		TSharedRef<SWidget> GetPanelWidget() const;
 
+		// Begin IAudioAnalyzerRackUnit overrides.
+		virtual void SetAudioBusInfo(const FAudioBusInfo& AudioBusInfo) override;
+		virtual TSharedRef<SDockTab> SpawnTab(const FSpawnTabArgs& Args) const override;
+		// End IAudioAnalyzerRackUnit overrides.
+
 	private:
+		static TSharedRef<IAudioAnalyzerRackUnit> MakeRackUnit(const FAudioAnalyzerRackUnitConstructParams& Params);
+		static constexpr float RackUnitTimeWindowMs = 10.0f;
+		static constexpr float RackUnitMaxTimeWindowMs = 10.0f;
+		static constexpr float RackUnitAnalysisPeriodMs = 10.0f;
+		static constexpr EAudioPanelLayoutType RackUnitPanelLayoutType = EAudioPanelLayoutType::Basic;
+
 		FAudioOscilloscopePanelStyle OscilloscopePanelStyle;
 
 		TSharedPtr<FWaveformAudioSamplesDataProvider> AudioSamplesDataProvider = nullptr;
diff --git a/Engine/Plugins/Runtime/AudioWidgets/Source/AudioWidgets/Public/AudioSpectrumAnalyzer.h b/Engine/Plugins/Runtime/AudioWidgets/Source/AudioWidgets/Public/AudioSpectrumAnalyzer.h
--- a/Engine/Plugins/Runtime/AudioWidgets/Source/AudioWidgets/Public/AudioSpectrumAnalyzer.h
+++ b/Engine/Plugins/Runtime/AudioWidgets/Source/AudioWidgets/Public/AudioSpectrumAnalyzer.h
@@ -1,6 +1,7 @@
 // Copyright Epic Games, Inc. All Rights Reserved.
 #pragma once
 
+#include "AudioAnalyzerRack.h"
 #include "ConstantQ.h"
 #include "SAudioSpectrumPlot.h"
 #include "Sound/AudioBus.h"
@@ -68,9 +69,11 @@ namespace AudioWidgets
 	 * Exponential time-smoothing is applied to the spectrum.
 	 * Can either create an Audio Bus to analyze, or analyze the given Bus.
 	 */
-	class AUDIOWIDGETS_API FAudioSpectrumAnalyzer : public TSharedFromThis<FAudioSpectrumAnalyzer>
+	class AUDIOWIDGETS_API FAudioSpectrumAnalyzer : public IAudioAnalyzerRackUnit, public TSharedFromThis<FAudioSpectrumAnalyzer>
 	{
 	public:
+		static const FAudioAnalyzerRackUnitTypeInfo RackUnitTypeInfo;
+
 		FAudioSpectrumAnalyzer(const FAudioSpectrumAnalyzerParams& Params);
 		FAudioSpectrumAnalyzer(int32 InNumChannels, Audio::FDeviceId InAudioDeviceId, TObjectPtr<UAudioBus> InExternalAudioBus = nullptr);
 		~FAudioSpectrumAnalyzer();
@@ -81,6 +84,15 @@ namespace AudioWidgets
 
 		void Init(int32 InNumChannels, Audio::FDeviceId InAudioDeviceId, TObjectPtr<UAudioBus> InExternalAudioBus = nullptr);
 
+		// Begin IAudioAnalyzerRackUnit overrides.
+		virtual void SetAudioBusInfo(const FAudioBusInfo& AudioBusInfo) override
+		{
+			Init(AudioBusInfo.AudioBus->GetNumChannels(), AudioBusInfo.AudioDeviceId, AudioBusInfo.AudioBus);
+		}
+
+		virtual TSharedRef<SDockTab> SpawnTab(const FSpawnTabArgs& Args) const override;
+		// End IAudioAnalyzerRackUnit overrides.
+
 	protected:
 		void StartAnalyzing(const EAudioSpectrumAnalyzerType InAnalyzerType);
 		void StopAnalyzing();
@@ -99,6 +111,8 @@ namespace AudioWidgets
 		void UpdateAnalyzerSettings();
 
 	private:
+		static TSharedRef<IAudioAnalyzerRackUnit> MakeRackUnit(const FAudioAnalyzerRackUnitConstructParams& Params);
+
 		void CreateSynesthesiaSpectrumAnalyzer();
 		void ReleaseSynesthesiaSpectrumAnalyzer();
 			
@@ -151,3 +165,36 @@ namespace AudioWidgets
 		FOnCQTAnalyzerFFTSizeMenuEntryClicked OnCQTAnalyzerFFTSizeMenuEntryClicked;
 	};
 } // namespace AudioWidgets
+
+USTRUCT()
+struct FSpectrumAnalyzerRackUnitSettings
+{
+	GENERATED_BODY()
+
+	UPROPERTY(EditAnywhere, config, Category = SpectrumAnalyzer)
+	EAudioSpectrumAnalyzerBallistics Ballistics = EAudioSpectrumAnalyzerBallistics::Digital;
+
+	UPROPERTY(EditAnywhere, config, Category = SpectrumAnalyzer)
+	EAudioSpectrumAnalyzerType AnalyzerType = EAudioSpectrumAnalyzerType::CQT;
+
+	UPROPERTY(EditAnywhere, config, Category = SpectrumAnalyzer, meta = (DisplayName = "FFT Size (FFT Analyzer)"))
+	EFFTSize FFTAnalyzerFFTSize = EFFTSize::Max;
+
+	UPROPERTY(EditAnywhere, config, Category = SpectrumAnalyzer, meta = (DisplayName = "FFT Size (CQT Analyzer)"))
+	EConstantQFFTSizeEnum CQTAnalyzerFFTSize = EConstantQFFTSizeEnum::XXLarge;
+
+	UPROPERTY(EditAnywhere, config, Category = SpectrumAnalyzer)
+	EAudioSpectrumPlotTilt TiltSpectrum = EAudioSpectrumPlotTilt::NoTilt;
+
+	UPROPERTY(EditAnywhere, config, Category = SpectrumAnalyzer)
+	EAudioSpectrumPlotFrequencyAxisPixelBucketMode PixelPlotMode = EAudioSpectrumPlotFrequencyAxisPixelBucketMode::Average;
+
+	UPROPERTY(EditAnywhere, config, Category = SpectrumAnalyzer)
+	EAudioSpectrumPlotFrequencyAxisScale FrequencyScale = EAudioSpectrumPlotFrequencyAxisScale::Logarithmic;
+
+	UPROPERTY(EditAnywhere, config, Category = SpectrumAnalyzer)
+	bool bDisplayFrequencyAxisLabels = false;
+
+	UPROPERTY(EditAnywhere, config, Category = SpectrumAnalyzer)
+	bool bDisplaySoundLevelAxisLabels = false;
+};
diff --git a/Engine/Plugins/Runtime/AudioWidgets/Source/AudioWidgets/Public/AudioVectorscope.h b/Engine/Plugins/Runtime/AudioWidgets/Source/AudioWidgets/Public/AudioVectorscope.h
--- a/Engine/Plugins/Runtime/AudioWidgets/Source/AudioWidgets/Public/AudioVectorscope.h
+++ b/Engine/Plugins/Runtime/AudioWidgets/Source/AudioWidgets/Public/AudioVectorscope.h
@@ -2,6 +2,7 @@
 
 #pragma once
 
+#include "AudioAnalyzerRack.h"
 #include "AudioDefines.h"
 #include "AudioVectorscopePanelStyle.h"
 #include "AudioWidgetsEnums.h"
@@ -15,28 +16,42 @@ namespace AudioWidgets
 {
 	class FWaveformAudioSamplesDataProvider;
 
-	class AUDIOWIDGETS_API FAudioVectorscope
+	class AUDIOWIDGETS_API FAudioVectorscope : public IAudioAnalyzerRackUnit
 	{
 	public:
+		static const FAudioAnalyzerRackUnitTypeInfo RackUnitTypeInfo;
+
 		FAudioVectorscope(Audio::FDeviceId InAudioDeviceId,
 			const uint32 InNumChannels, 
 			const float InTimeWindowMs, 
 			const float InMaxTimeWindowMs, 
 			const float InAnalysisPeriodMs, 
 			const EAudioPanelLayoutType InPanelLayoutType,
-			const FAudioVectorscopePanelStyle* PanelStyle = nullptr);
+			const FAudioVectorscopePanelStyle* PanelStyle = nullptr,
+			TObjectPtr<UAudioBus> InExternalAudioBus = nullptr);
 
 		void CreateAudioBus(const uint32 InNumChannels);
 		void CreateDataProvider(Audio::FDeviceId InAudioDeviceId, const float InTimeWindowMs, const float InMaxTimeWindowMs, const float InAnalysisPeriodMs);
 		void CreateVectorscopeWidget(const EAudioPanelLayoutType InPanelLayoutType, const FAudioVectorscopePanelStyle* PanelStyle = nullptr);
 
-		void StartProcessing();
-		void StopProcessing();
+		virtual void StartProcessing() override;
+		virtual void StopProcessing() override;
 
 		UAudioBus* GetAudioBus() const;
 		TSharedRef<SWidget> GetPanelWidget() const;
 
+		// Begin IAudioAnalyzerRackUnit overrides.
+		virtual void SetAudioBusInfo(const FAudioBusInfo& AudioBusInfo) override;
+		virtual TSharedRef<SDockTab> SpawnTab(const FSpawnTabArgs& Args) const override;
+		// End IAudioAnalyzerRackUnit overrides.
+
 	private:
+		static TSharedRef<IAudioAnalyzerRackUnit> MakeRackUnit(const FAudioAnalyzerRackUnitConstructParams& Params);
+		static constexpr float RackUnitTimeWindowMs = 30.0f;
+		static constexpr float RackUnitMaxTimeWindowMs = 30.0f;
+		static constexpr float RackUnitAnalysisPeriodMs = 10.0f;
+		static constexpr EAudioPanelLayoutType RackUnitPanelLayoutType = EAudioPanelLayoutType::Basic;
+
 		FAudioVectorscopePanelStyle VectorscopePanelStyle;
 
 		TSharedPtr<FWaveformAudioSamplesDataProvider> AudioSamplesDataProvider = nullptr;
diff --git a/Engine/Plugins/Runtime/AudioWidgets/Source/AudioWidgets/Public/AudioWidgets.h b/Engine/Plugins/Runtime/AudioWidgets/Source/AudioWidgets/Public/AudioWidgets.h
--- a/Engine/Plugins/Runtime/AudioWidgets/Source/AudioWidgets/Public/AudioWidgets.h
+++ b/Engine/Plugins/Runtime/AudioWidgets/Source/AudioWidgets/Public/AudioWidgets.h
@@ -4,6 +4,11 @@
 
 #include "Modules/ModuleInterface.h"
 
+namespace AudioWidgets
+{
+	struct FAudioAnalyzerRackUnitTypeInfo;
+}
+
 class FAudioWidgetsModule : public IModuleInterface
 {
 public:
@@ -11,6 +16,9 @@ public:
 	/** IModuleInterface implementation */
 	virtual void StartupModule() override;
 	virtual void ShutdownModule() override;
+
+	// Register a custom audio analyzer rack unit type
+	AUDIOWIDGETS_API void RegisterAudioAnalyzerRackUnitType(const AudioWidgets::FAudioAnalyzerRackUnitTypeInfo* RackUnitTypeInfo);
 };
 
 #if UE_ENABLE_INCLUDE_ORDER_DEPRECATED_IN_5_2
diff --git a/Engine/Plugins/Runtime/Metasound/Source/MetasoundEditor/Private/MetasoundEditor.cpp b/Engine/Plugins/Runtime/Metasound/Source/MetasoundEditor/Private/MetasoundEditor.cpp
--- a/Engine/Plugins/Runtime/Metasound/Source/MetasoundEditor/Private/MetasoundEditor.cpp
+++ b/Engine/Plugins/Runtime/Metasound/Source/MetasoundEditor/Private/MetasoundEditor.cpp
@@ -3,12 +3,8 @@
 #include "MetasoundEditor.h"
 
 #include "Algo/AnyOf.h"
+#include "AudioAnalyzerRack.h"
 #include "AudioDevice.h"
-#include "AudioMaterialSlate/SAudioMaterialMeter.h"
-#include "AudioMeterStyle.h"
-#include "AudioOscilloscope.h"
-#include "AudioSpectrumAnalyzer.h"
-#include "AudioVectorscope.h"
 #include "AudioWidgetsEnums.h"
 #include "Components/AudioComponent.h"
 #include "DetailLayoutBuilder.h"
@@ -987,9 +983,12 @@ namespace Metasound
 			.SetGroup(WorkspaceMenuCategoryRef)
 			.SetIcon(FSlateIcon("MetaSoundStyle", "MetasoundEditor.Metasound.Icon"));
 
-			InTabManager->RegisterTabSpawner(TabNamesPrivate::Analyzers, FOnSpawnTab::CreateLambda([InAnalyzerWidget = BuildAnalyzerWidget()](const FSpawnTabArgs&)
+			InTabManager->RegisterTabSpawner(TabNamesPrivate::Analyzers, FOnSpawnTab::CreateLambda([InAnalyzerRack = AnalyzerRack](const FSpawnTabArgs&)
 			{
-				return SNew(SDockTab).Label(LOCTEXT("MetasoundAnalyzersTitle", "Analyzers")) [ InAnalyzerWidget.ToSharedRef() ];
+				TSharedRef<SDockTab> DockTab = SNew(SDockTab).Label(LOCTEXT("MetasoundAnalyzersTitle", "Analyzers"));
+				TSharedRef<SWidget> AnalyzerRackWidget = InAnalyzerRack->CreateWidget(DockTab);
+				DockTab->SetContent(AnalyzerRackWidget);
+				return DockTab;
 			}))
 			.SetDisplayName(LOCTEXT("AnalyzersTab", "Analyzers"))
 			.SetGroup(WorkspaceMenuCategoryRef)
@@ -1050,59 +1049,6 @@ namespace Metasound
 			InTabManager->UnregisterTabSpawner(TabNamesPrivate::Find);
 		}
 
-		TSharedPtr<SWidget> FEditor::BuildAnalyzerWidget() const
-		{
-			if (!OutputMeter.IsValid() || !OutputOscilloscope.IsValid() || !OutputVectorscope.IsValid() || !OutputSpectrumAnalyzer.IsValid())
-			{
-				return SNullWidget::NullWidget->AsShared();
-			}
-
-			const ISlateStyle* MetaSoundStyle = FSlateStyleRegistry::FindSlateStyle("MetaSoundStyle");
-			FLinearColor BackgroundColor = FLinearColor::Transparent;
-			if (ensure(MetaSoundStyle))
-			{
-				BackgroundColor = MetaSoundStyle->GetColor("MetasoundEditor.Analyzers.BackgroundColor");
-			}
-
-			return SNew(SOverlay)
-			+ SOverlay::Slot()
-			[
-				SNew(SColorBlock)
-				.Color(BackgroundColor)
-			]
-			+ SOverlay::Slot()
-			[
-				SNew(SSplitter)
-				.Orientation(Orient_Vertical)
-				+ SSplitter::Slot()
-				.Value(0.5f)
-				[
-					SNew(SVerticalBox)
-					+ SVerticalBox::Slot()
-					.HAlign(HAlign_Center)
-					.VAlign(VAlign_Fill)
-					[
-						OutputMeter->GetWidget()
-					]
-				]
-				+ SSplitter::Slot()
-				.Value(0.15f)
-				[
-					OutputOscilloscope->GetPanelWidget()
-				]
-				+ SSplitter::Slot()
-				.Value(0.15f)
-				[
-					OutputVectorscope->GetPanelWidget()
-				]
-				+ SSplitter::Slot()
-				.Value(0.15f)
-				[
-					OutputSpectrumAnalyzer->GetWidget()
-				]
-			];
-		}
-
 		bool FEditor::IsPlaying() const
 		{
 			if (UObject* MetaSound = GetMetasoundObject())
@@ -1744,182 +1690,29 @@ namespace Metasound
 
 		void FEditor::CreateAnalyzers(UMetaSoundSource& MetaSoundSource)
 		{
-			const ISlateStyle* StyleSet = FSlateStyleRegistry::FindSlateStyle("MetaSoundStyle");
-			if (ensure(GEditor) && ensure(StyleSet))
+			if (!AnalyzerRack.IsValid())
 			{
-				const Audio::FDeviceId AudioDeviceId = GEditor->GetMainAudioDeviceID();
-				constexpr UAudioBus* DefaultBus = nullptr;
-
-				if (!OutputMeter.IsValid())
-				{
-					const FName MaterialMeterStyleName("AudioMaterialMeter.Style");
-					if (StyleSet->HasWidgetStyle<FAudioMaterialMeterStyle>(MaterialMeterStyleName))
-					{
-						const FAudioMaterialMeterStyle& AudioMaterialMeterStyle = StyleSet->GetWidgetStyle<FAudioMaterialMeterStyle>(MaterialMeterStyleName);
-						OutputMeter = MakeShared<AudioWidgets::FAudioMeter>(MetaSoundSource.NumChannels, AudioDeviceId, AudioMaterialMeterStyle, DefaultBus);
-					}
-					else
-					{
-						const FAudioMeterDefaultColorStyle& MeterDefaultColorStyle = StyleSet->GetWidgetStyle<FAudioMeterDefaultColorStyle>("AudioMeter.DefaultColorStyle");
-						OutputMeter = MakeShared<AudioWidgets::FAudioMeter>(MetaSoundSource.NumChannels, AudioDeviceId, DefaultBus, &MeterDefaultColorStyle);
-					}
-				}
-				else if (OutputMeter->GetAudioBus()->GetNumChannels() != MetaSoundSource.NumChannels)
-				{
-					OutputMeter->Init(MetaSoundSource.NumChannels, AudioDeviceId, DefaultBus);
-				}
+				using namespace AudioWidgets;
 
-				const uint32 MetaSoundNumChannels = static_cast<uint32>(MetaSoundSource.NumChannels);
-
-				// Init Oscilloscope
-				constexpr float OscilloscopeTimeWindowMs     = 10.0f;
-				constexpr float OscilloscopeMaxTimeWindowMs  = 10.0f;
-				constexpr float OscilloscopeAnalysisPeriodMs = 10.0f;
-				constexpr EAudioPanelLayoutType OscilloscopePanelLayoutType = EAudioPanelLayoutType::Basic;
-
-				if (!OutputOscilloscope.IsValid())
-				{
-					OutputOscilloscope = MakeShared<AudioWidgets::FAudioOscilloscope>(AudioDeviceId,
-						MetaSoundNumChannels,
-						OscilloscopeTimeWindowMs,
-						OscilloscopeMaxTimeWindowMs,
-						OscilloscopeAnalysisPeriodMs,
-						OscilloscopePanelLayoutType,
-						&StyleSet->GetWidgetStyle<FAudioOscilloscopePanelStyle>("AudioOscilloscope.PanelStyle"));
-				}
-				else if (OutputOscilloscope->GetAudioBus()->GetNumChannels() != MetaSoundSource.NumChannels)
-				{
-					OutputOscilloscope->CreateAudioBus(MetaSoundNumChannels);
-					OutputOscilloscope->CreateDataProvider(AudioDeviceId, OscilloscopeTimeWindowMs, OscilloscopeMaxTimeWindowMs, OscilloscopeAnalysisPeriodMs, OscilloscopePanelLayoutType);
-					OutputOscilloscope->CreateOscilloscopeWidget(MetaSoundNumChannels, OscilloscopePanelLayoutType, &StyleSet->GetWidgetStyle<FAudioOscilloscopePanelStyle>("AudioOscilloscope.PanelStyle"));
-				}
-
-				// Init Vectorscope
-				constexpr float VectorscopeTimeWindowMs     = 30.0f;
-				constexpr float VectorscopeMaxTimeWindowMs  = 30.0f;
-				constexpr float VectorscopeAnalysisPeriodMs = 10.0f;
-				constexpr EAudioPanelLayoutType VectorscopePanelLayoutType = EAudioPanelLayoutType::Basic;
-
-				if (!OutputVectorscope.IsValid())
-				{
-					OutputVectorscope = MakeShared<AudioWidgets::FAudioVectorscope>(AudioDeviceId,
-						MetaSoundNumChannels,
-						VectorscopeTimeWindowMs,
-						VectorscopeMaxTimeWindowMs,
-						VectorscopeAnalysisPeriodMs,
-						VectorscopePanelLayoutType,
-						&StyleSet->GetWidgetStyle<FAudioVectorscopePanelStyle>("AudioVectorscope.PanelStyle"));
-				}
-				else if (OutputVectorscope->GetAudioBus()->GetNumChannels() != MetaSoundSource.NumChannels)
-				{
-					OutputVectorscope->CreateAudioBus(MetaSoundNumChannels);
-					OutputVectorscope->CreateDataProvider(AudioDeviceId, VectorscopeTimeWindowMs, VectorscopeMaxTimeWindowMs, VectorscopeAnalysisPeriodMs);
-					OutputVectorscope->CreateVectorscopeWidget(VectorscopePanelLayoutType, &StyleSet->GetWidgetStyle<FAudioVectorscopePanelStyle>("AudioVectorscope.PanelStyle"));
-				}
-
-				if (!OutputSpectrumAnalyzer.IsValid())
+				// Set params so that rack layout is stored specific to Metasound Editor and the Metasound Editor styles are used for analyzer widgets:
+				const FAudioAnalyzerRack::FRackConstructParams Params
 				{
-					AudioWidgets::FAudioSpectrumAnalyzerParams Params;
-					Params.NumChannels = MetaSoundSource.NumChannels;
-					Params.AudioDeviceId = AudioDeviceId;
-
-					Params.Ballistics.BindLambda([]()
-						{
-							return GetDefault<UMetasoundEditorSettings>()->SpectrumAnalyzerSettings.Ballistics;
-						});
-					Params.AnalyzerType.BindLambda([]()
-						{
-							return GetDefault<UMetasoundEditorSettings>()->SpectrumAnalyzerSettings.AnalyzerType;
-						});
-					Params.FFTAnalyzerFFTSize.BindLambda([]()
-						{
-							return GetDefault<UMetasoundEditorSettings>()->SpectrumAnalyzerSettings.FFTAnalyzerFFTSize;
-						});
-					Params.CQTAnalyzerFFTSize.BindLambda([]()
-						{
-							return GetDefault<UMetasoundEditorSettings>()->SpectrumAnalyzerSettings.CQTAnalyzerFFTSize;
-						});
-					Params.TiltExponent.BindLambda([]()
-						{
-							const EAudioSpectrumPlotTilt TiltSpectrum = GetDefault<UMetasoundEditorSettings>()->SpectrumAnalyzerSettings.TiltSpectrum;
-							return SAudioSpectrumPlot::GetTiltExponentValue(TiltSpectrum);
-						});
-					Params.FrequencyAxisPixelBucketMode.BindLambda([]()
-						{
-							return GetDefault<UMetasoundEditorSettings>()->SpectrumAnalyzerSettings.PixelPlotMode;
-						});
-					Params.FrequencyAxisScale.BindLambda([]()
-						{
-							return GetDefault<UMetasoundEditorSettings>()->SpectrumAnalyzerSettings.FrequencyScale;
-						});
-					Params.bDisplayFrequencyAxisLabels.BindLambda([]()
-						{
-							return GetDefault<UMetasoundEditorSettings>()->SpectrumAnalyzerSettings.bDisplayFrequencyAxisLabels;
-						});
-					Params.bDisplaySoundLevelAxisLabels.BindLambda([]()
-						{
-							return GetDefault<UMetasoundEditorSettings>()->SpectrumAnalyzerSettings.bDisplaySoundLevelAxisLabels;
-						});
-
-					Params.OnBallisticsMenuEntryClicked.BindLambda([](EAudioSpectrumAnalyzerBallistics SelectedValue)
-						{
-							GetMutableDefault<UMetasoundEditorSettings>()->SpectrumAnalyzerSettings.Ballistics = SelectedValue;
-							GetMutableDefault<UMetasoundEditorSettings>()->SaveConfig();
-						});
-					Params.OnAnalyzerTypeMenuEntryClicked.BindLambda([](EAudioSpectrumAnalyzerType SelectedValue)
-						{
-							GetMutableDefault<UMetasoundEditorSettings>()->SpectrumAnalyzerSettings.AnalyzerType = SelectedValue;
-							GetMutableDefault<UMetasoundEditorSettings>()->SaveConfig();
-						});
-					Params.OnFFTAnalyzerFFTSizeMenuEntryClicked.BindLambda([](EFFTSize SelectedValue)
-						{
-							GetMutableDefault<UMetasoundEditorSettings>()->SpectrumAnalyzerSettings.FFTAnalyzerFFTSize = SelectedValue;
-							GetMutableDefault<UMetasoundEditorSettings>()->SaveConfig();
-						});
-					Params.OnCQTAnalyzerFFTSizeMenuEntryClicked.BindLambda([](EConstantQFFTSizeEnum SelectedValue)
-						{
-							GetMutableDefault<UMetasoundEditorSettings>()->SpectrumAnalyzerSettings.CQTAnalyzerFFTSize = SelectedValue;
-							GetMutableDefault<UMetasoundEditorSettings>()->SaveConfig();
-						});
-					Params.OnTiltSpectrumMenuEntryClicked.BindLambda([](EAudioSpectrumPlotTilt SelectedValue)
-						{
-							GetMutableDefault<UMetasoundEditorSettings>()->SpectrumAnalyzerSettings.TiltSpectrum = SelectedValue;
-							GetMutableDefault<UMetasoundEditorSettings>()->SaveConfig();
-						});
-					Params.OnFrequencyAxisPixelBucketModeMenuEntryClicked.BindLambda([](EAudioSpectrumPlotFrequencyAxisPixelBucketMode SelectedValue)
-						{
-							GetMutableDefault<UMetasoundEditorSettings>()->SpectrumAnalyzerSettings.PixelPlotMode = SelectedValue;
-							GetMutableDefault<UMetasoundEditorSettings>()->SaveConfig();
-						});
-					Params.OnFrequencyAxisScaleMenuEntryClicked.BindLambda([](EAudioSpectrumPlotFrequencyAxisScale SelectedValue)
-						{
-							GetMutableDefault<UMetasoundEditorSettings>()->SpectrumAnalyzerSettings.FrequencyScale = SelectedValue;
-							GetMutableDefault<UMetasoundEditorSettings>()->SaveConfig();
-						});
-					Params.OnDisplayFrequencyAxisLabelsButtonToggled.BindLambda([]()
-						{
-							FMetasoundEditorSpectrumAnalyzerSettings& SpectrumAnalyzerSettings = GetMutableDefault<UMetasoundEditorSettings>()->SpectrumAnalyzerSettings;
-							SpectrumAnalyzerSettings.bDisplayFrequencyAxisLabels = !SpectrumAnalyzerSettings.bDisplayFrequencyAxisLabels;
-							GetMutableDefault<UMetasoundEditorSettings>()->SaveConfig();
-						});
-					Params.OnDisplaySoundLevelAxisLabelsButtonToggled.BindLambda([]()
-						{
-							FMetasoundEditorSpectrumAnalyzerSettings& SpectrumAnalyzerSettings = GetMutableDefault<UMetasoundEditorSettings>()->SpectrumAnalyzerSettings;
-							SpectrumAnalyzerSettings.bDisplaySoundLevelAxisLabels = !SpectrumAnalyzerSettings.bDisplaySoundLevelAxisLabels;
-							GetMutableDefault<UMetasoundEditorSettings>()->SaveConfig();
-						});
-					Params.PlotStyle = &StyleSet->GetWidgetStyle<FAudioSpectrumPlotStyle>("AudioSpectrumPlot.Style");
-					OutputSpectrumAnalyzer = MakeShared<AudioWidgets::FAudioSpectrumAnalyzer>(Params);
-				}
-				else if (OutputSpectrumAnalyzer->GetAudioBus()->GetNumChannels() != MetaSoundSource.NumChannels)
-				{
-					OutputSpectrumAnalyzer->Init(MetaSoundSource.NumChannels, AudioDeviceId, nullptr);
-				}
-
-				return;
+					.TabManagerLayoutName = TEXT("MetasoundEditor_FAudioAnalyzerRack_v0"),
+					.StyleSet = FSlateStyleRegistry::FindSlateStyle("MetaSoundStyle"),
+					.EditorSettingsClass = UMetasoundEditorSettings::StaticClass(),
+				};
+				AnalyzerRack = MakeShared<FAudioAnalyzerRack>(Params);
 			}
 
-			DestroyAnalyzers();
+			if (ensure(GEditor))
+			{
+				// Init (or reinit) the analyzer rack with the correct num channels:
+				AnalyzerRack->Init(MetaSoundSource.NumChannels, GEditor->GetMainAudioDeviceID());
+			}
+			else
+			{
+				AnalyzerRack->DestroyAnalyzers();
+			}
 		}
 
 		TSharedRef<SWidget> FEditor::CreateAuditionMenuOptions()
@@ -2165,10 +1958,10 @@ namespace Metasound
 
 		void FEditor::DestroyAnalyzers()
 		{
-			OutputMeter.Reset();
-			OutputOscilloscope.Reset();
-			OutputVectorscope.Reset();
-			OutputSpectrumAnalyzer.Reset();
+			if (AnalyzerRack.IsValid())
+			{
+				AnalyzerRack->DestroyAnalyzers();
+			}
 		}
 
 		void FEditor::ExtendToolbarInternal()
@@ -2612,22 +2405,7 @@ namespace Metasound
 					{
 						SetPreviewID(PreviewComp->GetUniqueID());
 
-						if (UAudioBus* AudioBus = OutputMeter->GetAudioBus())
-						{
-							PreviewComp->SetAudioBusSendPostEffect(AudioBus, 1.0f);
-						}
-
-						if (UAudioBus* AudioBus = OutputOscilloscope->GetAudioBus())
-						{
-							PreviewComp->SetAudioBusSendPostEffect(AudioBus, 1.0f);
-						}
-
-						if (UAudioBus* AudioBus = OutputVectorscope->GetAudioBus())
-						{
-							PreviewComp->SetAudioBusSendPostEffect(AudioBus, 1.0f);
-						}
-
-						if (UAudioBus* AudioBus = OutputSpectrumAnalyzer->GetAudioBus())
+						if (UAudioBus* AudioBus = AnalyzerRack->GetAudioBus())
 						{
 							PreviewComp->SetAudioBusSendPostEffect(AudioBus, 1.0f);
 						}
@@ -2655,39 +2433,7 @@ namespace Metasound
 					})
 				);
 
-				TSharedPtr<SAudioMeterBase> OutputMeterWidget = OutputMeter->GetWidget<SAudioMeterBase>();
-				if (OutputMeterWidget.IsValid())
-				{
-					if (!OutputMeterWidget->bIsActiveTimerRegistered)
-					{
-						OutputMeterWidget->RegisterActiveTimer(0.0f,
-							FWidgetActiveTimerDelegate::CreateLambda([this](double InCurrentTime, float InDeltaTime)
-							{
-								if (IsPlaying())
-								{
-									return EActiveTimerReturnType::Continue;
-								}
-								else
-								{
-									TSharedRef<SAudioMeterBase> MeterRef = OutputMeter->GetWidget<SAudioMeterBase>();
-									MeterRef->bIsActiveTimerRegistered = false;
-									return EActiveTimerReturnType::Stop;
-								}
-							})
-						);
-						OutputMeterWidget->bIsActiveTimerRegistered = true;
-					}
-				}
-
-				if (OutputOscilloscope.IsValid())
-				{
-					OutputOscilloscope->StartProcessing();
-				}
-
-				if (OutputVectorscope.IsValid())
-				{
-					OutputVectorscope->StartProcessing();
-				}
+				AnalyzerRack->StartProcessing();
 			}
 		}
 
diff --git a/Engine/Plugins/Runtime/Metasound/Source/MetasoundEditor/Private/MetasoundEditor.h b/Engine/Plugins/Runtime/Metasound/Source/MetasoundEditor/Private/MetasoundEditor.h
--- a/Engine/Plugins/Runtime/Metasound/Source/MetasoundEditor/Private/MetasoundEditor.h
+++ b/Engine/Plugins/Runtime/Metasound/Source/MetasoundEditor/Private/MetasoundEditor.h
@@ -1,8 +1,6 @@
 // Copyright Epic Games, Inc. All Rights Reserved.
 #pragma once
 
-#include "AudioMeter.h"
-#include "AudioMeterStyle.h"
 #include "EdGraph/EdGraphNode.h"
 #include "EditorUndoClient.h"
 #include "Framework/Commands/UICommandList.h"
@@ -16,9 +14,7 @@
 #include "MetasoundFrontend.h"
 #include "MetasoundFrontendController.h"
 #include "MetasoundFrontendDocumentModifyDelegates.h"
-#include "Meter.h"
 #include "Misc/NotifyHook.h"
-#include "SAudioMeter.h"
 #include "SGraphActionMenu.h"
 #include "SGraphPanel.h"
 #include "SMetasoundPalette.h"
@@ -40,9 +36,7 @@
 // Forward Declarations
 namespace AudioWidgets
 {
-	class FAudioOscilloscope;
-	class FAudioVectorscope;
-	class FAudioSpectrumAnalyzer;
+	class FAudioAnalyzerRack;
 }
 
 class FTabManager;
@@ -69,7 +63,6 @@ class UMetasoundEditorGraph;
 class UMetasoundEditorGraphNode;
 
 struct FGraphActionNode;
-struct FMeterResults;
 struct FPropertyChangedEvent;
 
 
@@ -514,8 +507,6 @@ namespace Metasound::Editor
 		/** Create new graph editor widget */
 		void CreateGraphEditorWidget(UObject& MetaSound);
 
-		TSharedPtr<SWidget> BuildAnalyzerWidget() const;
-
 		void EditObjectSettings();
 
 		void NotifyAssetPrimeComplete();
@@ -556,17 +547,8 @@ namespace Metasound::Editor
 		/** Metasound graph members menu */
 		TSharedPtr<SGraphActionMenu> GraphMembersMenu;
 
-		/** Meter used in the analyzer tab for auditioning preview output. */
-		TSharedPtr<AudioWidgets::FAudioMeter> OutputMeter;
-
-		/** Oscilloscope used in the analyzer tab for visualizing preview output. */
-		TSharedPtr<AudioWidgets::FAudioOscilloscope> OutputOscilloscope;
-
-		/** Vectorscope used in the analyzer tab for visualizing preview output. */
-		TSharedPtr<AudioWidgets::FAudioVectorscope> OutputVectorscope;
-
-		/** Spectrum analyzer used in the analyzer tab for auditioning preview output. */
-		TSharedPtr<AudioWidgets::FAudioSpectrumAnalyzer> OutputSpectrumAnalyzer;
+		/** Displayed in the analyzer tab for visualizing preview output. */
+		TSharedPtr<AudioWidgets::FAudioAnalyzerRack> AnalyzerRack;
 
 		/** Find in MetaSound widget. */
 		TSharedPtr<SFindInMetasound> FindWidget;
diff --git a/Engine/Plugins/Runtime/Metasound/Source/MetasoundEditor/Private/MetasoundEditorModule.cpp b/Engine/Plugins/Runtime/Metasound/Source/MetasoundEditor/Private/MetasoundEditorModule.cpp
--- a/Engine/Plugins/Runtime/Metasound/Source/MetasoundEditor/Private/MetasoundEditorModule.cpp
+++ b/Engine/Plugins/Runtime/Metasound/Source/MetasoundEditor/Private/MetasoundEditorModule.cpp
@@ -201,6 +201,7 @@ namespace Metasound
 
 				// Audio Widgets
 				{
+					const FLinearColor& AnalyzerBackgroundColor = GetColor("MetasoundEditor.Analyzers.BackgroundColor");
 					const FLinearColor& AnalyzerForegroundColor = GetColor("MetasoundEditor.Analyzers.ForegroundColor");
 
 					// Add static overrides for these widget styles:
@@ -220,6 +221,8 @@ namespace Metasound
 						.SetVectorViewerStyle(FSampledSequenceVectorViewerStyle()
 							.SetLineColor(AnalyzerForegroundColor)));
 
+					Set("AudioAnalyzerRack.BackgroundColor", AnalyzerBackgroundColor);
+
 					// Add dynamic overrides for the widget style types with settings overrides:
 					AddDynamicLoadedWidgetStyle<FAudioMaterialKnobStyle>([]() { return GetDefault<UMetasoundEditorSettings>()->KnobStyleOverride; });
 					AddDynamicLoadedWidgetStyle<FAudioMaterialSliderStyle>([]() { return GetDefault<UMetasoundEditorSettings>()->SliderStyleOverride; });
diff --git a/Engine/Plugins/Runtime/Metasound/Source/MetasoundEditor/Private/MetasoundEditorSettings.h b/Engine/Plugins/Runtime/Metasound/Source/MetasoundEditor/Private/MetasoundEditorSettings.h
--- a/Engine/Plugins/Runtime/Metasound/Source/MetasoundEditor/Private/MetasoundEditorSettings.h
+++ b/Engine/Plugins/Runtime/Metasound/Source/MetasoundEditor/Private/MetasoundEditorSettings.h
@@ -100,39 +100,6 @@ struct FMetasoundAnalyzerAnimationSettings
 	float WireScalarMax = 4.5f;
 };
 
-USTRUCT()
-struct FMetasoundEditorSpectrumAnalyzerSettings
-{
-	GENERATED_BODY()
-	
-	UPROPERTY(EditAnywhere, config, Category = SpectrumAnalyzer)
-	EAudioSpectrumAnalyzerBallistics Ballistics = EAudioSpectrumAnalyzerBallistics::Digital;
-
-	UPROPERTY(EditAnywhere, config, Category = SpectrumAnalyzer)
-	EAudioSpectrumAnalyzerType AnalyzerType = EAudioSpectrumAnalyzerType::CQT;
-
-	UPROPERTY(EditAnywhere, config, Category = SpectrumAnalyzer, meta = (DisplayName = "FFT Size (FFT Analyzer)"))
-	EFFTSize FFTAnalyzerFFTSize = EFFTSize::Max;
-
-	UPROPERTY(EditAnywhere, config, Category = SpectrumAnalyzer, meta = (DisplayName = "FFT Size (CQT Analyzer)"))
-	EConstantQFFTSizeEnum CQTAnalyzerFFTSize = EConstantQFFTSizeEnum::XXLarge;
-
-	UPROPERTY(EditAnywhere, config, Category = SpectrumAnalyzer)
-	EAudioSpectrumPlotTilt TiltSpectrum = EAudioSpectrumPlotTilt::NoTilt;
-
-	UPROPERTY(EditAnywhere, config, Category = SpectrumAnalyzer)
-	EAudioSpectrumPlotFrequencyAxisPixelBucketMode PixelPlotMode = EAudioSpectrumPlotFrequencyAxisPixelBucketMode::Average;
-
-	UPROPERTY(EditAnywhere, config, Category = SpectrumAnalyzer)
-	EAudioSpectrumPlotFrequencyAxisScale FrequencyScale = EAudioSpectrumPlotFrequencyAxisScale::Logarithmic;
-
-	UPROPERTY(EditAnywhere, config, Category = SpectrumAnalyzer)
-	bool bDisplayFrequencyAxisLabels = false;
-
-	UPROPERTY(EditAnywhere, config, Category = SpectrumAnalyzer)
-	bool bDisplaySoundLevelAxisLabels = false;
-};
-
 UCLASS(config=EditorPerProjectUserSettings)
 class METASOUNDEDITOR_API UMetasoundEditorSettings : public UObject
 {
@@ -256,7 +223,7 @@ public:
 
 	/** Settings for metasound output spectrum analyzer widget */
 	UPROPERTY(EditAnywhere, config, Category = SpectrumAnalyzer, meta = (ShowOnlyInnerProperties))
-	FMetasoundEditorSpectrumAnalyzerSettings SpectrumAnalyzerSettings;
+	FSpectrumAnalyzerRackUnitSettings SpectrumAnalyzerSettings;
 
 	/** Widget type to show on input nodes by default */
 	UPROPERTY(EditAnywhere, config, Category = General)
