FGraphController:SetDefaultInput, actually set the default input after getting the default literal from the document and page id

#tests Harmonix Unit Tests
#rnx

diff --git a/Engine/Plugins/Runtime/Metasound/Source/MetasoundFrontend/Private/MetasoundFrontendGraphController.cpp b/Engine/Plugins/Runtime/Metasound/Source/MetasoundFrontend/Private/MetasoundFrontendGraphController.cpp
--- a/Engine/Plugins/Runtime/Metasound/Source/MetasoundFrontend/Private/MetasoundFrontendGraphController.cpp
+++ b/Engine/Plugins/Runtime/Metasound/Source/MetasoundFrontend/Private/MetasoundFrontendGraphController.cpp
@@ -1276,7 +1276,8 @@ namespace Metasound
 				if (ensure(IDataTypeRegistry::Get().IsLiteralTypeSupported(Desc->TypeName, InLiteral.GetType())))
 				{
 					const FGuid TargetPageID = IDocumentBuilderRegistry::GetChecked().ResolveTargetPageID(*Desc);
-					const FMetasoundFrontendLiteral& DefaultLiteral = Desc->FindConstDefaultChecked(TargetPageID);
+					FMetasoundFrontendLiteral& DefaultLiteral = Desc->FindDefaultChecked(TargetPageID);
+					DefaultLiteral = InLiteral;
 					return true;
 				}
 				else
