Fix for cook of MetaSound with SceneGraph enabled
#rnx

diff --git a/Engine/Plugins/Runtime/Metasound/Source/MetasoundEngine/Private/MetasoundSettings.cpp b/Engine/Plugins/Runtime/Metasound/Source/MetasoundEngine/Private/MetasoundSettings.cpp
--- a/Engine/Plugins/Runtime/Metasound/Source/MetasoundEngine/Private/MetasoundSettings.cpp
+++ b/Engine/Plugins/Runtime/Metasound/Source/MetasoundEngine/Private/MetasoundSettings.cpp
@@ -289,7 +289,7 @@ TArray<FGuid> UMetaSoundSettings::GetCookedTargetPageIDs(FName PlatformName) con
 
 const TArray<FGuid>& UMetaSoundSettings::GetCookedTargetPageIDsInternal(FName PlatformName) const
 {
-	if (PlatformName != CookPlatformTargetPage)
+	if ((PlatformName != CookPlatformTargetPage) || (PlatformName.IsNone() && CookPlatformTargetPageIDs.IsEmpty()))
 	{
 		CookPlatformTargetPage = PlatformName;
 		CookPlatformTargetPageIDs.Reset();
@@ -318,7 +318,10 @@ const TArray<FGuid>& UMetaSoundSettings::GetCookedTargetPageIDsInternal(FName Pl
 		if (CookPlatformTargetPageIDs.IsEmpty())
 		{
 #if WITH_EDITOR
-			const bool bCanTargetDefault = DefaultPageSettings.CanTarget.GetValueForPlatform(PlatformName);
+			// Allow default page to be a target if no platform name is provided.
+			// This can occur when generating reflection code for a MetaSound
+			// in a cloud cooker. 
+			const bool bCanTargetDefault = PlatformName.IsNone() || DefaultPageSettings.CanTarget.GetValueForPlatform(PlatformName);
 #else // !WITH_EDITOR
 			const bool bCanTargetDefault = DefaultPageSettings.CanTarget.GetValue();
 #endif // !WITH_EDITOR
