Fix for opening a MetaSound Patch crashes the editor.

diff --git a/Engine/Plugins/Runtime/Metasound/Source/MetasoundEditor/Private/MetasoundEditor.cpp b/Engine/Plugins/Runtime/Metasound/Source/MetasoundEditor/Private/MetasoundEditor.cpp
--- a/Engine/Plugins/Runtime/Metasound/Source/MetasoundEditor/Private/MetasoundEditor.cpp
+++ b/Engine/Plugins/Runtime/Metasound/Source/MetasoundEditor/Private/MetasoundEditor.cpp
@@ -986,8 +986,11 @@ namespace Metasound
 			InTabManager->RegisterTabSpawner(TabNamesPrivate::Analyzers, FOnSpawnTab::CreateLambda([InAnalyzerRack = AnalyzerRack](const FSpawnTabArgs&)
 			{
 				TSharedRef<SDockTab> DockTab = SNew(SDockTab).Label(LOCTEXT("MetasoundAnalyzersTitle", "Analyzers"));
-				TSharedRef<SWidget> AnalyzerRackWidget = InAnalyzerRack->CreateWidget(DockTab);
-				DockTab->SetContent(AnalyzerRackWidget);
+				if (InAnalyzerRack.IsValid()) // AnalyzerRack only created for MetaSoundSource, not MetaSoundPatch.
+				{
+					TSharedRef<SWidget> AnalyzerRackWidget = InAnalyzerRack->CreateWidget(DockTab);
+					DockTab->SetContent(AnalyzerRackWidget);
+				}
 				return DockTab;
 			}))
 			.SetDisplayName(LOCTEXT("AnalyzersTab", "Analyzers"))
