Fix for UnitTests that are still on actively-deprecating Controller API where builder registry is not implemented.
#rnx

diff --git a/Engine/Plugins/Runtime/Metasound/Source/MetasoundEditor/Private/MetasoundEditorGraphBuilder.cpp b/Engine/Plugins/Runtime/Metasound/Source/MetasoundEditor/Private/MetasoundEditorGraphBuilder.cpp
--- a/Engine/Plugins/Runtime/Metasound/Source/MetasoundEditor/Private/MetasoundEditorGraphBuilder.cpp
+++ b/Engine/Plugins/Runtime/Metasound/Source/MetasoundEditor/Private/MetasoundEditorGraphBuilder.cpp
@@ -2667,14 +2667,15 @@ namespace Metasound
 			// Collect all editor graph inputs with corresponding frontend inputs. 
 			GraphHandle->IterateConstNodes([&](FConstNodeHandle NodeHandle)
 			{
-				if (UMetasoundEditorGraphInput* Input = OutGraph.FindInput(NodeHandle->GetID()))
+				const FGuid NodeID = NodeHandle->GetID();
+				if (UMetasoundEditorGraphInput* Input = OutGraph.FindInput(NodeID))
 				{
 					Inputs.Add(Input);
 					return;
 				}
 
 				// Add an editor input if none exist for a frontend input.
-				Inputs.Add(OutGraph.FindOrAddInput(NodeHandle));
+				Inputs.Add(OutGraph.FindOrAddInput(NodeID));
 				constexpr bool bIncludeNamespace = true;
 				FText NodeDisplayName = FGraphBuilder::GetDisplayName(*NodeHandle, bIncludeNamespace);
 				UE_LOG(LogMetasoundEditor, Verbose, TEXT("Synchronizing Inputs: Added missing input '%s'."), *NodeDisplayName.ToString());
@@ -2710,14 +2711,15 @@ namespace Metasound
 			// Collect all editor graph outputs with corresponding frontend outputs. 
 			GraphHandle->IterateConstNodes([&](FConstNodeHandle NodeHandle)
 			{
-				if (UMetasoundEditorGraphOutput* Output = OutGraph.FindOutput(NodeHandle->GetID()))
+				const FGuid NodeID = NodeHandle->GetID();
+				if (UMetasoundEditorGraphOutput* Output = OutGraph.FindOutput(NodeID))
 				{
 					Outputs.Add(Output);
 					return;
 				}
 
 				// Add an editor output if none exist for a frontend output.
-				Outputs.Add(OutGraph.FindOrAddOutput(NodeHandle));
+				Outputs.Add(OutGraph.FindOrAddOutput(NodeID));
 				constexpr bool bIncludeNamespace = true;
 				FText NodeDisplayName = FGraphBuilder::GetDisplayName(*NodeHandle, bIncludeNamespace);
 				UE_LOG(LogMetasoundEditor, Verbose, TEXT("Synchronizing Outputs: Added missing output '%s'."), *NodeDisplayName.ToString());
diff --git a/Engine/Plugins/Runtime/Metasound/Source/MetasoundFrontend/Private/MetasoundFrontendGraphController.cpp b/Engine/Plugins/Runtime/Metasound/Source/MetasoundFrontend/Private/MetasoundFrontendGraphController.cpp
--- a/Engine/Plugins/Runtime/Metasound/Source/MetasoundFrontend/Private/MetasoundFrontendGraphController.cpp
+++ b/Engine/Plugins/Runtime/Metasound/Source/MetasoundFrontend/Private/MetasoundFrontendGraphController.cpp
@@ -1262,7 +1262,11 @@ namespace Metasound
 		{
 			if (const FMetasoundFrontendClassInput* Desc = FindInputDescriptionWithVertexID(InVertexID))
 			{
-				const FGuid TargetPageID = IDocumentBuilderRegistry::GetChecked().ResolveTargetPageID(*Desc);
+				FGuid TargetPageID = Frontend::DefaultPageID;
+				if (IDocumentBuilderRegistry* DocRegistry = IDocumentBuilderRegistry::Get())
+				{
+					TargetPageID = DocRegistry->ResolveTargetPageID(*Desc);
+				}
 				const FMetasoundFrontendLiteral& DefaultLiteral = Desc->FindConstDefaultChecked(TargetPageID);
 				return DefaultLiteral;
 			}
@@ -1275,7 +1279,11 @@ namespace Metasound
 			{
 				if (ensure(IDataTypeRegistry::Get().IsLiteralTypeSupported(Desc->TypeName, InLiteral.GetType())))
 				{
-					const FGuid TargetPageID = IDocumentBuilderRegistry::GetChecked().ResolveTargetPageID(*Desc);
+					FGuid TargetPageID = Frontend::DefaultPageID;
+					if (IDocumentBuilderRegistry* DocRegistry = IDocumentBuilderRegistry::Get())
+					{
+						TargetPageID = DocRegistry->ResolveTargetPageID(*Desc);
+					}
 					FMetasoundFrontendLiteral& DefaultLiteral = Desc->FindDefaultChecked(TargetPageID);
 					DefaultLiteral = InLiteral;
 					return true;
