Allow MetaSound parameter reflection during cook
#rnx

diff --git a/Engine/Plugins/Runtime/Metasound/Source/MetasoundEngine/Private/MetasoundSource.cpp b/Engine/Plugins/Runtime/Metasound/Source/MetasoundEngine/Private/MetasoundSource.cpp
--- a/Engine/Plugins/Runtime/Metasound/Source/MetasoundEngine/Private/MetasoundSource.cpp
+++ b/Engine/Plugins/Runtime/Metasound/Source/MetasoundEngine/Private/MetasoundSource.cpp
@@ -1118,18 +1118,36 @@ bool UMetaSoundSource::GetAllDefaultParameters(TArray<FAudioParameter>& OutParam
 	using namespace Metasound::Frontend;
 	using namespace Metasound::Engine;
 
-	
-	if(!RuntimeInputData.bIsValid.load())
+	if (Metasound::CanEverExecuteGraph())
 	{
-		UE_CLOG(SourcePrivate::HasNotBeenLoggedForThisObject(*this, __LINE__), LogMetaSound, Warning, TEXT("Default parameters will be ommitted. Accessing invalid runtime data on MetaSound %s. Ensure that UMetaSoundSource::InitResources() is executed on the game thread before calling UMetaSoundSource::GetAllDefaultParameters(...)"), *GetOwningAssetName());
-		return false;
-	}
+		// If we are in an executable context, we expect runtime input data to be
+		// cached as default parameters may be requested many times. 
+		if(!RuntimeInputData.bIsValid.load())
+		{
+			UE_CLOG(SourcePrivate::HasNotBeenLoggedForThisObject(*this, __LINE__), LogMetaSound, Warning, TEXT("Default parameters will be ommitted. Accessing invalid runtime data on MetaSound %s. Ensure that UMetaSoundSource::InitResources() is executed on the game thread before calling UMetaSoundSource::GetAllDefaultParameters(...)"), *GetOwningAssetName());
+			return false;
+		}
 
-	for(const TPair<FVertexName, FRuntimeInput>& Pair : RuntimeInputData.InputMap)
+		for(const TPair<FVertexName, FRuntimeInput>& Pair : RuntimeInputData.InputMap)
+		{
+			OutParameters.Add(Pair.Value.DefaultParameter);
+		}
+		return true;
+	}
+	else
 	{
-		OutParameters.Add(Pair.Value.DefaultParameter);
+		// Non executable contexts do not have cached runtime data and need their 
+		// parameters created each time they are requested. 
+		constexpr bool bCreateUObjectProxies = false;
+		TSortedVertexNameMap<FRuntimeInput> InputMap = CreateRuntimeInputMap(bCreateUObjectProxies);
+
+		for(const TPair<FVertexName, FRuntimeInput>& Pair : InputMap)
+		{
+			OutParameters.Add(Pair.Value.DefaultParameter);
+		}
+
+		return true;
 	}
-	return true;
 }
 
 void UMetaSoundSource::InitParametersInternal(const Metasound::TSortedVertexNameMap<FRuntimeInput>& InInputMap, TArray<FAudioParameter>& ParametersToInit, FName InFeatureName) const
