Fix crash when creating multiple MetaSound presets with the same name

#rnx

diff --git a/Engine/Plugins/Runtime/Metasound/Source/MetasoundEditor/Private/MetasoundAssetDefinitions.cpp b/Engine/Plugins/Runtime/Metasound/Source/MetasoundEditor/Private/MetasoundAssetDefinitions.cpp
--- a/Engine/Plugins/Runtime/Metasound/Source/MetasoundEditor/Private/MetasoundAssetDefinitions.cpp
+++ b/Engine/Plugins/Runtime/Metasound/Source/MetasoundEditor/Private/MetasoundAssetDefinitions.cpp
@@ -428,7 +428,14 @@ namespace MenuExtension_MetaSoundSourceTemplate
 				}
 			}
 
-			// Sync content browser to newly created assets
+			// Sync content browser to newly created valid assets
+			// Assets can be invalid if multiple assets are created with the same name 
+			// then force overwritten within the same operation
+			ObjectsToSync.RemoveAllSwap([](const UObject* InObject)
+			{
+				return !InObject || !InObject->IsValidLowLevelFast();
+			});
+
 			if (ObjectsToSync.Num() > 0)
 			{
 				FContentBrowserModule& ContentBrowserModule = FModuleManager::LoadModuleChecked<FContentBrowserModule>("ContentBrowser");
