Fix for crash when building an metasound operator
#rnx

diff --git a/Engine/Plugins/Runtime/Metasound/Source/MetasoundGraphCore/Private/MetasoundGraphOperator.cpp b/Engine/Plugins/Runtime/Metasound/Source/MetasoundGraphCore/Private/MetasoundGraphOperator.cpp
--- a/Engine/Plugins/Runtime/Metasound/Source/MetasoundGraphCore/Private/MetasoundGraphOperator.cpp
+++ b/Engine/Plugins/Runtime/Metasound/Source/MetasoundGraphCore/Private/MetasoundGraphOperator.cpp
@@ -91,7 +91,12 @@ namespace Metasound
 			// simply uses the pointer address as the ID and does not access the 
 			// actual underlying INode. 
 			const FOperatorID OperatorID = GetOperatorID(Node);
-			AppendOperator(MoveTemp(InOperatorState->OperatorMap[OperatorID].Operator));
+
+			FGraphOperatorData::FOperatorInfo* OperatorInfo = InOperatorState->OperatorMap.Find(OperatorID);
+			if (ensureMsgf(OperatorInfo, TEXT("Encountered possible corrupt operator data")))
+			{
+				AppendOperator(MoveTemp(OperatorInfo->Operator));
+			}
 		}
 
 		// Copy over vertex data
