MetaSound editor inspection via Builder BP API + ability to show user widget in MetaSound preset editor
- New class MetaSound Editor Builder Listener for subscribing to editor changes
- Various MetaSound builder BP API improvements for literals and graph inspection, new delegate for on data type changed
- MetaSoundPresetWidgetInterface when implemented and associated with a MetaSound, will show the widget in the preset editor
- Example editor utility widget that creates knobs for a given MetaSound's float inputs and updates along with editor changes

#virtualized

diff --git a/Engine/Plugins/Runtime/Metasound/Source/MetasoundEditor/MetasoundEditor.Build.cs b/Engine/Plugins/Runtime/Metasound/Source/MetasoundEditor/MetasoundEditor.Build.cs
--- a/Engine/Plugins/Runtime/Metasound/Source/MetasoundEditor/MetasoundEditor.Build.cs
+++ b/Engine/Plugins/Runtime/Metasound/Source/MetasoundEditor/MetasoundEditor.Build.cs
@@ -30,7 +30,9 @@ namespace UnrealBuildTool.Rules
 					"ToolMenus",
 					"ToolWidgets",
 					"WaveTable",
-					"WaveTableEditor"
+					"WaveTableEditor",
+					"UMG"
+
 				}
 			);
 
diff --git a/Engine/Plugins/Runtime/Metasound/Source/MetasoundEditor/Private/MetasoundEditor.cpp b/Engine/Plugins/Runtime/Metasound/Source/MetasoundEditor/Private/MetasoundEditor.cpp
--- a/Engine/Plugins/Runtime/Metasound/Source/MetasoundEditor/Private/MetasoundEditor.cpp
+++ b/Engine/Plugins/Runtime/Metasound/Source/MetasoundEditor/Private/MetasoundEditor.cpp
@@ -6,6 +6,8 @@
 #include "AudioAnalyzerRack.h"
 #include "AudioDevice.h"
 #include "AudioWidgetsEnums.h"
+#include "Audio/AudioWidgetSubsystem.h"
+#include "Blueprint/UserWidget.h"
 #include "Components/AudioComponent.h"
 #include "DetailLayoutBuilder.h"
 #include "EdGraph/EdGraphNode.h"
@@ -60,6 +62,7 @@
 #include "MetasoundGenerator.h"
 #include "MetasoundLog.h"
 #include "MetasoundNodeDetailCustomization.h"
+#include "MetasoundPresetWidgetInterface.h"
 #include "MetasoundSettings.h"
 #include "MetasoundSource.h"
 #include "MetasoundUObjectRegistry.h"
@@ -192,6 +195,7 @@ namespace Metasound
 			const FName Interfaces = "MetasoundEditor_Interfaces";
 			const FName Pages = "MetasoundEditor_Pages";
 			const FName Find = "MetasoundEditor_Find";
+			const FName UserPresetWidget = "MetasoundEditor_UserPresetWidget";
 		} // namespace TabNamesPrivate
 
 		static const TArray<FText> NodeSectionNames
@@ -918,7 +922,9 @@ namespace Metasound
 					InRenderStatsWidget = RenderStatsWidget
 				] (const FSpawnTabArgs& Args)
 			{
-				TSharedRef<SDockTab> SpawnedTab = SNew(SDockTab).Label(LOCTEXT("MetasoundGraphCanvasTitle", "MetaSound Graph"));
+				TSharedRef<SDockTab> SpawnedTab = SNew(SDockTab).Label(
+					FText::Format(LOCTEXT("GraphCanvasTabLabel", "Graph Page: {0}"), InPageStatsWidget.IsValid() ?
+						FText::FromName(InPageStatsWidget->GetDisplayedPageName()) : LOCTEXT("MetasoundGraphCanvasTitle", "MetaSound Graph")));
 
 				TSharedRef<SOverlay> Overlay = SNew(SOverlay)
 					+ SOverlay::Slot()
@@ -1020,6 +1026,27 @@ namespace Metasound
 				}
 			}
 
+			if (UserPresetWidget.IsValid())
+			{
+				const FCanSpawnTab CanSpawnTab = FCanSpawnTab::CreateLambda([this](const FSpawnTabArgs&)
+				{
+					return Builder.IsValid() && Builder->IsPreset();
+				});
+
+				InTabManager->RegisterTabSpawner(TabNamesPrivate::UserPresetWidget, FOnSpawnTab::CreateLambda([this, InUserPresetWidget = UserPresetWidget](const FSpawnTabArgs&)
+				{
+					IMetaSoundPresetWidgetInterface::Execute_OnConstructed(UserPresetWidget.Get(), Builder.Get());
+					return SNew(SDockTab)
+						.Label(LOCTEXT("UserPresetWidgetTitle", "Preset Widget"))
+						[
+							InUserPresetWidget.Get()->TakeWidget()
+						];
+				}), CanSpawnTab)
+				.SetDisplayName(LOCTEXT("UserPresetWidgetTab", "Preset Widget"))
+				.SetGroup(WorkspaceMenuCategoryRef)
+				.SetIcon(FSlateIcon(FAppStyle::GetAppStyleSetName(), "Kismet.Tabs.Palette"));
+			}
+
 			InTabManager->RegisterTabSpawner(TabNamesPrivate::Interfaces, FOnSpawnTab::CreateLambda([InInterfacesDetails = InterfacesDetails](const FSpawnTabArgs&)
 			{
 				return SNew(SDockTab).Label(LOCTEXT("MetasoundInterfacesDetailsTitle", "Interfaces")) [ InInterfacesDetails.ToSharedRef() ];
@@ -1050,6 +1077,7 @@ namespace Metasound
 			InTabManager->UnregisterTabSpawner(TabNamesPrivate::Pages);
 			InTabManager->UnregisterTabSpawner(TabNamesPrivate::Interfaces);
 			InTabManager->UnregisterTabSpawner(TabNamesPrivate::Find);
+			InTabManager->UnregisterTabSpawner(TabNamesPrivate::UserPresetWidget);
 		}
 
 		bool FEditor::IsPlaying() const
@@ -1177,7 +1205,24 @@ namespace Metasound
 				DetailsStack->AddTab(TabNamesPrivate::Pages, ETabState::InvalidTab);
 			}
 
-			const TSharedRef<FTabManager::FLayout> StandaloneDefaultLayout = FTabManager::NewLayout("Standalone_MetasoundEditor_Layout_v14")
+			GraphCanvasTabStack = FTabManager::NewStack()
+				->SetSizeCoefficient(0.8f)
+				->SetHideTabWell(true)
+				->SetExtensionId(TabNamesPrivate::GraphCanvas)
+				->AddTab(TabNamesPrivate::GraphCanvas, ETabState::OpenedTab)
+				->AddTab(TabNamesPrivate::UserPresetWidget, ETabState::InvalidTab);
+
+			if (Builder.IsValid() && Builder->IsPreset())
+			{
+				GraphCanvasTabStack->SetHideTabWell(false);
+				GraphCanvasTabStack->AddTab(TabNamesPrivate::UserPresetWidget, ETabState::OpenedTab);
+			}
+			else
+			{
+				GraphCanvasTabStack->AddTab(TabNamesPrivate::UserPresetWidget, ETabState::InvalidTab);
+			}
+
+			const TSharedRef<FTabManager::FLayout> StandaloneDefaultLayout = FTabManager::NewLayout("Standalone_MetasoundEditor_Layout_v15")
 			->AddArea
 			(
 				FTabManager::NewPrimaryArea()
@@ -1215,10 +1260,7 @@ namespace Metasound
 						->SetOrientation(Orient_Vertical)
 						->Split
 						(
-							FTabManager::NewStack()
-							->SetSizeCoefficient(0.8f)
-							->SetHideTabWell(true)
-							->AddTab(TabNamesPrivate::GraphCanvas, ETabState::OpenedTab)
+							GraphCanvasTabStack.ToSharedRef()
 						)
 						->Split
 						(
@@ -1680,6 +1722,11 @@ namespace Metasound
 				PagesDetails->SetVisibility(VisibilityAttr);
 			}
 
+			if (Builder.IsValid() && Builder->IsPreset())
+			{
+				UserPresetWidget = TStrongObjectPtr<UUserWidget>(CreateUserPresetWidget());
+			}
+
 			Palette = SNew(SMetasoundPalette);
 
 			FindWidget = SNew(SFindInMetasound, SharedThis(this));
@@ -1718,6 +1765,81 @@ namespace Metasound
 			}
 		}
 
+		UUserWidget* FEditor::CreateUserPresetWidget()
+		{
+			using namespace Metasound::Frontend;
+
+			UWorld* World = GEditor ? GEditor->GetEditorWorldContext().World() : nullptr;
+			UAudioWidgetSubsystem* WidgetSubsystem = GEngine ? GEngine->GetEngineSubsystem<UAudioWidgetSubsystem>() : nullptr;
+			if (!World || !WidgetSubsystem)
+			{
+				return nullptr;
+			}
+
+			// This asset 
+			const FMetaSoundFrontendDocumentBuilder& DocBuilder = Builder->GetConstBuilder();
+			const FMetasoundFrontendGraphClass& ThisGraphClass = DocBuilder.GetConstDocumentChecked().RootGraph;
+			const FNodeRegistryKey ThisRegistryKey(ThisGraphClass);
+
+			// Referenced asset 
+			FMetasoundAssetBase* ReferencedAsset = DocBuilder.GetReferencedPresetAsset();
+			const FMetasoundFrontendDocument& ReferencedDocument = ReferencedAsset->GetConstDocumentChecked();
+			const FMetasoundFrontendGraphClass& ReferencedGraphClass = ReferencedDocument.RootGraph;
+			const FNodeRegistryKey ReferencedRegistryKey(ReferencedGraphClass);
+
+			// Filter whether the widget supports this MetaSound
+			auto FilterFunction = [&](UUserWidget* UserWidget)
+			{
+				bool bSupportAllPresets = false;
+				TArray<TScriptInterface<IMetaSoundDocumentInterface>> ExcludedMetaSounds;
+				TArray<TScriptInterface<IMetaSoundDocumentInterface>> IncludedMetaSounds;
+				IMetaSoundPresetWidgetInterface::Execute_GetSupportedMetaSounds(UserWidget, bSupportAllPresets, ExcludedMetaSounds, IncludedMetaSounds);
+
+				// Show if all presets supported and this MetaSound is not excluded 
+				if (bSupportAllPresets)
+				{
+					if (!ExcludedMetaSounds.IsEmpty())
+					{
+						for (TScriptInterface<IMetaSoundDocumentInterface> ExcludedMetaSound : ExcludedMetaSounds)
+						{
+							// This MetaSound or the referenced MetaSound is excluded 
+							const FMetasoundFrontendGraphClass& ExcludedGraphClass = ExcludedMetaSound->GetConstDocument().RootGraph;
+							const FNodeRegistryKey ExcludedRegistryKey(ExcludedGraphClass);
+
+							if (ThisRegistryKey == ExcludedRegistryKey || ReferencedRegistryKey == ExcludedRegistryKey)
+							{
+								return false;
+							}
+						}
+					}
+					return true;
+				}
+
+				// This MetaSound is in the inclusion list 
+				for (TScriptInterface<IMetaSoundDocumentInterface> IncludedMetaSound : IncludedMetaSounds)
+				{
+					// This MetaSound or the referenced MetaSound is included 
+					const FMetasoundFrontendGraphClass& IncludedGraphClass = IncludedMetaSound->GetConstDocument().RootGraph;
+					const FNodeRegistryKey IncludedRegistryKey(IncludedGraphClass);
+
+					if (ThisRegistryKey == IncludedRegistryKey || ReferencedRegistryKey == IncludedRegistryKey)
+					{
+						return true;
+					}
+				}
+				return false;
+			};
+
+			// Create widget
+			TArray<UUserWidget*> UserWidgets = WidgetSubsystem->CreateUserWidgets(*World, UMetaSoundPresetWidgetInterface::StaticClass(), FilterFunction);
+			// Currently only support a single user widget per preset 
+			if (UserWidgets.Num() > 0)
+			{
+				return UserWidgets.Last();
+			}
+			return nullptr;
+		}
+
 		TSharedRef<SWidget> FEditor::CreateAuditionMenuOptions()
 		{
 			TSharedPtr<FUICommandList> Commands = MakeShared<FUICommandList>();
@@ -4914,6 +5036,15 @@ namespace Metasound
 				}
 
 				PageStatsWidget->Update(AuditionPageSettings, GraphPageSettings, Color);
+
+				// Update tab label with page name
+				TSharedPtr<SDockTab> GraphTab = TabManager->FindExistingLiveTab(TabNamesPrivate::GraphCanvas);
+				if (GraphTab.IsValid())
+				{
+					const FText GraphTabLabel = FText::Format(LOCTEXT("GraphCanvasTab", "Graph Page: {0}"), PageStatsWidget.IsValid() ?
+						FText::FromName(PageStatsWidget->GetDisplayedPageName()) : LOCTEXT("MetasoundGraphCanvasTitle", "MetaSound Graph"));
+					GraphTab->SetLabel(GraphTabLabel);
+				}
 			}
 		}
 
@@ -4933,6 +5064,7 @@ namespace Metasound
 		void FEditor::FDocumentListener::OnBuilderReloaded(Frontend::FDocumentModifyDelegates& OutDelegates)
 		{
 			OutDelegates.PageDelegates.OnPageSet.AddSP(this, &FDocumentListener::OnPageSet);
+			OutDelegates.InterfaceDelegates.OnInputDefaultChanged.AddSP(this, &FDocumentListener::OnInputDefaultChanged);
 		}
 
 		void FEditor::FDocumentListener::OnPageSet(const Frontend::FDocumentMutatePageArgs& Args)
@@ -4971,6 +5103,45 @@ namespace Metasound
 			}
 		}
 
+		void FEditor::FDocumentListener::OnInputDefaultChanged(int32 Index)
+		{
+			if (TSharedPtr<FEditor> ParentPtr = Parent.Pin())
+			{
+				const FMetaSoundFrontendDocumentBuilder& ConstBuilder = ParentPtr->Builder->GetConstBuilder();
+				const FMetasoundFrontendGraphClass& GraphClass = ConstBuilder.GetConstDocumentChecked().RootGraph;
+				const FMetasoundFrontendGraph& Graph = GraphClass.GetConstDefaultGraph();
+
+				const FMetasoundFrontendClassInput& GraphInput = GraphClass.Interface.Inputs[Index];
+				const FName InputName = GraphInput.Name;
+				UMetasoundEditorGraph& MetaSoundGraph = ParentPtr->GetMetaSoundGraphChecked();
+				UMetasoundEditorGraphInput* EditorGraphInput = MetaSoundGraph.FindInput(InputName);
+				if (EditorGraphInput)
+				{
+					UMetasoundEditorGraphMemberDefaultLiteral* Literal = EditorGraphInput->GetLiteral();
+					if (Literal)
+					{
+						// Synchronize editor literal and update preview instance
+						Literal->Synchronize();
+						if (GEditor && MetaSoundGraph.IsPreviewing())
+						{
+							UAudioComponent* PreviewComponent = GEditor->GetPreviewAudioComponent();
+							check(PreviewComponent);
+
+							if (TScriptInterface<IAudioParameterControllerInterface> ParamInterface = PreviewComponent)
+							{
+								Literal->UpdatePreviewInstance(InputName, ParamInterface);
+							}
+						}
+						
+						// Mark dirty (needed for when this is being called in response to a change from Builder BP API for example)
+						UObject* MetaSound = ParentPtr->GetMetasoundObject();
+						const FMetasoundAssetBase* MetaSoundAsset = IMetasoundUObjectRegistry::Get().GetObjectAsAssetBase(MetaSound);
+						MetaSoundAsset->MarkMetasoundDocumentDirty();
+					}
+				}
+			}
+		}
+
 		bool FEditor::CanPromoteToInput()
 		{
 			if (MetasoundGraphEditor.IsValid())
diff --git a/Engine/Plugins/Runtime/Metasound/Source/MetasoundEditor/Private/MetasoundEditor.h b/Engine/Plugins/Runtime/Metasound/Source/MetasoundEditor/Private/MetasoundEditor.h
--- a/Engine/Plugins/Runtime/Metasound/Source/MetasoundEditor/Private/MetasoundEditor.h
+++ b/Engine/Plugins/Runtime/Metasound/Source/MetasoundEditor/Private/MetasoundEditor.h
@@ -46,6 +46,7 @@ class SMetasoundPalette;
 class FSlateRect;
 class IDetailsView;
 class IToolkitHost;
+class UUserWidget;
 
 struct FMetaSoundFrontendDocumentBuilder;
 
@@ -448,6 +449,8 @@ namespace Metasound::Editor
 		/** Creates all internal widgets for the tabs to point at */
 		void CreateInternalWidgets(UObject& MetaSound);
 
+		UUserWidget* CreateUserPresetWidget();
+
 		/** Builds the toolbar widget for the Metasound editor */
 		void ExtendToolbarInternal();
 
@@ -495,6 +498,8 @@ namespace Metasound::Editor
 		/** Show the Metasound object's settings in the Details panel */
 		void EditMetasoundSettings();
 
+		TSharedPtr<FTabManager::FStack> GraphCanvasTabStack;
+
 		/** Add an input to the currently selected node */
 		void AddInput();
 
@@ -562,6 +567,9 @@ namespace Metasound::Editor
 		/** Widget showing render performance information that overlays the graph */
 		TSharedPtr<SRenderStats> RenderStatsWidget;
 
+		/** User defined widget shown when editing a preset.*/
+		TStrongObjectPtr<UUserWidget> UserPresetWidget;
+
 		TUniquePtr<FGraphConnectionManager> GraphConnectionManager;
 
 		/** Command list for this editor */
@@ -611,6 +619,7 @@ namespace Metasound::Editor
 		private:
 			virtual void OnBuilderReloaded(Frontend::FDocumentModifyDelegates& OutDelegates) override;
 			void OnPageSet(const Frontend::FDocumentMutatePageArgs& Args);
+			void OnInputDefaultChanged(int32 Index);
 		};
 
 		TSharedPtr<FDocumentListener> DocListener;
diff --git a/Engine/Plugins/Runtime/Metasound/Source/MetasoundEditor/Private/MetasoundEditorBuilderListener.cpp b/Engine/Plugins/Runtime/Metasound/Source/MetasoundEditor/Private/MetasoundEditorBuilderListener.cpp
new file mode 100644
--- /dev/null
+++ b/Engine/Plugins/Runtime/Metasound/Source/MetasoundEditor/Private/MetasoundEditorBuilderListener.cpp
@@ -0,0 +1,179 @@
+// Copyright Epic Games, Inc. All Rights Reserved.
+#include "MetasoundEditorBuilderListener.h"
+
+void UMetaSoundEditorBuilderListener::Init(const TWeakObjectPtr<UMetaSoundBuilderBase> InBuilder)
+{
+	if (InBuilder.IsValid())
+	{
+		Builder = InBuilder;
+		FMetaSoundFrontendDocumentBuilder& DocumentBuilder = Builder->GetBuilder();
+		Metasound::Frontend::FDocumentModifyDelegates& BuilderDelegates = DocumentBuilder.GetDocumentDelegates();
+		BuilderListener = MakeShared<FEditorBuilderListener>(this);
+		Builder->AddTransactionListener(BuilderListener->AsShared());
+	}
+}
+
+void UMetaSoundEditorBuilderListener::OnGraphInputAdded(int32 Index)
+{
+	if (Builder.IsValid())
+	{
+		const FMetaSoundFrontendDocumentBuilder& ConstBuilder = Builder->GetConstBuilder();
+		const FMetasoundFrontendGraphClass& GraphClass = ConstBuilder.GetConstDocumentChecked().RootGraph;
+		const FMetasoundFrontendClassInput& GraphInput = GraphClass.Interface.Inputs[Index];
+
+		OnGraphInputAddedDelegate.Broadcast(GraphInput.Name, GraphInput.TypeName);
+	}
+}
+
+void UMetaSoundEditorBuilderListener::OnGraphInputDefaultChanged(int32 Index)
+{
+	if (Builder.IsValid())
+	{
+		const FMetaSoundFrontendDocumentBuilder& ConstBuilder = Builder->GetConstBuilder();
+		const FMetasoundFrontendGraphClass& GraphClass = ConstBuilder.GetConstDocumentChecked().RootGraph;
+
+		// todo: support paged graph once frontend delegates support paged literals
+		const FMetasoundFrontendGraph& Graph = GraphClass.GetConstDefaultGraph();
+		FName PageName = Metasound::Frontend::DefaultPageName;
+
+		const FMetasoundFrontendClassInput& GraphInput = GraphClass.Interface.Inputs[Index];
+		const FMetasoundFrontendLiteral* DefaultLiteral = GraphInput.FindConstDefault(Metasound::Frontend::DefaultPageID);
+		if (ensure(DefaultLiteral))
+		{
+			OnGraphInputDefaultChangedDelegate.Broadcast(GraphInput.Name, *DefaultLiteral, PageName);
+		}
+	}
+}
+
+void UMetaSoundEditorBuilderListener::OnRemovingGraphInput(int32 Index)
+{
+	if (Builder.IsValid())
+	{
+		const FMetaSoundFrontendDocumentBuilder& ConstBuilder = Builder->GetConstBuilder();
+		const FMetasoundFrontendGraphClass& GraphClass = ConstBuilder.GetConstDocumentChecked().RootGraph;
+		const FMetasoundFrontendClassOutput& GraphInput = GraphClass.Interface.Inputs[Index];
+
+		OnRemovingGraphInputDelegate.Broadcast(GraphInput.Name, GraphInput.TypeName);
+	}
+}
+
+void UMetaSoundEditorBuilderListener::OnGraphOutputAdded(int32 Index)
+{
+	if (Builder.IsValid())
+	{
+		const FMetaSoundFrontendDocumentBuilder& ConstBuilder = Builder->GetConstBuilder();
+		const FMetasoundFrontendGraphClass& GraphClass = ConstBuilder.GetConstDocumentChecked().RootGraph;
+		const FMetasoundFrontendClassOutput& GraphOutput = GraphClass.Interface.Outputs[Index];
+
+		OnGraphOutputAddedDelegate.Broadcast(GraphOutput.Name, GraphOutput.TypeName);
+	}
+}
+
+void UMetaSoundEditorBuilderListener::OnRemovingGraphOutput(int32 Index)
+{
+	if (Builder.IsValid())
+	{
+		const FMetaSoundFrontendDocumentBuilder& ConstBuilder = Builder->GetConstBuilder();
+		const FMetasoundFrontendGraphClass& GraphClass = ConstBuilder.GetConstDocumentChecked().RootGraph;
+		const FMetasoundFrontendClassOutput& GraphOutput = GraphClass.Interface.Outputs[Index];
+
+		OnRemovingGraphOutputDelegate.Broadcast(GraphOutput.Name, GraphOutput.TypeName);
+	}
+}
+
+
+void UMetaSoundEditorBuilderListener::OnGraphInputDataTypeChanged(int32 Index)
+{
+	if (Builder.IsValid())
+	{
+		const FMetaSoundFrontendDocumentBuilder& ConstBuilder = Builder->GetConstBuilder();
+		const FMetasoundFrontendGraphClass& GraphClass = ConstBuilder.GetConstDocumentChecked().RootGraph;
+		const FMetasoundFrontendClassOutput& GraphInput = GraphClass.Interface.Inputs[Index];
+
+		OnGraphInputDataTypeChangedDelegate.Broadcast(GraphInput.Name, GraphInput.TypeName);
+	}
+}
+
+void UMetaSoundEditorBuilderListener::OnGraphOutputDataTypeChanged(int32 Index)
+{
+	if (Builder.IsValid())
+	{
+		const FMetaSoundFrontendDocumentBuilder& ConstBuilder = Builder->GetConstBuilder();
+		const FMetasoundFrontendGraphClass& GraphClass = ConstBuilder.GetConstDocumentChecked().RootGraph;
+		const FMetasoundFrontendClassOutput& GraphOutput = GraphClass.Interface.Outputs[Index];
+
+		OnGraphOutputDataTypeChangedDelegate.Broadcast(GraphOutput.Name, GraphOutput.TypeName);
+	}
+}
+
+void UMetaSoundEditorBuilderListener::OnGraphInputNameChanged(FName OldName, FName NewName)
+{
+	if (Builder.IsValid())
+	{
+		OnGraphInputNameChangedDelegate.Broadcast(OldName, NewName);
+	}
+}
+
+void UMetaSoundEditorBuilderListener::OnGraphOutputNameChanged(FName OldName, FName NewName)
+{
+	if (Builder.IsValid())
+	{
+		OnGraphOutputNameChangedDelegate.Broadcast(OldName, NewName);
+	}
+}
+
+void UMetaSoundEditorBuilderListener::RemoveAllDelegates()
+{
+	// Remove multicast BP delegates 
+	OnGraphInputAddedDelegate.RemoveAll(this);
+	OnGraphInputDefaultChangedDelegate.RemoveAll(this);
+	OnRemovingGraphInputDelegate.RemoveAll(this);
+
+	OnGraphOutputAddedDelegate.RemoveAll(this);
+	OnRemovingGraphOutputDelegate.RemoveAll(this);
+
+	OnGraphInputDataTypeChangedDelegate.RemoveAll(this);
+	OnGraphOutputDataTypeChangedDelegate.RemoveAll(this);
+	OnGraphInputNameChangedDelegate.RemoveAll(this);
+	OnGraphOutputNameChangedDelegate.RemoveAll(this);
+
+	// Remove document delegates 
+	if (Builder.IsValid())
+	{
+		FMetaSoundFrontendDocumentBuilder& DocumentBuilder = Builder->GetBuilder();
+		Metasound::Frontend::FDocumentModifyDelegates& BuilderDelegates = DocumentBuilder.GetDocumentDelegates();
+		Metasound::Frontend::FInterfaceModifyDelegates& InterfaceDelegates = BuilderDelegates.InterfaceDelegates;
+
+		InterfaceDelegates.OnInputAdded.Remove(OnInputAddedHandle);
+		InterfaceDelegates.OnInputDefaultChanged.Remove(OnInputDefaultChangedHandle);
+		InterfaceDelegates.OnRemovingInput.Remove(OnRemovingInputHandle);
+
+		InterfaceDelegates.OnOutputAdded.Remove(OnOutputAddedHandle);
+		InterfaceDelegates.OnRemovingOutput.Remove(OnRemovingOutputHandle);
+
+		InterfaceDelegates.OnInputDataTypeChanged.Remove(OnInputDataTypeChangedHandle);
+		InterfaceDelegates.OnOutputDataTypeChanged.Remove(OnOutputDataTypeChangedHandle);
+		InterfaceDelegates.OnInputNameChanged.Remove(OnInputNameChangedHandle);
+		InterfaceDelegates.OnOutputNameChanged.Remove(OnOutputNameChangedHandle);
+	}
+}
+
+void UMetaSoundEditorBuilderListener::FEditorBuilderListener::OnBuilderReloaded(Metasound::Frontend::FDocumentModifyDelegates& OutDelegates)
+{
+	if (Parent)
+	{
+		Metasound::Frontend::FInterfaceModifyDelegates& InterfaceDelegates = OutDelegates.InterfaceDelegates;
+
+		Parent->OnInputAddedHandle = InterfaceDelegates.OnInputAdded.AddUObject(Parent, &UMetaSoundEditorBuilderListener::OnGraphInputAdded);
+		Parent->OnInputDefaultChangedHandle = InterfaceDelegates.OnInputDefaultChanged.AddUObject(Parent, &UMetaSoundEditorBuilderListener::OnGraphInputDefaultChanged);
+		Parent->OnRemovingInputHandle = InterfaceDelegates.OnRemovingInput.AddUObject(Parent, &UMetaSoundEditorBuilderListener::OnRemovingGraphInput);
+
+		Parent->OnOutputAddedHandle = InterfaceDelegates.OnOutputAdded.AddUObject(Parent, &UMetaSoundEditorBuilderListener::OnGraphOutputAdded);
+		Parent->OnRemovingOutputHandle = InterfaceDelegates.OnRemovingOutput.AddUObject(Parent, &UMetaSoundEditorBuilderListener::OnRemovingGraphOutput);
+
+		Parent->OnInputDataTypeChangedHandle = InterfaceDelegates.OnInputDataTypeChanged.AddUObject(Parent, &UMetaSoundEditorBuilderListener::OnGraphInputDataTypeChanged);
+		Parent->OnOutputDataTypeChangedHandle = InterfaceDelegates.OnOutputDataTypeChanged.AddUObject(Parent, &UMetaSoundEditorBuilderListener::OnGraphOutputDataTypeChanged);
+		Parent->OnInputNameChangedHandle = InterfaceDelegates.OnInputNameChanged.AddUObject(Parent, &UMetaSoundEditorBuilderListener::OnGraphInputNameChanged);
+		Parent->OnOutputNameChangedHandle = InterfaceDelegates.OnOutputNameChanged.AddUObject(Parent, &UMetaSoundEditorBuilderListener::OnGraphOutputNameChanged);
+	}
+}
diff --git a/Engine/Plugins/Runtime/Metasound/Source/MetasoundEditor/Private/MetasoundEditorBuilderListener.h b/Engine/Plugins/Runtime/Metasound/Source/MetasoundEditor/Private/MetasoundEditorBuilderListener.h
new file mode 100644
--- /dev/null
+++ b/Engine/Plugins/Runtime/Metasound/Source/MetasoundEditor/Private/MetasoundEditorBuilderListener.h
@@ -0,0 +1,108 @@
+// Copyright Epic Games, Inc. All Rights Reserved.
+#pragma once
+#include "Delegates/DelegateCombinations.h"
+#include "MetasoundBuilderBase.h"
+#include "MetasoundDocumentInterface.h"
+#include "MetasoundFrontendDocumentModifyDelegates.h"
+
+#include "MetasoundEditorBuilderListener.generated.h"
+
+// Forward Declarations
+struct FMetaSoundNodeHandle;
+struct FMetasoundFrontendLiteral;
+
+// BP Delegates for builder changes 
+DECLARE_DYNAMIC_MULTICAST_DELEGATE_TwoParams(FOnMetaSoundBuilderGraphInterfaceMutate, FName, VertexName, FName, DataType);
+DECLARE_DYNAMIC_MULTICAST_DELEGATE_ThreeParams(FOnMetaSoundBuilderGraphLiteralMutate, FName, VertexName, FMetasoundFrontendLiteral, LiteralValue, FName, PageName);
+DECLARE_DYNAMIC_MULTICAST_DELEGATE_TwoParams(FOnMetaSoundBuilderGraphVertexRename, FName, OldName, FName, NewName);
+
+UCLASS(MinimalAPI, BlueprintType)
+class UMetaSoundEditorBuilderListener : public UObject
+{
+public: 
+	GENERATED_BODY()
+
+	virtual ~UMetaSoundEditorBuilderListener() {}
+	
+	void Init(const TWeakObjectPtr<UMetaSoundBuilderBase> InBuilder);
+
+	void OnGraphInputAdded(int32 Index);
+	void OnGraphInputDefaultChanged(int32 Index);
+	void OnRemovingGraphInput(int32 Index);
+
+	void OnGraphOutputAdded(int32 Index);
+	void OnRemovingGraphOutput(int32 Index);
+
+	void OnGraphInputDataTypeChanged(int32 Index);
+	void OnGraphOutputDataTypeChanged(int32 Index);
+	void OnGraphInputNameChanged(FName OldName, FName NewName);
+	void OnGraphOutputNameChanged(FName OldName, FName NewName);
+
+
+	UFUNCTION(BlueprintCallable, Category = "Audio|MetaSound|Builder|Editor|Delegates")
+	void RemoveAllDelegates();
+
+	UPROPERTY(BlueprintAssignable, BlueprintCallable, Category = "Audio|MetaSound|Builder|Editor|Delegates")
+	FOnMetaSoundBuilderGraphInterfaceMutate OnGraphInputAddedDelegate;
+	
+	UPROPERTY(BlueprintAssignable, BlueprintCallable, Category = "Audio|MetaSound|Builder|Editor|Delegates")
+	FOnMetaSoundBuilderGraphLiteralMutate OnGraphInputDefaultChangedDelegate;
+
+	UPROPERTY(BlueprintAssignable, BlueprintCallable, Category = "Audio|MetaSound|Builder|Editor|Delegates")
+	FOnMetaSoundBuilderGraphInterfaceMutate OnRemovingGraphInputDelegate;
+	
+	UPROPERTY(BlueprintAssignable, BlueprintCallable, Category= "Audio|MetaSound|Builder|Editor|Delegates")
+	FOnMetaSoundBuilderGraphInterfaceMutate OnGraphOutputAddedDelegate;
+
+	UPROPERTY(BlueprintAssignable, BlueprintCallable, Category = "Audio|MetaSound|Builder|Editor|Delegates")
+	FOnMetaSoundBuilderGraphInterfaceMutate OnRemovingGraphOutputDelegate;
+
+	UPROPERTY(BlueprintAssignable, BlueprintCallable, Category = "Audio|MetaSound|Builder|Editor|Delegates")
+	FOnMetaSoundBuilderGraphInterfaceMutate OnGraphInputDataTypeChangedDelegate;
+
+	UPROPERTY(BlueprintAssignable, BlueprintCallable, Category = "Audio|MetaSound|Builder|Editor|Delegates")
+	FOnMetaSoundBuilderGraphInterfaceMutate OnGraphOutputDataTypeChangedDelegate;
+
+	UPROPERTY(BlueprintAssignable, BlueprintCallable, Category = "Audio|MetaSound|Builder|Editor|Delegates")
+	FOnMetaSoundBuilderGraphVertexRename OnGraphInputNameChangedDelegate;
+
+	UPROPERTY(BlueprintAssignable, BlueprintCallable, Category = "Audio|MetaSound|Builder|Editor|Delegates")
+	FOnMetaSoundBuilderGraphVertexRename OnGraphOutputNameChangedDelegate;
+
+private:
+	// Handles for document delegates
+	FDelegateHandle OnInputAddedHandle;
+	FDelegateHandle OnInputDefaultChangedHandle;
+	FDelegateHandle OnRemovingInputHandle;
+
+	FDelegateHandle OnOutputAddedHandle;
+	FDelegateHandle OnRemovingOutputHandle;
+
+	FDelegateHandle OnInputDataTypeChangedHandle;
+	FDelegateHandle OnOutputDataTypeChangedHandle;
+	FDelegateHandle OnInputNameChangedHandle;
+	FDelegateHandle OnOutputNameChangedHandle;
+
+	TWeakObjectPtr<UMetaSoundBuilderBase> Builder;
+
+	class FEditorBuilderListener : public Metasound::Frontend::IDocumentBuilderTransactionListener
+	{
+	public: 
+
+		FEditorBuilderListener() = default;
+		FEditorBuilderListener(TObjectPtr<UMetaSoundEditorBuilderListener> InParent)
+			: Parent(InParent)
+		{
+		}
+
+		virtual ~FEditorBuilderListener() = default;
+
+		// IDocumentBuilderTransactionListener
+		virtual void OnBuilderReloaded(Metasound::Frontend::FDocumentModifyDelegates& OutDelegates) override;
+
+	private: 
+		TObjectPtr<UMetaSoundEditorBuilderListener> Parent;
+	};
+
+	TSharedPtr<FEditorBuilderListener> BuilderListener;
+};
\ No newline at end of file
diff --git a/Engine/Plugins/Runtime/Metasound/Source/MetasoundEditor/Private/MetasoundEditorGraph.cpp b/Engine/Plugins/Runtime/Metasound/Source/MetasoundEditor/Private/MetasoundEditorGraph.cpp
--- a/Engine/Plugins/Runtime/Metasound/Source/MetasoundEditor/Private/MetasoundEditorGraph.cpp
+++ b/Engine/Plugins/Runtime/Metasound/Source/MetasoundEditor/Private/MetasoundEditorGraph.cpp
@@ -1091,8 +1091,6 @@ void UMetasoundEditorGraphInput::UpdateFrontendDefaultLiteral(bool bPostTransact
 			});
 			Builder.SetGraphInputDefaults(MemberName, MoveTemp(NewDefaults));
 		}
-
-		GraphPrivate::UpdatePreviewParameter(GetOwningGraph(), MemberName, *Literal);
 	}
 }
 
diff --git a/Engine/Plugins/Runtime/Metasound/Source/MetasoundEditor/Private/MetasoundEditorGraphMemberDefaults.cpp b/Engine/Plugins/Runtime/Metasound/Source/MetasoundEditor/Private/MetasoundEditorGraphMemberDefaults.cpp
--- a/Engine/Plugins/Runtime/Metasound/Source/MetasoundEditor/Private/MetasoundEditorGraphMemberDefaults.cpp
+++ b/Engine/Plugins/Runtime/Metasound/Source/MetasoundEditor/Private/MetasoundEditorGraphMemberDefaults.cpp
@@ -829,7 +829,16 @@ void UMetasoundEditorGraphMemberDefaultFloat::SortPageDefaults()
 bool UMetasoundEditorGraphMemberDefaultFloat::Synchronize()
 {
 	using namespace Metasound::Editor;
-	return MemberDefaultsPrivate::SynchronizePageDefault<FMetasoundEditorMemberPageDefaultFloat, float>(FindMember(), Defaults);
+	const bool bModified = MemberDefaultsPrivate::SynchronizePageDefault<FMetasoundEditorMemberPageDefaultFloat, float>(FindMember(), Defaults);
+	if (bModified)
+	{
+		// Broadcast on value changed
+		for (const FMetasoundEditorMemberPageDefaultFloat& PageDefault : Defaults)
+		{
+			OnDefaultValueChanged.Broadcast(PageDefault.PageID, PageDefault.Value);
+		}
+	}
+	return bModified;
 }
 
 void UMetasoundEditorGraphMemberDefaultFloat::SetInitialRange()
diff --git a/Engine/Plugins/Runtime/Metasound/Source/MetasoundEditor/Private/MetasoundEditorSubsystem.cpp b/Engine/Plugins/Runtime/Metasound/Source/MetasoundEditor/Private/MetasoundEditorSubsystem.cpp
--- a/Engine/Plugins/Runtime/Metasound/Source/MetasoundEditor/Private/MetasoundEditorSubsystem.cpp
+++ b/Engine/Plugins/Runtime/Metasound/Source/MetasoundEditor/Private/MetasoundEditorSubsystem.cpp
@@ -7,6 +7,7 @@
 #include "MetasoundDocumentBuilderRegistry.h"
 #include "MetasoundDocumentInterface.h"
 #include "MetasoundEditorGraph.h"
+#include "MetasoundEditorBuilderListener.h"
 #include "MetasoundEditorGraphBuilder.h"
 #include "MetasoundEditorGraphSchema.h"
 #include "MetasoundEditorGraphMemberDefaults.h"
@@ -401,6 +402,21 @@ void UMetaSoundEditorSubsystem::RegisterToolbarExtender(TSharedRef<FExtender> In
 	EditorToolbarExtenders.AddUnique(InExtender);
 }
 
+UMetaSoundEditorBuilderListener* UMetaSoundEditorSubsystem::AddBuilderDelegateListener(UMetaSoundBuilderBase* InBuilder, EMetaSoundBuilderResult& OutResult)
+{
+	if (InBuilder)
+	{
+		UMetaSoundEditorBuilderListener* Listener = NewObject<UMetaSoundEditorBuilderListener>(GetTransientPackage(), FName(), RF_Public | RF_Transient);
+		Listener->Init(InBuilder);
+		OutResult = EMetaSoundBuilderResult::Succeeded;
+		return Listener;
+	}
+
+	UE_LOG(LogMetaSound, Warning, TEXT("Add Builder Delegate Listener called with invalid builder, listener will not be created."));
+	OutResult = EMetaSoundBuilderResult::Failed;
+	return nullptr;
+}
+
 void UMetaSoundEditorSubsystem::SetFocusedPage(UMetaSoundBuilderBase* Builder, FName PageName, bool bOpenEditor, EMetaSoundBuilderResult& OutResult) const
 {
 	using namespace Metasound::Frontend;
diff --git a/Engine/Plugins/Runtime/Metasound/Source/MetasoundEditor/Public/MetasoundEditorGraphMemberDefaults.h b/Engine/Plugins/Runtime/Metasound/Source/MetasoundEditor/Public/MetasoundEditorGraphMemberDefaults.h
--- a/Engine/Plugins/Runtime/Metasound/Source/MetasoundEditor/Public/MetasoundEditorGraphMemberDefaults.h
+++ b/Engine/Plugins/Runtime/Metasound/Source/MetasoundEditor/Public/MetasoundEditorGraphMemberDefaults.h
@@ -311,7 +311,7 @@ struct FMetasoundEditorMemberPageDefaultFloat : public FMetasoundEditorMemberPag
 };
 
 
-UCLASS(MinimalAPI)
+UCLASS(MinimalAPI, BlueprintType)
 class UMetasoundEditorGraphMemberDefaultFloat : public UMetasoundEditorGraphMemberDefaultLiteral
 {
 	GENERATED_BODY()
@@ -329,7 +329,7 @@ public:
 	UPROPERTY(EditAnywhere, Category = DefaultValue)
 	bool ClampDefault = false;
 
-	UPROPERTY(EditAnywhere, Category = DefaultValue)
+	UPROPERTY(EditAnywhere, BlueprintReadWrite, Category = DefaultValue)
 	FVector2D Range = FVector2D(0.0f, 1.0f);
 
 	UPROPERTY(EditAnywhere, Category = Widget, meta=(DisplayName = "Widget"))
diff --git a/Engine/Plugins/Runtime/Metasound/Source/MetasoundEditor/Public/MetasoundEditorSubsystem.h b/Engine/Plugins/Runtime/Metasound/Source/MetasoundEditor/Public/MetasoundEditorSubsystem.h
--- a/Engine/Plugins/Runtime/Metasound/Source/MetasoundEditor/Public/MetasoundEditorSubsystem.h
+++ b/Engine/Plugins/Runtime/Metasound/Source/MetasoundEditor/Public/MetasoundEditorSubsystem.h
@@ -13,6 +13,7 @@
 
 // Forward Declarations
 class UMetaSoundBuilderBase;
+class UMetaSoundEditorBuilderListener;
 class UMetasoundEditorGraphMember;
 class UMetasoundEditorGraphMemberDefaultLiteral;
 
@@ -88,6 +89,10 @@ public:
 	bool IsPageAuditionPlatformCookTarget(FName InPageName) const;
 	bool IsPageAuditionPlatformCookTarget(const FGuid& InPageID) const;
 
+	// Add a builder listener for a builder which is used to add and remove custom editor builder delegates.
+	UFUNCTION(BlueprintCallable, Category = "Audio|MetaSound|Builder|Editor", meta = (ExpandEnumAsExecs = "OutResult"))
+	UPARAM(DisplayName = "Builder Listener") UMetaSoundEditorBuilderListener* AddBuilderDelegateListener(UMetaSoundBuilderBase* InBuilder, EMetaSoundBuilderResult& OutResult);
+	
 	// Wraps RegisterGraphWithFrontend logic in Frontend with any additional logic required to refresh editor & respective editor object state.
 	// @param InMetaSound - MetaSound to register
 	// @param bInForceSynchronize - Forces the synchronize flag for all open graphs being registered by this call (all referenced graphs and
diff --git a/Engine/Plugins/Runtime/Metasound/Source/MetasoundEditor/Public/MetasoundPresetWidgetInterface.h b/Engine/Plugins/Runtime/Metasound/Source/MetasoundEditor/Public/MetasoundPresetWidgetInterface.h
new file mode 100644
--- /dev/null
+++ b/Engine/Plugins/Runtime/Metasound/Source/MetasoundEditor/Public/MetasoundPresetWidgetInterface.h
@@ -0,0 +1,34 @@
+// Copyright Epic Games, Inc. All Rights Reserved.
+#pragma once
+
+#include "Audio/AudioPanelWidgetInterface.h"
+#include "MetasoundBuilderSubsystem.h"
+#include "UObject/Class.h"
+#include "UObject/Interface.h"
+#include "UObject/ObjectMacros.h"
+#include "UObject/UObjectGlobals.h"
+#include "UObject/ScriptInterface.h"
+
+#include "MetasoundPresetWidgetInterface.generated.h"
+
+UINTERFACE(Blueprintable, MinimalAPI)
+class UMetaSoundPresetWidgetInterface : public UAudioPanelWidgetInterface
+{
+	GENERATED_BODY()
+};
+
+class IMetaSoundPresetWidgetInterface : public IAudioPanelWidgetInterface
+{
+	GENERATED_BODY()
+
+public:
+	// The MetaSounds whose presets are supported by this widget. If Support All Presets is true, this widget is supported by all presets except those in the Excluded array.
+	// MetaSounds in the include/exclude arrays can be MetaSound presets or non presets. 
+	// If a MetaSound is not a preset, then presets of that MetaSound will be supported/excluded by this widget.
+	UFUNCTION(BlueprintImplementableEvent)
+	TScriptInterface<IMetaSoundDocumentInterface> GetSupportedMetaSounds(bool& bSupportAllPresets, TArray<TScriptInterface<IMetaSoundDocumentInterface>>& ExcludedMetaSounds, TArray<TScriptInterface<IMetaSoundDocumentInterface>>& IncludedMetaSounds) const;
+
+	// Called when the preset widget is constructed, giving the builder of the associated MetaSound preset
+	UFUNCTION(BlueprintImplementableEvent, meta = (DisplayName = "On MetaSoundPreset Widget Constructed"))
+	void OnConstructed(UMetaSoundBuilderBase* Builder);
+};
diff --git a/Engine/Plugins/Runtime/Metasound/Source/MetasoundEngine/Private/MetasoundFrontendLiteralBlueprintAccess.cpp b/Engine/Plugins/Runtime/Metasound/Source/MetasoundEngine/Private/MetasoundFrontendLiteralBlueprintAccess.cpp
--- a/Engine/Plugins/Runtime/Metasound/Source/MetasoundEngine/Private/MetasoundFrontendLiteralBlueprintAccess.cpp
+++ b/Engine/Plugins/Runtime/Metasound/Source/MetasoundEngine/Private/MetasoundFrontendLiteralBlueprintAccess.cpp
@@ -1,6 +1,7 @@
 // Copyright Epic Games, Inc. All Rights Reserved.
 
 #include "MetasoundFrontendLiteralBlueprintAccess.h"
+#include "MetasoundDataReference.h"
 
 namespace MetasoundFrontendLiteralBlueprintAccessPrivate
 {
@@ -11,6 +12,24 @@ namespace MetasoundFrontendLiteralBlueprintAccessPrivate
 		Literal.Set(Value);
 		return Literal;
 	}
+
+	template <typename TLiteralType>
+	TLiteralType GetPODValueFromMetaSoundLiteral(const FMetasoundFrontendLiteral& Literal, EMetaSoundBuilderResult& OutResult)
+	{
+		TLiteralType Value;
+		OutResult = Literal.TryGet(Value) ? EMetaSoundBuilderResult::Succeeded : EMetaSoundBuilderResult::Failed;
+		return Value;
+	}
+}
+
+FString UMetasoundFrontendLiteralBlueprintAccess::Conv_MetaSoundLiteralToString(const FMetasoundFrontendLiteral& Literal)
+{
+	return Literal.ToString();
+}
+
+EMetasoundFrontendLiteralType UMetasoundFrontendLiteralBlueprintAccess::GetType(const FMetasoundFrontendLiteral& Literal)
+{
+	return Literal.GetType();
 }
 
 FMetasoundFrontendLiteral UMetasoundFrontendLiteralBlueprintAccess::CreateBoolMetaSoundLiteral(bool Value)
@@ -73,3 +92,53 @@ FMetasoundFrontendLiteral UMetasoundFrontendLiteralBlueprintAccess::CreateMetaSo
 {
 	return FMetasoundFrontendLiteral{ Param };
 }
+
+bool UMetasoundFrontendLiteralBlueprintAccess::GetBoolValueFromLiteral(const FMetasoundFrontendLiteral& Literal, EMetaSoundBuilderResult& OutResult)
+{
+	return MetasoundFrontendLiteralBlueprintAccessPrivate::GetPODValueFromMetaSoundLiteral<bool>(Literal, OutResult);
+}
+
+TArray<bool> UMetasoundFrontendLiteralBlueprintAccess::GetBoolArrayValueFromLiteral(const FMetasoundFrontendLiteral& Literal, EMetaSoundBuilderResult& OutResult)
+{
+	return MetasoundFrontendLiteralBlueprintAccessPrivate::GetPODValueFromMetaSoundLiteral<TArray<bool>>(Literal, OutResult);
+}
+
+float UMetasoundFrontendLiteralBlueprintAccess::GetFloatValueFromLiteral(const FMetasoundFrontendLiteral& Literal, EMetaSoundBuilderResult& OutResult)
+{
+	return MetasoundFrontendLiteralBlueprintAccessPrivate::GetPODValueFromMetaSoundLiteral<float>(Literal, OutResult);
+}
+
+TArray<float> UMetasoundFrontendLiteralBlueprintAccess::GetFloatArrayValueFromLiteral(const FMetasoundFrontendLiteral& Literal, EMetaSoundBuilderResult& OutResult)
+{
+	return MetasoundFrontendLiteralBlueprintAccessPrivate::GetPODValueFromMetaSoundLiteral<TArray<float>>(Literal, OutResult);
+}
+
+int32 UMetasoundFrontendLiteralBlueprintAccess::GetIntValueFromLiteral(const FMetasoundFrontendLiteral& Literal, EMetaSoundBuilderResult& OutResult)
+{
+	return MetasoundFrontendLiteralBlueprintAccessPrivate::GetPODValueFromMetaSoundLiteral<int>(Literal, OutResult);
+}
+
+TArray<int32> UMetasoundFrontendLiteralBlueprintAccess::GetIntArrayValueFromLiteral(const FMetasoundFrontendLiteral& Literal, EMetaSoundBuilderResult& OutResult)
+{
+	return MetasoundFrontendLiteralBlueprintAccessPrivate::GetPODValueFromMetaSoundLiteral<TArray<int32>>(Literal, OutResult);
+}
+
+UObject* UMetasoundFrontendLiteralBlueprintAccess::GetObjectValueFromLiteral(const FMetasoundFrontendLiteral& Literal, EMetaSoundBuilderResult& OutResult)
+{
+	return MetasoundFrontendLiteralBlueprintAccessPrivate::GetPODValueFromMetaSoundLiteral<UObject*>(Literal, OutResult);
+}
+
+TArray<UObject*> UMetasoundFrontendLiteralBlueprintAccess::GetObjectArrayValueFromLiteral(const FMetasoundFrontendLiteral& Literal, EMetaSoundBuilderResult& OutResult)
+{
+	return MetasoundFrontendLiteralBlueprintAccessPrivate::GetPODValueFromMetaSoundLiteral<TArray<UObject*>>(Literal, OutResult);
+}
+
+FString UMetasoundFrontendLiteralBlueprintAccess::GetStringValueFromLiteral(const FMetasoundFrontendLiteral& Literal, EMetaSoundBuilderResult& OutResult)
+{
+	return MetasoundFrontendLiteralBlueprintAccessPrivate::GetPODValueFromMetaSoundLiteral<FString>(Literal, OutResult);
+}
+
+TArray<FString> UMetasoundFrontendLiteralBlueprintAccess::GetStringArrayValueFromLiteral(const FMetasoundFrontendLiteral& Literal, EMetaSoundBuilderResult& OutResult)
+{
+	return MetasoundFrontendLiteralBlueprintAccessPrivate::GetPODValueFromMetaSoundLiteral<TArray<FString>>(Literal, OutResult);
+}
diff --git a/Engine/Plugins/Runtime/Metasound/Source/MetasoundEngine/Public/MetasoundBuilderBase.cpp b/Engine/Plugins/Runtime/Metasound/Source/MetasoundEngine/Public/MetasoundBuilderBase.cpp
--- a/Engine/Plugins/Runtime/Metasound/Source/MetasoundEngine/Public/MetasoundBuilderBase.cpp
+++ b/Engine/Plugins/Runtime/Metasound/Source/MetasoundEngine/Public/MetasoundBuilderBase.cpp
@@ -730,10 +730,27 @@ TArray<FMetaSoundNodeHandle> UMetaSoundBuilderBase::FindInterfaceOutputNodes(FNa
 
 FMetaSoundNodeHandle UMetaSoundBuilderBase::FindGraphInputNode(FName InputName, EMetaSoundBuilderResult& OutResult)
 {
-	if (const FMetasoundFrontendNode* GraphInputNode = Builder.FindGraphInputNode(InputName))
+	FMetaSoundBuilderNodeOutputHandle NodeOutputHandle;
+	FName DataTypeName;
+	return FindGraphInputNode(InputName, DataTypeName, NodeOutputHandle, OutResult);
+}
+
+FMetaSoundNodeHandle UMetaSoundBuilderBase::FindGraphInputNode(FName InputName, FName& DataTypeName, FMetaSoundBuilderNodeOutputHandle& NodeOutputHandle, EMetaSoundBuilderResult& OutResult)
+{
+	const FMetasoundFrontendNode* GraphInputNode = Builder.FindGraphInputNode(InputName);
+	const FMetasoundFrontendClassInput* ClassInput = Builder.FindGraphInput(InputName);
+
+	if (GraphInputNode && ClassInput)
 	{
 		OutResult = EMetaSoundBuilderResult::Succeeded;
-		return FMetaSoundNodeHandle { GraphInputNode->GetID() };
+		DataTypeName = ClassInput->TypeName;
+
+		check(!GraphInputNode->Interface.Outputs.IsEmpty())
+		const FMetasoundFrontendVertex& OutputVertex = GraphInputNode->Interface.Outputs[0];
+		const FGuid NodeID = GraphInputNode->GetID();
+		NodeOutputHandle.NodeID = NodeID;
+		NodeOutputHandle.VertexID = OutputVertex.VertexID;
+		return FMetaSoundNodeHandle { NodeID };
 	}
 
 	UE_LOG(LogMetaSound, Display, TEXT("Failed to find graph input by name '%s' with builder '%s'"), *InputName.ToString(), *GetName());
@@ -741,12 +758,22 @@ FMetaSoundNodeHandle UMetaSoundBuilderBase::FindGraphInputNode(FName InputName,
 	return { };
 }
 
-FMetaSoundNodeHandle UMetaSoundBuilderBase::FindGraphOutputNode(FName OutputName, EMetaSoundBuilderResult& OutResult)
+FMetaSoundNodeHandle UMetaSoundBuilderBase::FindGraphOutputNode(FName OutputName, FName& DataTypeName, FMetaSoundBuilderNodeInputHandle& NodeInputHandle, EMetaSoundBuilderResult& OutResult)
 {
-	if (const FMetasoundFrontendNode* GraphOutputNode = Builder.FindGraphOutputNode(OutputName))
+	const FMetasoundFrontendNode* GraphOutputNode = Builder.FindGraphOutputNode(OutputName);
+	const FMetasoundFrontendClassOutput* ClassOutput = Builder.FindGraphOutput(OutputName);
+	if (GraphOutputNode && ClassOutput)
 	{
 		OutResult = EMetaSoundBuilderResult::Succeeded;
-		return FMetaSoundNodeHandle{ GraphOutputNode->GetID() };
+		DataTypeName = ClassOutput->TypeName;
+
+		check(!GraphOutputNode->Interface.Inputs.IsEmpty())
+		const FMetasoundFrontendVertex& InputVertex = GraphOutputNode->Interface.Inputs[0];
+		const FGuid NodeID = GraphOutputNode->GetID();
+		NodeInputHandle.NodeID = NodeID;
+		NodeInputHandle.VertexID = InputVertex.VertexID;
+
+		return FMetaSoundNodeHandle{ NodeID };
 	}
 
 	UE_LOG(LogMetaSound, Display, TEXT("Failed to find graph output by name '%s' with builder '%s'"), *OutputName.ToString(), *GetName());
@@ -754,6 +781,13 @@ FMetaSoundNodeHandle UMetaSoundBuilderBase::FindGraphOutputNode(FName OutputName
 	return { };
 }
 
+FMetaSoundNodeHandle UMetaSoundBuilderBase::FindGraphOutputNode(FName OutputName, EMetaSoundBuilderResult& OutResult)
+{
+	FMetaSoundBuilderNodeInputHandle NodeInputHandle;
+	FName DataTypeName;
+	return FindGraphOutputNode(OutputName, DataTypeName, NodeInputHandle, OutResult);
+}
+
 #if WITH_EDITOR
 UMetaSoundFrontendMemberMetadata* UMetaSoundBuilderBase::FindMemberMetadata(const FGuid& InMemberID)
 {
@@ -783,6 +817,39 @@ const FMetaSoundFrontendDocumentBuilder& UMetaSoundBuilderBase::GetConstBuilder(
 	return Builder;
 }
 
+FMetasoundFrontendLiteral UMetaSoundBuilderBase::GetGraphInputDefault(FName InputName, EMetaSoundBuilderResult& OutResult) const
+{
+	if (const FMetasoundFrontendLiteral* Default = Builder.GetGraphInputDefault(InputName))
+	{
+		OutResult = EMetaSoundBuilderResult::Succeeded;
+		return *Default;
+	}
+	OutResult = EMetaSoundBuilderResult::Failed;
+	return { };
+}
+
+TArray<FName> UMetaSoundBuilderBase::GetGraphInputNames(EMetaSoundBuilderResult& OutResult) const
+{
+	TArray<FName> GraphInputNames;
+	const FMetasoundFrontendGraphClass& RootGraph = Builder.GetConstDocumentChecked().RootGraph;
+	for (const FMetasoundFrontendClassInput& Input : RootGraph.Interface.Inputs)
+	{
+		GraphInputNames.Emplace(Input.Name);
+	}
+	return GraphInputNames;
+}
+
+TArray<FName> UMetaSoundBuilderBase::GetGraphOutputNames(EMetaSoundBuilderResult& OutResult) const
+{
+	TArray<FName> GraphOutputNames;
+	const FMetasoundFrontendGraphClass& RootGraph = Builder.GetConstDocumentChecked().RootGraph;
+	for (const FMetasoundFrontendClassOutput& Output : RootGraph.Interface.Outputs)
+	{
+		GraphOutputNames.Emplace(Output.Name);
+	}
+	return GraphOutputNames;
+}
+
 int32 UMetaSoundBuilderBase::GetLastTransactionRegistered() const
 {
 	return LastTransactionRegistered;
diff --git a/Engine/Plugins/Runtime/Metasound/Source/MetasoundEngine/Public/MetasoundBuilderBase.h b/Engine/Plugins/Runtime/Metasound/Source/MetasoundEngine/Public/MetasoundBuilderBase.h
--- a/Engine/Plugins/Runtime/Metasound/Source/MetasoundEngine/Public/MetasoundBuilderBase.h
+++ b/Engine/Plugins/Runtime/Metasound/Source/MetasoundEngine/Public/MetasoundBuilderBase.h
@@ -8,6 +8,7 @@
 #include "MetasoundFrontendController.h"
 #include "MetasoundFrontendDocument.h"
 #include "MetasoundFrontendDocumentBuilder.h"
+#include "MetasoundFrontendDocumentModifyDelegates.h"
 #include "Subsystems/EngineSubsystem.h"
 #include "Templates/Function.h"
 #include "UObject/NoExportTypes.h"
@@ -226,10 +227,14 @@ public:
 
 	// Returns graph input node by the given name if it exists, or an invalid handle if not found.
 	UFUNCTION(BlueprintCallable, Category = "Audio|MetaSound|Builder", meta = (ExpandEnumAsExecs = "OutResult"))
+	UPARAM(DisplayName = "Node Handle") FMetaSoundNodeHandle FindGraphInputNode(FName InputName, UPARAM(DisplayName = "Data Type") FName& OutDataType, FMetaSoundBuilderNodeOutputHandle& NodeOutputHandle, EMetaSoundBuilderResult& OutResult);
+
 	UPARAM(DisplayName = "Node Handle") FMetaSoundNodeHandle FindGraphInputNode(FName InputName, EMetaSoundBuilderResult& OutResult);
 
 	// Returns graph output node by the given name if it exists, or an invalid handle if not found.
 	UFUNCTION(BlueprintCallable, Category = "Audio|MetaSound|Builder", meta = (ExpandEnumAsExecs = "OutResult"))
+	UPARAM(DisplayName = "Node Handle") FMetaSoundNodeHandle FindGraphOutputNode(FName OutputName, UPARAM(DisplayName = "Data Type") FName& OutDataType, FMetaSoundBuilderNodeInputHandle& NodeInputHandle, EMetaSoundBuilderResult& OutResult);
+	
 	UPARAM(DisplayName = "Node Handle") FMetaSoundNodeHandle FindGraphOutputNode(FName OutputName, EMetaSoundBuilderResult& OutResult);
 
 #if WITH_EDITOR
@@ -278,7 +283,19 @@ public:
 
 	// Returns output's parent node if the input is valid, otherwise returns invalid node handle.
 	UFUNCTION(BlueprintCallable, Category = "Audio|MetaSound|Builder", meta = (ExpandEnumAsExecs = "OutResult"))
-	UPARAM(DisplayName = "Node ClassVersion") FMetasoundFrontendVersion FindNodeClassVersion(const FMetaSoundNodeHandle& NodeHandle, EMetaSoundBuilderResult& OutResult);
+	UPARAM(DisplayName = "Node Class Version") FMetasoundFrontendVersion FindNodeClassVersion(const FMetaSoundNodeHandle& NodeHandle, EMetaSoundBuilderResult& OutResult);
+	
+	// Gets names of all graph inputs. 
+	UFUNCTION(BlueprintCallable, Category = "Audio|MetaSound|Builder", meta = (ExpandEnumAsExecs = "OutResult"))
+	UPARAM(DisplayName = "Graph Inputs") TArray<FName> GetGraphInputNames(EMetaSoundBuilderResult& OutResult) const;
+			
+	// Gets names of all graph outputs. 
+	UFUNCTION(BlueprintCallable, Category = "Audio|MetaSound|Builder", meta = (ExpandEnumAsExecs = "OutResult"))
+	UPARAM(DisplayName = "Graph Outputs") TArray<FName> GetGraphOutputNames(EMetaSoundBuilderResult& OutResult) const;
+	
+	// Gets the graph input's default literal value.
+	UFUNCTION(BlueprintCallable, Category = "Audio|MetaSound|Builder", meta = (ExpandEnumAsExecs = "OutResult"))
+	UPARAM(DisplayName = "Literal Value") FMetasoundFrontendLiteral GetGraphInputDefault(FName InputName, EMetaSoundBuilderResult& OutResult) const;
 
 	// Returns the MetaSound asset's graph class name (used by the MetaSound Node Class Registry)
 	UFUNCTION(BlueprintCallable, Category = "Audio|MetaSound|Builder|Advanced", meta = (DisplayName = "Get MetaSound Class Name"))
diff --git a/Engine/Plugins/Runtime/Metasound/Source/MetasoundEngine/Public/MetasoundFrontendLiteralBlueprintAccess.h b/Engine/Plugins/Runtime/Metasound/Source/MetasoundEngine/Public/MetasoundFrontendLiteralBlueprintAccess.h
--- a/Engine/Plugins/Runtime/Metasound/Source/MetasoundEngine/Public/MetasoundFrontendLiteralBlueprintAccess.h
+++ b/Engine/Plugins/Runtime/Metasound/Source/MetasoundEngine/Public/MetasoundFrontendLiteralBlueprintAccess.h
@@ -4,8 +4,8 @@
 
 #include "MetasoundFrontendLiteral.h"
 
+#include "MetasoundBuilderSubsystem.h"
 #include "Kismet/BlueprintFunctionLibrary.h"
-
 #include "MetasoundFrontendLiteralBlueprintAccess.generated.h"
 
 /**
@@ -15,6 +15,14 @@ UCLASS()
 class UMetasoundFrontendLiteralBlueprintAccess : public UBlueprintFunctionLibrary
 {
 	GENERATED_BODY()
+	
+	UFUNCTION(BlueprintPure, meta = (DisplayName = "To String (MetaSound Literal)", CompactNodeTitle = "->", BlueprintAutocast), Category = "Utilities|String")
+	static FString Conv_MetaSoundLiteralToString(const FMetasoundFrontendLiteral& Literal);
+	
+	UFUNCTION(BlueprintPure, Category = "Audio|MetaSound", meta = (DisplayName = "Get MetaSound Literal Type"))
+	static UPARAM(DisplayName = "Type") EMetasoundFrontendLiteralType GetType(const FMetasoundFrontendLiteral& Literal);
+
+	// Literal creation
 
 	UFUNCTION(BlueprintCallable, Category = "Audio|MetaSound", meta = (DisplayName = "Create MetaSound Bool Literal"))
 	static UPARAM(DisplayName = "Bool Literal") FMetasoundFrontendLiteral CreateBoolMetaSoundLiteral(bool Value);
@@ -48,4 +56,36 @@ class UMetasoundFrontendLiteralBlueprintAccess : public UBlueprintFunctionLibrar
 
 	UFUNCTION(BlueprintCallable, Category = "Audio|MetaSound", meta = (DisplayName = "Create MetaSound Literal From AudioParameter"))
 	static UPARAM(DisplayName = "Param Literal") FMetasoundFrontendLiteral CreateMetaSoundLiteralFromParam(const FAudioParameter& Param);
+
+	// Value accessors 
+
+	UFUNCTION(BlueprintCallable, Category = "Audio|MetaSound", meta = (ExpandEnumAsExecs = "OutResult"))
+	static bool GetBoolValueFromLiteral(const FMetasoundFrontendLiteral& Literal, EMetaSoundBuilderResult& OutResult);
+
+	UFUNCTION(BlueprintCallable, Category = "Audio|MetaSound", meta = (ExpandEnumAsExecs = "OutResult"))
+	static TArray<bool> GetBoolArrayValueFromLiteral(const FMetasoundFrontendLiteral& Literal, EMetaSoundBuilderResult& OutResult);
+
+	UFUNCTION(BlueprintCallable, Category = "Audio|MetaSound", meta = (ExpandEnumAsExecs = "OutResult"))
+	static float GetFloatValueFromLiteral(const FMetasoundFrontendLiteral& Literal, EMetaSoundBuilderResult& OutResult);
+
+	UFUNCTION(BlueprintCallable, Category = "Audio|MetaSound", meta = (ExpandEnumAsExecs = "OutResult"))
+	static TArray<float> GetFloatArrayValueFromLiteral(const FMetasoundFrontendLiteral& Literal, EMetaSoundBuilderResult& OutResult);
+
+	UFUNCTION(BlueprintCallable, Category = "Audio|MetaSound", meta = (ExpandEnumAsExecs = "OutResult"))
+	static int32 GetIntValueFromLiteral(const FMetasoundFrontendLiteral& Literal, EMetaSoundBuilderResult& OutResult);
+
+	UFUNCTION(BlueprintCallable, Category = "Audio|MetaSound", meta = (ExpandEnumAsExecs = "OutResult"))
+	static TArray<int32> GetIntArrayValueFromLiteral(const FMetasoundFrontendLiteral& Literal, EMetaSoundBuilderResult& OutResult);
+
+	UFUNCTION(BlueprintCallable, Category = "Audio|MetaSound", meta = (ExpandEnumAsExecs = "OutResult"))
+	static UObject* GetObjectValueFromLiteral(const FMetasoundFrontendLiteral& Literal, EMetaSoundBuilderResult& OutResult);
+	
+	UFUNCTION(BlueprintCallable, Category = "Audio|MetaSound", meta = (ExpandEnumAsExecs = "OutResult"))
+	static TArray<UObject*> GetObjectArrayValueFromLiteral(const FMetasoundFrontendLiteral& Literal, EMetaSoundBuilderResult& OutResult);
+	
+	UFUNCTION(BlueprintCallable, Category = "Audio|MetaSound", meta = (ExpandEnumAsExecs = "OutResult"))
+	static FString GetStringValueFromLiteral(const FMetasoundFrontendLiteral& Literal, EMetaSoundBuilderResult& OutResult);
+	
+	UFUNCTION(BlueprintCallable, Category = "Audio|MetaSound", meta = (ExpandEnumAsExecs = "OutResult"))
+	static TArray<FString> GetStringArrayValueFromLiteral(const FMetasoundFrontendLiteral& Literal, EMetaSoundBuilderResult& OutResult);
 };
diff --git a/Engine/Plugins/Runtime/Metasound/Source/MetasoundFrontend/Private/MetasoundFrontendDocumentBuilder.cpp b/Engine/Plugins/Runtime/Metasound/Source/MetasoundFrontend/Private/MetasoundFrontendDocumentBuilder.cpp
--- a/Engine/Plugins/Runtime/Metasound/Source/MetasoundFrontend/Private/MetasoundFrontendDocumentBuilder.cpp
+++ b/Engine/Plugins/Runtime/Metasound/Source/MetasoundFrontend/Private/MetasoundFrontendDocumentBuilder.cpp
@@ -2442,6 +2442,16 @@ const FGuid& FMetaSoundFrontendDocumentBuilder::GetBuildPageID() const
 	return BuildPageID;
 }
 
+const FMetasoundFrontendLiteral* FMetaSoundFrontendDocumentBuilder::GetGraphInputDefault(FName InputName, const FGuid* InPageID) const
+{
+	if (const FMetasoundFrontendClassInput* GraphInput = FindGraphInput(InputName))
+	{
+		const FGuid& PageID = InPageID ? *InPageID : BuildPageID;
+		return GraphInput->FindConstDefault(PageID);
+	}
+	return nullptr;
+}
+
 EMetasoundFrontendVertexAccessType FMetaSoundFrontendDocumentBuilder::GetNodeInputAccessType(const FGuid& InNodeID, const FGuid& InVertexID, const FGuid* InPageID) const
 {
 	using namespace Metasound::Frontend;
@@ -4026,6 +4036,10 @@ bool FMetaSoundFrontendDocumentBuilder::SetGraphInputDataType(FName InputName, F
 				return false;
 			}
 
+#if WITH_EDITOR
+			DocumentDelegates->InterfaceDelegates.OnInputDataTypeChanged.Broadcast(*Index);
+#endif // WITH_EDITOR
+
 			RemoveUnusedDependencies();
 
 #if WITH_EDITORONLY_DATA
@@ -4359,6 +4373,9 @@ bool FMetaSoundFrontendDocumentBuilder::SetGraphOutputDataType(FName OutputName,
 		{
 			return false;
 		}
+#if WITH_EDITOR
+		DocumentDelegates->InterfaceDelegates.OnOutputDataTypeChanged.Broadcast(*Index);
+#endif // WITH_EDITOR
 
 #if WITH_EDITORONLY_DATA
 		ClearMemberMetadata(GraphOutput.NodeID);
diff --git a/Engine/Plugins/Runtime/Metasound/Source/MetasoundFrontend/Private/MetasoundFrontendDocumentCache.cpp b/Engine/Plugins/Runtime/Metasound/Source/MetasoundFrontend/Private/MetasoundFrontendDocumentCache.cpp
--- a/Engine/Plugins/Runtime/Metasound/Source/MetasoundFrontend/Private/MetasoundFrontendDocumentCache.cpp
+++ b/Engine/Plugins/Runtime/Metasound/Source/MetasoundFrontend/Private/MetasoundFrontendDocumentCache.cpp
@@ -405,10 +405,16 @@ namespace Metasound::Frontend
 		OutDelegates.OnInputAdded.AddSP(this, &FDocumentGraphInterfaceCache::OnInputAdded);
 		OutDelegates.OnInputDefaultChanged.AddSP(this, &FDocumentGraphInterfaceCache::OnInputDefaultChanged);
 		OutDelegates.OnOutputAdded.AddSP(this, &FDocumentGraphInterfaceCache::OnOutputAdded);
+
 		OutDelegates.OnRemovingInput.AddSP(this, &FDocumentGraphInterfaceCache::OnRemovingInput);
 		OutDelegates.OnRemovingOutput.AddSP(this, &FDocumentGraphInterfaceCache::OnRemovingOutput);
+
 		OutDelegates.OnInputNameChanged.AddSP(this, &FDocumentGraphInterfaceCache::OnInputNameChanged);
 		OutDelegates.OnOutputNameChanged.AddSP(this, &FDocumentGraphInterfaceCache::OnOutputNameChanged);
+#if WITH_EDITOR
+		OutDelegates.OnInputDataTypeChanged.AddSP(this, &FDocumentGraphInterfaceCache::OnInputDataTypeChanged);
+		OutDelegates.OnOutputDataTypeChanged.AddSP(this, &FDocumentGraphInterfaceCache::OnOutputDataTypeChanged);
+#endif // WITH_EDITOR
 	}
 
 	void FDocumentGraphInterfaceCache::OnInputAdded(int32 NewIndex)
@@ -421,6 +427,13 @@ namespace Metasound::Frontend
 		++TransactionCount;
 	}
 
+#if WITH_EDITOR
+	void FDocumentGraphInterfaceCache::OnInputDataTypeChanged(int32 Index)
+	{
+		++TransactionCount;
+	}
+#endif // WITH_EDITOR
+
 	void FDocumentGraphInterfaceCache::OnInputDefaultChanged(int32 NewIndex)
 	{
 		++TransactionCount;
@@ -446,6 +459,13 @@ namespace Metasound::Frontend
 		++TransactionCount;
 	}
 
+#if WITH_EDITOR
+	void FDocumentGraphInterfaceCache::OnOutputDataTypeChanged(int32 Index)
+	{
+		++TransactionCount;
+	}
+#endif // WITH_EDITOR
+
 	void FDocumentGraphInterfaceCache::OnOutputNameChanged(FName OldName, FName NewName)
 	{
 		int32* Index = OutputNameToIndex.Find(OldName);
diff --git a/Engine/Plugins/Runtime/Metasound/Source/MetasoundFrontend/Private/MetasoundFrontendDocumentCache.h b/Engine/Plugins/Runtime/Metasound/Source/MetasoundFrontend/Private/MetasoundFrontendDocumentCache.h
--- a/Engine/Plugins/Runtime/Metasound/Source/MetasoundFrontend/Private/MetasoundFrontendDocumentCache.h
+++ b/Engine/Plugins/Runtime/Metasound/Source/MetasoundFrontend/Private/MetasoundFrontendDocumentCache.h
@@ -144,6 +144,11 @@ namespace Metasound::Frontend
 		void OnRemovingInput(int32 IndexBeingRemoved);
 		void OnRemovingOutput(int32 IndexBeingRemoved);
 
+#if WITH_EDITOR
+		void OnInputDataTypeChanged(int32 Index);
+		void OnOutputDataTypeChanged(int32 Index);
+#endif // WITH_EDITOR
+
 		// Cache of Input name to array index of input
 		TMap<FName, int32> InputNameToIndex;
 
diff --git a/Engine/Plugins/Runtime/Metasound/Source/MetasoundFrontend/Public/MetasoundFrontendDocument.h b/Engine/Plugins/Runtime/Metasound/Source/MetasoundFrontend/Public/MetasoundFrontendDocument.h
--- a/Engine/Plugins/Runtime/Metasound/Source/MetasoundFrontend/Public/MetasoundFrontendDocument.h
+++ b/Engine/Plugins/Runtime/Metasound/Source/MetasoundFrontend/Public/MetasoundFrontendDocument.h
@@ -1771,7 +1771,7 @@ public:
 #endif // WITH_EDITORONLY_DATA
 };
 
-UCLASS()
+UCLASS(BlueprintType)
 class METASOUNDFRONTEND_API UMetaSoundFrontendMemberMetadata : public UObject
 {
 	GENERATED_BODY()
diff --git a/Engine/Plugins/Runtime/Metasound/Source/MetasoundFrontend/Public/MetasoundFrontendDocumentBuilder.h b/Engine/Plugins/Runtime/Metasound/Source/MetasoundFrontend/Public/MetasoundFrontendDocumentBuilder.h
--- a/Engine/Plugins/Runtime/Metasound/Source/MetasoundFrontend/Public/MetasoundFrontendDocumentBuilder.h
+++ b/Engine/Plugins/Runtime/Metasound/Source/MetasoundFrontend/Public/MetasoundFrontendDocumentBuilder.h
@@ -318,6 +318,8 @@ public:
 	const bool GetIsAdvancedDisplay(const FName MemberName, const EMetasoundFrontendClassType Type) const;
 #endif // WITH_EDITORONLY_DATA
 
+	const FMetasoundFrontendLiteral* GetGraphInputDefault(FName InputName, const FGuid* InPageID = nullptr) const;
+	
 	// Initializes the builder's document, using the (optional) provided document template, (optional) class name, and (optionally) whether or not to reset the existing class version.
 	void InitDocument(const FMetasoundFrontendDocument* InDocumentTemplate = nullptr, const FMetasoundFrontendClassName* InNewClassName = nullptr, bool bResetVersion = true);
 
diff --git a/Engine/Plugins/Runtime/Metasound/Source/MetasoundFrontend/Public/MetasoundFrontendDocumentModifyDelegates.h b/Engine/Plugins/Runtime/Metasound/Source/MetasoundFrontend/Public/MetasoundFrontendDocumentModifyDelegates.h
--- a/Engine/Plugins/Runtime/Metasound/Source/MetasoundFrontend/Public/MetasoundFrontendDocumentModifyDelegates.h
+++ b/Engine/Plugins/Runtime/Metasound/Source/MetasoundFrontend/Public/MetasoundFrontendDocumentModifyDelegates.h
@@ -38,11 +38,22 @@ namespace Metasound::Frontend
 		FOnMetaSoundFrontendDocumentMutateInterfaceArray OnRemovingInterface;
 
 		FOnMetaSoundFrontendDocumentMutateArray OnInputAdded;
+
+// Currently only used in editor contexts  
+// so only enabled there to avoid unnecessary delegate overhead
+// but may be changed in the future 
+#if WITH_EDITOR
+		FOnMetaSoundFrontendDocumentMutateArray OnInputDataTypeChanged;
+#endif // WITH_EDITOR
 		FOnMetaSoundFrontendDocumentMutateArray OnInputDefaultChanged;
+
 		FOnMetaSoundFrontendDocumentRenameVertex OnInputNameChanged;
 		FOnMetaSoundFrontendDocumentMutateArray OnRemovingInput;
 
 		FOnMetaSoundFrontendDocumentMutateArray OnOutputAdded;
+#if WITH_EDITOR
+		FOnMetaSoundFrontendDocumentMutateArray OnOutputDataTypeChanged;
+#endif // WITH_EDITOR
 		FOnMetaSoundFrontendDocumentRenameVertex OnOutputNameChanged;
 		FOnMetaSoundFrontendDocumentMutateArray OnRemovingOutput;
 	};
diff --git a/Engine/Plugins/Runtime/Metasound/Source/MetasoundFrontend/Public/MetasoundFrontendLiteral.h b/Engine/Plugins/Runtime/Metasound/Source/MetasoundFrontend/Public/MetasoundFrontendLiteral.h
--- a/Engine/Plugins/Runtime/Metasound/Source/MetasoundFrontend/Public/MetasoundFrontendLiteral.h
+++ b/Engine/Plugins/Runtime/Metasound/Source/MetasoundFrontend/Public/MetasoundFrontendLiteral.h
@@ -3,7 +3,6 @@
 #pragma once
 
 #include "AudioParameterControllerInterface.h"
-#include "CoreMinimal.h"
 #include "MetasoundLiteral.h"
 
 #include <type_traits>
@@ -15,7 +14,7 @@
 // The EMetasoundFrontendLiteralType's are matched to Metasound::ELiteralType`s 
 // by giving them the same value. This supports easy conversion from one type to
 // another.
-UENUM()
+UENUM(BlueprintType)
 enum class EMetasoundFrontendLiteralType : uint8
 {
 
@@ -163,6 +162,12 @@ public:
 
 	static FMetasoundFrontendLiteral GetInvalid();
 
+	friend bool operator==(const FMetasoundFrontendLiteral& A, const FMetasoundFrontendLiteral& B)
+	{
+		return A.IsEqual(B);
+	}
+
+
 private:
 	// Remove all values.
 	void Empty();
diff --git a/Engine/Source/Runtime/Engine/Private/Audio/AudioWidgetSubsystem.cpp b/Engine/Source/Runtime/Engine/Private/Audio/AudioWidgetSubsystem.cpp
--- a/Engine/Source/Runtime/Engine/Private/Audio/AudioWidgetSubsystem.cpp
+++ b/Engine/Source/Runtime/Engine/Private/Audio/AudioWidgetSubsystem.cpp
@@ -2,8 +2,9 @@
 #include "Audio/AudioWidgetSubsystem.h"
 
 #include "AssetRegistry/ARFilter.h"
-#include "AssetRegistry/AssetRegistryModule.h"
 #include "AssetRegistry/AssetData.h"
+#include "AssetRegistry/AssetRegistryHelpers.h"
+#include "AssetRegistry/AssetRegistryModule.h"
 #include "Blueprint/BlueprintSupport.h"
 #include "Blueprint/UserWidget.h"
 #include "Engine/Blueprint.h"
@@ -65,7 +66,8 @@ const TArray<FAssetData> UAudioWidgetSubsystem::GetBlueprintAssetData(UClass* In
 	FARFilter Filter;
 	Filter.ClassPaths.Add(UBlueprint::StaticClass()->GetClassPathName());
 	Filter.ClassPaths.Add(FTopLevelAssetPath(TEXT("/Script/UMGEditor"), TEXT("WidgetBlueprint")));
-	Filter.TagsAndValues.Add(FBlueprintTags::ImplementedInterfaces);
+	Filter.ClassPaths.Add(FTopLevelAssetPath(TEXT("/Script/Blutility"), TEXT("EditorUtilityWidget")));
+	Filter.bRecursiveClasses = true;
 
 	AssetRegistryModule.Get().EnumerateAssets(Filter, [FoundAssets = &AssetData, InInterfaceClass](const FAssetData& InAssetData)
 	{
@@ -82,6 +84,13 @@ const TArray<FAssetData> UAudioWidgetSubsystem::GetBlueprintAssetData(UClass* In
 
 bool UAudioWidgetSubsystem::ImplementsInterface(const FAssetData& InAssetData, UClass* InInterfaceClass)
 {
+	UClass* NativeClass = UAssetRegistryHelpers::FindAssetNativeClass(InAssetData);
+	// Native interfaces can't use the BP tag system
+	if (InAssetData.GetClass()->ImplementsInterface(InInterfaceClass))
+	{
+		return true;
+	}
+
 	const FString ImplementedInterfaces = InAssetData.GetTagValueRef<FString>(FBlueprintTags::ImplementedInterfaces);
 	if (!ImplementedInterfaces.IsEmpty())
 	{
