MetaSound Vertex Deprecation Cleanup
#rnx

diff --git a/Engine/Plugins/Runtime/Metasound/Source/MetasoundGraphCore/Public/MetasoundVertex.h b/Engine/Plugins/Runtime/Metasound/Source/MetasoundGraphCore/Public/MetasoundVertex.h
--- a/Engine/Plugins/Runtime/Metasound/Source/MetasoundGraphCore/Public/MetasoundVertex.h
+++ b/Engine/Plugins/Runtime/Metasound/Source/MetasoundGraphCore/Public/MetasoundVertex.h
@@ -75,26 +75,6 @@ namespace Metasound
 
 		/** Access type of the vertex. */
 		EVertexAccessType AccessType;
-
-		UE_DEPRECATED(5.1, "GetVertexName() is deprecated. Use VertexName")
-		const FVertexName& GetVertexName() const
-		{
-			return VertexName;
-		}
-
-		UE_DEPRECATED(5.1, "GetDataTypeName() is deprecated. Use DataTypeName")
-		const FName& GetDataTypeName() const
-		{
-			return DataTypeName;
-		}
-
-#if WITH_EDITORONLY_DATA
-		UE_DEPRECATED(5.1, "GetMetadata() is deprecated. Use Metadata")
-		const FDataVertexMetadata& GetMetadata() const
-		{
-			return Metadata;
-		}
-#endif // WITH_EDITORONLY_DATA
 	};
 
 	/** FInputDataVertex */
@@ -152,18 +132,8 @@ namespace Metasound
 		: FInputDataVertex(InVertexName, GetMetasoundDataTypeName<DataType>(), InMetadata, EVertexAccessType::Reference, Forward<ArgTypes>(Args)...)
 		{
 		}
-
-		template<typename... ArgTypes>
- 		UE_DEPRECATED(5.1, "TInputDataVertex constructor without explicit FDataVertexMatadata is deprecated. Use constructor which accepts TDataVertexMetadata")
-		TInputDataVertex(const FVertexName& InVertexName, const FText& InDescription, ArgTypes&&... Args)
-		: FInputDataVertex(InVertexName, GetMetasoundDataTypeName<DataType>(), { InDescription }, EVertexAccessType::Reference, Forward<ArgTypes>(Args)...)
-		{
-		}
 	};
 
-	template<typename DataType>
-	using TInputDataVertexModel UE_DEPRECATED(5.1, "TInputDataVertexModel<DataType> is deprecated. Replace with TInputDataVertex<DataType>") = TInputDataVertex<DataType>;
-
 	/** Create a FInputDataVertex with a templated MetaSound data type which only
 	 * reads data at operator time. */
 	template<typename DataType>
@@ -206,18 +176,8 @@ namespace Metasound
 		: FOutputDataVertex(InVertexName, GetMetasoundDataTypeName<DataType>(), InMetadata, EVertexAccessType::Reference, Forward<ArgTypes>(Args)...)
 		{
 		}
-
-		template<typename... ArgTypes>
- 		UE_DEPRECATED(5.1, "TOutputDataVertex constructor without explicit FDataVertexMatadata is deprecated. Use constructor which accepts TDataVertexMetadata")
-		TOutputDataVertex(const FVertexName& InVertexName, const FText& InDescription, ArgTypes&&... Args)
-		: FOutputDataVertex(InVertexName, GetMetasoundDataTypeName<DataType>(), { InDescription }, EVertexAccessType::Reference, Forward<ArgTypes>(Args)...)
-		{
-		}
 	};
 
-	template<typename DataType>
-	using TOutputDataVertexModel UE_DEPRECATED(5.1, "TOutputDataVertexModel<DataType> is deprecated. Replace with TOutputDataVertex<DataType>") = TOutputDataVertex<DataType>;
-
 	/** Create a FOutputDataVertex with a templated MetaSound data type which is only
 	 * sets data at operator construction time. 
 	 */
@@ -265,22 +225,6 @@ namespace Metasound
 		FText Description;
 #endif // WITH_EDITORONLY_DATA
 
-		/** Name of vertex. */
-		UE_DEPRECATED(5.1, "GetVertexName() is deprecated. Use VertexName")
-		const FVertexName& GetVertexName() const
-		{
-			return VertexName;
-		}
-
-#if WITH_EDITORONLY_DATA
-		/** Description of the vertex. */
-		UE_DEPRECATED(5.1, "GetDescription() is deprecated. Use Description")
-		const FText& GetDescription() const
-		{
-			return Description;
-		}
-#endif // WITH_EDITORONLY_DATA
-
 		friend bool METASOUNDGRAPHCORE_API operator==(const FEnvironmentVertex& InLHS, const FEnvironmentVertex& InRHS);
 		friend bool METASOUNDGRAPHCORE_API operator!=(const FEnvironmentVertex& InLHS, const FEnvironmentVertex& InRHS);
 		friend bool METASOUNDGRAPHCORE_API operator<(const FEnvironmentVertex& InLHS, const FEnvironmentVertex& InRHS);
