Loctext CIS fix in MetaSound editor

#rnx

diff --git a/Engine/Plugins/Runtime/Metasound/Source/MetasoundEditor/Private/MetasoundEditor.cpp b/Engine/Plugins/Runtime/Metasound/Source/MetasoundEditor/Private/MetasoundEditor.cpp
--- a/Engine/Plugins/Runtime/Metasound/Source/MetasoundEditor/Private/MetasoundEditor.cpp
+++ b/Engine/Plugins/Runtime/Metasound/Source/MetasoundEditor/Private/MetasoundEditor.cpp
@@ -5041,7 +5041,7 @@ namespace Metasound
 				TSharedPtr<SDockTab> GraphTab = TabManager->FindExistingLiveTab(TabNamesPrivate::GraphCanvas);
 				if (GraphTab.IsValid())
 				{
-					const FText GraphTabLabel = FText::Format(LOCTEXT("GraphCanvasTab", "Graph Page: {0}"), PageStatsWidget.IsValid() ?
+					const FText GraphTabLabel = FText::Format(LOCTEXT("GraphCanvasPageTab", "Graph Page: {0}"), PageStatsWidget.IsValid() ?
 						FText::FromName(PageStatsWidget->GetDisplayedPageName()) : LOCTEXT("MetasoundGraphCanvasTitle", "MetaSound Graph"));
 					GraphTab->SetLabel(GraphTabLabel);
 				}
