Optimization for typical case at runtime where only a single page value is provided

[RN] Optimized MetaSound page selection at runtime to be more performant in case where only one page value implementation is cooked.

diff --git a/Engine/Plugins/Runtime/Metasound/Source/MetasoundEngine/Private/MetasoundDocumentBuilderRegistry.cpp b/Engine/Plugins/Runtime/Metasound/Source/MetasoundEngine/Private/MetasoundDocumentBuilderRegistry.cpp
--- a/Engine/Plugins/Runtime/Metasound/Source/MetasoundEngine/Private/MetasoundDocumentBuilderRegistry.cpp
+++ b/Engine/Plugins/Runtime/Metasound/Source/MetasoundEngine/Private/MetasoundDocumentBuilderRegistry.cpp
@@ -441,6 +441,21 @@ namespace Metasound::Engine
 
 	FGuid FDocumentBuilderRegistry::ResolveTargetPageID(const FMetasoundFrontendGraphClass& InGraphClass) const
 	{
+		METASOUND_LLM_SCOPE;
+		METASOUND_TRACE_CPUPROFILER_EVENT_SCOPE(FDocumentBuilderRegistry::ResolveTargetPageID_GraphClass);
+
+#if !WITH_EDITORONLY_DATA
+		// No resolution required if only one item. This is the typical scenario and saves from tying up a lock
+		// and calling resolution delegate(s). For larger graphs, this can add up, and most implementation
+		// only has a single page value. With editor-only data, go ahead and take the perf hit in favor of
+		// resolution reporting (for example if page data is invalid and needs to be fixed up).
+		const TArray<FMetasoundFrontendGraph>& GraphPages = InGraphClass.GetConstGraphPages();
+		if (GraphPages.Num() == 1)
+		{
+			return GraphPages.Last().PageID;
+		}
+#endif // !WITH_EDITORONLY_DATA
+
 		FScopeLock Lock(&TargetPageResolveScratchCritSec);
 		TargetPageResolveScratch.Reset();
 		InGraphClass.IterateGraphPages([this](const FMetasoundFrontendGraph& PageGraph)
@@ -453,6 +468,21 @@ namespace Metasound::Engine
 
 	FGuid FDocumentBuilderRegistry::ResolveTargetPageID(const FMetasoundFrontendClassInput& InClassInput) const
 	{
+		METASOUND_LLM_SCOPE;
+		METASOUND_TRACE_CPUPROFILER_EVENT_SCOPE(FDocumentBuilderRegistry::ResolveTargetPageID_ClassInput);
+
+#if !WITH_EDITORONLY_DATA
+		// No resolution required if only one item. This is the typical scenario and saves from tying up a lock
+		// and calling resolution delegate(s). For larger graphs, this can add up, and most implementation
+		// only has a single page value. With editor-only data, go ahead and take the perf hit in favor of
+		// resolution reporting (for example if page data is invalid and needs to be fixed up).
+		const TArray<FMetasoundFrontendClassInputDefault>& ClassDefaults = InClassInput.GetDefaults();
+		if (ClassDefaults.Num() == 1)
+		{
+			return ClassDefaults.Last().PageID;
+		}
+#endif // !WITH_EDITORONLY_DATA
+
 		FScopeLock Lock(&TargetPageResolveScratchCritSec);
 		TargetPageResolveScratch.Reset();
 		InClassInput.IterateDefaults([this](const FGuid& PageID, const FMetasoundFrontendLiteral&)
@@ -465,6 +495,20 @@ namespace Metasound::Engine
 
 	FGuid FDocumentBuilderRegistry::ResolveTargetPageID(const TArray<FMetasoundFrontendClassInputDefault>& InClassDefaults) const
 	{
+		METASOUND_LLM_SCOPE;
+		METASOUND_TRACE_CPUPROFILER_EVENT_SCOPE(FDocumentBuilderRegistry::ResolveTargetPageID_ClassDefaults);
+
+#if !WITH_EDITORONLY_DATA
+		// No resolution required if only one item. This is the typical scenario and saves from tying up a lock
+		// and calling resolution delegate(s). For larger graphs, this can add up, and most implementation
+		// only has a single page value. With editor-only data, go ahead and take the perf hit in favor of
+		// resolution reporting (for example if page data is invalid and needs to be fixed up).
+		if (InClassDefaults.Num() == 1)
+		{
+			return InClassDefaults.Last().PageID;
+		}
+#endif // !WITH_EDITORONLY_DATA
+
 		FScopeLock Lock(&TargetPageResolveScratchCritSec);
 		TargetPageResolveScratch.Reset();
 		Algo::Transform(InClassDefaults, TargetPageResolveScratch, [](const FMetasoundFrontendClassInputDefault& ClassDefault) { return ClassDefault.PageID; });
@@ -473,9 +517,6 @@ namespace Metasound::Engine
 
 	FGuid FDocumentBuilderRegistry::ResolveTargetPageIDInternal(const TArray<FGuid>& InPageIDsToResolve) const
 	{
-		METASOUND_LLM_SCOPE;
-		METASOUND_TRACE_CPUPROFILER_EVENT_SCOPE(FDocumentBuilderRegistry::ResolveTargetPageID);
-
 		FName PlatformName = FPlatformProperties::IniPlatformName();
 
 #if WITH_EDITOR
