Put noisy metasound trace scopes behind au.metasound.VerboseBuildGraphTraceEvents CVar.

There can be up to 5k+ scopes when building a single operator which can overwhelm some profilers.
This change reduces that to < 100 while retaining information about how many sub operators are being created.

#rnx

diff --git a/Engine/Plugins/Runtime/Metasound/Source/MetasoundGraphCore/Private/MetasoundOperatorBuilder.cpp b/Engine/Plugins/Runtime/Metasound/Source/MetasoundGraphCore/Private/MetasoundOperatorBuilder.cpp
--- a/Engine/Plugins/Runtime/Metasound/Source/MetasoundGraphCore/Private/MetasoundOperatorBuilder.cpp
+++ b/Engine/Plugins/Runtime/Metasound/Source/MetasoundGraphCore/Private/MetasoundOperatorBuilder.cpp
@@ -4,6 +4,7 @@
 
 #include "Algo/Count.h"
 #include "Algo/Transform.h"
+#include "HAL/IConsoleManager.h"
 #include "MetasoundBuildError.h"
 #include "MetasoundBuilderInterface.h"
 #include "MetasoundBuildError.h"
@@ -24,6 +25,13 @@
 
 namespace Metasound
 {
+	static bool bVerboseTraceEvents = false;
+	FAutoConsoleVariableRef CVarVerboseTraceEvents(
+		TEXT("au.metasound.VerboseBuildGraphTraceEvents"),
+		bVerboseTraceEvents,
+		TEXT("Enable additional trace scopes for operator building")
+	);
+
 	namespace OperatorBuilder
 	{
 		// Shared context used in the builder to maintain state of current build.
@@ -83,7 +91,7 @@ namespace Metasound
 
 	TUniquePtr<IOperator> FOperatorBuilder::BuildGraphOperator(const FBuildGraphOperatorParams& InParams, FBuildResults& OutResults) const
 	{
-		METASOUND_TRACE_CPUPROFILER_EVENT_SCOPE(Metasound::FOperatorBuilder::BuildGraphOperator);
+		METASOUND_TRACE_CPUPROFILER_EVENT_SCOPE_COND(Metasound::FOperatorBuilder::BuildGraphOperator, bVerboseTraceEvents);
 
 		if (BuilderSettings.bEnableOperatorRebind)
 		{
@@ -132,7 +140,7 @@ namespace Metasound
 
 	TUniquePtr<IOperator> FOperatorBuilder::BuildStaticGraphOperator(const FBuildGraphOperatorParams& InParams, FBuildResults& OutResults) const
 	{
-		METASOUND_TRACE_CPUPROFILER_EVENT_SCOPE(Metasound::FOperatorBuilder::BuildStaticGraphOperator);
+		METASOUND_TRACE_CPUPROFILER_EVENT_SCOPE_COND(Metasound::FOperatorBuilder::BuildStaticGraphOperator, bVerboseTraceEvents);
 
 		TUniquePtr<DirectedGraphAlgo::FStaticGraphOperatorData> GraphOperatorData = MakeUnique<DirectedGraphAlgo::FStaticGraphOperatorData>(InParams.OperatorSettings);
 
@@ -295,7 +303,7 @@ namespace Metasound
 
 	FOperatorBuilder::FBuildStatus FOperatorBuilder::DepthFirstTopologicalSort(OperatorBuilder::FBuildContext& InOutContext, TArray<const INode*>& OutNodes) const
 	{
-		METASOUND_TRACE_CPUPROFILER_EVENT_SCOPE(Metasound::FOperatorBuilder::DepthFirstTopologicalSort);
+		METASOUND_TRACE_CPUPROFILER_EVENT_SCOPE_COND(Metasound::FOperatorBuilder::DepthFirstTopologicalSort, bVerboseTraceEvents);
 		using namespace OperatorBuilderPrivate;
 
 		bool bSuccess = DirectedGraphAlgo::DepthFirstTopologicalSort(InOutContext.AlgoAdapter, OutNodes);
@@ -313,7 +321,7 @@ namespace Metasound
 
 	FOperatorBuilder::FBuildStatus FOperatorBuilder::KahnsTopologicalSort(OperatorBuilder::FBuildContext& InOutContext, TArray<const INode*>& OutNodes) const
 	{
-		METASOUND_TRACE_CPUPROFILER_EVENT_SCOPE(Metasound::FOperatorBuilder::KahnsTopologicalSort);
+		METASOUND_TRACE_CPUPROFILER_EVENT_SCOPE_COND(Metasound::FOperatorBuilder::KahnsTopologicalSort, bVerboseTraceEvents);
 		using namespace OperatorBuilderPrivate;
 
 		bool bSuccess = DirectedGraphAlgo::KahnTopologicalSort(InOutContext.AlgoAdapter, OutNodes);
@@ -413,7 +421,7 @@ namespace Metasound
 
 	void FOperatorBuilder::InitializeOperatorInfo(const IGraph& InGraph, TArray<const INode*>& InSortedNodes, DirectedGraphAlgo::FGraphOperatorData& InOutGraphOperatorData) const
 	{
-		METASOUND_TRACE_CPUPROFILER_EVENT_SCOPE(Metasound::FOperatorBuilder::InitializeOperatorInfo);
+		METASOUND_TRACE_CPUPROFILER_EVENT_SCOPE_COND(Metasound::FOperatorBuilder::InitializeOperatorInfo, bVerboseTraceEvents);
 		using namespace DirectedGraphAlgo;
 
 		// Create FOperatorInfos from Nodes
@@ -423,7 +431,7 @@ namespace Metasound
 		OperatorMap.Reserve(OperatorMap.Num() + NumNodes);
 
 		{
-			METASOUND_TRACE_CPUPROFILER_EVENT_SCOPE(Metasound::FOperatorBuilder::InitializeOperatorInfo::Nodes);
+			METASOUND_TRACE_CPUPROFILER_EVENT_SCOPE_COND(Metasound::FOperatorBuilder::InitializeOperatorInfo::Nodes, bVerboseTraceEvents);
 			int32 Ordinal = 0;
 			for (const INode* Node : InSortedNodes)
 			{
@@ -434,7 +442,7 @@ namespace Metasound
 		}
 
 		{
-			METASOUND_TRACE_CPUPROFILER_EVENT_SCOPE(Metasound::FOperatorBuilder::InitializeOperatorInfo::OutputDestinations);
+			METASOUND_TRACE_CPUPROFILER_EVENT_SCOPE_COND(Metasound::FOperatorBuilder::InitializeOperatorInfo::OutputDestinations, bVerboseTraceEvents);
 			// Set the output destinations on operator infos
 			for (const FDataEdge& Edge : InGraph.GetDataEdges())
 			{
@@ -448,7 +456,7 @@ namespace Metasound
 
 	FOperatorBuilder::FBuildStatus FOperatorBuilder::GatherExternalInputDataReferences(OperatorBuilder::FBuildContext& InOutContext, const FInputVertexInterfaceData& InExternalInputData) const
 	{
-		METASOUND_TRACE_CPUPROFILER_EVENT_SCOPE(Metasound::FOperatorBuilder::GatherExternalInputDataReferences);
+		METASOUND_TRACE_CPUPROFILER_EVENT_SCOPE_COND(Metasound::FOperatorBuilder::GatherExternalInputDataReferences, bVerboseTraceEvents);
 		using namespace DirectedGraphAlgo;
 
 		FBuildStatus BuildStatus;
@@ -492,7 +500,7 @@ namespace Metasound
 
 	void FOperatorBuilder::GatherInternalGraphDataReferences(OperatorBuilder::FBuildContext& InOutContext, const TArray<const INode*>& InNodes, TMap<FGuid, FDataReferenceCollection>& OutNodeVertexData) const
 	{
-		METASOUND_TRACE_CPUPROFILER_EVENT_SCOPE(Metasound::FOperatorBuilder::GatherInternalGraphDataReferences);
+		METASOUND_TRACE_CPUPROFILER_EVENT_SCOPE_COND(Metasound::FOperatorBuilder::GatherInternalGraphDataReferences, bVerboseTraceEvents);
 		using namespace DirectedGraphAlgo;
 
 		for (const INode* NodePtr : InNodes)
@@ -507,7 +515,7 @@ namespace Metasound
 
 	FOperatorBuilder::FBuildStatus FOperatorBuilder::ValidateOperatorOutputsAreBound(const INode& InNode, const FOutputVertexInterfaceData& InVertexData) const
 	{
-		METASOUND_TRACE_CPUPROFILER_EVENT_SCOPE(Metasound::FOperatorBuilder::ValidateOperatorOutputsAreBound);
+		METASOUND_TRACE_CPUPROFILER_EVENT_SCOPE_COND(Metasound::FOperatorBuilder::ValidateOperatorOutputsAreBound, bVerboseTraceEvents);
 		using FOutputVertexBinding = MetasoundVertexDataPrivate::FOutputBinding;
 
 		FBuildStatus BuildStatus = FBuildStatus::NoError;
@@ -559,20 +567,20 @@ namespace Metasound
 				const FNodeClassName& NodeClassName = Node->GetMetadata().ClassName;
 				const FString NodeTraceName = NodeClassName.IsValid() ? NodeClassName.ToString() : Node->GetInstanceName().ToString();
 				TraceNamePtr << "Metasound::FOperatorBuilder::CreateOperators::CreateAndBind " << NodeTraceName;
-				METASOUND_TRACE_CPUPROFILER_EVENT_SCOPE_TEXT(*TraceNamePtr);
+				METASOUND_TRACE_CPUPROFILER_EVENT_SCOPE_TEXT_COND(*TraceNamePtr, bVerboseTraceEvents);
 #endif // METASOUND_CPUPROFILERTRACE_ENABLED
 
 				FBuildOperatorParams CreateParams{*Node, InOutContext.Settings, OperatorInfo.VertexData.GetInputs(), InOutContext.Environment, this, InOutContext.GraphRenderCost};
 				FOperatorFactorySharedRef Factory = Node->GetDefaultOperatorFactory();
 				if (ProfileOperators && Profiling::OperatorShouldBeProfiled(Node->GetMetadata()))
 				{
-					METASOUND_TRACE_CPUPROFILER_EVENT_SCOPE(Metasound::FOperatorBuilder::CreateOperators::CreateOperator);
+					METASOUND_TRACE_CPUPROFILER_EVENT_SCOPE_COND(Metasound::FOperatorBuilder::CreateOperators::CreateOperator, bVerboseTraceEvents);
 
 					OperatorInfo.Operator = MakeUnique<FProfilingOperator>(Factory->CreateOperator(CreateParams, InOutContext.Results), Node);
 				}
 				else
 				{
-					METASOUND_TRACE_CPUPROFILER_EVENT_SCOPE(Metasound::FOperatorBuilder::CreateOperators::CreateOperator);
+					METASOUND_TRACE_CPUPROFILER_EVENT_SCOPE_COND(Metasound::FOperatorBuilder::CreateOperators::CreateOperator, bVerboseTraceEvents);
 					OperatorInfo.Operator = Factory->CreateOperator(CreateParams, InOutContext.Results);
 				}
 
@@ -583,7 +591,7 @@ namespace Metasound
 
 				// Bind vertex to operator data
 				{
-					METASOUND_TRACE_CPUPROFILER_EVENT_SCOPE(Metasound::FOperatorBuilder::CreateOperators::BindInputsAndOutputs);
+					METASOUND_TRACE_CPUPROFILER_EVENT_SCOPE_COND(Metasound::FOperatorBuilder::CreateOperators::BindInputsAndOutputs, bVerboseTraceEvents);
 					// Inputs don't need to be bound for all nodes unless they are dynamic 
 					// Inputs for input nodes will be bound separately in GatherGraphDataReferences
 					if (BuilderSettings.bEnableOperatorRebind)
@@ -605,7 +613,7 @@ namespace Metasound
 
 			// Route outputs of operator to downstream operators' VertexData
 			{
-				METASOUND_TRACE_CPUPROFILER_EVENT_SCOPE(Metasound::FOperatorBuilder::CreateOperators::RouteOutputs);
+				METASOUND_TRACE_CPUPROFILER_EVENT_SCOPE_COND(Metasound::FOperatorBuilder::CreateOperators::RouteOutputs, bVerboseTraceEvents);
 				for (const TPair<FVertexName, TArray<FGraphOperatorData::FVertexDestination>>& OutputRouting : OperatorInfo.OutputConnections)
 				{
 					const FVertexName& VertexName = OutputRouting.Key;
@@ -636,7 +644,7 @@ namespace Metasound
 
 	FOperatorBuilder::FBuildStatus FOperatorBuilder::GatherGraphDataReferences(OperatorBuilder::FBuildContext& InOutContext, FVertexInterfaceData& OutVertexData) const
 	{
-		METASOUND_TRACE_CPUPROFILER_EVENT_SCOPE(Metasound::FOperatorBuilder::GatherGraphDataReferences);
+		METASOUND_TRACE_CPUPROFILER_EVENT_SCOPE_COND(Metasound::FOperatorBuilder::GatherGraphDataReferences, bVerboseTraceEvents);
 		using namespace OperatorBuilderPrivate;
 		using namespace DirectedGraphAlgo;
 		using FDestinationElement = FInputDataDestinationCollection::ElementType;
diff --git a/Engine/Plugins/Runtime/Metasound/Source/MetasoundGraphCore/Public/MetasoundTrace.h b/Engine/Plugins/Runtime/Metasound/Source/MetasoundGraphCore/Public/MetasoundTrace.h
--- a/Engine/Plugins/Runtime/Metasound/Source/MetasoundGraphCore/Public/MetasoundTrace.h
+++ b/Engine/Plugins/Runtime/Metasound/Source/MetasoundGraphCore/Public/MetasoundTrace.h
@@ -18,13 +18,19 @@
 
 // Copied from SCOPED_NAMED_EVENT but modified
 // to accommodate event names containing ::
+#define METASOUND_TRACE_CPUPROFILER_EVENT_SCOPE_COND(Name, Cond)\
+	FScopedNamedEventConditionalStatic PREPROCESSOR_JOIN(MetaSound_NamedEvent_,__LINE__)(FColor::Green, NAMED_EVENT_STR(#Name), GCycleStatsShouldEmitNamedEvents > 0 && (Cond));\
+	TRACE_CPUPROFILER_EVENT_SCOPE_CONDITIONAL(Name, Cond);
+
 #define METASOUND_TRACE_CPUPROFILER_EVENT_SCOPE(Name)\
-	FScopedNamedEventConditionalStatic PREPROCESSOR_JOIN(MetaSound_NamedEvent_,__LINE__)(FColor::Green, NAMED_EVENT_STR(#Name), GCycleStatsShouldEmitNamedEvents > 0);\
-	TRACE_CPUPROFILER_EVENT_SCOPE(Name);
+	METASOUND_TRACE_CPUPROFILER_EVENT_SCOPE_COND(Name, true)
+
+#define METASOUND_TRACE_CPUPROFILER_EVENT_SCOPE_TEXT_COND(Name, Cond) \
+	FScopedNamedEventConditional ANONYMOUS_VARIABLE(NamedEvent_)(FColor::Green, Name, GCycleStatsShouldEmitNamedEvents > 0 && (Cond));\
+	TRACE_CPUPROFILER_EVENT_SCOPE_TEXT_CONDITIONAL(Name, Cond);
 
 #define METASOUND_TRACE_CPUPROFILER_EVENT_SCOPE_TEXT(Name) \
-	FScopedNamedEventConditional ANONYMOUS_VARIABLE(NamedEvent_)(FColor::Green, Name, GCycleStatsShouldEmitNamedEvents > 0);\
-	TRACE_CPUPROFILER_EVENT_SCOPE_TEXT(Name);
+	METASOUND_TRACE_CPUPROFILER_EVENT_SCOPE_TEXT_COND(Name, true)
 
 // Uses cached Insights SpecId to avoid string lookup
 #define METASOUND_TRACE_CPUPROFILER_EVENT_SCOPE_FAST(TraceSpecId, Name) \
@@ -34,7 +40,9 @@
 #else
 // Metasound CPU profiler trace *not* enabled
 
+#define METASOUND_TRACE_CPUPROFILER_EVENT_SCOPE_COND(Name, Cond)
 #define METASOUND_TRACE_CPUPROFILER_EVENT_SCOPE(Name)
+#define METASOUND_TRACE_CPUPROFILER_EVENT_SCOPE_TEXT_COND(Name, Cond)
 #define METASOUND_TRACE_CPUPROFILER_EVENT_SCOPE_TEXT(Name)
 #define METASOUND_TRACE_CPUPROFILER_EVENT_SCOPE_FAST(TraceSpecId, Name)
 
diff --git a/Engine/Source/Runtime/AudioMixer/Public/AudioMixerSubmix.h b/Engine/Source/Runtime/AudioMixer/Public/AudioMixerSubmix.h
--- a/Engine/Source/Runtime/AudioMixer/Public/AudioMixerSubmix.h
+++ b/Engine/Source/Runtime/AudioMixer/Public/AudioMixerSubmix.h
@@ -36,7 +36,7 @@ DECLARE_CYCLE_STAT_EXTERN(TEXT("Submix Graph Source Mixing"), STAT_AudioMixerSub
 DECLARE_CYCLE_STAT_EXTERN(TEXT("Submix Graph Effect Processing"), STAT_AudioMixerSubmixEffectProcessing, STATGROUP_AudioMixer, AUDIOMIXER_API);
 
 // The time it takes to process the submix buffer listeners. 
-DECLARE_CYCLE_STAT_EXTERN(TEXT("Submix Buffer Listeners"), STAT_AudioMixerSubmixBufferListeners, STATGROUP_AudioMixer, AUDIOMIXER_API);
+DECLARE_CYCLE_STAT_WITH_FLAGS_EXTERN(TEXT("Submix Buffer Listeners"), STAT_AudioMixerSubmixBufferListeners, STATGROUP_AudioMixer, EStatFlags::Verbose, AUDIOMIXER_API);
 
 // The time it takes to process the submix soundfield child submixes. 
 DECLARE_CYCLE_STAT_EXTERN(TEXT("Submix Soundfield Children"), STAT_AudioMixerSubmixSoundfieldChildren, STATGROUP_AudioMixer, AUDIOMIXER_API);
