Add final variable support to Builder API, more MetaSound controller deprecation related therein

#rnx

diff --git a/Engine/Plugins/Runtime/Metasound/Source/MetasoundEditor/Private/MetasoundEditor.cpp b/Engine/Plugins/Runtime/Metasound/Source/MetasoundEditor/Private/MetasoundEditor.cpp
--- a/Engine/Plugins/Runtime/Metasound/Source/MetasoundEditor/Private/MetasoundEditor.cpp
+++ b/Engine/Plugins/Runtime/Metasound/Source/MetasoundEditor/Private/MetasoundEditor.cpp
@@ -344,20 +344,25 @@ namespace Metasound
 					return FReply::Unhandled();
 				}
 
+				if (!Editor.IsValid())
+				{
+					return FReply::Unhandled();
+				}
+
 				UMetasoundEditorGraph* MetasoundGraph = GraphMember->GetOwningGraph();
 				check(MetasoundGraph);
-				UObject& ParentMetasound = MetasoundGraph->GetMetasoundChecked();
+				UObject& ParentMetaSound = MetasoundGraph->GetMetasoundChecked();
+				FMetaSoundFrontendDocumentBuilder& DocBuilder = FDocumentBuilderRegistry::GetChecked().FindOrBeginBuilding(&ParentMetaSound);
 
 				if (UMetasoundEditorGraphInput* Input = Cast<UMetasoundEditorGraphInput>(GraphMember.Get()))
 				{
 					const FScopedTransaction Transaction(LOCTEXT("DropAddNewInputNode", "Drop New MetaSound Input Node"));
-					ParentMetasound.Modify();
+					ParentMetaSound.Modify();
 					MetasoundGraph->Modify();
 					Input->Modify();
 
-					FMetaSoundFrontendDocumentBuilder& Builder = FDocumentBuilderRegistry::GetChecked().FindOrBeginBuilding(&ParentMetasound);
-					const FMetasoundFrontendNode* TemplateNode = FInputNodeTemplate::CreateNode(Builder, Input->GetMemberName());
-					if (UMetasoundEditorGraphNode* NewGraphNode = FGraphBuilder::AddInputNode(ParentMetasound, TemplateNode->GetID()))
+					const FMetasoundFrontendNode* TemplateNode = FInputNodeTemplate::CreateNode(DocBuilder, Input->GetMemberName());
+					if (UMetasoundEditorGraphNode* NewGraphNode = FGraphBuilder::AddInputNode(ParentMetaSound, TemplateNode->GetID()))
 					{
 						NewGraphNode->Modify();
 						NewGraphNode->UpdateFrontendNodeLocation(InGraphPosition);
@@ -365,12 +370,8 @@ namespace Metasound
 
 						TryConnectToHoveredPin(NewGraphNode);
 
-						FGraphBuilder::RegisterGraphWithFrontend(ParentMetasound);
-						TSharedPtr<FEditor> MetasoundEditor = FGraphBuilder::GetEditorForGraph(*MetasoundGraph);
-						if (MetasoundEditor.IsValid())
-						{
-							MetasoundEditor->ClearSelectionAndSelectNode(NewGraphNode);
-						}
+						FGraphBuilder::RegisterGraphWithFrontend(ParentMetaSound);
+						Editor->ClearSelectionAndSelectNode(NewGraphNode);
 						return FReply::Handled();
 					}
 				}
@@ -381,11 +382,11 @@ namespace Metasound
 					if (Nodes.IsEmpty())
 					{
 						const FScopedTransaction Transaction(LOCTEXT("DropAddNewOutputNode", "Drop New MetaSound Output Node"));
-						ParentMetasound.Modify();
+						ParentMetaSound.Modify();
 						MetasoundGraph->Modify();
 						Output->Modify();
 
-						if (UMetasoundEditorGraphOutputNode* NewGraphNode = FGraphBuilder::AddOutputNode(ParentMetasound, Output->NodeID))
+						if (UMetasoundEditorGraphOutputNode* NewGraphNode = FGraphBuilder::AddOutputNode(ParentMetaSound, Output->NodeID))
 						{
 							NewGraphNode->Modify();
 							NewGraphNode->UpdateFrontendNodeLocation(InGraphPosition);
@@ -393,7 +394,7 @@ namespace Metasound
 
 							TryConnectToHoveredPin(NewGraphNode);
 
-							FGraphBuilder::RegisterGraphWithFrontend(ParentMetasound);
+							FGraphBuilder::RegisterGraphWithFrontend(ParentMetaSound);
 							TSharedPtr<FEditor> MetasoundEditor = FGraphBuilder::GetEditorForGraph(*MetasoundGraph);
 							if (MetasoundEditor.IsValid())
 							{
@@ -415,32 +416,26 @@ namespace Metasound
 				if (UMetasoundEditorGraphVariable* Variable = Cast<UMetasoundEditorGraphVariable>(GraphMember.Get()))
 				{
 					const FScopedTransaction Transaction(LOCTEXT("DropAddNewVariableNode", "Drop New MetaSound Variable Node"));
-					ParentMetasound.Modify();
+					ParentMetaSound.Modify();
 					MetasoundGraph->Modify();
 					Variable->Modify();
 
-					FVariableHandle VariableHandle = Variable->GetVariableHandle();
 					FMetasoundFrontendClass VariableClass;
-
 					const bool bMakeOrJumpToMutator = FSlateApplication::Get().GetModifierKeys().AreModifersDown(EModifierKey::Shift);
 					if (bMakeOrJumpToMutator)
 					{
-						FConstNodeHandle MutatorNodeHandle = Variable->GetConstVariableHandle()->FindMutatorNode();
-						if (MutatorNodeHandle->IsValid())
+						if (const FMetasoundFrontendVariable* FrontendVariable = DocBuilder.FindGraphVariable(Variable->GetMemberName()); FrontendVariable && FrontendVariable->MutatorNodeID.IsValid())
 						{
-							if (Editor.IsValid())
+							auto IsMutatorNode = [&FrontendVariable](const UMetasoundEditorGraphMemberNode* Node)
 							{
-								auto IsMutatorNode = [&MutatorNodeHandle](const UMetasoundEditorGraphMemberNode* Node)
-								{
-									return Node->GetNodeID() == MutatorNodeHandle->GetID();
-								};
-								TArray<UMetasoundEditorGraphMemberNode*> Nodes = Variable->GetNodes();
-								if (UMetasoundEditorGraphMemberNode** MutatorNode = Nodes.FindByPredicate(IsMutatorNode))
-								{
-									check(*MutatorNode);
-									Editor->JumpToNodes<UMetasoundEditorGraphMemberNode>({ *MutatorNode });
-									return FReply::Handled();
-								}
+								return Node->GetNodeID() == FrontendVariable->MutatorNodeID;
+							};
+							TArray<UMetasoundEditorGraphMemberNode*> Nodes = Variable->GetNodes();
+							if (UMetasoundEditorGraphMemberNode** MutatorNode = Nodes.FindByPredicate(IsMutatorNode))
+							{
+								check(*MutatorNode);
+								Editor->JumpToNodes<UMetasoundEditorGraphMemberNode>({ *MutatorNode });
+								return FReply::Handled();
 							}
 						}
 						else
@@ -481,23 +476,20 @@ namespace Metasound
 						}
 					}
 
-					const FNodeClassName ClassName = VariableClass.Metadata.GetClassName().ToNodeClassName();
-					FConstNodeHandle NodeHandle = FGraphBuilder::AddVariableNodeHandle(ParentMetasound, Variable->GetVariableID(), ClassName);
-					if (UMetasoundEditorGraphVariableNode* NewGraphNode = FGraphBuilder::AddVariableNode(ParentMetasound, NodeHandle))
+					const FMetasoundFrontendNode* NewVariableNode = DocBuilder.AddGraphVariableNode(Variable->GetMemberName(), VariableClass.Metadata.GetType());
+					if (ensure(NewVariableNode))
 					{
-						NewGraphNode->Modify();
-						NewGraphNode->UpdateFrontendNodeLocation(InGraphPosition);
-						NewGraphNode->SyncLocationFromFrontendNode();
+						if (UMetasoundEditorGraphVariableNode* NewGraphNode = FGraphBuilder::AddVariableNode(ParentMetaSound, NewVariableNode->GetID()))
+						{
+							NewGraphNode->UpdateFrontendNodeLocation(InGraphPosition);
+							NewGraphNode->SyncLocationFromFrontendNode();
 
-						TryConnectToHoveredPin(NewGraphNode);
+							TryConnectToHoveredPin(NewGraphNode);
 
-						FGraphBuilder::RegisterGraphWithFrontend(ParentMetasound);
-						TSharedPtr<FEditor> MetasoundEditor = FGraphBuilder::GetEditorForGraph(*MetasoundGraph);
-						if (MetasoundEditor.IsValid())
-						{
-							MetasoundEditor->ClearSelectionAndSelectNode(NewGraphNode);
+							FGraphBuilder::RegisterGraphWithFrontend(ParentMetaSound);
+							Editor->ClearSelectionAndSelectNode(NewGraphNode);
+							return FReply::Handled();
 						}
-						return FReply::Handled();
 					}
 				}
 
@@ -579,7 +571,7 @@ namespace Metasound
 				FText Message;
 				if (GraphMember.IsValid())
 				{
-					UMetasoundEditorGraph* OwningGraph = GraphMember->GetOwningGraph();
+					const UMetasoundEditorGraph* OwningGraph = GraphMember->GetOwningGraph();
 					Message = GraphMember->GetDisplayName();
 					if (GetHoveredGraph() && OwningGraph)
 					{
@@ -3076,12 +3068,10 @@ namespace Metasound
 						if (GraphMember->IsA<UMetasoundEditorGraphInput>())
 						{
 							Builder->RemoveGraphInput(MemberName, Result);
-							ensure(Result == EMetaSoundBuilderResult::Succeeded);
 						}
 						else if (GraphMember->IsA<UMetasoundEditorGraphOutput>())
 						{
 							Builder->RemoveGraphOutput(MemberName, Result);
-							ensure(Result == EMetaSoundBuilderResult::Succeeded);
 						}
 						// TODO: Move to builder API
 						else if (UMetasoundEditorGraphVariable* Variable = Cast<UMetasoundEditorGraphVariable>(GraphMember))
@@ -3089,9 +3079,10 @@ namespace Metasound
 							const FGuid VariableID = Variable->GetVariableID();
 							if (VariableID.IsValid())
 							{
-								ensure(Graph.GetGraphHandle()->RemoveVariable(VariableID));
+								Builder->RemoveGraphVariable(MemberName, Result);
 							}
 						}
+						ensure(Result == EMetaSoundBuilderResult::Succeeded);
 					}
 				}
 
@@ -3678,14 +3669,16 @@ namespace Metasound
 					MetaSound->Modify();
 
 					UMetasoundEditorGraphMember* NewGraphMember = nullptr;
+					FMetaSoundFrontendDocumentBuilder& DocBuilder = Graph.GetBuilderChecked().GetBuilder();
 
-					//Duplicate the Sources NodeHandle and add a new member from it
+					// Duplicate the Sources Frontend Node and add a new member from it
 					if (const UMetasoundEditorGraphVariable* SourceGraphVariable = Cast<UMetasoundEditorGraphVariable>(SourceGraphMember))
 					{
-						FConstVariableHandle VariableHandle = FGraphBuilder::DuplicateVariableHandle(Graph.GetMetasoundChecked(), SourceGraphVariable->GetConstVariableHandle());
-						if (ensure(VariableHandle->IsValid()))
+						const FName NewName = FGraphBuilder::GenerateUniqueVariableName(DocBuilder, SourceGraphVariable->GetMemberName().ToString());
+						const FMetasoundFrontendVariable* NewFrontendVariable = DocBuilder.DuplicateGraphVariable(SourceGraphVariable->GetMemberName(), NewName);
+						if (ensure(NewFrontendVariable))
 						{
-							NewGraphMember = Graph.FindOrAddVariable(VariableHandle);
+							NewGraphMember = Graph.FindOrAddVariable(NewFrontendVariable->Name);
 						}
 					}
 					else if (const UMetasoundEditorGraphVertex* SourceGraphVertex = Cast<UMetasoundEditorGraphVertex>(SourceGraphMember))
@@ -3693,34 +3686,27 @@ namespace Metasound
 						const FName SourceMemberName = SourceGraphVertex->GetMemberName();
 						const EMetasoundFrontendClassType ClassType = SourceGraphVertex->GetClassType();
 						
-						FMetaSoundFrontendDocumentBuilder& DocumentBuilder = IDocumentBuilderRegistry::GetChecked().FindOrBeginBuilding(&Graph.GetMetasoundChecked());
 						const FName Name = FGraphBuilder::GenerateUniqueNameByClassType(Graph.GetMetasoundChecked(), ClassType, SourceMemberName.ToString());
 
 						if (ClassType == EMetasoundFrontendClassType::Input)
 						{
-							if (const FMetasoundFrontendClassInput* SourceInput = DocumentBuilder.FindGraphInput(SourceMemberName))
+							if (const FMetasoundFrontendClassInput* NewInput = DocBuilder.DuplicateGraphInput(SourceMemberName, Name))
 							{
-								if (const FMetasoundFrontendNode* FrontendNode = DocumentBuilder.DuplicateGraphInput(*SourceInput, Name))
-								{
-									FGraphBuilder::SynchronizeGraphMembers(DocumentBuilder, Graph);
-									NewGraphMember = Graph.FindInput(FrontendNode->Name);
-								}
+								FGraphBuilder::SynchronizeGraphMembers(DocBuilder, Graph);
+								NewGraphMember = Graph.FindInput(NewInput->Name);
 							}
 						}
 						else if (ClassType == EMetasoundFrontendClassType::Output)
 						{
-							if (const FMetasoundFrontendClassOutput* SourceOutput = DocumentBuilder.FindGraphOutput(SourceMemberName))
+							if (const FMetasoundFrontendClassOutput* NewOutput = DocBuilder.DuplicateGraphOutput(SourceMemberName, Name))
 							{
-								if (const FMetasoundFrontendNode* FrontendNode = DocumentBuilder.DuplicateGraphOutput(*SourceOutput, Name))
-								{
-									FGraphBuilder::SynchronizeGraphMembers(DocumentBuilder, Graph);
-									NewGraphMember = Graph.FindOutput(FrontendNode->Name);
-								}
+								FGraphBuilder::SynchronizeGraphMembers(DocBuilder, Graph);
+								NewGraphMember = Graph.FindOutput(NewOutput->Name);
 							}
 						}
 					}
 
-					//Duplicate the literal from the SourceGraphMember to the NewGraphMember added
+					// Duplicate the literal from the SourceGraphMember to the NewGraphMember added
 					if (NewGraphMember)
 					{
 						if (UMetaSoundEditorSubsystem* EditorSubsystem = GEditor->GetEditorSubsystem<UMetaSoundEditorSubsystem>())
@@ -3849,7 +3835,7 @@ namespace Metasound
 
 		void FEditor::PasteSelectedMemberItems()
 		{
-			const UMetasoundEditorGraphMember* Member = FDocumentClipboardUtils::GetMemberFromClipboard();		
+			const UMetasoundEditorGraphMember* Member = FDocumentClipboardUtils::GetMemberFromClipboard();
 			if (Member)
 			{
 				UObject* MetaSound = GetMetasoundObject();
@@ -3858,14 +3844,13 @@ namespace Metasound
 					return;
 				}
 
+				UMetasoundEditorGraph& Graph = GetMetaSoundGraphChecked();
+				FMetaSoundFrontendDocumentBuilder& DocumentBuilder = Frontend::IDocumentBuilderRegistry::GetChecked().FindOrBeginBuilding(&Graph.GetMetasoundChecked());
+				MetaSound->Modify();
+				Graph.Modify();
+
 				if (const UMetasoundEditorGraphVertex* SourceGraphVertex = Cast<UMetasoundEditorGraphVertex>(Member))
 				{
-
-					UMetasoundEditorGraph& Graph = GetMetaSoundGraphChecked();
-					MetaSound->Modify();
-					Graph.Modify();
-
-					FMetaSoundFrontendDocumentBuilder& DocumentBuilder = Frontend::IDocumentBuilderRegistry::GetChecked().FindOrBeginBuilding(&Graph.GetMetasoundChecked());				
 					const FMetasoundEditorGraphVertexBreadcrumb Breadcrumb = SourceGraphVertex->GetBreadcrumb();
 
 					FCreateNodeVertexParams VertexParams;
@@ -3875,43 +3860,45 @@ namespace Metasound
 					const EMetasoundFrontendClassType ClassType = SourceGraphVertex->GetClassType();
 
 					if (ClassType == EMetasoundFrontendClassType::Input)
-					{				
+					{
 						TArray<FMetasoundFrontendClassInputDefault> InputDefaults;
 						Algo::Transform(Breadcrumb.DefaultLiterals, InputDefaults, [](const TPair<FGuid, FMetasoundFrontendLiteral>& Pair)
-							{
-								return FMetasoundFrontendClassInputDefault(Pair.Key, Pair.Value);
-							});
+						{
+							return FMetasoundFrontendClassInputDefault(Pair.Key, Pair.Value);
+						});
 
 						FMetasoundFrontendClassInput ClassInput = FGraphBuilder::CreateUniqueClassInput(*MetaSound, VertexParams, InputDefaults, &Breadcrumb.MemberName);
 						ClassInput.Metadata.SetDescription(Breadcrumb.Description);
 						ClassInput.Metadata.SetIsAdvancedDisplay(Breadcrumb.bIsAdvancedDisplay);
 						ClassInput.Metadata.SortOrderIndex = Breadcrumb.SortOrderIndex;
 						
-						DocumentBuilder.AddGraphInput(ClassInput);	
+						DocumentBuilder.AddGraphInput(MoveTemp(ClassInput));
 					}
 					else if (ClassType == EMetasoundFrontendClassType::Output)
 					{
 						FMetasoundFrontendClassOutput ClassOutput = FGraphBuilder::CreateUniqueClassOutput(*MetaSound, VertexParams, &Breadcrumb.MemberName);
 						ClassOutput.Metadata.SetDescription(Breadcrumb.Description);
 						ClassOutput.Metadata.SetIsAdvancedDisplay(Breadcrumb.bIsAdvancedDisplay);
-						ClassOutput.Metadata.SortOrderIndex = Breadcrumb.SortOrderIndex;					
+						ClassOutput.Metadata.SortOrderIndex = Breadcrumb.SortOrderIndex;
 
-						DocumentBuilder.AddGraphOutput(ClassOutput);
+						DocumentBuilder.AddGraphOutput(MoveTemp(ClassOutput));
 					}
 				}
 				else if (const UMetasoundEditorGraphVariable* SourceGraphVariable = Cast<UMetasoundEditorGraphVariable>(Member))
 				{
-
 					const FMetasoundEditorGraphVariableBreadcrumb Breadcrumb = SourceGraphVariable->GetBreadcrumb();
-
-					Frontend::FVariableHandle FrontendVariable = FGraphBuilder::AddVariableHandle(*MetaSound, SourceGraphVariable->GetDataType(), &Breadcrumb.MemberName);
-					FrontendVariable->SetDescription(Breadcrumb.Description);
-					FrontendVariable->SetLiteral(Breadcrumb.DefaultLiteral);
-
-					if (ensure(FrontendVariable->IsValid()))
+					const FName VariableName = FGraphBuilder::GenerateUniqueVariableName(DocumentBuilder, Breadcrumb.MemberName.ToString());
+					const FMetasoundFrontendVariable* FrontendVariable = DocumentBuilder.AddGraphVariable(
+						VariableName,
+						SourceGraphVariable->GetDataType(),
+						&Breadcrumb.DefaultLiteral,
+						nullptr, // TODO: Add DisplayName to breadcrumb, and only paste if it doesn't conflict with existing variable
+						&Breadcrumb.Description
+					);
+					
+					if (ensure(FrontendVariable))
 					{
-						UMetasoundEditorGraph& Graph = GetMetaSoundGraphChecked();
-						Graph.FindOrAddVariable(FrontendVariable);
+						Graph.FindOrAddVariable(FrontendVariable->Name);
 					}
 				}
 
@@ -3997,7 +3984,9 @@ namespace Metasound
 			{
 				if (UMetasoundEditorGraphExternalNode* ExternalNode = Cast<UMetasoundEditorGraphExternalNode>(Object))
 				{
+PRAGMA_DISABLE_DEPRECATION_WARNINGS
 					Metasound::Frontend::FNodeHandle NodeHandle = ExternalNode->GetNodeHandle();
+PRAGMA_ENABLE_DEPRECATION_WARNINGS
 					const FMetasoundFrontendClassMetadata& Metadata = NodeHandle->GetClassMetadata();
 
 					// Check for new version
@@ -4023,7 +4012,9 @@ namespace Metasound
 
 			if (bReplacedNodes)
 			{
+PRAGMA_DISABLE_DEPRECATION_WARNINGS
 				FDocumentHandle DocumentHandle = Graph.GetDocumentHandle();
+PRAGMA_ENABLE_DEPRECATION_WARNINGS
 				DocumentHandle->RemoveUnreferencedDependencies();
 				DocumentHandle->SynchronizeDependencyMetadata();
 				FMetasoundFrontendDocumentModifyContext& ModifyContext = FGraphBuilder::GetOutermostMetaSoundChecked(Graph).GetModifyContext();
@@ -4840,18 +4831,20 @@ namespace Metasound
 			TArray<TObjectPtr<UObject>> SelectedObjects;
 
 			FName NameToSelect;
+			FMetaSoundFrontendDocumentBuilder& DocBuilder = Builder->GetBuilder();
 			switch (static_cast<ENodeSection>(InSectionID))
 			{
 				case ENodeSection::Inputs:
 				{
 					const FScopedTransaction Transaction(LOCTEXT("AddInputNode", "Add MetaSound Input"));
 					MetaSound->Modify();
+					Graph.Modify();
 
 					FCreateNodeVertexParams VertexParams;
 					VertexParams.DataType = DataTypeName;
 
 					FMetasoundFrontendClassInput ClassInput = FGraphBuilder::CreateUniqueClassInput(*MetaSound, VertexParams);
-					if (const FMetasoundFrontendNode* NewNode = Builder->GetBuilder().AddGraphInput(ClassInput))
+					if (const FMetasoundFrontendNode* NewNode = DocBuilder.AddGraphInput(MoveTemp(ClassInput)))
 					{
 						NameToSelect = NewNode->Name;
 
@@ -4868,12 +4861,13 @@ namespace Metasound
 				{
 					const FScopedTransaction Transaction(LOCTEXT("AddOutputNode", "Add MetaSound Output"));
 					MetaSound->Modify();
+					Graph.Modify();
 
 					FCreateNodeVertexParams VertexParams;
 					VertexParams.DataType = DataTypeName;
 
 					FMetasoundFrontendClassOutput ClassOutput = FGraphBuilder::CreateUniqueClassOutput(*MetaSound, VertexParams);
-					if (const FMetasoundFrontendNode* NewNode = Builder->GetBuilder().AddGraphOutput(ClassOutput))
+					if (const FMetasoundFrontendNode* NewNode = DocBuilder.AddGraphOutput(MoveTemp(ClassOutput)))
 					{
 						NameToSelect = NewNode->Name;
 
@@ -4890,11 +4884,13 @@ namespace Metasound
 				{
 					const FScopedTransaction Transaction(TEXT(""), LOCTEXT("AddVariableNode", "Add MetaSound Variable"), MetaSound);
 					MetaSound->Modify();
+					Graph.Modify();
 
-					Frontend::FVariableHandle FrontendVariable = FGraphBuilder::AddVariableHandle(*MetaSound, DataTypeName);
-					if (ensure(FrontendVariable->IsValid()))
+					const FName VariableName = FGraphBuilder::GenerateUniqueVariableName(DocBuilder, TEXT("Variable"));
+					const FMetasoundFrontendVariable* FrontendVariable = DocBuilder.AddGraphVariable(VariableName, DataTypeName);
+					if (ensure(FrontendVariable))
 					{
-						TObjectPtr<UMetasoundEditorGraphVariable> EditorVariable = Graph.FindOrAddVariable(FrontendVariable);
+						TObjectPtr<UMetasoundEditorGraphVariable> EditorVariable = Graph.FindOrAddVariable(FrontendVariable->Name);
 						if (ensure(EditorVariable))
 						{
 							SelectedObjects.Add(EditorVariable);
diff --git a/Engine/Plugins/Runtime/Metasound/Source/MetasoundEditor/Private/MetasoundEditorDocumentClipboardUtils.cpp b/Engine/Plugins/Runtime/Metasound/Source/MetasoundEditor/Private/MetasoundEditorDocumentClipboardUtils.cpp
--- a/Engine/Plugins/Runtime/Metasound/Source/MetasoundEditor/Private/MetasoundEditorDocumentClipboardUtils.cpp
+++ b/Engine/Plugins/Runtime/Metasound/Source/MetasoundEditor/Private/MetasoundEditorDocumentClipboardUtils.cpp
@@ -78,10 +78,14 @@ namespace Metasound::Editor
 								return FMetasoundFrontendClassInputDefault(Pair.Key, Pair.Value);
 							});
 
-							FMetasoundFrontendClassInput ClassInput = FGraphBuilder::CreateUniqueClassInput(*OutAsset.GetOwningAsset(), VertexParams, InputDefaults, &Breadcrumb.MemberName);
-							ClassInput.Metadata = Breadcrumb.VertexMetadata;
+							const FMetasoundFrontendNode* NewNode = nullptr;
+							{
+								FMetasoundFrontendClassInput ClassInput = FGraphBuilder::CreateUniqueClassInput(*OutAsset.GetOwningAsset(), VertexParams, InputDefaults, &Breadcrumb.MemberName);
+								ClassInput.Metadata = Breadcrumb.VertexMetadata;
+								NewNode = Builder.AddGraphInput(MoveTemp(ClassInput));
+							}
 
-							if (const FMetasoundFrontendNode* NewNode = Builder.AddGraphInput(ClassInput))
+							if (NewNode)
 							{
 								Input = Graph.FindOrAddInput(NewNode->GetID());
 								if (Breadcrumb.MemberMetadataPath.IsSet())
@@ -89,7 +93,7 @@ namespace Metasound::Editor
 									UObject* MemberMetadata = Breadcrumb.MemberMetadataPath->TryLoad();
 									if (UMetasoundEditorGraphMemberDefaultLiteral* DefaultLiteral = Cast<UMetasoundEditorGraphMemberDefaultLiteral>(MemberMetadata))
 									{
-										Builder.ClearMemberMetadata(ClassInput.NodeID);
+										Builder.ClearMemberMetadata(NewNode->GetID());
 										UMetaSoundEditorSubsystem& MetaSoundEditorSubsystem = UMetaSoundEditorSubsystem::GetChecked();
 										TSubclassOf<UMetasoundEditorGraphMemberDefaultLiteral> LiteralClass = MetaSoundEditorSubsystem.GetLiteralClassForType(Breadcrumb.DataType);
 										MetaSoundEditorSubsystem.BindMemberMetadata(Builder, *Input, LiteralClass, DefaultLiteral);
@@ -225,6 +229,7 @@ namespace Metasound::Editor
 		OutNotifications.bPastedNodesAddMultipleVariableSetters = false;
 
 		UMetasoundEditorGraph& Graph = *CastChecked<UMetasoundEditorGraph>(&OutAsset.GetGraphChecked());
+		FMetaSoundFrontendDocumentBuilder& DocBuilder = Graph.GetBuilderChecked().GetBuilder();
 		TMap<FName, TObjectPtr<UMetasoundEditorGraphVariable>> MappedGeneratedVariableNames;
 		Graph.IterateVariables([&MappedGeneratedVariableNames](UMetasoundEditorGraphVariable& Variable)
 		{
@@ -242,50 +247,55 @@ namespace Metasound::Editor
 			VariableNode->CreateNewGuid();
 
 			TObjectPtr<UMetasoundEditorGraphVariable>& Variable = VariableNode->Variable;
-			if (!Variable || !Graph.ContainsVariable(*Variable))
+			if (!Variable || !Graph.ContainsVariable(*Variable) || Variable->GetFrontendVariable() == nullptr)
 			{
 				const FMetasoundEditorGraphMemberNodeBreadcrumb& Breadcrumb = VariableNode->Breadcrumb;
 				const FName BaseName = Breadcrumb.MemberName;
 				TObjectPtr<UMetasoundEditorGraphVariable> CachedVariable = MappedGeneratedVariableNames.FindRef(BaseName);
-				Frontend::FVariableHandle FrontendVariable = IVariableController::GetInvalidHandle();
+				const FMetasoundFrontendVariable* FrontendVariable = nullptr;
 				if (CachedVariable)
 				{
-					FrontendVariable = CachedVariable->GetVariableHandle();
+					FrontendVariable = CachedVariable->GetFrontendVariable();
 				}
 
-				if (FrontendVariable->IsValid() && FrontendVariable->GetDataType() == Breadcrumb.DataType)
+				if (FrontendVariable && FrontendVariable->TypeName == Breadcrumb.DataType)
 				{
 					Variable = CachedVariable;
 				}
 				else
 				{
-					FrontendVariable = FGraphBuilder::AddVariableHandle(*OutAsset.GetOwningAsset(), Breadcrumb.DataType);
-					FrontendVariable->SetDescription(Breadcrumb.VertexMetadata.GetDescription());
-					FrontendVariable->SetDisplayName(Breadcrumb.VertexMetadata.GetDisplayName());
-					// Hack to reuse the default literals breadcrumb property for variables, which only have a single (rather than paged) literals
-					const FMetasoundFrontendLiteral* Literal = Breadcrumb.DefaultLiterals.Find(DefaultPageID);
-					if (Literal)
-					{
-						FrontendVariable->SetLiteral(*Literal);
-					}
-					// TODO: copy member metadata when variables have been moved to use the builder API
-					Variable = Graph.FindOrAddVariable(FrontendVariable);
+					const FMetasoundFrontendLiteral* Literal = Breadcrumb.DefaultLiterals.Find(Frontend::DefaultPageID);
+					const FName VariableName = FGraphBuilder::GenerateUniqueVariableName(DocBuilder, BaseName.ToString());
+					FrontendVariable = DocBuilder.AddGraphVariable(
+						VariableName,
+						Breadcrumb.DataType,
+						Literal,
+						&Breadcrumb.VertexMetadata.GetDisplayName(),
+						&Breadcrumb.VertexMetadata.GetDescription()
+					);
+
+					Variable = Graph.FindOrAddVariable(FrontendVariable->Name);
 					check(Variable);
 
-					constexpr bool bPostSubTransaction = false;
-					const FName VariableName = FGraphBuilder::GenerateUniqueNameByClassType(*OutAsset.GetOwningAsset(), EMetasoundFrontendClassType::Variable, BaseName.ToString());
-					Variable->SetMemberName(VariableName, bPostSubTransaction);
 					VariableNode->CacheBreadcrumb(); // Name of referenced variable/node state has changed so make sure up-to-date in case breadcrumb is used later
 					MappedGeneratedVariableNames.Add(BaseName, Variable);
 				}
 			}
 
-			FConstVariableHandle VariableHandle = Variable->GetConstVariableHandle();
-			if (VariableHandle->IsValid())
+			const FMetasoundFrontendVariable* FrontendVariable = Variable->GetFrontendVariable();
+			if (ensure(FrontendVariable))
 			{
-				// Can only have one setter node
-				FConstNodeHandle VariableMutatorNodeHandle = VariableHandle->FindMutatorNode();
-				if (VariableNode->GetNodeID() == VariableMutatorNodeHandle->GetID())
+				// Can only have one mutator/setter node
+				const FNodeClassName NodeClassName = VariableNode->GetClassName().ToNodeClassName();
+				const bool bMatchesMutatorNodeID = VariableNode->GetNodeID() == FrontendVariable->MutatorNodeID;
+				bool bIsDuplicateMutatorNode = false;
+				if (const FMetasoundFrontendNode* MutatorNode = DocBuilder.FindNode(FrontendVariable->MutatorNodeID))
+				{
+					const FMetasoundFrontendClass* MutatorClass = DocBuilder.FindDependency(MutatorNode->ClassID);
+					check(MutatorClass);
+					bIsDuplicateMutatorNode = MutatorClass->Metadata.GetClassName() == NodeClassName;
+				}
+				if (bMatchesMutatorNodeID || bIsDuplicateMutatorNode)
 				{
 					OutNotifications.bPastedNodesAddMultipleVariableSetters = true;
 					OutPastedNodes.RemoveAtSwap(Index, 1, EAllowShrinking::No);
@@ -294,15 +304,27 @@ namespace Metasound::Editor
 				else
 				{
 					// Add new variable node
-					const FNodeClassName NodeClassName = VariableNode->GetClassName().ToNodeClassName();
-					FNodeHandle NodeHandle = FGraphBuilder::AddVariableNodeHandle(*OutAsset.GetOwningAsset(), Variable->GetVariableID(), NodeClassName, VariableNode);
-					if (!NodeHandle->IsValid())
+					FMetasoundFrontendClass FrontendClass;
+					bool bDidFindClassWithName = ISearchEngine::Get().FindClassWithHighestVersion(NodeClassName, FrontendClass);
+					if (ensure(bDidFindClassWithName))
 					{
-						OutPastedNodes.RemoveAtSwap(Index, 1, EAllowShrinking::No);
-						Graph.RemoveNode(VariableNode);
+						if (const FMetasoundFrontendNode* NewNode = DocBuilder.AddGraphVariableNode(FrontendVariable->Name, FrontendClass.Metadata.GetType()))
+						{
+							VariableNode->SetNodeID(NewNode->GetID());
+						}
+						else
+						{
+							OutPastedNodes.RemoveAtSwap(Index, 1, EAllowShrinking::No);
+							Graph.RemoveNode(VariableNode);
+						}
 					}
 				}
 			}
+			else
+			{
+				OutPastedNodes.RemoveAtSwap(Index, 1, EAllowShrinking::No);
+				Graph.RemoveNode(VariableNode);
+			}
 		}
 	}
 
@@ -451,7 +473,9 @@ namespace Metasound::Editor
 			GraphNode->SnapToGrid(SNodePanel::GetSnapGridSize());
 			if (UMetasoundEditorGraphNode* MetasoundGraphNode = Cast<UMetasoundEditorGraphNode>(GraphNode))
 			{
+PRAGMA_DISABLE_DEPRECATION_WARNINGS
 				FNodeHandle NodeHandle = MetasoundGraphNode->GetNodeHandle();
+PRAGMA_ENABLE_DEPRECATION_WARNINGS
 				if (ensure(NodeHandle->IsValid()))
 				{
 					const FVector2D NewNodeLocation = FVector2D(GraphNode->NodePosX, GraphNode->NodePosY);
diff --git a/Engine/Plugins/Runtime/Metasound/Source/MetasoundEditor/Private/MetasoundEditorGraph.cpp b/Engine/Plugins/Runtime/Metasound/Source/MetasoundEditor/Private/MetasoundEditorGraph.cpp
--- a/Engine/Plugins/Runtime/Metasound/Source/MetasoundEditor/Private/MetasoundEditorGraph.cpp
+++ b/Engine/Plugins/Runtime/Metasound/Source/MetasoundEditor/Private/MetasoundEditorGraph.cpp
@@ -279,7 +279,9 @@ void UMetasoundEditorGraphVertex::SetDescription(const FText& InDescription, boo
 			UObject& MetaSound = Graph->GetMetasoundChecked();
 			MetaSound.Modify();
 
+PRAGMA_DISABLE_DEPRECATION_WARNINGS
 			FNodeHandle NodeHandle = GetNodeHandle();
+PRAGMA_ENABLE_DEPRECATION_WARNINGS
 			NodeHandle->SetDescription(InDescription);
 
 			Graph->RegisterGraphWithFrontend();
@@ -297,6 +299,11 @@ FName UMetasoundEditorGraphVertex::GetMemberName() const
 {
 	using namespace Metasound::Frontend;
 
+	if (!Breadcrumb.MemberName.IsNone())
+	{
+		return Breadcrumb.MemberName;
+	}
+
 	if (const FMetasoundFrontendNode* FrontendNode = GetFrontendNode())
 	{
 		return FrontendNode->Name;
@@ -369,8 +376,10 @@ void UMetasoundEditorGraphVertex::SetDisplayName(const FText& InNewName, bool bP
 		return;
 	}
 
-
+PRAGMA_DISABLE_DEPRECATION_WARNINGS
 	FNodeHandle NodeHandle = GetNodeHandle();
+PRAGMA_ENABLE_DEPRECATION_WARNINGS
+
 	if (NodeHandle->GetDisplayName().EqualTo(InNewName))
 	{
 		return;
@@ -445,11 +454,19 @@ void UMetasoundEditorGraphVertex::CacheBreadcrumb()
 {
 	Breadcrumb = { };
 
-	Breadcrumb.MemberName = GetMemberName();
+	if (const FMetasoundFrontendNode* FrontendNode = GetFrontendNode())
+	{
+		Breadcrumb.MemberName = FrontendNode->Name;
+	}
+
+	if (const FMetasoundFrontendClassVertex* FrontendVertex = GetFrontendClassVertex())
+	{
+		Breadcrumb.bIsAdvancedDisplay = FrontendVertex->GetIsAdvancedDisplay();
+	}
+
 	Breadcrumb.AccessType = GetVertexAccessType();
 	Breadcrumb.Description = GetDescription();
 	Breadcrumb.SortOrderIndex = GetSortOrderIndex();
-	Breadcrumb.bIsAdvancedDisplay = GetFrontendClassVertex()->GetIsAdvancedDisplay();
 }
 
 Metasound::Frontend::FNodeHandle UMetasoundEditorGraphVertex::GetNodeHandle()
@@ -826,6 +843,11 @@ const FMetasoundFrontendClassVertex* UMetasoundEditorGraphInput::GetFrontendClas
 
 FText UMetasoundEditorGraphInput::GetDescription() const
 {
+	if (!Breadcrumb.Description.IsEmpty())
+	{
+		return Breadcrumb.Description;
+	}
+
 	const FMetaSoundFrontendDocumentBuilder& Builder = GetFrontendBuilderChecked();
 	if (const FMetasoundFrontendClassInput* Input = Builder.FindGraphInput(GetMemberName()))
 	{
@@ -924,7 +946,9 @@ void UMetasoundEditorGraphInput::SetSortOrderIndex(int32 InSortOrderIndex)
 	UMetasoundEditorGraph* MetaSoundGraph = GetOwningGraph();
 	check(MetaSoundGraph);
 
+PRAGMA_DISABLE_DEPRECATION_WARNINGS
 	FGraphHandle GraphHandle = MetaSoundGraph->GetGraphHandle();
+PRAGMA_ENABLE_DEPRECATION_WARNINGS
 	FConstNodeHandle NodeHandle = GetConstNodeHandle();
 	const Metasound::FVertexName& NodeName = NodeHandle->GetNodeName();
 
@@ -1096,6 +1120,11 @@ void UMetasoundEditorGraphInput::UpdateFrontendDefaultLiteral(bool bPostTransact
 
 EMetasoundFrontendVertexAccessType UMetasoundEditorGraphInput::GetVertexAccessType() const
 {
+	if (Breadcrumb.AccessType != EMetasoundFrontendVertexAccessType::Unset)
+	{
+		return Breadcrumb.AccessType;
+	}
+
 	const FName MemberName = GetMemberName();
 	const FMetaSoundFrontendDocumentBuilder& Builder = GetFrontendBuilderChecked();
 	if (const FMetasoundFrontendClassInput* Input = Builder.FindGraphInput(MemberName))
@@ -1136,6 +1165,11 @@ UMetasoundEditorGraphNode* UMetasoundEditorGraphOutput::AddNode(Metasound::Front
 
 FText UMetasoundEditorGraphOutput::GetDescription() const
 {
+	if (!Breadcrumb.Description.IsEmpty())
+	{
+		return Breadcrumb.Description;
+	}
+
 	const FMetaSoundFrontendDocumentBuilder& Builder = GetFrontendBuilderChecked();
 	if (const FMetasoundFrontendClassOutput* Output = Builder.FindGraphOutput(GetMemberName()))
 	{
@@ -1210,7 +1244,10 @@ void UMetasoundEditorGraphOutput::SetSortOrderIndex(int32 InSortOrderIndex)
 	UMetasoundEditorGraph* MetaSoundGraph = GetOwningGraph();
 	check(MetaSoundGraph);
 
+PRAGMA_DISABLE_DEPRECATION_WARNINGS
 	FGraphHandle GraphHandle = MetaSoundGraph->GetGraphHandle();
+PRAGMA_ENABLE_DEPRECATION_WARNINGS
+
 	FConstNodeHandle NodeHandle = GetConstNodeHandle();
 	const Metasound::FVertexName& NodeName = NodeHandle->GetNodeName();
 
@@ -1338,6 +1375,11 @@ void UMetasoundEditorGraphOutput::UpdateFrontendDefaultLiteral(bool bPostTransac
 
 EMetasoundFrontendVertexAccessType UMetasoundEditorGraphOutput::GetVertexAccessType() const
 {
+	if (Breadcrumb.AccessType != EMetasoundFrontendVertexAccessType::Unset)
+	{
+		return Breadcrumb.AccessType;
+	}
+
 	const FName MemberName = GetMemberName();
 	const FMetaSoundFrontendDocumentBuilder& Builder = GetFrontendBuilderChecked();
 	if (const FMetasoundFrontendClassOutput* Output = Builder.FindGraphOutput(MemberName))
@@ -1429,9 +1471,12 @@ void UMetasoundEditorGraphVariable::CacheBreadcrumb()
 {
 	Breadcrumb = { };
 
-	Breadcrumb.MemberName = GetMemberName();
-	Breadcrumb.DefaultLiteral = GetVariableHandle()->GetLiteral();
-	Breadcrumb.Description = GetDescription();
+	if (const FMetasoundFrontendVariable* Variable = GetFrontendVariable())
+	{
+		Breadcrumb.MemberName = Variable->Name;
+		Breadcrumb.DefaultLiteral = Variable->Literal;
+		Breadcrumb.Description = Variable->Description;
+	}
 }
 
 Metasound::Frontend::FVariableHandle UMetasoundEditorGraphVariable::GetVariableHandle()
@@ -1455,7 +1500,7 @@ UMetasoundEditorGraphNode* UMetasoundEditorGraphVariable::AddNode(Metasound::Fro
 	using namespace Metasound::Editor;
 
 	FMetasoundAssetBase& MetaSound = FGraphBuilder::GetOutermostMetaSoundChecked(*this);
-	return FGraphBuilder::AddVariableNode(*MetaSound.GetOwningAsset(), InNodeHandle, bInSelectNewNode);
+	return FGraphBuilder::AddVariableNode(*MetaSound.GetOwningAsset(), InNodeHandle->GetID(), bInSelectNewNode);
 }
 
 void UMetasoundEditorGraphVariable::SetMemberName(const FName& InNewName, bool bPostTransaction)
@@ -1472,12 +1517,10 @@ void UMetasoundEditorGraphVariable::SetMemberName(const FName& InNewName, bool b
 	const FScopedTransaction Transaction(TransactionLabel, bPostTransaction);
 
 	Graph->Modify();
-	Graph->GetMetasoundChecked().Modify();
-
-	GetVariableHandle()->SetName(InNewName);
+	UObject& MetaSound = Graph->GetMetasoundChecked();
+	MetaSound.Modify();
 
-	Graph->RegisterGraphWithFrontend();
-	FGraphBuilder::GetOutermostMetaSoundChecked(*this).GetModifyContext().AddMemberIDsModified({ GetMemberID() });
+	Graph->GetBuilderChecked().GetBuilder().SetGraphVariableName(GetMemberName(), InNewName);
 }
 
 FGuid UMetasoundEditorGraphVariable::GetMemberID() const 
@@ -1487,7 +1530,17 @@ FGuid UMetasoundEditorGraphVariable::GetMemberID() const
 
 FName UMetasoundEditorGraphVariable::GetMemberName() const
 {
-	return GetConstVariableHandle()->GetName();
+	if (!Breadcrumb.MemberName.IsNone())
+	{
+		return Breadcrumb.MemberName;
+	}
+
+	if (const FMetasoundFrontendVariable* Variable = GetFrontendVariable())
+	{
+		return Variable->Name;
+	}
+
+	return { };
 }
 
 Metasound::Editor::ENodeSection UMetasoundEditorGraphVariable::GetSectionID() const 
@@ -1497,7 +1550,17 @@ Metasound::Editor::ENodeSection UMetasoundEditorGraphVariable::GetSectionID() co
 
 FText UMetasoundEditorGraphVariable::GetDescription() const
 {
-	return GetConstVariableHandle()->GetDescription();
+	if (!Breadcrumb.Description.IsEmpty())
+	{
+		return Breadcrumb.Description;
+	}
+
+	if (const FMetasoundFrontendVariable* Variable = GetFrontendVariable())
+	{
+		return Variable->Description;
+	}
+
+	return { };
 }
 
 void UMetasoundEditorGraphVariable::SetDescription(const FText& InDescription, bool bPostTransaction)
@@ -1511,7 +1574,7 @@ void UMetasoundEditorGraphVariable::SetDescription(const FText& InDescription, b
 		UObject& MetaSound = Graph->GetMetasoundChecked();
 		MetaSound.Modify();
 
-		GetVariableHandle()->SetDescription(InDescription);
+		Graph->GetBuilderChecked().GetBuilder().SetGraphVariableDescription(GetMemberName(), InDescription);
 	}
 }
 
@@ -1590,7 +1653,8 @@ TArray<UMetasoundEditorGraphMemberNode*> UMetasoundEditorGraphVariable::GetNodes
 
 FText UMetasoundEditorGraphVariable::GetDisplayName() const
 {
-	return Metasound::Editor::FGraphBuilder::GetDisplayName(*GetConstVariableHandle());
+	constexpr bool bIncludeNamespace = true;
+	return Metasound::Editor::FGraphBuilder::GetDisplayName(*GetConstVariableHandle(), bIncludeNamespace);
 }
 
 void UMetasoundEditorGraphVariable::SetDisplayName(const FText& InNewName, bool bPostTransaction)
@@ -1604,17 +1668,14 @@ void UMetasoundEditorGraphVariable::SetDisplayName(const FText& InNewName, bool
 		return;
 	}
 
-	const FText TransactionLabel = FText::Format(LOCTEXT("RenameGraphVariableDisplayNameFormat", "Set Metasound {0} DisplayName"), GetGraphMemberLabel());
+	const FText TransactionLabel = FText::Format(LOCTEXT("RenameGraphVariableDisplayNameFormat", "Set Metasound {0} DisplayName"), GetGraphMemberLabel(), bPostTransaction);
 	const FScopedTransaction Transaction(TransactionLabel, bPostTransaction);
 	{
 		Graph->Modify();
 		Graph->GetMetasoundChecked().Modify();
 	}
 
-	FVariableHandle VariableHandle = GetVariableHandle();
-	VariableHandle->SetDisplayName(InNewName);
-
-	FGraphBuilder::GetOutermostMetaSoundChecked(*this).GetModifyContext().AddMemberIDsModified({ GetMemberID() });
+	GetFrontendBuilderChecked().SetGraphVariableDisplayName(GetMemberName(), InNewName);
 }
 
 void UMetasoundEditorGraphVariable::SetDataType(FName InNewType, bool bPostTransaction)
@@ -1651,38 +1712,44 @@ void UMetasoundEditorGraphVariable::SetDataType(FName InNewType, bool bPostTrans
 			FVariableNodeLocations Locations;
 		} CachedData;
 
-		FConstVariableHandle OrigVariable = GetConstVariableHandle();
+		const FMetasoundFrontendVariable* OrigVariable = GetFrontendVariable();
+		check(OrigVariable);
 
 		// Cache variable metadata
-		CachedData.MemberName = OrigVariable->GetName();
-		CachedData.DisplayName = OrigVariable->GetDisplayName();
-		CachedData.Description = OrigVariable->GetDescription();
+		CachedData.MemberName = OrigVariable->Name;
+		CachedData.DisplayName = OrigVariable->DisplayName;
+		CachedData.Description = OrigVariable->Description;
 		CachedData.Locations = GetVariableNodeLocations();
 
-
 		// Remove the current variable
-		FGraphBuilder::DeleteMemberNodes(*this);
-		FGraphHandle FrontendGraph = Graph->GetGraphHandle();
-		FrontendGraph->RemoveVariable(VariableID);
+		{
+			const TArray<UMetasoundEditorGraphMemberNode*> Nodes = GetNodes();
+			for (UMetasoundEditorGraphMemberNode* Node : Nodes)
+			{
+				Graph->RemoveNode(Node);
+			}
+		}
+
+		FMetaSoundFrontendDocumentBuilder& DocBuilder = GetFrontendBuilderChecked();
+		DocBuilder.RemoveGraphVariable(CachedData.MemberName);
 		VariableID = FGuid();
 
 		// Add variable with new type to frontend
-		FVariableHandle FrontendVariable = FrontendGraph->AddVariable(InNewType);
+		const FMetasoundFrontendVariable* FrontendVariable = DocBuilder.AddGraphVariable(
+			CachedData.MemberName,
+			InNewType,
+			nullptr,
+			&CachedData.DisplayName,
+			&CachedData.Description);
 
-		if (!ensure(FrontendVariable->IsValid()))
+		if (!ensure(FrontendVariable))
 		{
 			// Failed to add a new variable with the given data type. 
 			return;
 		}
 
 		// Setup this object with new variable data
-		VariableID = FrontendVariable->GetID();
-
-		constexpr bool bPostSubTransaction = false;
-		SetMemberName(CachedData.MemberName, bPostSubTransaction);
-		SetDisplayName(CachedData.DisplayName, bPostSubTransaction);
-		SetDescription(CachedData.Description, bPostSubTransaction);
-
+		VariableID = FrontendVariable->ID;
 		TypeName = InNewType;
 		InitializeLiteral();
 
@@ -1695,8 +1762,7 @@ void UMetasoundEditorGraphVariable::SetDataType(FName InNewType, bool bPostTrans
 		}
 
 		// Add the nodes with the same identifier data but new datatype.
-		UObject& Metasound = Graph->GetMetasoundChecked();
-		AddVariableNodes(Metasound, FrontendGraph, CachedData.Locations);
+		AddVariableNodes(Graph->GetMetasoundChecked(), CachedData.Locations);
 	}
 
 }
@@ -1782,37 +1848,44 @@ UMetasoundEditorGraphVariable::FVariableNodeLocations UMetasoundEditorGraphVaria
 	return Locations;
 }
 
-void UMetasoundEditorGraphVariable::AddVariableNodes(UObject& InMetasound, Metasound::Frontend::FGraphHandle& InFrontendGraph, const FVariableNodeLocations& InNodeLocs)
+void UMetasoundEditorGraphVariable::AddVariableNodes(UObject& InMetaSound, const FVariableNodeLocations& InNodeLocs)
 {
 	using namespace Metasound::Frontend;
 	using namespace Metasound::Editor;
 
+	auto AddAndSyncEdGraphNode = [&](const FMetasoundFrontendNode* InNode, const FVector2D& Location)
+	{
+		check(InNode);
+
+		UMetasoundEditorGraphNode* NewGraphNode = FGraphBuilder::AddVariableNode(InMetaSound, InNode->GetID(), false /* bInSelectNewNode */);
+		check(NewGraphNode);
+
+		NewGraphNode->UpdateFrontendNodeLocation(Location);
+		NewGraphNode->SyncLocationFromFrontendNode();
+	};
+
+	FMetaSoundFrontendDocumentBuilder& DocBuilder = GetFrontendBuilderChecked();
+	const FMetasoundFrontendVariable* FrontendVariable = GetFrontendVariable();
+	check(FrontendVariable);
 	if (InNodeLocs.MutatorLocation)
 	{
-		bool bMutatorNodeAlreadyExists = GetConstVariableHandle()->FindMutatorNode()->IsValid();
-		if (ensure(!bMutatorNodeAlreadyExists))
+		if (ensure(!FrontendVariable->MutatorNodeID.IsValid()))
 		{
-			FConstNodeHandle MutatorFrontendNode = InFrontendGraph->FindOrAddVariableMutatorNode(VariableID);
-			UMetasoundEditorGraphNode* NewGraphNode = FGraphBuilder::AddVariableNode(InMetasound, MutatorFrontendNode, false /* bInSelectNewNode */);
-			NewGraphNode->UpdateFrontendNodeLocation(*InNodeLocs.MutatorLocation);
-			NewGraphNode->SyncLocationFromFrontendNode();
+			const FMetasoundFrontendNode* NewFrontendNode = DocBuilder.AddGraphVariableMutatorNode(FrontendVariable->Name);
+			AddAndSyncEdGraphNode(NewFrontendNode, *InNodeLocs.MutatorLocation);
 		}
 	}
 
 	for (const FVector2D& Location : InNodeLocs.AccessorLocations)
 	{
-		FConstNodeHandle AccessorFrontendNode = InFrontendGraph->AddVariableAccessorNode(VariableID);
-		UMetasoundEditorGraphNode* NewGraphNode = FGraphBuilder::AddVariableNode(InMetasound, AccessorFrontendNode, false /* bInSelectNewNode */);
-		NewGraphNode->UpdateFrontendNodeLocation(Location);
-		NewGraphNode->SyncLocationFromFrontendNode();
+		const FMetasoundFrontendNode* NewFrontendNode = DocBuilder.AddGraphVariableAccessorNode(FrontendVariable->Name);
+		AddAndSyncEdGraphNode(NewFrontendNode, Location);
 	}
 
 	for (const FVector2D& Location : InNodeLocs.DeferredAccessorLocations)
 	{
-		FConstNodeHandle DeferredAccessorFrontendNode = InFrontendGraph->AddVariableDeferredAccessorNode(VariableID);
-		UMetasoundEditorGraphNode* NewGraphNode = FGraphBuilder::AddVariableNode(InMetasound, DeferredAccessorFrontendNode, false /* bInSelectNewNode */);
-		NewGraphNode->UpdateFrontendNodeLocation(Location);
-		NewGraphNode->SyncLocationFromFrontendNode();
+		const FMetasoundFrontendNode* NewFrontendNode = DocBuilder.AddGraphVariableDeferredAccessorNode(FrontendVariable->Name);
+		AddAndSyncEdGraphNode(NewFrontendNode, Location);
 	}
 }
 
@@ -1821,10 +1894,14 @@ const FGuid& UMetasoundEditorGraphVariable::GetVariableID() const
 	return VariableID;
 }
 
+const FMetasoundFrontendVariable* UMetasoundEditorGraphVariable::GetFrontendVariable() const
+{
+	return GetFrontendBuilderChecked().FindGraphVariable(VariableID);
+}
+
 void UMetasoundEditorGraphVariable::ResetToClassDefault()
 {
 	using namespace Metasound;
-	using namespace Metasound::Editor;
 	using namespace Metasound::Frontend;
 	using namespace Metasound::VariableNames;
 
@@ -1834,20 +1911,20 @@ void UMetasoundEditorGraphVariable::ResetToClassDefault()
 	Literal->Modify();
 	Literal->SetFromLiteral(DefaultLiteral);
 
-	FVariableHandle VariableHandle = GetVariableHandle();
-	VariableHandle->SetLiteral(DefaultLiteral);
+	UObject* MetaSound = Literal->GetOutermostObject();
+	check(MetaSound);
+	MetaSound->Modify();
 
-	FNodeHandle MutatorNode = VariableHandle->FindMutatorNode();
-	if (MutatorNode->IsValid())
+	FMetaSoundFrontendDocumentBuilder& DocBuilder = GetFrontendBuilderChecked();
+	DocBuilder.SetGraphVariableDefault(GetMemberName(), MoveTemp(DefaultLiteral));
+
+	if (const FMetasoundFrontendVariable* FrontendVariable = DocBuilder.FindGraphVariable(GetMemberName()); ensure(FrontendVariable))
 	{
-		FInputHandle InputHandle = MutatorNode->GetInputWithVertexName(METASOUND_GET_PARAM_NAME(InputData));
-		if (ensure(InputHandle->IsValid()))
+		if (const FMetasoundFrontendVertex* FrontendVertex = DocBuilder.FindNodeInput(FrontendVariable->MutatorNodeID, METASOUND_GET_PARAM_NAME(InputData)))
 		{
-			InputHandle->ClearLiteral();
+			DocBuilder.RemoveNodeInputDefault(FrontendVariable->MutatorNodeID, FrontendVertex->VertexID);
 		}
 	}
-
-	FGraphBuilder::GetOutermostMetaSoundChecked(*this).GetModifyContext().AddMemberIDsModified({ GetMemberID() });
 }
 
 bool UMetasoundEditorGraphVariable::Synchronize()
@@ -1857,53 +1934,55 @@ bool UMetasoundEditorGraphVariable::Synchronize()
 	bool bModified = Super::Synchronize();
 	FMetaSoundFrontendDocumentBuilder& Builder = GetFrontendBuilderChecked();
 	const FMetasoundFrontendGraph& Graph = Builder.FindConstBuildGraphChecked();
-	const FName MemberName = GetMemberName();
-	if (const FMetasoundFrontendVariable* Variable = Builder.FindGraphVariable(MemberName))
+	const FMetasoundFrontendVariable* Variable = GetFrontendVariable();
+	if (!Variable)
 	{
-		if (TypeName != Variable->TypeName)
-		{
-			bModified = true;
-			TypeName = Variable->TypeName;
+		return false;
+	}
 
-			InitializeLiteral();
-		}
+	if (TypeName != Variable->TypeName)
+	{
+		bModified = true;
+		TypeName = Variable->TypeName;
 
-		if (ensure(Literal))
+		InitializeLiteral();
+	}
+
+	if (ensure(Literal))
+	{
+		TOptional<FMetasoundFrontendLiteral> NewDefault;
+		FMetasoundFrontendLiteral DefaultLiteral;
+		Literal->TryFindDefault(DefaultLiteral);
+		if (const FMetasoundFrontendNode* MutatorNode = Builder.FindNode(Variable->MutatorNodeID))
 		{
-			TOptional<FMetasoundFrontendLiteral> NewDefault;
-			FMetasoundFrontendLiteral DefaultLiteral;
-			Literal->TryFindDefault(DefaultLiteral);
-			if (const FMetasoundFrontendNode* MutatorNode = Builder.FindNode(Variable->MutatorNodeID))
+			if (!MutatorNode->InputLiterals.IsEmpty())
 			{
-				if (!MutatorNode->InputLiterals.IsEmpty())
+				const FMetasoundFrontendVertexLiteral& VertexLiteral = MutatorNode->InputLiterals.Last();
+				if (!VertexLiteral.Value.IsEqual(DefaultLiteral))
 				{
-					const FMetasoundFrontendVertexLiteral& VertexLiteral = MutatorNode->InputLiterals.Last();
-					if (!VertexLiteral.Value.IsEqual(DefaultLiteral))
-					{
-						NewDefault = VertexLiteral.Value;
-					}
-				}
-				else
-				{
-					FMetasoundFrontendLiteral TypeDefault;
-					TypeDefault.SetFromLiteral(Frontend::IDataTypeRegistry::Get().CreateDefaultLiteral(TypeName));
-					if (!TypeDefault.IsEqual(DefaultLiteral))
-					{
-						NewDefault = TypeDefault;
-					}
+					NewDefault = VertexLiteral.Value;
 				}
 			}
-			else if (!Variable->Literal.IsEqual(DefaultLiteral))
+			else
 			{
-				NewDefault = Variable->Literal;
+				FMetasoundFrontendLiteral TypeDefault;
+				TypeDefault.SetFromLiteral(Frontend::IDataTypeRegistry::Get().CreateDefaultLiteral(TypeName));
+				if (!TypeDefault.IsEqual(DefaultLiteral))
+				{
+					NewDefault = TypeDefault;
+				}
 			}
+		}
+		else if (!Variable->Literal.IsEqual(DefaultLiteral))
+		{
+			NewDefault = Variable->Literal;
+		}
 
-			if (NewDefault.IsSet())
-			{
-				bModified = true;
-				Literal->ResetDefaults();
-				Literal->SetFromLiteral(*NewDefault);
-			}
+		if (NewDefault.IsSet())
+		{
+			bModified = true;
+			Literal->ResetDefaults();
+			Literal->SetFromLiteral(*NewDefault);
 		}
 	}
 
@@ -1922,14 +2001,14 @@ void UMetasoundEditorGraphVariable::UpdateFrontendDefaultLiteral(bool bPostTrans
 		return;
 	}
 
-	FMetaSoundFrontendDocumentBuilder& Builder = GetFrontendBuilderChecked();
-	const FMetasoundFrontendVariable* Variable = Builder.FindGraphVariable(GetMemberName());
+	const FMetasoundFrontendVariable* Variable = GetFrontendVariable();
 	if (!ensure(Variable))
 	{
 		return;
 	}
 
 	const FScopedTransaction Transaction(LOCTEXT("Set Variable Default", "Set MetaSound Variable Default"), bPostTransaction);
+	FMetaSoundFrontendDocumentBuilder& Builder = GetFrontendBuilderChecked();
 	Builder.CastDocumentObjectChecked<UObject>().Modify();
 
 	// Use the default page ID here as variables do *not* support paged defaults
@@ -1939,7 +2018,7 @@ void UMetasoundEditorGraphVariable::UpdateFrontendDefaultLiteral(bool bPostTrans
 
 	// Page ID is passed along to the builder from here because the builder needs the current BuildPageID to access the appropriate in-graph variable
 	// (variables can have the same IDs/names in different paged graphs).
-	Builder.SetGraphVariableDefault(GetMemberName(), DefaultLiteral, InPageID);
+	Builder.SetGraphVariableDefault(Variable->Name, DefaultLiteral, InPageID);
 
 	if (const FMetasoundFrontendNode* MutatorNode = Builder.FindNode(Variable->MutatorNodeID, InPageID))
 	{
@@ -1960,7 +2039,9 @@ UMetasoundEditorGraphInputNode* UMetasoundEditorGraph::CreateInputNode(Metasound
 
 Metasound::Frontend::FDocumentHandle UMetasoundEditorGraph::GetDocumentHandle()
 {
+PRAGMA_DISABLE_DEPRECATION_WARNINGS
 	return GetGraphHandle()->GetOwningDocument();
+PRAGMA_ENABLE_DEPRECATION_WARNINGS
 }
 
 Metasound::Frontend::FConstDocumentHandle UMetasoundEditorGraph::GetDocumentHandle() const
@@ -1992,7 +2073,10 @@ void UMetasoundEditorGraph::PreSave(FObjectPreSaveContext InSaveContext)
 	GetNodesOfClass<UMetasoundEditorGraphNode>(MetaSoundNodes);
 	for (UMetasoundEditorGraphNode* Node : MetaSoundNodes)
 	{
+PRAGMA_DISABLE_DEPRECATION_WARNINGS
 		FNodeHandle NodeHandle = Node->GetNodeHandle();
+PRAGMA_ENABLE_DEPRECATION_WARNINGS
+
 		FMetasoundFrontendNodeStyle Style = NodeHandle->GetNodeStyle();
 		Style.bMessageNodeUpdated = false;
 		NodeHandle->SetNodeStyle(Style);
@@ -2246,6 +2330,31 @@ UMetasoundEditorGraphVariable* UMetasoundEditorGraph::FindVariable(const FGuid&
 	return Variable ? Variable->Get() : nullptr;
 }
 
+UMetasoundEditorGraphVariable* UMetasoundEditorGraph::FindOrAddVariable(FName VariableName)
+{
+	using namespace Metasound::Editor;
+	using namespace Metasound::Frontend;
+
+	if (const FMetasoundFrontendVariable* Variable = GetBuilderChecked().GetBuilder().FindGraphVariable(VariableName))
+	{
+		if (TObjectPtr<UMetasoundEditorGraphVariable> EditorVariable = FindVariable(Variable->ID))
+		{
+			ensure(EditorVariable->TypeName == Variable->TypeName);
+			return EditorVariable;
+		}
+
+		UMetasoundEditorGraphVariable* NewVariable = NewObject<UMetasoundEditorGraphVariable>(this, GraphPrivate::GetUniqueTransientMemberName(), RF_Transactional);
+		if (ensure(NewVariable))
+		{
+			NewVariable->InitMember(Variable->TypeName, Variable->Literal, Variable->ID);
+			Variables.Add(NewVariable);
+			return NewVariable;
+		}
+	}
+
+	return nullptr;
+}
+
 UMetasoundEditorGraphVariable* UMetasoundEditorGraph::FindOrAddVariable(const Metasound::Frontend::FConstVariableHandle& InVariableHandle)
 {
 	using namespace Metasound::Editor;
diff --git a/Engine/Plugins/Runtime/Metasound/Source/MetasoundEditor/Private/MetasoundEditorGraphBuilder.cpp b/Engine/Plugins/Runtime/Metasound/Source/MetasoundEditor/Private/MetasoundEditorGraphBuilder.cpp
--- a/Engine/Plugins/Runtime/Metasound/Source/MetasoundEditor/Private/MetasoundEditorGraphBuilder.cpp
+++ b/Engine/Plugins/Runtime/Metasound/Source/MetasoundEditor/Private/MetasoundEditorGraphBuilder.cpp
@@ -524,28 +524,36 @@ namespace Metasound
 			return INodeController::GetInvalidHandle();
 		}
 
-		UMetasoundEditorGraphVariableNode* FGraphBuilder::AddVariableNode(UObject& InMetaSound, const Frontend::FConstNodeHandle& InNodeHandle, bool bInSelectNewNode)
+		UMetasoundEditorGraphVariableNode* FGraphBuilder::AddVariableNode(UObject& InMetaSound, const FGuid& InNodeID, bool bInSelectNewNode)
 		{
 			using namespace Frontend;
 
-			const FMetasoundFrontendClassMetadata& Metadata = InNodeHandle->GetClassMetadata();
+			FMetaSoundFrontendDocumentBuilder& Builder = IDocumentBuilderRegistry::GetChecked().FindOrBeginBuilding(&InMetaSound);
+
+			const FMetasoundFrontendNode* FrontendNode = Builder.FindNode(InNodeID);
+			check(FrontendNode);
+
+			const FMetasoundFrontendClass* Class = Builder.FindDependency(FrontendNode->ClassID);
+			check(Class);
+
+			const FMetasoundFrontendClassMetadata& Metadata = Class->Metadata;
 			EMetasoundFrontendClassType ClassType = Metadata.GetType();
-			const bool bIsSupportedClassType = (ClassType == EMetasoundFrontendClassType::VariableAccessor) 
+			const bool bIsSupportedClassType = (ClassType == EMetasoundFrontendClassType::VariableAccessor)
 				|| (ClassType == EMetasoundFrontendClassType::VariableDeferredAccessor)
 				|| (ClassType == EMetasoundFrontendClassType::VariableMutator);
 
 			if (ensureMsgf(bIsSupportedClassType, TEXT("Cannot call 'AddVariableNode' with node of class type '%s'"), LexToString(ClassType)))
 			{
-				FConstVariableHandle FrontendVariable = InNodeHandle->GetOwningGraph()->FindVariableContainingNode(InNodeHandle->GetID());
-				if (ensure(FrontendVariable->IsValid()))
+				const FMetasoundFrontendVariable* FrontendVariable = Builder.FindGraphVariableByNodeID(InNodeID);
+				if (ensure(FrontendVariable))
 				{
-					auto InitNodeFunc = [&FrontendVariable, &InNodeHandle, &Metadata](UMetasoundEditorGraph& MetasoundGraph, UMetasoundEditorGraphVariableNode& NewGraphNode)
+					auto InitNodeFunc = [&FrontendVariable, &InNodeID, &Metadata](UMetasoundEditorGraph& MetasoundGraph, UMetasoundEditorGraphVariableNode& NewGraphNode)
 					{
-						UMetasoundEditorGraphVariable* Variable = MetasoundGraph.FindOrAddVariable(FrontendVariable);
+						UMetasoundEditorGraphVariable* Variable = MetasoundGraph.FindOrAddVariable(FrontendVariable->Name);
 						if (ensure(Variable))
 						{
 							NewGraphNode.Variable = Variable;
-							NewGraphNode.NodeID = InNodeHandle->GetID();
+							NewGraphNode.NodeID = InNodeID;
 							NewGraphNode.ClassName = Metadata.GetClassName();
 							NewGraphNode.ClassType = Metadata.GetType();
 						}
@@ -808,11 +816,15 @@ namespace Metasound
 					// purposes instead of the underlying vertex's name. The frontend vertices
 					// of a variable node have consistent names no matter what the 
 					// variable is named.
+PRAGMA_DISABLE_DEPRECATION_WARNINGS
 					return EdVariableNode->GetNodeHandle()->GetInputWithVertexName(METASOUND_GET_PARAM_NAME(InputData));
+PRAGMA_ENABLE_DEPRECATION_WARNINGS
 				}
 				else if (UMetasoundEditorGraphNode* EdNode = CastChecked<UMetasoundEditorGraphNode>(InPin->GetOwningNode()))
 				{
+PRAGMA_DISABLE_DEPRECATION_WARNINGS
 					return EdNode->GetNodeHandle()->GetInputWithVertexName(InPin->GetFName());
+PRAGMA_ENABLE_DEPRECATION_WARNINGS
 				}
 			}
 
@@ -1014,11 +1026,15 @@ namespace Metasound
 					// purposes instead of the underlying vertex's name. The frontend vertices
 					// of a variable node have consistent names no matter what the 
 					// variable is named.
+PRAGMA_DISABLE_DEPRECATION_WARNINGS
 					return EdVariableNode->GetNodeHandle()->GetOutputWithVertexName(METASOUND_GET_PARAM_NAME(OutputData));
+PRAGMA_ENABLE_DEPRECATION_WARNINGS
 				}
 				else if (UMetasoundEditorGraphNode* EdNode = CastChecked<UMetasoundEditorGraphNode>(InPin->GetOwningNode()))
 				{
+PRAGMA_DISABLE_DEPRECATION_WARNINGS
 					return EdNode->GetNodeHandle()->GetOutputWithVertexName(InPin->GetFName());
+PRAGMA_ENABLE_DEPRECATION_WARNINGS
 				}
 			}
 
@@ -1419,15 +1435,15 @@ namespace Metasound
 			return ClassOutput;
 		}
 
-		FName FGraphBuilder::GenerateUniqueVariableName(const Frontend::FConstGraphHandle& InFrontendGraph, const FString& InBaseName)
+		FName FGraphBuilder::GenerateUniqueVariableName(const FMetaSoundFrontendDocumentBuilder& InBuilder, const FString& InBaseName)
 		{
 			using namespace Frontend;
 
 			TSet<FName> ExistingVariableNames;
 
-			// Get all the names from the existing variables on the graph
-			// and place into the ExistingVariableNames array.
-			Algo::Transform(InFrontendGraph->GetVariables(), ExistingVariableNames, [](const FConstVariableHandle& Var) { return Var->GetName(); });
+			// Get all the names from the existing variables on
+			// the build graph and place into the ExistingVariableNames array.
+			Algo::Transform(InBuilder.FindConstBuildGraphChecked().Variables, ExistingVariableNames, [](const FMetasoundFrontendVariable& Var) { return Var.Name; });
 
 			return GraphBuilderPrivate::GenerateUniqueName(ExistingVariableNames, InBaseName);
 		}
@@ -1446,93 +1462,6 @@ namespace Metasound
 			return *MetaSoundAsset;
 		}
 
-		Frontend::FVariableHandle FGraphBuilder::AddVariableHandle(UObject& InMetaSound, const FName& InTypeName, const FName* InNameBase)
-		{
-			using namespace Frontend;
-
-			FMetasoundAssetBase* MetaSoundAsset = IMetasoundUObjectRegistry::Get().GetObjectAsAssetBase(&InMetaSound);
-			check(MetaSoundAsset);
-
-			FGraphHandle FrontendGraph = MetaSoundAsset->GetRootGraphHandle();
-
-			FName VariableName = GenerateUniqueVariableName(FrontendGraph, InNameBase ? InNameBase->ToString() : LOCTEXT("VariableDefaultDisplayName", "Variable").ToString());
-			FVariableHandle Variable = FrontendGraph->AddVariable(InTypeName);
-
-			Variable->SetDisplayName(FText::GetEmpty());
-			Variable->SetName(VariableName);
-
-			return Variable;
-		}
-
-		Frontend::FVariableHandle FGraphBuilder::DuplicateVariableHandle(UObject& InMetaSound, const Frontend::FConstVariableHandle& InHandle)
-		{
-			Frontend::FVariableHandle NewHandle = AddVariableHandle(InMetaSound, InHandle->GetDataType(), &InHandle->GetName());
-			NewHandle->SetDisplayName(FText::GetEmpty());
-			NewHandle->SetDescription(InHandle->GetDescription());
-			NewHandle->SetLiteral(InHandle->GetLiteral());
-			return NewHandle;
-		}
-
-		Frontend::FNodeHandle FGraphBuilder::AddVariableNodeHandle(UObject& InMetaSound, const FGuid& InVariableID, const Metasound::FNodeClassName& InVariableNodeClassName, UMetasoundEditorGraphVariableNode* InVariableNode)
-		{
-			using namespace Frontend;
-
-			FNodeHandle FrontendNode = INodeController::GetInvalidHandle();
-
-			FMetasoundAssetBase* MetaSoundAsset = IMetasoundUObjectRegistry::Get().GetObjectAsAssetBase(&InMetaSound);
-			
-			if (ensure(MetaSoundAsset))
-			{
-				FMetasoundFrontendClass FrontendClass;
-				bool bDidFindClassWithName = ISearchEngine::Get().FindClassWithHighestVersion(InVariableNodeClassName, FrontendClass);
-				if (ensure(bDidFindClassWithName))
-				{
-					FGraphHandle Graph = MetaSoundAsset->GetRootGraphHandle();
-
-					switch (FrontendClass.Metadata.GetType())
-					{
-						case EMetasoundFrontendClassType::VariableDeferredAccessor:
-							FrontendNode = Graph->AddVariableDeferredAccessorNode(InVariableID);
-							break;
-
-						case EMetasoundFrontendClassType::VariableAccessor:
-							FrontendNode = Graph->AddVariableAccessorNode(InVariableID);
-							break;
-
-						case EMetasoundFrontendClassType::VariableMutator:
-						{
-							FConstVariableHandle Variable = Graph->FindVariable(InVariableID);
-							FConstNodeHandle ExistingMutator = Variable->FindMutatorNode();
-							if (!ExistingMutator->IsValid())
-							{
-								FrontendNode = Graph->FindOrAddVariableMutatorNode(InVariableID);
-							}
-							else
-							{
-								UE_LOG(LogMetaSound, Error, TEXT("Cannot add node because \"%s\" already exists for variable \"%s\""), *ExistingMutator->GetDisplayName().ToString(), *Variable->GetDisplayName().ToString());
-							}
-						}
-						break;
-
-						default:
-						{
-							checkNoEntry();
-						}
-					}
-				}
-			}
-
-			if (InVariableNode)
-			{
-				InVariableNode->ClassName = FrontendNode->GetClassMetadata().GetClassName();
-				InVariableNode->ClassType = FrontendNode->GetClassMetadata().GetType();
-				InVariableNode->NodeID = FrontendNode->GetID();
-				MetaSoundAsset->GetModifyContext().AddNodeIDsModified({ InVariableNode->GetNodeID() });
-			}
-
-			return FrontendNode;
-		}
-
 		bool FGraphBuilder::ConnectNodes(UEdGraphPin& InInputPin, UEdGraphPin& InOutputPin, bool bInConnectEdPins)
 		{
 			using namespace Frontend;
@@ -1588,20 +1517,6 @@ namespace Metasound
 			}
 		}
 
-		bool FGraphBuilder::DeleteMemberNodes(const UMetasoundEditorGraphMember& InGraphMember)
-		{
-			bool bSuccess = true;
-			const TArray<UMetasoundEditorGraphMemberNode*> Nodes = InGraphMember.GetNodes();
-			for (UMetasoundEditorGraphMemberNode* Node : Nodes)
-			{
-				if (ensure(Node))
-				{
-					bSuccess &= Metasound::Editor::FGraphBuilder::DeleteNode(*Node);
-				}
-			}
-			return bSuccess;
-		}
-
 		bool FGraphBuilder::DeleteNode(UEdGraphNode& InNode, bool bRemoveUnusedDependencies)
 		{
 			using namespace Engine;
@@ -1664,7 +1579,7 @@ namespace Metasound
 				InGraphNode.RemovePin(InGraphNode.Pins[i]);
 			}
 
-			FConstNodeHandle NodeHandle = InGraphNode.GetNodeHandle();
+			FConstNodeHandle NodeHandle = InGraphNode.GetConstNodeHandle();
 			TArray<FConstInputHandle> InputHandles = NodeHandle->GetConstInputs();
 			NodeHandle->GetInputStyle().SortDefaults(InputHandles, [](const Frontend::FConstInputHandle& Handle) { return FGraphBuilder::GetDisplayName(*Handle); });
 			for (const FConstInputHandle& InputHandle : InputHandles)
@@ -2291,6 +2206,7 @@ namespace Metasound
 				for (int32 j = EditorNodes.Num() - 1; j >= 0; --j)
 				{
 					UMetasoundEditorGraphNode* EditorNode = EditorNodes[j];
+					EditorNode->CacheBreadcrumb();
 					if (EditorNode->GetNodeID() == Node->GetID())
 					{
 						// Editor node may have the same Frontend NodeID as another page,
@@ -2377,7 +2293,7 @@ namespace Metasound
 						case EMetasoundFrontendClassType::VariableDeferredAccessor:
 						case EMetasoundFrontendClassType::Variable:
 						{
-							NewGraphNode = CastChecked<UMetasoundEditorGraphNode>(AddVariableNode(MetaSoundObject, Node, false));
+							NewGraphNode = CastChecked<UMetasoundEditorGraphNode>(AddVariableNode(MetaSoundObject, Node->GetID(), false));
 						}
 						break;
 
@@ -2655,130 +2571,114 @@ namespace Metasound
 			using namespace Frontend;
 
 			bool bEditorGraphModified = false;
+			auto GetMemberID = [](const TObjectPtr<UMetasoundEditorGraphMember>& Member) { return Member->GetMemberID(); };
+			auto IsValidMember = [](const TObjectPtr<UMetasoundEditorGraphMember>& Member) { return Member != nullptr; };
 
-			const UObject& MetaSoundObject = InBuilder.CastDocumentObjectChecked<const UObject>();
-			const FMetasoundAssetBase* MetaSoundAsset = IMetasoundUObjectRegistry::Get().GetObjectAsAssetBase(&MetaSoundObject);
-			check(MetaSoundAsset);
-			FConstGraphHandle GraphHandle = MetaSoundAsset->GetRootGraphHandle();
-
-			TSet<UMetasoundEditorGraphInput*> Inputs;
-			TSet<UMetasoundEditorGraphOutput*> Outputs;
-
-			// Collect all editor graph inputs with corresponding frontend inputs. 
-			GraphHandle->IterateConstNodes([&](FConstNodeHandle NodeHandle)
+			auto DeleteMemberNodes = [&OutGraph](const TObjectPtr<UMetasoundEditorGraphMember>& Member)
 			{
-				const FGuid NodeID = NodeHandle->GetID();
-				if (UMetasoundEditorGraphInput* Input = OutGraph.FindInput(NodeID))
+				if (Member)
 				{
-					Inputs.Add(Input);
-					return;
+					const TArray<UMetasoundEditorGraphMemberNode*> Nodes = Member->GetNodes();
+					for (UMetasoundEditorGraphMemberNode* Node : Nodes)
+					{
+						OutGraph.RemoveNode(Node);
+					}
 				}
+			};
 
-				// Add an editor input if none exist for a frontend input.
-				Inputs.Add(OutGraph.FindOrAddInput(NodeID));
-				constexpr bool bIncludeNamespace = true;
-				FText NodeDisplayName = FGraphBuilder::GetDisplayName(*NodeHandle, bIncludeNamespace);
-				UE_LOG(LogMetasoundEditor, Verbose, TEXT("Synchronizing Inputs: Added missing input '%s'."), *NodeDisplayName.ToString());
-				bEditorGraphModified = true;
-			}, EMetasoundFrontendClassType::Input);
+			const FMetasoundFrontendGraphClass& RootGraph = InBuilder.GetConstDocumentChecked().RootGraph;
 
-			// Remove stale inputs.
+			// Input Sync
 			{
-				for (int32 Index = OutGraph.Inputs.Num() - 1; Index >= 0; --Index)
+				TSet<FGuid> EdMemberIDs;
+				Algo::TransformIf(OutGraph.Inputs, EdMemberIDs, IsValidMember, GetMemberID);
+				for (const FMetasoundFrontendClassInput& FrontendInput : RootGraph.Interface.Inputs)
 				{
-					const UMetasoundEditorGraphInput* Input = OutGraph.Inputs[Index];
-					check(Input);
-					if (!Inputs.Contains(Input))
+					if (EdMemberIDs.Remove(FrontendInput.NodeID) == 0)
 					{
-						TArray<UMetasoundEditorGraphMemberNode*> InputNodes = Input->GetNodes();
-						for (UMetasoundEditorGraphNode* Node : InputNodes)
-						{
-							FGraphBuilder::DeleteMemberNodes(*Input);
-						}
-
 						bEditorGraphModified = true;
-						UE_LOG(LogMetasoundEditor, Verbose, TEXT("Synchronizing Outputs: Removing stale output '%s'."), *Input->GetName());
-						OutGraph.Inputs.RemoveAtSwap(Index, EAllowShrinking::No);
+						OutGraph.FindOrAddInput(FrontendInput.NodeID);
 					}
 				}
-			}
 
-			for (TObjectPtr<UMetasoundEditorGraphInput> Input : OutGraph.Inputs)
-			{
-				bEditorGraphModified |= Input->Synchronize();
+				for (int32 Index = OutGraph.Inputs.Num() - 1; Index >= 0; --Index)
+				{
+					TObjectPtr<UMetasoundEditorGraphMember> Member = OutGraph.Inputs[Index];
+					if (!Member || EdMemberIDs.Contains(Member->GetMemberID()))
+					{
+						DeleteMemberNodes(Member);
+						OutGraph.Inputs.RemoveAtSwap(Index, EAllowShrinking::No);
+						bEditorGraphModified = true;
+					}
+					else
+					{
+						Member->CacheBreadcrumb();
+						bEditorGraphModified |= Member->Synchronize();
+					}
+				}
+				OutGraph.Inputs.Shrink();
 			}
 
-			// Collect all editor graph outputs with corresponding frontend outputs. 
-			GraphHandle->IterateConstNodes([&](FConstNodeHandle NodeHandle)
+			// Output Sync
 			{
-				const FGuid NodeID = NodeHandle->GetID();
-				if (UMetasoundEditorGraphOutput* Output = OutGraph.FindOutput(NodeID))
+				TSet<FGuid> EdMemberIDs;
+				Algo::TransformIf(OutGraph.Outputs, EdMemberIDs, IsValidMember, GetMemberID);
+				for (const FMetasoundFrontendClassOutput& FrontendOutput : RootGraph.Interface.Outputs)
 				{
-					Outputs.Add(Output);
-					return;
+					if (EdMemberIDs.Remove(FrontendOutput.NodeID) == 0)
+					{
+						bEditorGraphModified = true;
+						OutGraph.FindOrAddOutput(FrontendOutput.NodeID);
+					}
 				}
 
-				// Add an editor output if none exist for a frontend output.
-				Outputs.Add(OutGraph.FindOrAddOutput(NodeID));
-				constexpr bool bIncludeNamespace = true;
-				FText NodeDisplayName = FGraphBuilder::GetDisplayName(*NodeHandle, bIncludeNamespace);
-				UE_LOG(LogMetasoundEditor, Verbose, TEXT("Synchronizing Outputs: Added missing output '%s'."), *NodeDisplayName.ToString());
-				bEditorGraphModified = true;
-			}, EMetasoundFrontendClassType::Output);
-
-			// Remove stale outputs.
-			{
 				for (int32 Index = OutGraph.Outputs.Num() - 1; Index >= 0; --Index)
 				{
-					const UMetasoundEditorGraphOutput* Output = OutGraph.Outputs[Index];
-					check(Output);
-					if (!Outputs.Contains(Output))
+					TObjectPtr<UMetasoundEditorGraphMember> Member = OutGraph.Outputs[Index];
+					if (!Member || EdMemberIDs.Contains(Member->GetMemberID()))
 					{
-						bEditorGraphModified = true;
-						UE_LOG(LogMetasoundEditor, Verbose, TEXT("Synchronizing Outputs: Removing stale output '%s'."), *Output->GetName());
+						DeleteMemberNodes(Member);
 						OutGraph.Outputs.RemoveAtSwap(Index, EAllowShrinking::No);
 					}
+					else
+					{
+						Member->CacheBreadcrumb();
+						bEditorGraphModified |= Member->Synchronize();
+					}
 				}
+				OutGraph.Outputs.Shrink();
 			}
 
-			for (TObjectPtr<UMetasoundEditorGraphOutput> Output : OutGraph.Outputs)
+			// Variable Sync
 			{
-				bEditorGraphModified |= Output->Synchronize();
-			}
-
-			// Enforce all variables exists in the edgraph
-			TSet<FGuid> VariableIDs;
-			Algo::Transform(OutGraph.Variables, VariableIDs, [](const TObjectPtr<UMetasoundEditorGraphVariable>& Variable) { return Variable->GetVariableID(); });
-			TArray<FConstVariableHandle> FrontendVariables = GraphHandle->GetVariables();
-			for (const FConstVariableHandle& VariableHandle : FrontendVariables)
-			{
-				const FGuid VariableID = VariableHandle->GetID();
-				if (VariableIDs.Remove(VariableID) == 0)
+				TSet<FGuid> EdMemberIDs;
+				Algo::TransformIf(OutGraph.Variables, EdMemberIDs, IsValidMember, GetMemberID);
+				for (const FMetasoundFrontendVariable& FrontendVariable : InBuilder.FindConstBuildGraphChecked().Variables)
 				{
-					UMetasoundEditorGraphVariable* NewVariable = NewObject<UMetasoundEditorGraphVariable>(&OutGraph, FName(), RF_Transactional);
-					if (ensure(NewVariable))
+					if (EdMemberIDs.Remove(FrontendVariable.ID) == 0)
 					{
-						const FMetasoundFrontendLiteral DefaultLiteral = VariableHandle->GetLiteral();
-						NewVariable->InitMember(VariableHandle->GetDataType(), DefaultLiteral, VariableID);
-						OutGraph.Variables.Add(NewVariable);
 						bEditorGraphModified = true;
+						OutGraph.FindOrAddVariable(FrontendVariable.Name);
 					}
 				}
-			}
 
-			for (TObjectPtr<UMetasoundEditorGraphVariable> Variable : OutGraph.Variables)
-			{
-				bEditorGraphModified |= Variable->Synchronize();
+				for (int32 Index = OutGraph.Variables.Num() - 1; Index >= 0; --Index)
+				{
+					TObjectPtr<UMetasoundEditorGraphMember> Member = OutGraph.Variables[Index];
+					if (!Member || EdMemberIDs.Contains(Member->GetMemberID()))
+					{
+						DeleteMemberNodes(Member);
+						OutGraph.Variables.RemoveAtSwap(Index, EAllowShrinking::No);
+					}
+					else
+					{
+						Member->CacheBreadcrumb();
+						bEditorGraphModified |= Member->Synchronize();
+					}
+				}
+				OutGraph.Variables.Shrink();
 			}
 
-			// Remove empty entries
-			bEditorGraphModified |= OutGraph.Inputs.RemoveAllSwap([](const TObjectPtr<UMetasoundEditorGraphInput>& Input) { return !Input; }) > 0;
-			bEditorGraphModified |= OutGraph.Outputs.RemoveAllSwap([](const TObjectPtr<UMetasoundEditorGraphOutput>& Output) { return !Output; }) > 0;
-			bEditorGraphModified |= OutGraph.Variables.RemoveAllSwap([&VariableIDs](const TObjectPtr<UMetasoundEditorGraphVariable>& Variable)
-			{
-				return !Variable || VariableIDs.Contains(Variable->GetVariableID());
-			}) > 0;
-
 			return bEditorGraphModified;
 		}
 	} // namespace Editor
diff --git a/Engine/Plugins/Runtime/Metasound/Source/MetasoundEditor/Private/MetasoundEditorGraphBuilder.h b/Engine/Plugins/Runtime/Metasound/Source/MetasoundEditor/Private/MetasoundEditorGraphBuilder.h
--- a/Engine/Plugins/Runtime/Metasound/Source/MetasoundEditor/Private/MetasoundEditorGraphBuilder.h
+++ b/Engine/Plugins/Runtime/Metasound/Source/MetasoundEditor/Private/MetasoundEditorGraphBuilder.h
@@ -130,8 +130,8 @@ namespace Metasound
 
 			static Frontend::FNodeHandle AddExternalNodeHandle(UObject& InMetaSound, const FMetasoundFrontendClassName& InClassName);
 
-			// Adds an variable node with the given node handle to the editor graph.
-			static UMetasoundEditorGraphVariableNode* AddVariableNode(UObject& InMetaSound, const Frontend::FConstNodeHandle& InNodeHandle, bool bInSelectNewNode = true);
+			// Adds an variable editor node with the given variable node (ex. mutator, accessor) to the editor graph.
+			static UMetasoundEditorGraphVariableNode* AddVariableNode(UObject& InMetaSound, const FGuid& InNodeID, bool bInSelectNewNode = true);
 
 			// Adds an input node to the editor graph that corresponds to the provided input template node in the document with the given ID.
 			static UMetasoundEditorGraphInputNode* AddInputNode(UObject& InMetaSound, const FGuid& InInputTemplateNodeID, bool bInSelectNewNode = true);
@@ -140,34 +140,12 @@ namespace Metasound
 			static UMetasoundEditorGraphOutputNode* AddOutputNode(UObject& InMetaSound, const FGuid& InNodeID, bool bInSelectNewNode = true);
 
 			// Create a unique name for the variable.
-			static FName GenerateUniqueVariableName(const Frontend::FConstGraphHandle& InFrontendGraph, const FString& InBaseName);
+			static FName GenerateUniqueVariableName(const FMetaSoundFrontendDocumentBuilder& InBuilder, const FString& InBaseName);
 
 			// Convenience method for walking to the outermost object and transforming to a base MetaSound
 			static FMetasoundAssetBase& GetOutermostMetaSoundChecked(UObject& InSubObject);
 			static const FMetasoundAssetBase& GetOutermostConstMetaSoundChecked(const UObject& InSubObject);
 
-			// Adds a frontend variable to the root graph of the MetaSound
-			//
-			// @param InMetaSound - FMetasoundAssetBase derived UObject.
-			// @param InTypeName - Data type of variable.
-			// @param InNameBase - (optional) Name of the Variable, if not given default naming is used.
-			// 
-			// @return The added frontend variable handle. On error, the returned handle is invalid.
-			static Frontend::FVariableHandle AddVariableHandle(UObject& InMetaSound, const FName& InTypeName, const FName* InNameBase = nullptr);
-
-			// Generates Duplicate FVariableHandle from the given handle.
-			static Frontend::FVariableHandle DuplicateVariableHandle(UObject& InMetaSound, const Frontend::FConstVariableHandle& InHandle);
-
-			// Adds a frontend variable node to root graph using the supplied node class name.
-			//
-			// @param InMetaSound - FMetasoundAssetBase derived UObject.
-			// @param InVariableID - ID of variable existing on the root graph.
-			// @param InVariableNodeClassName - FNodeClassName of the variable node to add.
-			// @param InEditorNode - (Optional) Editor node to set as referencing new variable node handle.
-			//
-			// @return The added frontend node handle. On error, the returned handle is invalid.
-			static Frontend::FNodeHandle AddVariableNodeHandle(UObject& InMetaSound, const FGuid& InVariableID, const Metasound::FNodeClassName& InVariableNodeClassName, UMetasoundEditorGraphVariableNode* InEditorNode = nullptr);
-
 			// Attempts to connect Frontend node counterparts together for provided pins.  Returns true if succeeded,
 			// and breaks pin link and returns false if failed.  If bConnectEdPins is set, will attempt to connect
 			// the Editor Graph representation of the pins.
@@ -233,9 +211,6 @@ namespace Metasound
 			// Rebuilds all editor node pins based on the provided node handle's class definition.
 			static void RebuildNodePins(UMetasoundEditorGraphNode& InGraphNode);
 
-			// Deletes all nodes from a parent graph that are linked to a given member.
-			static bool DeleteMemberNodes(const UMetasoundEditorGraphMember& InGraphMember);
-
 			// Deletes both the editor graph & frontend nodes from respective graphs
 			static bool DeleteNode(UEdGraphNode& InNode, bool bRemoveUnusedDependencies = true);
 
diff --git a/Engine/Plugins/Runtime/Metasound/Source/MetasoundEditor/Private/MetasoundEditorGraphNode.cpp b/Engine/Plugins/Runtime/Metasound/Source/MetasoundEditor/Private/MetasoundEditorGraphNode.cpp
--- a/Engine/Plugins/Runtime/Metasound/Source/MetasoundEditor/Private/MetasoundEditorGraphNode.cpp
+++ b/Engine/Plugins/Runtime/Metasound/Source/MetasoundEditor/Private/MetasoundEditorGraphNode.cpp
@@ -25,6 +25,7 @@
 #include "MetasoundEditorSettings.h"
 #include "MetasoundFrontend.h"
 #include "MetasoundFrontendDataTypeRegistry.h"
+#include "MetasoundFrontendDocument.h"
 #include "MetasoundFrontendNodeTemplateRegistry.h"
 #include "MetasoundFrontendRegistries.h"
 #include "MetasoundFrontendSearchEngine.h"
@@ -643,7 +644,11 @@ void UMetasoundEditorGraphNode::OnUpdateCommentText(const FString& NewComment)
 		const FScopedTransaction Transaction(LOCTEXT("CommentCommitted", "Comment Changed"));
 		Modify();
 		NodeComment = NewComment;
+
+PRAGMA_DISABLE_DEPRECATION_WARNINGS
 		FNodeHandle NodeHandle = GetNodeHandle();
+PRAGMA_ENABLE_DEPRECATION_WARNINGS
+
 		FMetasoundFrontendNodeStyle Style = NodeHandle->GetNodeStyle();
 		Style.Display.Comment = NewComment;
 		NodeHandle->SetNodeStyle(Style);
@@ -818,7 +823,7 @@ FString UMetasoundEditorGraphMemberNode::GetFindReferenceSearchString_Impl(EGetF
 	
 		if (!GraphMember->GetDisplayName().IsEmpty())
 		{
-			NameToSearch = GraphMember->GetDisplayName().ToString();		
+			NameToSearch = GraphMember->GetDisplayName().ToString();
 		}
 		else
 		{
@@ -1541,22 +1546,26 @@ void UMetasoundEditorGraphVariableNode::CacheBreadcrumb()
 	using namespace Metasound::Frontend;
 
 	Breadcrumb = { };
+
 	if (Variable)
 	{
-		FConstVariableHandle VariableHandle = Variable->GetConstVariableHandle();
+		if (const FMetasoundFrontendVariable* FrontendVariable = Variable->GetFrontendVariable())
+		{
+			Breadcrumb.MemberName = FrontendVariable->Name;
+			Breadcrumb.DataType = FrontendVariable->TypeName;
 
-		Breadcrumb.MemberName = VariableHandle->GetName();
-		Breadcrumb.DataType = VariableHandle->GetDataType();
-		// Hack to reuse the default literals breadcrumb property for variables, which only have a single (rather than paged) literals
-		TMap<FGuid, FMetasoundFrontendLiteral> DefaultLiterals;
-		DefaultLiterals.Add(DefaultPageID, VariableHandle->GetLiteral());
-		Breadcrumb.DefaultLiterals = DefaultLiterals;
+			// Hack to reuse the default literals breadcrumb property for variables, which only have a single (rather than paged) literals
+			Breadcrumb.DefaultLiterals = TMap<FGuid, FMetasoundFrontendLiteral>
+			{
+				{ DefaultPageID, FrontendVariable->Literal }
+			};
 
-		FMetasoundFrontendVertexMetadata VertexMetadata;
-		VertexMetadata.SetDisplayName(VariableHandle->GetDisplayName());
-		VertexMetadata.SetDescription(VariableHandle->GetDescription());
-		
-		Breadcrumb.VertexMetadata = VertexMetadata;
+			FMetasoundFrontendVertexMetadata VertexMetadata;
+			VertexMetadata.SetDisplayName(FrontendVariable->DisplayName);
+			VertexMetadata.SetDescription(FrontendVariable->Description);
+
+			Breadcrumb.VertexMetadata = VertexMetadata;
+		}
 	}
 }
 
@@ -1625,6 +1634,22 @@ FName UMetasoundEditorGraphVariableNode::GetCornerIcon() const
 	return Super::GetCornerIcon();
 }
 
+void UMetasoundEditorGraphVariableNode::GetPinHoverText(const UEdGraphPin& Pin, FString& OutHoverText) const
+{
+	using namespace Metasound::Editor;
+	using namespace Metasound::Frontend;
+
+	if (Variable)
+	{
+		OutHoverText = Variable->GetBreadcrumb().Description.ToString();
+	}
+
+	if (OutHoverText.IsEmpty())
+	{
+		Super::GetPinHoverText(Pin, OutHoverText);
+	}
+}
+
 void UMetasoundEditorGraphVariableNode::PinDefaultValueChanged(UEdGraphPin* Pin)
 {
 	using namespace Metasound::Editor;
diff --git a/Engine/Plugins/Runtime/Metasound/Source/MetasoundEditor/Private/MetasoundEditorGraphSchema.cpp b/Engine/Plugins/Runtime/Metasound/Source/MetasoundEditor/Private/MetasoundEditorGraphSchema.cpp
--- a/Engine/Plugins/Runtime/Metasound/Source/MetasoundEditor/Private/MetasoundEditorGraphSchema.cpp
+++ b/Engine/Plugins/Runtime/Metasound/Source/MetasoundEditor/Private/MetasoundEditorGraphSchema.cpp
@@ -275,7 +275,8 @@ PRAGMA_ENABLE_DEPRECATION_WARNINGS
 				const FMetasoundFrontendClassName& ClassName = InVariableClass.Metadata.GetClassName();
 
 				UMetasoundEditorGraph* MetaSoundGraph = CastChecked<UMetasoundEditorGraph>(ConnectedNode->GetGraph());
-				const FName NodeName = FGraphBuilder::GenerateUniqueVariableName(MetaSoundGraph->GetGraphHandle(), BaseName.ToString());
+				FMetaSoundFrontendDocumentBuilder& DocBuilder = MetaSoundGraph->GetBuilderChecked().GetBuilder();
+				const FName NodeName = FGraphBuilder::GenerateUniqueVariableName(DocBuilder, BaseName.ToString());
 
 				const FScopedTransaction Transaction(FText::Format(
 					LOCTEXT("PromoteNodeVertexToGraphVariableFormat", "Promote MetaSound Node {0} to {1}"),
@@ -293,49 +294,49 @@ PRAGMA_ENABLE_DEPRECATION_WARNINGS
 					FGraphBuilder::GetPinLiteral(FromPin, DefaultValue);
 				}
 
-				FVariableHandle VariableHandle = FGraphBuilder::AddVariableHandle(ParentMetaSound, DataType);
-				if (ensure(VariableHandle->IsValid()))
+				const FMetasoundFrontendVariable* FrontendVariable = DocBuilder.AddGraphVariable(NodeName, DataType);
+				if (ensure(FrontendVariable))
 				{
-					UMetasoundEditorGraphVariable* Variable = MetaSoundGraph->FindOrAddVariable(VariableHandle);
+					UMetasoundEditorGraphVariable* Variable = MetaSoundGraph->FindOrAddVariable(FrontendVariable->Name);
 					if (ensure(Variable))
 					{
-						constexpr bool bPostSubTransaction = false;
-						Variable->SetMemberName(NodeName, bPostSubTransaction);
-
-						FConstNodeHandle NodeHandle = FGraphBuilder::AddVariableNodeHandle(ParentMetaSound, Variable->GetVariableID(), ClassName.ToNodeClassName());
-
-						if (UMetasoundEditorGraphVariableNode* NewGraphNode = FGraphBuilder::AddVariableNode(ParentMetaSound, NodeHandle))
+						const FMetasoundFrontendNode* NewVariableNode = DocBuilder.AddGraphVariableNode(FrontendVariable->Name, InVariableClass.Metadata.GetType());
+						if (ensure(NewVariableNode))
 						{
-							NewGraphNode->UpdateFrontendNodeLocation(InLocation);
-							NewGraphNode->SyncLocationFromFrontendNode();
-
-							// Set the literal using the new value if connecting to an input
-							if (FromPin.Direction == EGPD_Input)
+							UMetasoundEditorGraphVariableNode* NewGraphNode = FGraphBuilder::AddVariableNode(ParentMetaSound, NewVariableNode->GetID());
+							if (ensure(NewGraphNode))
 							{
-								UMetasoundEditorGraphMemberDefaultLiteral* DefaultLiteral = Variable->GetLiteral();
-								if (ensure(DefaultLiteral))
-								{
-									DefaultLiteral->SetFromLiteral(DefaultValue);
-								}
-
-								// Ensures the setter node value is synced with the editor literal value
-								constexpr bool bPostTransaction = false;
-								Variable->UpdateFrontendDefaultLiteral(bPostTransaction);
-							}
+								NewGraphNode->UpdateFrontendNodeLocation(InLocation);
+								NewGraphNode->SyncLocationFromFrontendNode();
 
-							UEdGraphNode* EdGraphNode = CastChecked<UEdGraphNode>(NewGraphNode);
-							if (ensure(SchemaPrivate::TryConnectNewNodeToMatchingDataTypePin(*EdGraphNode, &FromPin)))
-							{
-								FGraphBuilder::RegisterGraphWithFrontend(ParentMetaSound);
-								if (bSelectNode)
+								// Set the literal using the new value if connecting to an input
+								if (FromPin.Direction == EGPD_Input)
 								{
-									SchemaPrivate::SelectNodeInEditorForRename(*MetaSoundGraph, *NewGraphNode);
+									UMetasoundEditorGraphMemberDefaultLiteral* DefaultLiteral = Variable->GetLiteral();
+									if (ensure(DefaultLiteral))
+									{
+										DefaultLiteral->SetFromLiteral(DefaultValue);
+									}
+
+									// Ensures the setter node value is synced with the editor literal value
+									constexpr bool bPostTransaction = false;
+									Variable->UpdateFrontendDefaultLiteral(bPostTransaction);
 								}
-								else
+
+								UEdGraphNode* EdGraphNode = CastChecked<UEdGraphNode>(NewGraphNode);
+								if (ensure(SchemaPrivate::TryConnectNewNodeToMatchingDataTypePin(*EdGraphNode, &FromPin)))
 								{
-									SchemaPrivate::SelectNodeInEditor(*MetaSoundGraph, *NewGraphNode);
+									FGraphBuilder::RegisterGraphWithFrontend(ParentMetaSound);
+									if (bSelectNode)
+									{
+										SchemaPrivate::SelectNodeInEditorForRename(*MetaSoundGraph, *NewGraphNode);
+									}
+									else
+									{
+										SchemaPrivate::SelectNodeInEditor(*MetaSoundGraph, *NewGraphNode);
+									}
+									return EdGraphNode;
 								}
-								return EdGraphNode;
 							}
 						}
 					}
@@ -474,13 +475,19 @@ PRAGMA_ENABLE_DEPRECATION_WARNINGS
 				InputVertex->TypeName,
 				Builder.GetNodeInputAccessType(InputVertexHandle.NodeID, InputVertex->VertexID)
 			};
-			FMetasoundFrontendClassInput ClassInput = FGraphBuilder::CreateUniqueClassInput(ParentMetasound, VertexParams, DefaultLiterals, &InputVertex->Name);
-			if (const FMetasoundFrontendNode* NewNode = Builder.AddGraphInput(ClassInput))
+
+			const FMetasoundFrontendNode* NewNode = nullptr;
+			{
+				FMetasoundFrontendClassInput ClassInput = FGraphBuilder::CreateUniqueClassInput(ParentMetasound, VertexParams, DefaultLiterals, &InputVertex->Name);
+				NewNode = Builder.AddGraphInput(MoveTemp(ClassInput));
+			}
+
+			if (ensure(NewNode))
 			{
 				UMetasoundEditorGraphInput* Input = MetasoundGraph->FindOrAddInput(NewNode->GetID());
 				if (ensure(Input))
 				{
-					if (const FMetasoundFrontendNode* NewTemplateNode = FInputNodeTemplate::CreateNode(Builder, ClassInput.Name))
+					if (const FMetasoundFrontendNode* NewTemplateNode = FInputNodeTemplate::CreateNode(Builder, InputVertex->Name))
 					{
 						if (UMetasoundEditorGraphNode* NewGraphNode = FGraphBuilder::AddInputNode(ParentMetasound, NewTemplateNode->GetID()))
 						{
@@ -772,7 +779,10 @@ UEdGraphNode* FMetasoundGraphSchemaAction_NewInput::PerformAction(UEdGraph* Pare
 		return nullptr;
 	}
 
+PRAGMA_DISABLE_DEPRECATION_WARNINGS
 	FNodeHandle InputNodeHandle = Input->GetNodeHandle();
+PRAGMA_ENABLE_DEPRECATION_WARNINGS
+
 	if (!ensure(InputNodeHandle->IsValid()))
 	{
 		return nullptr;
@@ -946,10 +956,10 @@ UEdGraphNode* FMetasoundGraphSchemaAction_NewVariableNode::PerformAction(UEdGrap
 				MetasoundGraph->Modify();
 				Variable->Modify();
 				
-				FConstNodeHandle FrontendNode = CreateFrontendVariableNode(MetasoundGraph->GetGraphHandle(), VariableID);
-				if (ensure(FrontendNode->IsValid()))
+				const FMetasoundFrontendNode* FrontendNode = CreateFrontendVariableNode(MetasoundGraph->GetBuilderChecked().GetBuilder());
+				if (ensure(FrontendNode))
 				{
-					if (UMetasoundEditorGraphVariableNode* NewGraphNode = FGraphBuilder::AddVariableNode(*ParentMetasound, FrontendNode, bSelectNewNode))
+					if (UMetasoundEditorGraphVariableNode* NewGraphNode = FGraphBuilder::AddVariableNode(*ParentMetasound, FrontendNode->GetID(), bSelectNewNode))
 					{
 						NewGraphNode->Modify();
 						NewGraphNode->UpdateFrontendNodeLocation(Location);
@@ -985,9 +995,14 @@ FMetasoundGraphSchemaAction_NewVariableAccessorNode::FMetasoundGraphSchemaAction
 {
 }
 
-Metasound::Frontend::FNodeHandle FMetasoundGraphSchemaAction_NewVariableAccessorNode::CreateFrontendVariableNode(const Metasound::Frontend::FGraphHandle& InFrontendGraph, const FGuid& InVariableID) const
+const FMetasoundFrontendNode* FMetasoundGraphSchemaAction_NewVariableAccessorNode::CreateFrontendVariableNode(FMetaSoundFrontendDocumentBuilder& DocBuilder) const
 {
-	return InFrontendGraph->AddVariableAccessorNode(InVariableID);
+	if (const FMetasoundFrontendVariable* Variable = DocBuilder.FindGraphVariable(VariableID))
+	{
+		return DocBuilder.AddGraphVariableAccessorNode(Variable->Name);
+	}
+
+	return nullptr;
 }
 
 FMetasoundGraphSchemaAction_NewVariableDeferredAccessorNode::FMetasoundGraphSchemaAction_NewVariableDeferredAccessorNode(FText InNodeCategory, FText InDisplayName, FGuid InVariableID, FText InToolTip)
@@ -995,9 +1010,14 @@ FMetasoundGraphSchemaAction_NewVariableDeferredAccessorNode::FMetasoundGraphSche
 {
 }
 
-Metasound::Frontend::FNodeHandle FMetasoundGraphSchemaAction_NewVariableDeferredAccessorNode::CreateFrontendVariableNode(const Metasound::Frontend::FGraphHandle& InFrontendGraph, const FGuid& InVariableID) const
+const FMetasoundFrontendNode* FMetasoundGraphSchemaAction_NewVariableDeferredAccessorNode::CreateFrontendVariableNode(FMetaSoundFrontendDocumentBuilder& DocBuilder) const
 {
-	return InFrontendGraph->AddVariableDeferredAccessorNode(InVariableID);
+	if (const FMetasoundFrontendVariable* Variable = DocBuilder.FindGraphVariable(VariableID))
+	{
+		return DocBuilder.AddGraphVariableDeferredAccessorNode(Variable->Name);
+	}
+
+	return nullptr;
 }
 
 FMetasoundGraphSchemaAction_NewVariableMutatorNode::FMetasoundGraphSchemaAction_NewVariableMutatorNode(FText InNodeCategory, FText InDisplayName, FGuid InVariableID, FText InToolTip)
@@ -1005,19 +1025,14 @@ FMetasoundGraphSchemaAction_NewVariableMutatorNode::FMetasoundGraphSchemaAction_
 {
 }
 
-Metasound::Frontend::FNodeHandle FMetasoundGraphSchemaAction_NewVariableMutatorNode::CreateFrontendVariableNode(const Metasound::Frontend::FGraphHandle& InFrontendGraph, const FGuid& InVariableID) const
+const FMetasoundFrontendNode* FMetasoundGraphSchemaAction_NewVariableMutatorNode::CreateFrontendVariableNode(FMetaSoundFrontendDocumentBuilder& DocBuilder) const
 {
-	using namespace Metasound::Frontend;
-
-	// Only one mutator node should exist per variable. Check to
-	// make sure that one does not already exist for this variable.
-	FConstVariableHandle Variable = InFrontendGraph->FindVariable(InVariableID);
-	FConstNodeHandle MutatorNode = Variable->FindMutatorNode();
-	if (ensure(!MutatorNode->IsValid()))
+	if (const FMetasoundFrontendVariable* Variable = DocBuilder.FindGraphVariable(VariableID))
 	{
-		return InFrontendGraph->FindOrAddVariableMutatorNode(InVariableID);
+		return DocBuilder.AddGraphVariableMutatorNode(Variable->Name);
 	}
-	return INodeController::GetInvalidHandle();
+
+	return nullptr;
 }
 
 UEdGraphNode* FMetasoundGraphSchemaAction_NewFromSelected::PerformAction(UEdGraph* ParentGraph, UEdGraphPin* FromPin, const FVector2D Location, bool bSelectNewNode /* = true*/)
@@ -1904,7 +1919,7 @@ FText UMetasoundEditorGraphSchema::GetPinDisplayName(const UEdGraphPin* Pin) con
 				UMetasoundEditorGraphMember* Member = MemberNode->GetMember();
 				if (ensure(Member))
 				{
-					return MemberNode->GetMember()->GetDisplayName();
+					return Member->GetDisplayName();
 				}
 			}
 		}
diff --git a/Engine/Plugins/Runtime/Metasound/Source/MetasoundEditor/Private/MetasoundNodeDetailCustomization.cpp b/Engine/Plugins/Runtime/Metasound/Source/MetasoundEditor/Private/MetasoundNodeDetailCustomization.cpp
--- a/Engine/Plugins/Runtime/Metasound/Source/MetasoundEditor/Private/MetasoundNodeDetailCustomization.cpp
+++ b/Engine/Plugins/Runtime/Metasound/Source/MetasoundEditor/Private/MetasoundNodeDetailCustomization.cpp
@@ -1275,19 +1275,6 @@ PRAGMA_ENABLE_DEPRECATION_WARNINGS
 			return FText::GetEmpty();
 		}
 
-		Frontend::FDocumentHandle FMetasoundMemberDetailCustomization::GetDocumentHandle() const
-		{
-			if (GraphMember.IsValid())
-			{
-				if (UMetasoundEditorGraph* Graph = GraphMember->GetOwningGraph())
-				{
-					return Graph->GetDocumentHandle();
-				}
-			}
-
-			return Frontend::IDocumentController::GetInvalidHandle();
-		}
-
 		bool FMetasoundMemberDetailCustomization::IsGraphEditable() const
 		{
 			if (GraphMember.IsValid())
@@ -1544,11 +1531,13 @@ PRAGMA_ENABLE_DEPRECATION_WARNINGS
 
 		bool FMetasoundInputDetailCustomization::GetInputInheritsDefault() const
 		{
-			if (UMetasoundEditorGraphInput* Input = Cast<UMetasoundEditorGraphInput>(GraphMember.Get()))
+			if (const UMetasoundEditorGraphInput* Input = Cast<UMetasoundEditorGraphInput>(GraphMember.Get()))
 			{
-				const TSet<FName>& InputsInheritingDefault = GetDocumentHandle()->GetRootGraph()->GetInputsInheritingDefault();
-				FName NodeName = Input->GetConstNodeHandle()->GetNodeName();
-				return InputsInheritingDefault.Contains(NodeName);
+				if (const TSet<FName>* InputsInheritingDefault = Input->GetFrontendBuilderChecked().GetGraphInputsInheritingDefault())
+				{
+					FName MemberName = Input->GetMemberName();
+					return InputsInheritingDefault->Contains(MemberName);
+				}
 			}
 
 			return false;
@@ -1567,8 +1556,16 @@ PRAGMA_ENABLE_DEPRECATION_WARNINGS
 					MemberDefaultLiteral->Modify();
 
 					constexpr bool bDefaultIsInherited = true;
-					const FName NodeName = Input->GetConstNodeHandle()->GetNodeName();
-					GetDocumentHandle()->GetRootGraph()->SetInputInheritsDefault(NodeName, bDefaultIsInherited);
+					const FName MemberName = Input->GetMemberName();
+					FMetaSoundFrontendDocumentBuilder& DocumentBuilder = GraphMember->GetFrontendBuilderChecked();
+					DocumentBuilder.SetGraphInputInheritsDefault(MemberName, bDefaultIsInherited);
+
+					Input->UpdateFrontendDefaultLiteral(false /* bPostTransaction */);
+
+					if (UMetasoundEditorGraphMemberDefaultLiteral* Literal = Input->GetLiteral())
+					{
+						Literal->ForceRefresh();
+					}
 
 					if (UObject* Metasound = Input->GetOutermostObject())
 					{
@@ -1591,8 +1588,9 @@ PRAGMA_ENABLE_DEPRECATION_WARNINGS
 					MemberDefaultLiteral->Modify();
 
 					constexpr bool bDefaultIsInherited = false;
-					const FName NodeName = Input->GetConstNodeHandle()->GetNodeName();
-					GetDocumentHandle()->GetRootGraph()->SetInputInheritsDefault(NodeName, bDefaultIsInherited);
+					const FName MemberName = Input->GetMemberName();
+					FMetaSoundFrontendDocumentBuilder& DocumentBuilder = GraphMember->GetFrontendBuilderChecked();
+					DocumentBuilder.SetGraphInputInheritsDefault(MemberName, bDefaultIsInherited);
 
 					Input->UpdateFrontendDefaultLiteral(false /* bPostTransaction */);
 
@@ -1662,13 +1660,13 @@ PRAGMA_ENABLE_DEPRECATION_WARNINGS
 		{
 			if (InGraphVertex.IsValid() && AdvancedPinCheckbox.IsValid())
 			{
-				bool Checked = InNewState == ECheckBoxState::Checked;		
+				const bool Checked = InNewState == ECheckBoxState::Checked;
 				InGraphVertex->SetIsAdvancedDisplay(Checked);
 
 				if (FMetasoundAssetBase* MetasoundAsset = Metasound::IMetasoundUObjectRegistry::Get().GetObjectAsAssetBase(GraphMember->GetOutermostObject()))
 				{
 					MetasoundAsset->GetModifyContext().AddMemberIDsModified({ GraphMember->GetMemberID() });
-				}				
+				}
 			}
 		}
 #endif // WITH_EDITORONLY_DATA
@@ -1690,7 +1688,7 @@ PRAGMA_ENABLE_DEPRECATION_WARNINGS
 			}
 
 			// Build preset row first if graph has managed interface, not default constructed, & not a trigger
-			const bool bIsPreset = GetDocumentHandle()->GetRootGraphClass().PresetOptions.bIsPreset;
+			const bool bIsPreset = GraphMember->GetFrontendBuilderChecked().IsPreset();
 			const bool bIsDefaultConstructed = MemberDefaultLiteral->GetLiteralType() == EMetasoundFrontendLiteralType::None;
 			const bool bIsTriggerDataType = GraphMember->GetDataType() == GetMetasoundDataTypeName<FTrigger>();
 
diff --git a/Engine/Plugins/Runtime/Metasound/Source/MetasoundEditor/Private/MetasoundNodeDetailCustomization.h b/Engine/Plugins/Runtime/Metasound/Source/MetasoundEditor/Private/MetasoundNodeDetailCustomization.h
--- a/Engine/Plugins/Runtime/Metasound/Source/MetasoundEditor/Private/MetasoundNodeDetailCustomization.h
+++ b/Engine/Plugins/Runtime/Metasound/Source/MetasoundEditor/Private/MetasoundNodeDetailCustomization.h
@@ -357,7 +357,6 @@ namespace Metasound
 			// End of IDetailCustomization interface
 
 			void OnNameChanged(const FText& InNewName);
-			Frontend::FDocumentHandle GetDocumentHandle() const;
 			FText GetName() const;
 			bool IsGraphEditable() const;
 			FText GetDisplayName() const;
diff --git a/Engine/Plugins/Runtime/Metasound/Source/MetasoundEditor/Public/MetasoundEditorGraph.h b/Engine/Plugins/Runtime/Metasound/Source/MetasoundEditor/Public/MetasoundEditorGraph.h
--- a/Engine/Plugins/Runtime/Metasound/Source/MetasoundEditor/Public/MetasoundEditorGraph.h
+++ b/Engine/Plugins/Runtime/Metasound/Source/MetasoundEditor/Public/MetasoundEditorGraph.h
@@ -240,7 +240,7 @@ public:
 	/** Whether the displayed default supports paged values or is a single default value (i.e. characterized per page). */
 	virtual bool IsDefaultPaged() const { return false; }
 
-	/** Cache breadcrumb data before copying member to the clipboard*/
+	/** Cache breadcrumb data before copying member to the clipboard */
 	virtual void CacheBreadcrumb() { };
 
 #if WITH_EDITOR
@@ -352,7 +352,7 @@ public:
 	bool SetIsAdvancedDisplay(const bool IsAdvancedDisplay);
 #endif // WITH_EDITORONLY_DATA
 
-	/** Returns the node handle associated with the vertex. */
+	UE_DEPRECATED(5.6, "Node Handles are actively being deprecated, use the MetaSound Frontend Document Builder API")
 	Metasound::Frontend::FNodeHandle GetNodeHandle();
 
 	/** Returns the node handle associated with the vertex. */
@@ -477,10 +477,12 @@ public:
 	/* ~EndUMetasoundEditorGraphMember interface */
 
 	const FGuid& GetVariableID() const;
+	const FMetasoundFrontendVariable* GetFrontendVariable() const;
 
 	/**Returns Breadcrumb data of this Variable*/
 	const FMetasoundEditorGraphVariableBreadcrumb& GetBreadcrumb() const { return Breadcrumb; };
 
+	UE_DEPRECATED(5.6, "")
 	Metasound::Frontend::FVariableHandle GetVariableHandle();
 	Metasound::Frontend::FConstVariableHandle GetConstVariableHandle() const;
 
@@ -513,7 +515,7 @@ private:
 
 	FVariableEditorNodes GetVariableNodes() const;
 	FVariableNodeLocations GetVariableNodeLocations() const;
-	void AddVariableNodes(UObject& InMetasound, Metasound::Frontend::FGraphHandle& InFrontendGraph, const FVariableNodeLocations& InNodeLocs);
+	void AddVariableNodes(UObject& InMetaSound, const FVariableNodeLocations& InNodeLocs);
 };
 
 UCLASS()
@@ -525,9 +527,14 @@ public:
 	UE_DEPRECATED(5.5, "Input node creation is no longer supported publically via the MetasoundEditorGraph.")
 	UMetasoundEditorGraphInputNode* CreateInputNode(Metasound::Frontend::FNodeHandle InNodeHandle, bool bInSelectNewNode);
 
+	UE_DEPRECATED(5.6, "Document Handles are actively being deprecated, use the MetaSound Frontend Document Builder API")
 	Metasound::Frontend::FDocumentHandle GetDocumentHandle();
+
 	Metasound::Frontend::FConstDocumentHandle GetDocumentHandle() const;
+
+	UE_DEPRECATED(5.6, "Document Handles are actively being deprecated, use the MetaSound Frontend Document Builder API")
 	Metasound::Frontend::FGraphHandle GetGraphHandle();
+
 	Metasound::Frontend::FConstGraphHandle GetGraphHandle() const;
 
 	virtual void PreSave(FObjectPreSaveContext InSaveContext) override;
@@ -579,14 +586,21 @@ public:
 	UMetasoundEditorGraphInput* FindInput(FGuid InNodeID) const;
 	UMetasoundEditorGraphInput* FindInput(FName InName) const;
 	UMetasoundEditorGraphInput* FindOrAddInput(const FGuid& InNodeID);
+
+	UE_DEPRECATED(5.6, "ConstNodeHandle is being deprecated, use overload which is provided NodeID instead")
 	UMetasoundEditorGraphInput* FindOrAddInput(Metasound::Frontend::FConstNodeHandle InNodeHandle);
 
 	UMetasoundEditorGraphOutput* FindOutput(FGuid InNodeID) const;
 	UMetasoundEditorGraphOutput* FindOutput(FName InName) const;
 	UMetasoundEditorGraphOutput* FindOrAddOutput(const FGuid& InNodeID);
+
+	UE_DEPRECATED(5.6, "ConstNodeHandle is being deprecated, use overload which is provided NodeID instead")
 	UMetasoundEditorGraphOutput* FindOrAddOutput(Metasound::Frontend::FConstNodeHandle InNodeHandle);
 
 	UMetasoundEditorGraphVariable* FindVariable(const FGuid& InVariableID) const;
+	UMetasoundEditorGraphVariable* FindOrAddVariable(FName VariableName);
+
+	UE_DEPRECATED(5.6, "ConstVariableHandle is being deprecated, use overload which is provided VariableName instead")
 	UMetasoundEditorGraphVariable* FindOrAddVariable(const Metasound::Frontend::FConstVariableHandle& InVariableHandle);
 
 	UMetasoundEditorGraphMember* FindMember(FGuid InMemberID) const;
@@ -625,13 +639,13 @@ private:
 	}
 
 public:
-	UE_DEPRECATED(5.5, "Use the associated Frontend builder and synchronize the graph remove a given member.")
+	UE_DEPRECATED(5.5, "Use the associated Frontend builder and synchronize the graph to remove a given editor member.")
 	bool RemoveMember(UMetasoundEditorGraphMember& InGraphMember) { return false; }
 
-	UE_DEPRECATED(5.5, "Use the associated Frontend builder and synchronize the graph remove all member nodes.")
+	UE_DEPRECATED(5.5, "Use the associated Frontend builder and synchronize the graph to remove all editor member nodes.")
 	bool RemoveMemberNodes(UMetasoundEditorGraphMember& InGraphMember) { return false; }
 
-	UE_DEPRECATED(5.5, "Use RemoveMember instead")
+	UE_DEPRECATED(5.5, "Use the associated Frontend builder and synchronize the graph to remove all editor member nodes")
 	bool RemoveFrontendMember(UMetasoundEditorGraphMember& InGraphMember) { return false; }
 
 	friend class UMetaSoundFactory;
diff --git a/Engine/Plugins/Runtime/Metasound/Source/MetasoundEditor/Public/MetasoundEditorGraphSchema.h b/Engine/Plugins/Runtime/Metasound/Source/MetasoundEditor/Public/MetasoundEditorGraphSchema.h
--- a/Engine/Plugins/Runtime/Metasound/Source/MetasoundEditor/Public/MetasoundEditorGraphSchema.h
+++ b/Engine/Plugins/Runtime/Metasound/Source/MetasoundEditor/Public/MetasoundEditorGraphSchema.h
@@ -232,11 +232,16 @@ struct METASOUNDEDITOR_API FMetasoundGraphSchemaAction_NewVariableNode : public
 	//~ End FMetasoundGraphSchemaAction Interface
 	
 protected:
-	// Derived classes should override this method to create the desired frontend node.
+	UE_DEPRECATED(5.6, "Deprecated in favor of direct creation via Document Builder API")
 	virtual Metasound::Frontend::FNodeHandle CreateFrontendVariableNode(const Metasound::Frontend::FGraphHandle& InFrontendGraph, const FGuid& InVariableID) const
 	{
 		return Metasound::Frontend::INodeController::GetInvalidHandle();
 	}
+
+	virtual const FMetasoundFrontendNode* CreateFrontendVariableNode(FMetaSoundFrontendDocumentBuilder& DocBuilder) const
+	{
+		return nullptr;
+	}
 };
 
 /** Adds a variable node to the graph */
@@ -254,7 +259,7 @@ struct METASOUNDEDITOR_API FMetasoundGraphSchemaAction_NewVariableAccessorNode :
 	virtual ~FMetasoundGraphSchemaAction_NewVariableAccessorNode() = default;
 
 protected:
-	virtual Metasound::Frontend::FNodeHandle CreateFrontendVariableNode(const Metasound::Frontend::FGraphHandle& InFrontendGraph, const FGuid& InVariableID) const override;
+	virtual const FMetasoundFrontendNode* CreateFrontendVariableNode(FMetaSoundFrontendDocumentBuilder& DocBuilder) const override;
 
 };
 
@@ -273,7 +278,7 @@ struct METASOUNDEDITOR_API FMetasoundGraphSchemaAction_NewVariableDeferredAccess
 	FMetasoundGraphSchemaAction_NewVariableDeferredAccessorNode(FText InNodeCategory, FText InDisplayName, FGuid InVariableID, FText InToolTip);
 	
 protected:
-	virtual Metasound::Frontend::FNodeHandle CreateFrontendVariableNode(const Metasound::Frontend::FGraphHandle& InFrontendGraph, const FGuid& InVariableID) const override;
+	virtual const FMetasoundFrontendNode* CreateFrontendVariableNode(FMetaSoundFrontendDocumentBuilder& DocBuilder) const override;
 };
 
 /** Adds a variable node to the graph */
@@ -289,9 +294,9 @@ struct METASOUNDEDITOR_API FMetasoundGraphSchemaAction_NewVariableMutatorNode :
 	virtual ~FMetasoundGraphSchemaAction_NewVariableMutatorNode() = default;
 
 	FMetasoundGraphSchemaAction_NewVariableMutatorNode(FText InNodeCategory, FText InDisplayName, FGuid InVariableID, FText InToolTip);
-	
+
 protected:
-	virtual Metasound::Frontend::FNodeHandle CreateFrontendVariableNode(const Metasound::Frontend::FGraphHandle& InFrontendGraph, const FGuid& InVariableID) const override;
+	virtual const FMetasoundFrontendNode* CreateFrontendVariableNode(FMetaSoundFrontendDocumentBuilder& DocBuilder) const override;
 };
 
 /** Promotes an input to a graph variable & respective getter node, using its respective literal value as the default value */
diff --git a/Engine/Plugins/Runtime/Metasound/Source/MetasoundEngine/Public/MetasoundBuilderBase.cpp b/Engine/Plugins/Runtime/Metasound/Source/MetasoundEngine/Public/MetasoundBuilderBase.cpp
--- a/Engine/Plugins/Runtime/Metasound/Source/MetasoundEngine/Public/MetasoundBuilderBase.cpp
+++ b/Engine/Plugins/Runtime/Metasound/Source/MetasoundEngine/Public/MetasoundBuilderBase.cpp
@@ -88,7 +88,7 @@ FMetaSoundBuilderNodeOutputHandle UMetaSoundBuilderBase::AddGraphInputNode(FName
 
 			Description.InitDefault(MoveTemp(DefaultValue));
 
-			Node = Builder.AddGraphInput(Description);
+			Node = Builder.AddGraphInput(MoveTemp(Description));
 		}
 
 		if (Node)
@@ -133,7 +133,7 @@ FMetaSoundBuilderNodeInputHandle UMetaSoundBuilderBase::AddGraphOutputNode(FName
 			Description.NodeID = IDGenerator.CreateNodeID(Doc);
 			Description.VertexID = IDGenerator.CreateVertexID(Doc);
 			Description.AccessType = bIsConstructorOutput ? EMetasoundFrontendVertexAccessType::Value : EMetasoundFrontendVertexAccessType::Reference;
-			Node = Builder.AddGraphOutput(Description);
+			Node = Builder.AddGraphOutput(MoveTemp(Description));
 		}
 
 		if (Node)
@@ -171,6 +171,102 @@ void UMetaSoundBuilderBase::AddGraphPage(FName PageName, bool bDuplicateLastGrap
 }
 #endif // WITH_EDITORONLY_DATA
 
+void UMetaSoundBuilderBase::AddGraphVariable(FName Name, FName DataType, FMetasoundFrontendLiteral DefaultValue, EMetaSoundBuilderResult& OutResult)
+{
+	using namespace Metasound;
+
+	OutResult = EMetaSoundBuilderResult::Failed;
+	if (Frontend::IDataTypeRegistry::Get().FindDataTypeRegistryEntry(DataType) == nullptr)
+	{
+		UE_LOG(LogMetaSound, Error, TEXT("AddGraphVariable Failed: Attempted creation of variable '%s' with unregistered DataType '%s'"), *Name.ToString(), *DataType.ToString());
+	}
+	else
+	{
+		const FMetasoundFrontendVariable* Variable = Builder.FindGraphVariable(Name);
+		if (Variable)
+		{
+			UE_LOG(LogMetaSound, Warning, TEXT("AddGraphVariable Failed: Variable already exists with name '%s' (existing DataType '%s', requested DataType '%s')"), *Name.ToString(), *Variable->TypeName.ToString(), *DataType.ToString());
+		}
+		else
+		{
+			Variable = Builder.AddGraphVariable(Name, DataType, &DefaultValue);
+			OutResult = Variable ? EMetaSoundBuilderResult::Succeeded : EMetaSoundBuilderResult::Failed;
+		}
+	}
+}
+
+FMetaSoundNodeHandle UMetaSoundBuilderBase::AddGraphVariableGetNode(FName Name, EMetaSoundBuilderResult& OutResult)
+{
+	using namespace Metasound;
+
+	OutResult = EMetaSoundBuilderResult::Failed;
+	FMetaSoundNodeHandle NodeHandle;
+	if (const FMetasoundFrontendVariable* Variable = Builder.FindGraphVariable(Name))
+	{
+		if (const FMetasoundFrontendNode* AccessorNode = Builder.AddGraphVariableAccessorNode(Name))
+		{
+			NodeHandle.NodeID = AccessorNode->GetID();
+			OutResult = EMetaSoundBuilderResult::Succeeded;
+		}
+	}
+	else
+	{
+		UE_LOG(LogMetaSound, Warning, TEXT("AddGraphVariableGetNode Failed: Variable does not exists with name '%s'"), *Name.ToString());
+	}
+
+	return NodeHandle;
+}
+
+FMetaSoundNodeHandle UMetaSoundBuilderBase::AddGraphVariableGetDelayedNode(FName Name, EMetaSoundBuilderResult& OutResult)
+{
+	using namespace Metasound;
+
+	OutResult = EMetaSoundBuilderResult::Failed;
+	FMetaSoundNodeHandle NodeHandle;
+	if (const FMetasoundFrontendVariable* Variable = Builder.FindGraphVariable(Name))
+	{
+		if (const FMetasoundFrontendNode* AccessorNode = Builder.AddGraphVariableDeferredAccessorNode(Name))
+		{
+			NodeHandle.NodeID = AccessorNode->GetID();
+			OutResult = EMetaSoundBuilderResult::Succeeded;
+		}
+	}
+	else
+	{
+		UE_LOG(LogMetaSound, Warning, TEXT("AddGraphVariableGetDelayedNode Failed: Variable does not exists with name '%s'"), *Name.ToString());
+	}
+
+	return NodeHandle;
+}
+
+FMetaSoundNodeHandle UMetaSoundBuilderBase::AddGraphVariableSetNode(FName Name, EMetaSoundBuilderResult& OutResult)
+{
+	using namespace Metasound;
+
+	OutResult = EMetaSoundBuilderResult::Failed;
+	FMetaSoundNodeHandle NodeHandle;
+	if (const FMetasoundFrontendVariable* Variable = Builder.FindGraphVariable(Name))
+	{
+		const bool bAlreadyContainsMutator = Variable->MutatorNodeID.IsValid();
+		if (const FMetasoundFrontendNode* MutatorNode = Builder.AddGraphVariableMutatorNode(Name))
+		{
+			NodeHandle.NodeID = MutatorNode->GetID();
+
+			// Frontend Document builder already logs warning if
+			// attempting to add addition mutator and returns the
+			// existing one, so just handle the OutResult but no
+			// log necessary.
+			OutResult = bAlreadyContainsMutator ? EMetaSoundBuilderResult::Failed : EMetaSoundBuilderResult::Succeeded;
+		}
+	}
+	else
+	{
+		UE_LOG(LogMetaSound, Warning, TEXT("AddGraphVariableSetNode Failed: Variable does not exists with name '%s'"), *Name.ToString());
+	}
+
+	return NodeHandle;
+}
+
 void UMetaSoundBuilderBase::AddInterface(FName InterfaceName, EMetaSoundBuilderResult& OutResult)
 {
 	const bool bInterfaceAdded = Builder.AddInterface(InterfaceName);
@@ -828,6 +924,18 @@ FMetasoundFrontendLiteral UMetaSoundBuilderBase::GetGraphInputDefault(FName Inpu
 	return { };
 }
 
+FMetasoundFrontendLiteral UMetaSoundBuilderBase::GetGraphVariableDefault(FName VariableName, EMetaSoundBuilderResult& OutResult) const
+{
+	if (const FMetasoundFrontendLiteral* Default = Builder.GetGraphVariableDefault(VariableName))
+	{
+		OutResult = EMetaSoundBuilderResult::Succeeded;
+		return *Default;
+	}
+
+	OutResult = EMetaSoundBuilderResult::Failed;
+	return { };
+}
+
 TArray<FName> UMetaSoundBuilderBase::GetGraphInputNames(EMetaSoundBuilderResult& OutResult) const
 {
 	TArray<FName> GraphInputNames;
@@ -1069,6 +1177,22 @@ void UMetaSoundBuilderBase::RemoveGraphPage(FName Name, EMetaSoundBuilderResult&
 }
 #endif // WITH_EDITORONLY_DATA
 
+void UMetaSoundBuilderBase::RemoveGraphVariable(FName Name, EMetaSoundBuilderResult& OutResult)
+{
+	using namespace Metasound;
+
+	if (const FMetasoundFrontendVariable* Variable = Builder.FindGraphVariable(Name))
+	{
+		const bool bVariableRemoved = Builder.RemoveGraphVariable(Name);
+		OutResult = bVariableRemoved ? EMetaSoundBuilderResult::Succeeded : EMetaSoundBuilderResult::Failed;
+	}
+	else
+	{
+		UE_LOG(LogMetaSound, Warning, TEXT("RemoveGraphVariable Failed: Variable not found with name '%s'"), *Name.ToString());
+		OutResult = EMetaSoundBuilderResult::Failed;
+	}
+}
+
 void UMetaSoundBuilderBase::RemoveInterface(FName InterfaceName, EMetaSoundBuilderResult& OutResult)
 {
 	const bool bInterfaceRemoved = Builder.RemoveInterface(InterfaceName);
diff --git a/Engine/Plugins/Runtime/Metasound/Source/MetasoundEngine/Public/MetasoundBuilderBase.h b/Engine/Plugins/Runtime/Metasound/Source/MetasoundEngine/Public/MetasoundBuilderBase.h
--- a/Engine/Plugins/Runtime/Metasound/Source/MetasoundEngine/Public/MetasoundBuilderBase.h
+++ b/Engine/Plugins/Runtime/Metasound/Source/MetasoundEngine/Public/MetasoundBuilderBase.h
@@ -144,6 +144,22 @@ public:
 	void AddGraphPage(FName PageName, bool bDuplicateLastGraph, bool bSetAsBuildGraph, EMetaSoundBuilderResult& OutResult);
 #endif // WITH_EDITORONLY_DATA
 
+	// Adds a graph variable node with the given name, DataType, and sets to default value.
+	UFUNCTION(BlueprintCallable, Category = "Audio|MetaSound|Builder", meta = (ExpandEnumAsExecs = "OutResult"))
+	void AddGraphVariable(FName Name, FName DataType, FMetasoundFrontendLiteral DefaultValue, EMetaSoundBuilderResult& OutResult);
+
+	// Adds a graph variable node with the given name, DataType, and sets to default value.
+	UFUNCTION(BlueprintCallable, Category = "Audio|MetaSound|Builder", meta = (ExpandEnumAsExecs = "OutResult"))
+	UPARAM(DisplayName = "Get Node Handle") FMetaSoundNodeHandle AddGraphVariableGetNode(FName Name, EMetaSoundBuilderResult& OutResult);
+
+	// Adds a graph variable node with the given name, DataType, and sets to default value.
+	UFUNCTION(BlueprintCallable, Category = "Audio|MetaSound|Builder", meta = (ExpandEnumAsExecs = "OutResult"))
+	UPARAM(DisplayName = "Get Delayed Node Handle") FMetaSoundNodeHandle AddGraphVariableGetDelayedNode(FName Name, EMetaSoundBuilderResult& OutResult);
+
+	// Adds a graph variable node with the given name, DataType, and sets to default value.
+	UFUNCTION(BlueprintCallable, Category = "Audio|MetaSound|Builder", meta = (ExpandEnumAsExecs = "OutResult"))
+	UPARAM(DisplayName = "Set Node Handle") FMetaSoundNodeHandle AddGraphVariableSetNode(FName Name, EMetaSoundBuilderResult& OutResult);
+
 	// Adds an interface registered with the given name to the graph, adding associated input and output nodes.
 	UFUNCTION(BlueprintCallable, Category = "Audio|MetaSound|Builder", meta = (ExpandEnumAsExecs = "OutResult"))
 	void AddInterface(FName InterfaceName, EMetaSoundBuilderResult& OutResult);
@@ -297,6 +313,10 @@ public:
 	UFUNCTION(BlueprintCallable, Category = "Audio|MetaSound|Builder", meta = (ExpandEnumAsExecs = "OutResult"))
 	UPARAM(DisplayName = "Literal Value") FMetasoundFrontendLiteral GetGraphInputDefault(FName InputName, EMetaSoundBuilderResult& OutResult) const;
 
+	// Gets the graph variable's default literal value.
+	UFUNCTION(BlueprintCallable, Category = "Audio|MetaSound|Builder", meta = (ExpandEnumAsExecs = "OutResult"))
+	UPARAM(DisplayName = "Literal Value") FMetasoundFrontendLiteral GetGraphVariableDefault(FName VariableName, EMetaSoundBuilderResult& OutResult) const;
+
 	// Returns the MetaSound asset's graph class name (used by the MetaSound Node Class Registry)
 	UFUNCTION(BlueprintCallable, Category = "Audio|MetaSound|Builder|Advanced", meta = (DisplayName = "Get MetaSound Class Name"))
 	UPARAM(DisplayName = "Class Name") FMetasoundFrontendClassName GetRootGraphClassName() const;
@@ -362,11 +382,11 @@ public:
 	void ResetGraphPages(bool bClearDefaultPage);
 #endif // WITH_EDITORONLY_DATA
 
-	// Removes graph input if it exists; sets result to succeeded if it was removed and failed if it was not.
+	// Removes input from all paged graphs if it exists; sets result to succeeded if it was removed and failed if it was not.
 	UFUNCTION(BlueprintCallable, Category = "Audio|MetaSound|Builder", meta = (ExpandEnumAsExecs = "OutResult"))
 	void RemoveGraphInput(FName Name, EMetaSoundBuilderResult& OutResult);
 
-	// Removes graph output if it exists; sets result to succeeded if it was removed and failed if it was not.
+	// Removes output from all paged graphs if it exists; sets result to succeeded if it was removed and failed if it was not.
 	UFUNCTION(BlueprintCallable, Category = "Audio|MetaSound|Builder", meta = (ExpandEnumAsExecs = "OutResult"))
 	void RemoveGraphOutput(FName Name, EMetaSoundBuilderResult& OutResult);
 
@@ -375,13 +395,17 @@ public:
 	void RemoveGraphPage(FName Name, EMetaSoundBuilderResult& OutResult);
 #endif // WITH_EDITORONLY_DATA
 
+	// Removes graph variable from the current build graph if it exists; sets result to succeeded if it was removed and failed if it was not.
+	UFUNCTION(BlueprintCallable, Category = "Audio|MetaSound|Builder", meta = (ExpandEnumAsExecs = "OutResult"))
+	void RemoveGraphVariable(FName Name, EMetaSoundBuilderResult& OutResult);
+
 	// Removes the interface with the given name from the builder's MetaSound. Removes any graph inputs
 	// and outputs associated with the given interface and their respective connections (if any).
 	UFUNCTION(BlueprintCallable, Category = "Audio|MetaSound|Builder", meta = (ExpandEnumAsExecs = "OutResult"))
 	void RemoveInterface(FName InterfaceName, EMetaSoundBuilderResult& OutResult);
 
-	// Removes node and any associated connections from the builder's MetaSound. (Advanced) Optionally, remove unused dependencies
-	// from the internal dependendency list on successful removal of node.
+	// Removes node and any associated connections from the current build graph. (Advanced) Optionally, remove unused dependencies
+	// from the internal dependency list on successful removal of node.
 	UFUNCTION(BlueprintCallable, Category = "Audio|MetaSound|Builder", meta = (ExpandEnumAsExecs = "OutResult", AdvancedDisplay = "2"))
 	void RemoveNode(const FMetaSoundNodeHandle& NodeHandle, EMetaSoundBuilderResult& OutResult, bool bRemoveUnusedDependencies = true);
 
diff --git a/Engine/Plugins/Runtime/Metasound/Source/MetasoundFrontend/MetasoundFrontend.Build.cs b/Engine/Plugins/Runtime/Metasound/Source/MetasoundFrontend/MetasoundFrontend.Build.cs
--- a/Engine/Plugins/Runtime/Metasound/Source/MetasoundFrontend/MetasoundFrontend.Build.cs
+++ b/Engine/Plugins/Runtime/Metasound/Source/MetasoundFrontend/MetasoundFrontend.Build.cs
@@ -9,6 +9,8 @@ namespace UnrealBuildTool.Rules
 	{
 		public MetasoundFrontend(ReadOnlyTargetRules Target) : base(Target)
 		{
+			OptimizeCode = CodeOptimization.Never;
+
 			NumIncludedBytesPerUnityCPPOverride = 294912; // best unity size found from using UBT ProfileUnitySizes mode
 
 			PublicDependencyModuleNames.AddRange
diff --git a/Engine/Plugins/Runtime/Metasound/Source/MetasoundFrontend/Private/MetasoundFrontendDocumentBuilder.cpp b/Engine/Plugins/Runtime/Metasound/Source/MetasoundFrontend/Private/MetasoundFrontendDocumentBuilder.cpp
--- a/Engine/Plugins/Runtime/Metasound/Source/MetasoundFrontend/Private/MetasoundFrontendDocumentBuilder.cpp
+++ b/Engine/Plugins/Runtime/Metasound/Source/MetasoundFrontend/Private/MetasoundFrontendDocumentBuilder.cpp
@@ -794,15 +794,13 @@ FMetaSoundFrontendDocumentBuilder::~FMetaSoundFrontendDocumentBuilder()
 	FinishBuilding();
 }
 
-const FMetasoundFrontendClass* FMetaSoundFrontendDocumentBuilder::AddDependency(const FMetasoundFrontendClass& InClass)
+const FMetasoundFrontendClass* FMetaSoundFrontendDocumentBuilder::AddDependency(FMetasoundFrontendClass NewDependency)
 {
 	using namespace Metasound::Frontend;
 
 	FMetasoundFrontendDocument& Document = GetDocumentChecked();
 	const FMetasoundFrontendClass* Dependency = nullptr;
 
-	FMetasoundFrontendClass NewDependency = InClass;
-
 	// All 'Graph' dependencies are listed as 'External' from the perspective of the owning document.
 	// This makes them implementation agnostic to accommodate nativization of assets.
 	if (NewDependency.Metadata.GetType() == EMetasoundFrontendClassType::Graph)
@@ -1011,9 +1009,9 @@ bool FMetaSoundFrontendDocumentBuilder::AddEdgesFromMatchingInterfaceNodeOutputs
 		const FInterfaceRegistryKey InterfaceKey = GetInterfaceRegistryKey(Version);
 		if (const IInterfaceRegistryEntry* RegistryEntry = IInterfaceRegistry::Get().FindInterfaceRegistryEntry(InterfaceKey))
 		{
-			Algo::Transform(RegistryEntry->GetInterface().Outputs, EdgesToMake, [this, &NodeCache, &InterfaceCache, InNodeID](const FMetasoundFrontendClassOutput& Output)
+			Algo::Transform(RegistryEntry->GetInterface().Outputs, EdgesToMake, [this, &NodeCache, &InterfaceCache, &PageID, InNodeID](const FMetasoundFrontendClassOutput& Output)
 			{
-				const FMetasoundFrontendGraph& Graph = FindConstBuildGraphChecked();
+				const FMetasoundFrontendGraph& Graph = GetDocumentChecked().RootGraph.FindConstGraphChecked(PageID);
 				const FMetasoundFrontendVertex* NodeVertex = NodeCache.FindOutputVertex(InNodeID, Output.Name);
 				check(NodeVertex);
 				const FMetasoundFrontendClassOutput* OutputClass = InterfaceCache.FindOutput(Output.Name);
@@ -1076,204 +1074,575 @@ bool FMetaSoundFrontendDocumentBuilder::AddEdgesFromMatchingInterfaceNodeInputsT
 	return AddNamedEdges(EdgesToMake, &OutEdgesCreated, bReplaceExistingConnections, &PageID);
 }
 
-const FMetasoundFrontendNode* FMetaSoundFrontendDocumentBuilder::AddGraphInput(const FMetasoundFrontendClassInput& InClassInput, const FGuid* InPageID)
+const FMetasoundFrontendNode* FMetaSoundFrontendDocumentBuilder::AddGraphInput(FMetasoundFrontendClassInput ClassInput, const FGuid* InPageID)
 {
 	using namespace Metasound::Frontend;
 
-	checkf(InClassInput.NodeID.IsValid(), TEXT("Unassigned NodeID when adding graph input"));
-	checkf(InClassInput.VertexID.IsValid(), TEXT("Unassigned VertexID when adding graph input"));
+	checkf(ClassInput.NodeID.IsValid(), TEXT("Unassigned NodeID when adding graph input"));
+	checkf(ClassInput.VertexID.IsValid(), TEXT("Unassigned VertexID when adding graph input"));
 
 	const FGuid& PageID = InPageID ? *InPageID : BuildPageID;
-	if (InClassInput.TypeName.IsNone())
+	if (ClassInput.TypeName.IsNone())
 	{
-		UE_LOG(LogMetaSound, Error, TEXT("TypeName unset when attempting to add class input '%s'"), *InClassInput.Name.ToString());
+		UE_LOG(LogMetaSound, Error, TEXT("TypeName unset when attempting to add class input '%s'"), *ClassInput.Name.ToString());
 		return nullptr;
 	}
-	else if (const FMetasoundFrontendClassInput* Input = DocumentCache->GetInterfaceCache().FindInput(InClassInput.Name))
+	else if (const FMetasoundFrontendClassInput* Input = DocumentCache->GetInterfaceCache().FindInput(ClassInput.Name))
 	{
-		UE_LOG(LogMetaSound, Error, TEXT("Attempting to add MetaSound graph input '%s' when input with name already exists"), *InClassInput.Name.ToString());
+		UE_LOG(LogMetaSound, Error, TEXT("Attempting to add MetaSound graph input '%s' when input with name already exists"), *ClassInput.Name.ToString());
 		const FMetasoundFrontendNode* OutputNode = DocumentCache->GetNodeCache(PageID).FindNode(Input->NodeID);
 		check(OutputNode);
 		return OutputNode;
 	}
-	else if (!IDataTypeRegistry::Get().IsRegistered(InClassInput.TypeName))
+	else if (!IDataTypeRegistry::Get().IsRegistered(ClassInput.TypeName))
 	{
-		UE_LOG(LogMetaSound, Error, TEXT("Cannot add MetaSound graph input '%s' with unregistered TypeName '%s'"), *InClassInput.Name.ToString(), *InClassInput.TypeName.ToString());
+		UE_LOG(LogMetaSound, Error, TEXT("Cannot add MetaSound graph input '%s' with unregistered TypeName '%s'"), *ClassInput.Name.ToString(), *ClassInput.TypeName.ToString());
 		return nullptr;
 	}
 
-	FMetasoundFrontendClass Class;
-	if (DocumentBuilderPrivate::FindInputRegistryClass(InClassInput.TypeName, InClassInput.AccessType, Class))
+	FNodeRegistryKey ClassKey;
 	{
-		if(!FindDependency(Class.Metadata))
+		FMetasoundFrontendClass Class;
+		if (!DocumentBuilderPrivate::FindInputRegistryClass(ClassInput.TypeName, ClassInput.AccessType, Class))
 		{
-			AddDependency(Class);
+			return nullptr;
 		}
 
-		FMetasoundFrontendDocument& Document = GetDocumentChecked();
-		FMetasoundFrontendGraphClass& RootGraph = Document.RootGraph;
-
-		auto FinalizeNode = [this, &InClassInput](FMetasoundFrontendNode& InOutNode, const Metasound::Frontend::FNodeRegistryKey&)
+		ClassKey = FNodeRegistryKey(Class.Metadata);
+		if (!FindDependency(Class.Metadata))
 		{
-			// Sets the name of the node an vertices on the node to match the class vertex name
-			DocumentBuilderPrivate::SetNodeAndVertexNames(InOutNode, InClassInput);
+			AddDependency(MoveTemp(Class));
+		}
+	}
 
-			// Set the default literal on the nodes inputs so that it gets passed to the instantiated TInputNode on a live
-			// auditioned MetaSound.
-			DocumentBuilderPrivate::SetDefaultLiteralOnInputNode(InOutNode, InClassInput);
-		};
+	FMetasoundFrontendDocument& Document = GetDocumentChecked();
+	FMetasoundFrontendGraphClass& RootGraph = Document.RootGraph;
+
+	auto FinalizeNode = [this, &ClassInput](FMetasoundFrontendNode& InOutNode, const Metasound::Frontend::FNodeRegistryKey&)
+	{
+		// Sets the name of the node an vertices on the node to match the class vertex name
+		DocumentBuilderPrivate::SetNodeAndVertexNames(InOutNode, ClassInput);
+
+		// Set the default literal on the nodes inputs so that it gets passed to the instantiated TInputNode on a live
+		// auditioned MetaSound.
+		DocumentBuilderPrivate::SetDefaultLiteralOnInputNode(InOutNode, ClassInput);
+	};
 
 #if WITH_EDITORONLY_DATA
-		bool bIsRequired = false;
-		FMetasoundFrontendInterface Interface;
-		if (DocumentBuilderPrivate::IsInterfaceInput(InClassInput.Name, InClassInput.TypeName, &Interface))
+	bool bIsRequired = false;
+	FMetasoundFrontendInterface Interface;
+	if (DocumentBuilderPrivate::IsInterfaceInput(ClassInput.Name, ClassInput.TypeName, &Interface))
+	{
+		if (Document.Interfaces.Contains(Interface.Version))
 		{
-			if (Document.Interfaces.Contains(Interface.Version))
-			{
-				FText RequiredText;
-				bIsRequired = Interface.IsMemberInputRequired(InClassInput.Name, RequiredText);
-			}
+			FText RequiredText;
+			bIsRequired = Interface.IsMemberInputRequired(ClassInput.Name, RequiredText);
 		}
+	}
 #endif // WITH_EDITORONLY_DATA
 
-		// Must add input node to all paged graphs to maintain API parity for all page implementations
-		FMetasoundFrontendNode* NewNode = nullptr;
-		RootGraph.IterateGraphPages([&](const FMetasoundFrontendGraph& Graph)
+	// Must add input node to all paged graphs to maintain API parity for all page implementations
+	FMetasoundFrontendNode* NewNode = nullptr;
+	RootGraph.IterateGraphPages([&](const FMetasoundFrontendGraph& Graph)
+	{
+		constexpr int32* NewNodeIndex = nullptr;
+		FMetasoundFrontendNode* NewPageNode = AddNodeInternal(ClassKey, FinalizeNode, Graph.PageID, ClassInput.NodeID, NewNodeIndex);
+		if (Graph.PageID == PageID)
 		{
-			constexpr int32* NewNodeIndex = nullptr;
-			FMetasoundFrontendNode* NewPageNode = AddNodeInternal(Class.Metadata, FinalizeNode, Graph.PageID, InClassInput.NodeID, NewNodeIndex);
-			if (Graph.PageID == PageID)
-			{
-				NewNode = NewPageNode;
-			}
+			NewNode = NewPageNode;
+		}
 
 #if WITH_EDITORONLY_DATA
-			if (bIsRequired)
-			{
-				// LocationGuid corresponds with the assigned editor graph node guid when dynamically created.
-				// This is added if this is an interface member that is required to force page to create visual
-				// representation that can inform the user of its required state.
-				FGuid LocationGuid = FDocumentIDGenerator::Get().CreateVertexID(Document);
-				SetNodeLocation(InClassInput.NodeID, FVector2D::ZeroVector, &LocationGuid, &Graph.PageID);
-			}
+		if (bIsRequired)
+		{
+			// LocationGuid corresponds with the assigned editor graph node guid when dynamically created.
+			// This is added if this is an interface member that is required to force page to create visual
+			// representation that can inform the user of its required state.
+			FGuid LocationGuid = FDocumentIDGenerator::Get().CreateVertexID(Document);
+			SetNodeLocation(ClassInput.NodeID, FVector2D::ZeroVector, &LocationGuid, &Graph.PageID);
+		}
 #endif // WITH_EDITORONLY_DATA
 
-			// Remove the default literal on the node added during the "FinalizeNode" call. This matches how 
-			// nodes are serialized in editor. The default literals are only stored on the FMetasoundFrontendClassInputs.
-			NewPageNode->InputLiterals.Reset();
-		});
+		// Remove the default literal on the node added during the "FinalizeNode" call. This matches how 
+		// nodes are serialized in editor. The default literals are only stored on the FMetasoundFrontendClassInputs.
+		NewPageNode->InputLiterals.Reset();
+	});
 
-		if (NewNode)
+	if (NewNode)
+	{
+		const int32 NewIndex = RootGraph.Interface.Inputs.Num();
+		FMetasoundFrontendClassInput& NewInput = RootGraph.Interface.Inputs.Add_GetRef(MoveTemp(ClassInput));
+		if (!NewInput.VertexID.IsValid())
 		{
-			const int32 NewIndex = RootGraph.Interface.Inputs.Num();
-			FMetasoundFrontendClassInput& NewInput = RootGraph.Interface.Inputs.Add_GetRef(InClassInput);
-			if (!NewInput.VertexID.IsValid())
-			{
-				NewInput.VertexID = FDocumentIDGenerator::Get().CreateVertexID(Document);
-			}
+			NewInput.VertexID = FDocumentIDGenerator::Get().CreateVertexID(Document);
+		}
 
-			DocumentDelegates->InterfaceDelegates.OnInputAdded.Broadcast(NewIndex);
+		DocumentDelegates->InterfaceDelegates.OnInputAdded.Broadcast(NewIndex);
 #if WITH_EDITORONLY_DATA
-			Document.Metadata.ModifyContext.AddMemberIDModified(InClassInput.NodeID);
+		Document.Metadata.ModifyContext.AddMemberIDModified(NewInput.NodeID);
 #endif // WITH_EDITORONLY_DATA
 
-			return NewNode;
-		}
+		return NewNode;
 	}
 
 	return nullptr;
 }
 
-const FMetasoundFrontendNode* FMetaSoundFrontendDocumentBuilder::AddGraphOutput(const FMetasoundFrontendClassOutput& InClassOutput, const FGuid* InPageID)
+const FMetasoundFrontendNode* FMetaSoundFrontendDocumentBuilder::AddGraphOutput(FMetasoundFrontendClassOutput ClassOutput, const FGuid* InPageID)
 {
 	using namespace Metasound::Frontend;
 
-	checkf(InClassOutput.NodeID.IsValid(), TEXT("Unassigned NodeID when adding graph output"));
-	checkf(InClassOutput.VertexID.IsValid(), TEXT("Unassigned VertexID when adding graph output"));
+	checkf(ClassOutput.NodeID.IsValid(), TEXT("Unassigned NodeID when adding graph output"));
+	checkf(ClassOutput.VertexID.IsValid(), TEXT("Unassigned VertexID when adding graph output"));
 
 	const FGuid& PageID = InPageID ? *InPageID : BuildPageID;
-	if (InClassOutput.TypeName.IsNone())
+	if (ClassOutput.TypeName.IsNone())
 	{
-		UE_LOG(LogMetaSound, Error, TEXT("TypeName unset when attempting to add class output '%s'"), *InClassOutput.Name.ToString());
+		UE_LOG(LogMetaSound, Error, TEXT("TypeName unset when attempting to add class output '%s'"), *ClassOutput.Name.ToString());
 		return nullptr;
 	}
-	else if (const FMetasoundFrontendClassOutput* Output = DocumentCache->GetInterfaceCache().FindOutput(InClassOutput.Name))
+	else if (const FMetasoundFrontendClassOutput* Output = DocumentCache->GetInterfaceCache().FindOutput(ClassOutput.Name))
 	{
-		UE_LOG(LogMetaSound, Error, TEXT("Attempting to add MetaSound graph output '%s' when output with name already exists"), *InClassOutput.Name.ToString());
+		UE_LOG(LogMetaSound, Error, TEXT("Attempting to add MetaSound graph output '%s' when output with name already exists"), *ClassOutput.Name.ToString());
 		return DocumentCache->GetNodeCache(PageID).FindNode(Output->NodeID);
 	}
-	else if (!IDataTypeRegistry::Get().IsRegistered(InClassOutput.TypeName))
+	else if (!IDataTypeRegistry::Get().IsRegistered(ClassOutput.TypeName))
 	{
-		UE_LOG(LogMetaSound, Error, TEXT("Cannot add MetaSound graph output '%s' with unregistered TypeName '%s'"), *InClassOutput.Name.ToString(), *InClassOutput.TypeName.ToString());
+		UE_LOG(LogMetaSound, Error, TEXT("Cannot add MetaSound graph output '%s' with unregistered TypeName '%s'"), *ClassOutput.Name.ToString(), *ClassOutput.TypeName.ToString());
 		return nullptr;
 	}
 
-	FMetasoundFrontendClass Class;
-	if (DocumentBuilderPrivate::FindOutputRegistryClass(InClassOutput.TypeName, InClassOutput.AccessType, Class))
+	FNodeRegistryKey ClassKey;
 	{
+		FMetasoundFrontendClass Class;
+		if (!DocumentBuilderPrivate::FindOutputRegistryClass(ClassOutput.TypeName, ClassOutput.AccessType, Class))
+		{
+			return nullptr;
+		}
+
+		ClassKey = FNodeRegistryKey(Class.Metadata);
 		if (!FindDependency(Class.Metadata))
 		{
-			AddDependency(Class);
+			AddDependency(MoveTemp(Class));
+		}
+	}
+
+	auto FinalizeNode = [&ClassOutput](FMetasoundFrontendNode& InOutNode, const Metasound::Frontend::FNodeRegistryKey&)
+	{
+		DocumentBuilderPrivate::SetNodeAndVertexNames(InOutNode, ClassOutput);
+	};
+
+#if WITH_EDITORONLY_DATA
+	bool bIsRequired = false;
+	FMetasoundFrontendInterface Interface;
+	if (DocumentBuilderPrivate::IsInterfaceOutput(ClassOutput.Name, ClassOutput.TypeName, &Interface))
+	{
+		FText RequiredText;
+		bIsRequired = Interface.IsMemberOutputRequired(ClassOutput.Name, RequiredText);
+	}
+#endif // WITH_EDITORONLY_DATA
+
+	bool bAddedNodes = true;
+	FMetasoundFrontendNode* NewNodeToReturn = nullptr;
+	FMetasoundFrontendDocument& Document = GetDocumentChecked();
+	FMetasoundFrontendGraphClass& RootGraph = Document.RootGraph;
+	Document.RootGraph.IterateGraphPages([&](FMetasoundFrontendGraph& Graph)
+	{
+		FMetasoundFrontendNode* NewNode = AddNodeInternal(ClassKey, FinalizeNode, Graph.PageID, ClassOutput.NodeID);
+		if (Graph.PageID == PageID)
+		{
+			NewNodeToReturn = NewNode;
+		}
+
+#if WITH_EDITORONLY_DATA
+		if (bIsRequired)
+		{
+			// LocationGuid corresponds with the assigned editor graph node guid when dynamically created.
+			// This is added if this is an interface member that is required to force page to create visual
+			// representation that can inform the user of its required state.
+			FGuid LocationGuid = FDocumentIDGenerator::Get().CreateVertexID(Document);
+			SetNodeLocation(ClassOutput.NodeID, FVector2D::ZeroVector, &LocationGuid, &Graph.PageID);
 		}
+#endif // WITH_EDITORONLY_DATA
+
+		bAddedNodes &= NewNode != nullptr;
+	});
+
+	if (bAddedNodes)
+	{
+		const int32 NewIndex = RootGraph.Interface.Outputs.Num();
+		FMetasoundFrontendClassOutput& NewOutput = RootGraph.Interface.Outputs.Add_GetRef(MoveTemp(ClassOutput));
+		if (!NewOutput.VertexID.IsValid())
+		{
+			NewOutput.VertexID = FDocumentIDGenerator::Get().CreateVertexID(Document);
+		}
+
+		DocumentDelegates->InterfaceDelegates.OnOutputAdded.Broadcast(NewIndex);
+#if WITH_EDITORONLY_DATA
+		Document.Metadata.ModifyContext.AddMemberIDModified(NewOutput.NodeID);
+#endif // WITH_EDITORONLY_DATA
+	}
+
+	check(NewNodeToReturn);
+	return NewNodeToReturn;
+}
+
+const FMetasoundFrontendVariable* FMetaSoundFrontendDocumentBuilder::AddGraphVariable(FName VariableName, FName DataType, const FMetasoundFrontendLiteral* Literal, const FText* DisplayName, const FText* Description, const FGuid* InPageID)
+{
+	using namespace Metasound::Frontend;
 
-		auto FinalizeNode = [&InClassOutput](FMetasoundFrontendNode& InOutNode, const Metasound::Frontend::FNodeRegistryKey&)
+	checkf(FindGraphVariable(VariableName, InPageID) == nullptr, TEXT("Attempted to add variable '%s' that already exists"), *VariableName.ToString());
+
+	const IDataTypeRegistry& Registry = IDataTypeRegistry::Get();
+	FDataTypeRegistryInfo Info;
+	if (ensure(Registry.GetDataTypeInfo(DataType, Info)))
+	{
+		FMetasoundFrontendVariable Variable
 		{
-			DocumentBuilderPrivate::SetNodeAndVertexNames(InOutNode, InClassOutput);
+			.Name = VariableName,
+			.ID = FGuid::NewGuid()
 		};
 
+		Variable.TypeName = Info.DataTypeName;
+		if (Literal)
+		{
+			Variable.Literal = *Literal;
+		}
+		else
+		{
+			Variable.Literal.SetFromLiteral(Registry.CreateDefaultLiteral(DataType));
+		}
+
 #if WITH_EDITORONLY_DATA
-		bool bIsRequired = false;
-		FMetasoundFrontendInterface Interface;
-		if (DocumentBuilderPrivate::IsInterfaceOutput(InClassOutput.Name, InClassOutput.TypeName, &Interface))
+		if (DisplayName)
 		{
-			FText RequiredText;
-			bIsRequired = Interface.IsMemberOutputRequired(InClassOutput.Name, RequiredText);
+			Variable.DisplayName = *DisplayName;
+		}
+
+		if (Description)
+		{
+			Variable.Description = *Description;
 		}
 #endif // WITH_EDITORONLY_DATA
 
-		bool bAddedNodes = true;
-		FMetasoundFrontendNode* NewNodeToReturn = nullptr;
-		FMetasoundFrontendDocument& Document = GetDocumentChecked();
-		FMetasoundFrontendGraphClass& RootGraph = Document.RootGraph;
-		Document.RootGraph.IterateGraphPages([&](FMetasoundFrontendGraph& Graph)
+#if WITH_EDITOR
+		GetDocumentChecked().Metadata.ModifyContext.AddMemberIDModified(Variable.ID);
+#endif // WITH_EDITOR
+
+		FNodeRegistryKey VariableNodeClassKey;
 		{
-			FMetasoundFrontendNode* NewNode = AddNodeInternal(Class.Metadata, FinalizeNode, Graph.PageID, InClassOutput.NodeID);
-			if (Graph.PageID == PageID)
+			FMetasoundFrontendClass VariableNodeClass;
+			if (!IDataTypeRegistry::Get().GetFrontendVariableClass(Variable.TypeName, VariableNodeClass))
 			{
-				NewNodeToReturn = NewNode;
+				return nullptr;
 			}
 
-#if WITH_EDITORONLY_DATA
-			if (bIsRequired)
+			VariableNodeClassKey = FNodeRegistryKey(VariableNodeClass.Metadata);
+			const FMetasoundFrontendClass* Dependency = FindDependency(VariableNodeClass.Metadata);
+			if (!Dependency)
 			{
-				// LocationGuid corresponds with the assigned editor graph node guid when dynamically created.
-				// This is added if this is an interface member that is required to force page to create visual
-				// representation that can inform the user of its required state.
-				FGuid LocationGuid = FDocumentIDGenerator::Get().CreateVertexID(Document);
-				SetNodeLocation(InClassOutput.NodeID, FVector2D::ZeroVector, &LocationGuid, &Graph.PageID);
+				Dependency = AddDependency(MoveTemp(VariableNodeClass));
 			}
-#endif // WITH_EDITORONLY_DATA
+			check(Dependency);
+		}
 
-			bAddedNodes &= NewNode != nullptr;
-		});
+		auto FinalizeNode = [](FMetasoundFrontendNode& InOutNode, const Metasound::Frontend::FNodeRegistryKey& ClassKey)
+		{
+#if WITH_EDITOR
+			using namespace Metasound::Frontend;
+
+			// Cache the asset name on the node if it node is reference to asset-defined graph.
+			const FTopLevelAssetPath Path = IMetaSoundAssetManager::GetChecked().FindAssetPath(FAssetKey(ClassKey.ClassName, ClassKey.Version));
+			if (Path.IsValid())
+			{
+				InOutNode.Name = Path.GetAssetName();
+				return;
+			}
+
+			InOutNode.Name = ClassKey.ClassName.GetFullName();
+#endif // WITH_EDITOR
+		};
+
+		const FGuid& PageID = InPageID ? *InPageID : BuildPageID;
+		if (FMetasoundFrontendNode* VariableNode = AddNodeInternal(VariableNodeClassKey, FinalizeNode, PageID))
+		{
+			Variable.VariableNodeID = VariableNode->GetID();
+			FMetasoundFrontendGraph& Graph = GetDocumentChecked().RootGraph.FindGraphChecked(PageID);
+			return &Graph.Variables.Add_GetRef(MoveTemp(Variable));
+		}
+	}
+
+	return nullptr;
+}
 
-		if (bAddedNodes)
+const FMetasoundFrontendNode* FMetaSoundFrontendDocumentBuilder::AddGraphVariableNode(FName VariableName, EMetasoundFrontendClassType ClassType, FGuid InNodeID,const FGuid* InPageID)
+{
+	using namespace Metasound::Frontend;
+
+	switch (ClassType)
+	{
+		case EMetasoundFrontendClassType::VariableDeferredAccessor:
+			return AddGraphVariableDeferredAccessorNode(VariableName, InNodeID, InPageID);
+
+		case EMetasoundFrontendClassType::VariableAccessor:
+			return AddGraphVariableAccessorNode(VariableName, InNodeID, InPageID);
+
+		case EMetasoundFrontendClassType::VariableMutator:
 		{
-			const int32 NewIndex = RootGraph.Interface.Outputs.Num();
-			FMetasoundFrontendClassOutput& NewOutput = RootGraph.Interface.Outputs.Add_GetRef(InClassOutput);
-			if (!NewOutput.VertexID.IsValid())
+			const FMetasoundFrontendVariable* Variable = FindGraphVariable(VariableName, InPageID);
+			if (Variable && Variable->MutatorNodeID.IsValid())
 			{
-				NewOutput.VertexID = FDocumentIDGenerator::Get().CreateVertexID(Document);
+				UE_LOG(LogMetaSound, Error, TEXT("Cannot add node: variable mutator already exists for variable \"%s\""), *Variable->Name.ToString());
 			}
+			else
+			{
+				return AddGraphVariableMutatorNode(VariableName, InNodeID, InPageID);
+			}
+		}
+		break;
 
-			DocumentDelegates->InterfaceDelegates.OnOutputAdded.Broadcast(NewIndex);
-#if WITH_EDITORONLY_DATA
-			Document.Metadata.ModifyContext.AddMemberIDModified(InClassOutput.NodeID);
-#endif // WITH_EDITORONLY_DATA
+		default:
+		{
+			checkNoEntry();
+		}
+	}
+
+	return nullptr;
+}
+
+const FMetasoundFrontendNode* FMetaSoundFrontendDocumentBuilder::AddGraphVariableAccessorNode(FName VariableName, FGuid InNodeID,const FGuid* InPageID)
+{
+	using namespace Metasound::Frontend;
+	using namespace Metasound::VariableNames;
+
+	if (FMetasoundFrontendVariable* Variable = FindGraphVariableInternal(VariableName, InPageID))
+	{
+		FNodeRegistryKey VariableNodeClassKey;
+		{
+			FMetasoundFrontendClass NodeClass;
+			if (!IDataTypeRegistry::Get().GetFrontendVariableAccessorClass(Variable->TypeName, NodeClass))
+			{
+				UE_LOG(LogMetaSound, Error, TEXT("Could not find registered \"get variable\" node class for data type \"%s\""), *Variable->TypeName.ToString());
+				return nullptr;
+			}
+
+			VariableNodeClassKey = FNodeRegistryKey(NodeClass.Metadata);
+			const FMetasoundFrontendClass* Dependency = FindDependency(NodeClass.Metadata);
+			if (!Dependency)
+			{
+				Dependency = AddDependency(MoveTemp(NodeClass));
+			}
+			check(Dependency);
+		}
+
+		auto FinalizeNodeFunction = [](const FMetasoundFrontendNode&, const Metasound::Frontend::FNodeRegistryKey&) { };
+		const FGuid& PageID = InPageID ? *InPageID : BuildPageID;
+		if (const FMetasoundFrontendNode* NewNode = AddNodeInternal(VariableNodeClassKey, FinalizeNodeFunction, PageID, InNodeID))
+		{
+			// Connect new node.
+			const FMetasoundFrontendVertex* NewInput = FindNodeInput(NewNode->GetID(), METASOUND_GET_PARAM_NAME(InputVariable), InPageID);
+			check(NewInput);
+
+			const FMetasoundFrontendNode* TailNode = FindTailNodeInVariableStack(VariableName, InPageID);
+			if (!TailNode)
+			{
+				// variable stack is empty. Connect to init variable node.
+				TailNode = FindNode(Variable->VariableNodeID, InPageID);
+			}
+
+			if (ensure(TailNode))
+			{
+				// connect new node to the last "get" node.
+				const FMetasoundFrontendVertex* TailNodeOutput = FindNodeOutput(TailNode->GetID(), METASOUND_GET_PARAM_NAME(OutputVariable), InPageID);
+				check(TailNodeOutput);
+
+				FMetasoundFrontendEdge NewEdge;
+				NewEdge.FromNodeID = TailNode->GetID();
+				NewEdge.FromVertexID = TailNodeOutput->VertexID;
+				NewEdge.ToNodeID = NewNode->GetID();
+				NewEdge.ToVertexID = NewInput->VertexID;
+				AddEdge(MoveTemp(NewEdge), InPageID);
+			}
+
+			// Add node ID to variable after connecting since the array
+			// order of node ids is used to determine whether a node 
+			// is the tail node.
+			Variable->AccessorNodeIDs.Add(NewNode->GetID());
+			return NewNode;
+		}
+	}
+
+	return nullptr;
+}
+
+const FMetasoundFrontendNode* FMetaSoundFrontendDocumentBuilder::AddGraphVariableDeferredAccessorNode(FName VariableName, FGuid InNodeID,const FGuid* InPageID)
+{
+	using namespace Metasound::Frontend;
+	using namespace Metasound::VariableNames;
+
+	if (FMetasoundFrontendVariable* Variable = FindGraphVariableInternal(VariableName, InPageID))
+	{
+		FNodeRegistryKey ClassKey;
+		{
+			FMetasoundFrontendClass NodeClass;
+			if (!IDataTypeRegistry::Get().GetFrontendVariableDeferredAccessorClass(Variable->TypeName, NodeClass))
+			{
+				UE_LOG(LogMetaSound, Error, TEXT("Could not find registered \"get variable\" node class for data type \"%s\""), *Variable->TypeName.ToString());
+				return nullptr;
+			}
+
+			ClassKey = FNodeRegistryKey(NodeClass.Metadata);
+			const FMetasoundFrontendClass* Dependency = FindDependency(NodeClass.Metadata);
+			if (!Dependency)
+			{
+				Dependency = AddDependency(MoveTemp(NodeClass));
+			}
+			check(Dependency);
+		}
+
+		auto FinalizeNodeFunction = [](const FMetasoundFrontendNode&, const Metasound::Frontend::FNodeRegistryKey&) { };
+		const FGuid& PageID = InPageID ? *InPageID : BuildPageID;
+		if (const FMetasoundFrontendNode* NewNode = AddNodeInternal(ClassKey, FinalizeNodeFunction, PageID, InNodeID))
+		{
+			// Connect new node.
+			const FMetasoundFrontendVertex* NewNodeOutput = FindNodeOutput(NewNode->GetID(), METASOUND_GET_PARAM_NAME(OutputVariable), InPageID);
+			const FMetasoundFrontendNode* HeadNode = FindHeadNodeInVariableStack(VariableName, InPageID);
+			if (HeadNode)
+			{
+				const FMetasoundFrontendVertex* HeadNodeInput = FindNodeInput(HeadNode->GetID(), METASOUND_GET_PARAM_NAME(InputVariable), InPageID);
+				check(HeadNodeInput);
+
+				RemoveEdgeToNodeInput(HeadNode->GetID(), HeadNodeInput->VertexID, InPageID);
+
+				FMetasoundFrontendEdge NewEdge;
+				NewEdge.FromNodeID = NewNode->GetID();
+				NewEdge.FromVertexID = NewNodeOutput->VertexID;
+				NewEdge.ToNodeID = HeadNode->GetID();
+				NewEdge.ToVertexID = HeadNodeInput->VertexID;
+				AddEdge(MoveTemp(NewEdge), InPageID);
+			}
+
+			const FMetasoundFrontendVertex* NewNodeInput = FindNodeInput(NewNode->GetID(), METASOUND_GET_PARAM_NAME(InputVariable), InPageID);
+			check(NewNodeInput);
+
+			const FMetasoundFrontendNode* VariableNode = FindNode(Variable->VariableNodeID, InPageID);
+			check(VariableNode);
+
+			const FMetasoundFrontendVertex* VariableNodeOutput = FindNodeOutput(VariableNode->GetID(), METASOUND_GET_PARAM_NAME(OutputVariable), InPageID);
+			check(VariableNodeOutput);
+
+			FMetasoundFrontendEdge NewEdge;
+			NewEdge.FromNodeID = VariableNode->GetID();
+			NewEdge.FromVertexID = VariableNodeOutput->VertexID;
+			NewEdge.ToNodeID = NewNode->GetID();
+			NewEdge.ToVertexID = NewNodeInput->VertexID;
+			AddEdge(MoveTemp(NewEdge), InPageID);
+
+			// Add node ID to variable after connecting since the array
+			// order of node ids is used to determine whether a node 
+			// is the tail node.
+			Variable->DeferredAccessorNodeIDs.Add(NewNode->GetID());
+			return NewNode;
+		}
+	}
+
+	return nullptr;
+}
+
+const FMetasoundFrontendNode* FMetaSoundFrontendDocumentBuilder::AddGraphVariableMutatorNode(FName VariableName, FGuid InNodeID, const FGuid* InPageID)
+{
+	using namespace Metasound::Frontend;
+	using namespace Metasound::VariableNames;
+
+	if (FMetasoundFrontendVariable* Variable = FindGraphVariableInternal(VariableName, InPageID))
+	{
+		const FMetasoundFrontendNode* MutatorNode = FindNode(Variable->MutatorNodeID, InPageID);
+		if (MutatorNode)
+		{
+			UE_LOG(LogMetaSound, Warning, TEXT("Cannot add mutator node as one already exists for variable '%s'. Returning existing mutator"), *VariableName.ToString());
+			return MutatorNode;
 		}
+		else
+		{
+			FNodeRegistryKey ClassKey;
+			{
+				FMetasoundFrontendClass MutatorNodeClass;
+				if (!IDataTypeRegistry::Get().GetFrontendVariableMutatorClass(Variable->TypeName, MutatorNodeClass))
+				{
+					UE_LOG(LogMetaSound, Error, TEXT("Could not find registered \"set variable\" node class for data type \"%s\""), *Variable->TypeName.ToString());
+					return nullptr;
+				}
 
-		check(NewNodeToReturn);
-		return NewNodeToReturn;
+				ClassKey = FNodeRegistryKey(MutatorNodeClass.Metadata);
+				const FMetasoundFrontendClass* Dependency = FindDependency(MutatorNodeClass.Metadata);
+				if (!Dependency)
+				{
+					Dependency = AddDependency(MoveTemp(MutatorNodeClass));
+				}
+				check(Dependency);
+			}
+
+			auto FinalizeNodeFunction = [](const FMetasoundFrontendNode&, const Metasound::Frontend::FNodeRegistryKey&) { };
+			const FGuid& PageID = InPageID ? *InPageID : BuildPageID;
+			MutatorNode = AddNodeInternal(ClassKey, FinalizeNodeFunction, PageID, InNodeID);
+			if (MutatorNode)
+			{
+				// Initialize mutator default literal value to that of the variable
+				const FMetasoundFrontendVertex* MutatorDataInput = FindNodeInput(MutatorNode->GetID(), METASOUND_GET_PARAM_NAME(InputData), InPageID);
+				check(MutatorDataInput);
+				SetNodeInputDefault(MutatorNode->GetID(), MutatorDataInput->VertexID, Variable->Literal, InPageID);
+
+				Variable->MutatorNodeID = MutatorNode->GetID();
+				FGuid SourceVariableNodeID = Variable->VariableNodeID;
+
+				// Connect last delayed getter in variable stack.
+				if (!Variable->DeferredAccessorNodeIDs.IsEmpty())
+				{
+					SourceVariableNodeID = Variable->DeferredAccessorNodeIDs.Last();
+				}
+
+				const FMetasoundFrontendNode* SourceVariableNode = FindNode(SourceVariableNodeID, InPageID);
+				if (ensure(SourceVariableNode))
+				{
+					const FMetasoundFrontendVertex* MutatorNodeInput = FindNodeInput(MutatorNode->GetID(), METASOUND_GET_PARAM_NAME(InputVariable), InPageID);
+					check(MutatorNodeInput);
+
+					const FMetasoundFrontendNode* VariableSourceNode = FindNode(SourceVariableNodeID, InPageID);
+					check(VariableSourceNode);
+
+					const FMetasoundFrontendVertex* SourceVariableNodeOutput = FindNodeOutput(VariableSourceNode->GetID(), METASOUND_GET_PARAM_NAME(OutputVariable), InPageID);
+					check(SourceVariableNodeOutput);
+
+					FMetasoundFrontendEdge NewEdge;
+					NewEdge.FromNodeID = SourceVariableNodeID;
+					NewEdge.FromVertexID = SourceVariableNodeOutput->VertexID;
+					NewEdge.ToNodeID = MutatorNode->GetID();
+					NewEdge.ToVertexID = MutatorNodeInput->VertexID;
+					AddEdge(MoveTemp(NewEdge), InPageID);
+				}
+
+				// Connect to first inline getter in variable stack
+				if (!Variable->AccessorNodeIDs.IsEmpty())
+				{
+					const FGuid& HeadAccessorNodeID = Variable->AccessorNodeIDs[0];
+					const FMetasoundFrontendVertex* MutatorNodeOutput = FindNodeOutput(MutatorNode->GetID(), METASOUND_GET_PARAM_NAME(OutputVariable), InPageID);
+					check(MutatorNodeOutput);
+
+					const FMetasoundFrontendVertex* AccessorNodeInput = FindNodeInput(HeadAccessorNodeID, METASOUND_GET_PARAM_NAME(InputVariable), InPageID);
+					check(AccessorNodeInput);
+
+					RemoveEdgeToNodeInput(HeadAccessorNodeID, AccessorNodeInput->VertexID, InPageID);
+
+					FMetasoundFrontendEdge NewEdge;
+					NewEdge.FromNodeID = MutatorNode->GetID();
+					NewEdge.FromVertexID = MutatorNodeOutput->VertexID;
+					NewEdge.ToNodeID = HeadAccessorNodeID;
+					NewEdge.ToVertexID = AccessorNodeInput->VertexID;
+					AddEdge(MoveTemp(NewEdge), InPageID);
+				}
+			}
+		}
+		return MutatorNode;
 	}
 
 	return nullptr;
@@ -1315,6 +1684,8 @@ bool FMetaSoundFrontendDocumentBuilder::AddInterface(FName InterfaceName)
 
 const FMetasoundFrontendNode* FMetaSoundFrontendDocumentBuilder::AddGraphNode(const FMetasoundFrontendGraphClass& InGraphClass, FGuid InNodeID, const FGuid* InPageID)
 {
+	using FNodeRegistryKey = Metasound::Frontend::FNodeRegistryKey;
+
 	auto FinalizeNode = [](FMetasoundFrontendNode& InOutNode, const Metasound::Frontend::FNodeRegistryKey& ClassKey)
 	{
 #if WITH_EDITOR
@@ -1332,18 +1703,23 @@ const FMetasoundFrontendNode* FMetaSoundFrontendDocumentBuilder::AddGraphNode(co
 #endif // WITH_EDITOR
 	};
 
-	// Dependency is considered "External" when looked up or added on another graph
-	FMetasoundFrontendClassMetadata NewClassMetadata = InGraphClass.Metadata;
-	NewClassMetadata.SetType(EMetasoundFrontendClassType::External);
-
-	if (!FindDependency(NewClassMetadata))
+	FNodeRegistryKey ClassKey;
 	{
-		AddDependency(InGraphClass);
+		// Dependency is considered "External" when looked up or added
+		// on another graph Cast strips GraphClass-specific data as well
+		FMetasoundFrontendClass NewClass = InGraphClass;
+		NewClass.Metadata.SetType(EMetasoundFrontendClassType::External);
+
+		ClassKey = FNodeRegistryKey(NewClass.Metadata);
+		if (!FindDependency(NewClass.Metadata))
+		{
+			AddDependency(MoveTemp(NewClass));
+		}
 	}
 
 	constexpr int32* NewNodeIndex = nullptr;
 	const FGuid& PageID = InPageID ? *InPageID : BuildPageID;
-	return AddNodeInternal(NewClassMetadata, FinalizeNode, PageID, InNodeID, NewNodeIndex);
+	return AddNodeInternal(ClassKey, FinalizeNode, PageID, InNodeID, NewNodeIndex);
 }
 
 const FMetasoundFrontendNode* FMetaSoundFrontendDocumentBuilder::AddNodeByClassName(const FMetasoundFrontendClassName& InClassName, int32 InMajorVersion, FGuid InNodeID, const FGuid* InPageID)
@@ -1351,9 +1727,16 @@ const FMetasoundFrontendNode* FMetaSoundFrontendDocumentBuilder::AddNodeByClassN
 	using namespace Metasound;
 	using namespace Metasound::Frontend;
 
-	FMetasoundFrontendClass RegisteredClass;
-	if (ISearchEngine::Get().FindClassWithHighestMinorVersion(InClassName, InMajorVersion, RegisteredClass))
+	const FMetasoundFrontendClass* Dependency = nullptr;
+	FNodeRegistryKey ClassKey;
 	{
+		FMetasoundFrontendClass RegisteredClass;
+		if (!ISearchEngine::Get().FindClassWithHighestMinorVersion(InClassName, InMajorVersion, RegisteredClass))
+		{
+			UE_LOG(LogMetaSound, Error, TEXT("Failed to add new node by class name '%s' and major version '%d': Class not found"), *InClassName.ToString(), InMajorVersion);
+			return nullptr;
+		}
+
 		const EMetasoundFrontendClassType ClassType = RegisteredClass.Metadata.GetType();
 		if (ClassType != EMetasoundFrontendClassType::External && ClassType != EMetasoundFrontendClassType::Graph)
 		{
@@ -1365,23 +1748,22 @@ const FMetasoundFrontendNode* FMetaSoundFrontendDocumentBuilder::AddNodeByClassN
 
 		// Dependency is considered "External" when looked up or added as a dependency to a graph
 		RegisteredClass.Metadata.SetType(EMetasoundFrontendClassType::External);
-
-		const FMetasoundFrontendClass* Dependency = FindDependency(RegisteredClass.Metadata);
+		ClassKey = FNodeRegistryKey(RegisteredClass.Metadata);
+		Dependency = FindDependency(RegisteredClass.Metadata);
 		if (!Dependency)
 		{
-			Dependency = AddDependency(RegisteredClass);
+			Dependency = AddDependency(MoveTemp(RegisteredClass));
 		}
+	}
 
-		if (Dependency)
-		{
-			auto FinalizeNode = [](const FMetasoundFrontendNode& Node, const Metasound::Frontend::FNodeRegistryKey& ClassKey) { return Node.Name; };
-			constexpr int32* NewNodeIndex = nullptr;
-			const FGuid& PageID = InPageID ? *InPageID : BuildPageID;
-			return AddNodeInternal(Dependency->Metadata, FinalizeNode, PageID, InNodeID, NewNodeIndex);
-		}
+	if (Dependency)
+	{
+		auto FinalizeNode = [](const FMetasoundFrontendNode& Node, const Metasound::Frontend::FNodeRegistryKey& ClassKey) { return Node.Name; };
+		constexpr int32* NewNodeIndex = nullptr;
+		const FGuid& PageID = InPageID ? *InPageID : BuildPageID;
+		return AddNodeInternal(Dependency->Metadata, FinalizeNode, PageID, InNodeID, NewNodeIndex);
 	}
 
-	UE_LOG(LogMetaSound, Warning, TEXT("Failed to add new node by class name '%s' and major version '%d': Class not found"), *InClassName.ToString(), InMajorVersion);
 	return nullptr;
 }
 
@@ -1399,11 +1781,7 @@ const FMetasoundFrontendNode* FMetaSoundFrontendDocumentBuilder::AddNodeByTempla
 	}
 	check(Dependency);
 
-	auto FinalizeNodeFunction = [](const FMetasoundFrontendNode& Node, const Metasound::Frontend::FNodeRegistryKey& ClassKey)
-	{
-		return Node.Name;
-	};
-
+	auto FinalizeNodeFunction = [](const FMetasoundFrontendNode&, const Metasound::Frontend::FNodeRegistryKey&) { };
 	constexpr int32* NewNodeIndex = nullptr;
 	const FGuid & PageID = InPageID ? *InPageID : BuildPageID;
 	FMetasoundFrontendNode* NewNode = AddNodeInternal(Dependency->Metadata, FinalizeNodeFunction, PageID, InNodeID, NewNodeIndex);
@@ -1420,14 +1798,23 @@ FMetasoundFrontendNode* FMetaSoundFrontendDocumentBuilder::AddNodeInternal(const
 	using namespace Metasound::Frontend;
 
 	const FNodeRegistryKey ClassKey = FNodeRegistryKey(InClassMetadata);
-	if (const FMetasoundFrontendClass* Dependency = DocumentCache->FindDependency(ClassKey))
+	return AddNodeInternal(ClassKey, FinalizeNode, InPageID, InNodeID, NewNodeIndex);
+}
+
+FMetasoundFrontendNode* FMetaSoundFrontendDocumentBuilder::AddNodeInternal(const Metasound::Frontend::FNodeRegistryKey& InClassKey, FFinalizeNodeFunctionRef FinalizeNode, const FGuid& InPageID, FGuid InNodeID, int32* NewNodeIndex)
+{
+	METASOUND_TRACE_CPUPROFILER_EVENT_SCOPE(FMetaSoundFrontendDocumentBuilder::AddNodeInternal);
+
+	using namespace Metasound::Frontend;
+
+	if (const FMetasoundFrontendClass* Dependency = DocumentCache->FindDependency(InClassKey))
 	{
 		FMetasoundFrontendDocument& Document = GetDocumentChecked();
 		FMetasoundFrontendGraph& Graph = Document.RootGraph.FindGraphChecked(InPageID);
 		TArray<FMetasoundFrontendNode>& Nodes = Graph.Nodes;
 		FMetasoundFrontendNode& Node = Nodes.Emplace_GetRef(*Dependency);
 		Node.UpdateID(InNodeID);
-		FinalizeNode(Node, ClassKey);
+		FinalizeNode(Node, InClassKey);
 
 		const int32 NewIndex = Nodes.Num() - 1;
 		const IDocumentGraphNodeCache& NodeCache = DocumentCache->GetNodeCache(InPageID);
@@ -1546,7 +1933,7 @@ bool FMetaSoundFrontendDocumentBuilder::ConformGraphInputNodeToClass(const FMeta
 		const FMetasoundFrontendClass* Dependency = FindDependency(Class.Metadata);
 		if (!Dependency)
 		{
-			Dependency = AddDependency(Class);
+			Dependency = AddDependency(MoveTemp(Class));
 		}
 
 		if (ensureAlways(Dependency))
@@ -1604,7 +1991,7 @@ bool FMetaSoundFrontendDocumentBuilder::ConformGraphOutputNodeToClass(const FMet
 		const FMetasoundFrontendClass* Dependency = FindDependency(Class.Metadata);
 		if (!Dependency)
 		{
-			Dependency = AddDependency(Class);
+			Dependency = AddDependency(MoveTemp(Class));
 		}
 
 		if (ensureAlways(Dependency))
@@ -1666,10 +2053,7 @@ bool FMetaSoundFrontendDocumentBuilder::ConvertFromPreset()
 
 	if (IsPreset())
 	{
-		FMetasoundFrontendDocument& Document = GetDocumentChecked();
-		FMetasoundFrontendGraphClass& RootGraphClass = Document.RootGraph;
-		FMetasoundFrontendGraphClassPresetOptions& PresetOptions = RootGraphClass.PresetOptions;
-		PresetOptions.bIsPreset = false;
+		GetDocumentChecked().RootGraph.PresetOptions = { };
 
 #if WITH_EDITOR
 		FMetasoundFrontendGraphStyle& Style = FindBuildGraphChecked().Style;
@@ -1742,7 +2126,39 @@ const FMetasoundFrontendNode* FMetaSoundFrontendDocumentBuilder::DuplicateGraphI
 #endif // WITH_EDITORONLY_DATA
 	ClassInput.Name = InName;
 
-	return AddGraphInput(ClassInput, &PageID);
+	return AddGraphInput(MoveTemp(ClassInput), &PageID);
+}
+
+const FMetasoundFrontendClassInput* FMetaSoundFrontendDocumentBuilder::DuplicateGraphInput(FName ExistingName, FName NewName)
+{
+	using namespace Metasound;
+
+	const FMetasoundFrontendClassInput* ExistingInput = FindGraphInput(ExistingName);
+	if (!ExistingInput)
+	{
+		UE_LOG(LogMetaSound, Warning, TEXT("Failed to duplicate graph input '%s': input does not exist"), *ExistingName.ToString());
+		return nullptr;
+	}
+
+	if (FindGraphInput(NewName))
+	{
+		UE_LOG(LogMetaSound, Warning, TEXT("Failed to duplicate graph input '%s': input with name '%s' already exists"), *ExistingName.ToString(), *NewName.ToString());
+		return nullptr;
+	}
+
+	Frontend::FDocumentIDGenerator& IDGenerator = Frontend::FDocumentIDGenerator::Get();
+	const FMetasoundFrontendDocument& Doc = GetConstDocumentChecked();
+
+	FMetasoundFrontendClassInput ClassInput = *ExistingInput;
+	ClassInput.NodeID = IDGenerator.CreateNodeID(Doc);
+	ClassInput.VertexID = IDGenerator.CreateVertexID(Doc);
+#if WITH_EDITORONLY_DATA
+	ClassInput.Metadata.SetDisplayName(FText::GetEmpty());
+#endif // WITH_EDITORONLY_DATA
+	ClassInput.Name = NewName;
+
+	AddGraphInput(MoveTemp(ClassInput));
+	return FindGraphInput(NewName);
 }
 
 const FMetasoundFrontendNode* FMetaSoundFrontendDocumentBuilder::DuplicateGraphOutput(const FMetasoundFrontendClassOutput& InClassOutput, const FName InName, const FGuid* InPageID)
@@ -1762,7 +2178,77 @@ const FMetasoundFrontendNode* FMetaSoundFrontendDocumentBuilder::DuplicateGraphO
 #endif // WITH_EDITORONLY_DATA
 	ClassOutput.Name = InName;
 
-	return AddGraphOutput(ClassOutput, &PageID);
+	return AddGraphOutput(MoveTemp(ClassOutput), &PageID);
+}
+
+const FMetasoundFrontendClassOutput* FMetaSoundFrontendDocumentBuilder::DuplicateGraphOutput(FName ExistingName, FName NewName)
+{
+	using namespace Metasound::Frontend;
+
+	const FMetasoundFrontendClassOutput* ExistingOutput = FindGraphOutput(ExistingName);
+	if (!ExistingOutput)
+	{
+		UE_LOG(LogMetaSound, Warning, TEXT("Failed to duplicate graph output '%s', output does not exist"), *ExistingName.ToString());
+		return nullptr;
+
+	}
+
+	if (FindGraphOutput(NewName))
+	{
+		UE_LOG(LogMetaSound, Warning, TEXT("Failed to duplicate graph output '%s', output with name '%s' already exists"), *ExistingName.ToString(), *NewName.ToString());
+		return nullptr;
+	}
+
+	FDocumentIDGenerator& IDGenerator = FDocumentIDGenerator::Get();
+	const FMetasoundFrontendDocument& Doc = GetConstDocumentChecked();
+
+	FMetasoundFrontendClassOutput ClassOutput = *ExistingOutput;
+	ClassOutput.NodeID = IDGenerator.CreateNodeID(Doc);
+	ClassOutput.VertexID = IDGenerator.CreateVertexID(Doc);
+#if WITH_EDITORONLY_DATA
+	ClassOutput.Metadata.SetDisplayName(FText::GetEmpty());
+#endif // WITH_EDITORONLY_DATA
+	ClassOutput.Name = NewName;
+
+	AddGraphOutput(MoveTemp(ClassOutput));
+	return FindGraphOutput(NewName);
+}
+
+const FMetasoundFrontendVariable* FMetaSoundFrontendDocumentBuilder::DuplicateGraphVariable(FName ExistingName, FName NewName, const FGuid* InPageID)
+{
+	using namespace Metasound::Frontend;
+
+	if (FindGraphVariable(NewName, InPageID))
+	{
+		UE_LOG(LogMetaSound, Warning, TEXT("Failed to duplicate graph variable '%s': variable with name '%s' already exists"), *ExistingName.ToString(), *NewName.ToString());
+		return nullptr;
+	}
+
+	if (const FMetasoundFrontendVariable* ExistingVariable = FindGraphVariable(ExistingName, InPageID))
+	{
+#if WITH_EDITORONLY_DATA
+		const FText* Description = &ExistingVariable->Description;
+#else
+		const FText* Description = nullptr;
+#endif // WITH_EDITORONLY_DATA
+
+		const FMetasoundFrontendVariable* NewVariable = AddGraphVariable(
+			NewName,
+			ExistingVariable->TypeName,
+			&ExistingVariable->Literal,
+			&FText::GetEmpty(), // Don't copy to ensure no confusion over identical display names
+			Description,
+			InPageID
+		);
+		return NewVariable;
+	}
+	else
+	{
+		const FGuid& PageID = InPageID ? *InPageID : BuildPageID;
+		UE_LOG(LogMetaSound, Warning, TEXT("Failed to duplicate graph variable '%s' on page '%s': variable does not exist"), *ExistingName.ToString(), *PageID.ToString());
+	}
+
+	return nullptr;
 }
 
 FMetasoundFrontendGraph& FMetaSoundFrontendDocumentBuilder::FindBuildGraphChecked() const
@@ -1889,6 +2375,54 @@ FMetaSoundFrontendGraphComment* FMetaSoundFrontendDocumentBuilder::FindGraphComm
 }
 #endif // WITH_EDITORONLY_DATA
 
+const FMetasoundFrontendNode* FMetaSoundFrontendDocumentBuilder::FindHeadNodeInVariableStack(FName VariableName, const FGuid* InPageID) const
+{
+	// The variable "stack" is [GetDelayedNodes, SetNode, GetNodes].
+	if (const FMetasoundFrontendVariable* Variable = FindGraphVariable(VariableName, InPageID))
+	{
+		if (!Variable->DeferredAccessorNodeIDs.IsEmpty())
+		{
+			return FindNode(Variable->DeferredAccessorNodeIDs[0], InPageID);
+		}
+
+		if (Variable->MutatorNodeID.IsValid())
+		{
+			return FindNode(Variable->MutatorNodeID, InPageID);
+		}
+
+		if (!Variable->AccessorNodeIDs.IsEmpty())
+		{
+			return FindNode(Variable->AccessorNodeIDs[0], InPageID);
+		}
+	}
+
+	return nullptr;
+}
+
+const FMetasoundFrontendNode* FMetaSoundFrontendDocumentBuilder::FindTailNodeInVariableStack(FName VariableName, const FGuid* InPageID) const
+{
+	// The variable "stack" is [GetDelayedNodes, SetNode, GetNodes].
+	if (const FMetasoundFrontendVariable* Variable = FindGraphVariable(VariableName, InPageID))
+	{
+		if (!Variable->AccessorNodeIDs.IsEmpty())
+		{
+			return FindNode(Variable->AccessorNodeIDs.Last(), InPageID);
+		}
+
+		if (Variable->MutatorNodeID.IsValid())
+		{
+			return FindNode(Variable->MutatorNodeID, InPageID);
+		}
+
+		if (!Variable->DeferredAccessorNodeIDs.IsEmpty())
+		{
+			return FindNode(Variable->DeferredAccessorNodeIDs.Last(), InPageID);
+		}
+	}
+
+	return nullptr;
+}
+
 bool FMetaSoundFrontendDocumentBuilder::FindInterfaceInputNodes(FName InterfaceName, TArray<const FMetasoundFrontendNode*>& OutInputs, const FGuid* InPageID) const
 {
 	using namespace Metasound::Frontend;
@@ -2014,6 +2548,15 @@ const FMetasoundFrontendNode* FMetaSoundFrontendDocumentBuilder::FindGraphOutput
 	return nullptr;
 }
 
+const FMetasoundFrontendVariable* FMetaSoundFrontendDocumentBuilder::FindGraphVariable(const FGuid& InVariableID, const FGuid* InPageID) const
+{
+	const FGuid& PageID = InPageID ? *InPageID : BuildPageID;
+	const FMetasoundFrontendDocument& Document = GetDocumentChecked();
+	const FMetasoundFrontendGraph& Graph = Document.RootGraph.FindConstGraphChecked(PageID);
+	auto MatchesID = [&InVariableID](const FMetasoundFrontendVariable& Variable) { return Variable.ID == InVariableID; };
+	return Graph.Variables.FindByPredicate(MatchesID);
+}
+
 const FMetasoundFrontendVariable* FMetaSoundFrontendDocumentBuilder::FindGraphVariable(FName VariableName, const FGuid* InPageID) const
 {
 	const FGuid& PageID = InPageID ? *InPageID : BuildPageID;
@@ -2023,6 +2566,30 @@ const FMetasoundFrontendVariable* FMetaSoundFrontendDocumentBuilder::FindGraphVa
 	return Graph.Variables.FindByPredicate(MatchesName);
 }
 
+FMetasoundFrontendVariable* FMetaSoundFrontendDocumentBuilder::FindGraphVariableInternal(FName VariableName, const FGuid* InPageID)
+{
+	const FGuid& PageID = InPageID ? *InPageID : BuildPageID;
+	FMetasoundFrontendDocument& Document = GetDocumentChecked();
+	FMetasoundFrontendGraph& Graph = Document.RootGraph.FindGraphChecked(PageID);
+	auto MatchesName = [&VariableName](const FMetasoundFrontendVariable& Variable) { return Variable.Name == VariableName; };
+	return Graph.Variables.FindByPredicate(MatchesName);
+}
+
+const FMetasoundFrontendVariable* FMetaSoundFrontendDocumentBuilder::FindGraphVariableByNodeID(const FGuid& InNodeID, const FGuid* InPageID) const
+{
+	const FGuid& PageID = InPageID ? *InPageID : BuildPageID;
+	const FMetasoundFrontendDocument& Document = GetDocumentChecked();
+	const FMetasoundFrontendGraph& Graph = Document.RootGraph.FindConstGraphChecked(PageID);
+	auto ContainsNodeWithID = [&InNodeID](const FMetasoundFrontendVariable& Variable)
+	{
+		return Variable.VariableNodeID == InNodeID
+		|| Variable.MutatorNodeID == InNodeID
+		|| Variable.DeferredAccessorNodeIDs.Contains(InNodeID)
+		|| Variable.AccessorNodeIDs.Contains(InNodeID);
+	};
+	return Graph.Variables.FindByPredicate(ContainsNodeWithID);
+}
+
 #if WITH_EDITOR
 UMetaSoundFrontendMemberMetadata* FMetaSoundFrontendDocumentBuilder::FindMemberMetadata(const FGuid& InMemberID)
 {
@@ -2387,6 +2954,18 @@ TArray<const FMetasoundFrontendNode*> FMetaSoundFrontendDocumentBuilder::GetGrap
 	return TemplateNodes;
 }
 
+const TSet<FName>* FMetaSoundFrontendDocumentBuilder::GetGraphInputsInheritingDefault() const
+{
+	FMetasoundFrontendGraphClassPresetOptions& PresetOptions = GetDocumentChecked().RootGraph.PresetOptions;
+	if (PresetOptions.bIsPreset)
+	{
+		return &PresetOptions.InputsInheritingDefault;
+	}
+
+	return nullptr;
+}
+
+
 FMetasoundAssetBase& FMetaSoundFrontendDocumentBuilder::GetMetasoundAsset() const
 {
 	using namespace Metasound::Frontend;
@@ -2452,6 +3031,15 @@ const FMetasoundFrontendLiteral* FMetaSoundFrontendDocumentBuilder::GetGraphInpu
 	return nullptr;
 }
 
+const FMetasoundFrontendLiteral* FMetaSoundFrontendDocumentBuilder::GetGraphVariableDefault(FName VariableName, const FGuid* InPageID) const
+{
+	if (const FMetasoundFrontendVariable* Variable = FindGraphVariable(VariableName, InPageID))
+	{
+		return &Variable->Literal;
+	}
+	return nullptr;
+}
+
 EMetasoundFrontendVertexAccessType FMetaSoundFrontendDocumentBuilder::GetNodeInputAccessType(const FGuid& InNodeID, const FGuid& InVertexID, const FGuid* InPageID) const
 {
 	using namespace Metasound::Frontend;
@@ -2886,6 +3474,26 @@ Metasound::Frontend::EInvalidEdgeReason FMetaSoundFrontendDocumentBuilder::IsVal
 	return EInvalidEdgeReason::None;
 }
 
+void FMetaSoundFrontendDocumentBuilder::IterateNodesDownstreamDepthFirst(const FGuid& InNodeID, FDepthFirstVisitFunction VisitFunction)
+{
+	TArray<FGuid> Stack({ InNodeID });
+	TSet<FGuid> Visited;
+	while (!Stack.IsEmpty())
+	{
+		// Do not revisit a node that has already been visited.
+		const FGuid CurrentNodeID = Stack.Pop();
+		if (!Visited.Contains(CurrentNodeID))
+		{
+			if (const FMetasoundFrontendNode* CurrentNode = FindNode(CurrentNodeID))
+			{
+				TArray<FGuid> Children = VisitFunction(CurrentNode->GetID()).Array();
+				Stack.Append(Children);
+				Visited.Add(CurrentNodeID);
+			}
+		}
+	}
+}
+
 void FMetaSoundFrontendDocumentBuilder::IterateNodesConnectedWithVertex(const FMetasoundFrontendVertexHandle& Vertex, TFunctionRef<void(const FMetasoundFrontendEdge&, FMetasoundFrontendNode&)> NodeIndexIterFunc, const FGuid& InPageID)
 {
 	using namespace Metasound::Frontend;
@@ -3295,10 +3903,9 @@ bool FMetaSoundFrontendDocumentBuilder::RemoveGraphInputDefault(FName InputName,
 
 			if (bClearInheritsDefault)
 			{
-				// Set the input as no longer inheriting default for presets
-				// (No-ops if MetaSound isn't preset or isn't set to inherit default).
 				constexpr bool bInputInheritsDefault = false;
-				SetGraphInputInheritsDefault(InputName, bInputInheritsDefault);
+				constexpr bool bForceUpdate = true;
+				SetGraphInputInheritsDefault(InputName, bInputInheritsDefault, bForceUpdate);
 			}
 
 			return true;
@@ -3493,18 +4100,18 @@ bool FMetaSoundFrontendDocumentBuilder::RemoveGraphComment(const FGuid& InCommen
 }
 #endif // WITH_EDITORONLY_DATA
 
-bool FMetaSoundFrontendDocumentBuilder::RemoveGraphInput(FName InInputName)
+bool FMetaSoundFrontendDocumentBuilder::RemoveGraphInput(FName InputName)
 {
 	using namespace Metasound::Frontend;
 
 	FMetasoundFrontendDocument& Document = GetDocumentChecked();
-	if (const int32* IndexPtr = DocumentCache->GetInterfaceCache().FindInputIndex(InInputName))
+	if (const int32* IndexPtr = DocumentCache->GetInterfaceCache().FindInputIndex(InputName))
 	{
 		TArray<FMetasoundFrontendClassInput>& Inputs = Document.RootGraph.Interface.Inputs;
 		const FGuid NodeID = Inputs[*IndexPtr].NodeID;
 		FGuid ClassID;
 		bool bNodesRemoved = true;
-		Document.RootGraph.IterateGraphPages([this, &Document, &ClassID, &NodeID, &InInputName, &bNodesRemoved](const FMetasoundFrontendGraph& Graph)
+		Document.RootGraph.IterateGraphPages([this, &Document, &ClassID, &NodeID, &InputName, &bNodesRemoved](const FMetasoundFrontendGraph& Graph)
 		{
 			TArray<FGuid> NodeIDsToRemove { NodeID };
 
@@ -3518,7 +4125,7 @@ bool FMetaSoundFrontendDocumentBuilder::RemoveGraphInput(FName InInputName)
 				return;
 			}
 
-			const TArray<const FMetasoundFrontendNode*> TemplateNodes = GetGraphInputTemplateNodes(InInputName, &Graph.PageID);
+			const TArray<const FMetasoundFrontendNode*> TemplateNodes = GetGraphInputTemplateNodes(InputName, &Graph.PageID);
 			Algo::Transform(TemplateNodes, NodeIDsToRemove, [](const FMetasoundFrontendNode* Node) { return Node->GetID(); });
 
 			for (const FGuid& ToRemove : NodeIDsToRemove)
@@ -3557,6 +4164,10 @@ bool FMetaSoundFrontendDocumentBuilder::RemoveGraphInput(FName InInputName)
 			Document.Metadata.ModifyContext.AddMemberIDModified(NodeID);
 #endif // WITH_EDITORONLY_DATA
 
+			constexpr bool bInputInheritsDefault = false;
+			constexpr bool bForceUpdate = true;
+			SetGraphInputInheritsDefault(InputName, bInputInheritsDefault, bForceUpdate);
+
 			const bool bDependencyReferenced = IsDependencyReferenced(ClassID);
 			if (bDependencyReferenced || RemoveDependency(ClassID))
 			{
@@ -3568,15 +4179,15 @@ bool FMetaSoundFrontendDocumentBuilder::RemoveGraphInput(FName InInputName)
 	return false;
 }
 
-bool FMetaSoundFrontendDocumentBuilder::RemoveGraphOutput(FName InOutputName)
+bool FMetaSoundFrontendDocumentBuilder::RemoveGraphOutput(FName OutputName)
 {
 	bool bNodesRemoved = true;
 	FGuid ClassID;
 	FGuid NodeID;
 	FMetasoundFrontendDocument& Document = GetDocumentChecked();
-	Document.RootGraph.IterateGraphPages([this, &InOutputName, &Document, &ClassID, &NodeID, &bNodesRemoved](const FMetasoundFrontendGraph& Graph)
+	Document.RootGraph.IterateGraphPages([this, &OutputName, &Document, &ClassID, &NodeID, &bNodesRemoved](const FMetasoundFrontendGraph& Graph)
 	{
-		if (const FMetasoundFrontendNode* Node = FindGraphOutputNode(InOutputName, &Graph.PageID))
+		if (const FMetasoundFrontendNode* Node = FindGraphOutputNode(OutputName, &Graph.PageID))
 		{
 			ClassID = Node->ClassID;
 			NodeID = Node->GetID();
@@ -3595,7 +4206,7 @@ bool FMetaSoundFrontendDocumentBuilder::RemoveGraphOutput(FName InOutputName)
 	if (bNodesRemoved)
 	{
 		TArray<FMetasoundFrontendClassOutput>& Outputs = Document.RootGraph.Interface.Outputs;
-		auto OutputNameMatches = [InOutputName](const FMetasoundFrontendClassOutput& Output) { return Output.Name == InOutputName; };
+		auto OutputNameMatches = [OutputName](const FMetasoundFrontendClassOutput& Output) { return Output.Name == OutputName; };
 		const int32 Index = Outputs.IndexOfByPredicate(OutputNameMatches);
 		if (Index != INDEX_NONE)
 		{
@@ -3654,6 +4265,48 @@ bool FMetaSoundFrontendDocumentBuilder::RemoveGraphPage(const FGuid& InPageID)
 }
 #endif // WITH_EDITORONLY_DATA
 
+bool FMetaSoundFrontendDocumentBuilder::RemoveGraphVariable(FName VariableName, const FGuid* InPageID)
+{
+	if (const FMetasoundFrontendVariable* Variable = FindGraphVariable(VariableName, InPageID))
+	{
+		RemoveNode(Variable->VariableNodeID);
+		RemoveNode(Variable->MutatorNodeID);
+
+		// Copy ids as node removal will update variable node IDs
+		TArray<FGuid> AccessorNodeIDs = Variable->AccessorNodeIDs;
+		for (const FGuid& NodeID : AccessorNodeIDs)
+		{
+			RemoveNode(NodeID, InPageID);
+		}
+
+		// Copy ids as node removal will update variable node IDs
+		TArray<FGuid> DeferredAccessorNodeIDs = Variable->DeferredAccessorNodeIDs;
+		for (const FGuid& NodeID : DeferredAccessorNodeIDs)
+		{
+			RemoveNode(NodeID, InPageID);
+		}
+
+		FMetasoundFrontendDocument& Document = GetDocumentChecked();
+		FMetasoundFrontendGraph& Graph = Document.RootGraph.FindGraphChecked(InPageID ? *InPageID : BuildPageID);
+
+#if WITH_EDITOR
+		Document.Metadata.ModifyContext.AddMemberIDModified(Variable->ID);
+#endif // WITH_EDITOR
+
+		auto IsVariableWithID = [&Variable](const FMetasoundFrontendVariable& InVariable)
+		{
+			return InVariable.ID == Variable->ID;
+		};
+		Graph.Variables.RemoveAllSwap(IsVariableWithID);
+
+		// Clean-up/remove variable dependencies that may or may no longer be referenced.
+		RemoveUnusedDependencies();
+		return true;
+	}
+
+	return false;
+}
+
 bool FMetaSoundFrontendDocumentBuilder::RemoveInterface(FName InterfaceName)
 {
 	using namespace Metasound::Frontend;
@@ -4083,10 +4736,10 @@ bool FMetaSoundFrontendDocumentBuilder::SetGraphInputDefault(FName InputName, FM
 			}
 			DocumentDelegates->InterfaceDelegates.OnInputDefaultChanged.Broadcast(Index);
 
-			// Set the input as inheriting default for presets
-			// (No-ops if MetaSound isn't preset or is already set to inherit default).
+			// Set the input as no longer inheriting default
 			constexpr bool bInputInheritsDefault = false;
-			SetGraphInputInheritsDefault(InputName, bInputInheritsDefault);
+			constexpr bool bForceUpdate = true;
+			SetGraphInputInheritsDefault(InputName, bInputInheritsDefault, bForceUpdate);
 
 			return true;
 		}
@@ -4118,10 +4771,10 @@ bool FMetaSoundFrontendDocumentBuilder::SetGraphInputDefaults(FName InputName, T
 			Input.SetDefaults(MoveTemp(Defaults));
 			DocumentDelegates->InterfaceDelegates.OnInputDefaultChanged.Broadcast(Index);
 
-			// Set the input as no longer inheriting default for presets
-			// (No-ops if MetaSound isn't preset or isn't set to inherit default).
+			// Set the input as no longer inheriting default
 			constexpr bool bInputInheritsDefault = false;
-			SetGraphInputInheritsDefault(InputName, bInputInheritsDefault);
+			constexpr bool bForceUpdate = true;
+			SetGraphInputInheritsDefault(InputName, bInputInheritsDefault, bForceUpdate);
 			return true;
 		}
 		UE_LOG(LogMetaSound, Error, TEXT("Attempting to set graph input of type '%s' with unsupported literal type(s)"), *Input.TypeName.ToString());
@@ -4151,19 +4804,19 @@ bool FMetaSoundFrontendDocumentBuilder::SetGraphOutputAdvancedDisplay(const FNam
 }
 #endif // WITH_EDITORONLY_DATA
 
-bool FMetaSoundFrontendDocumentBuilder::SetGraphInputInheritsDefault(FName InName, bool bInputInheritsDefault)
+bool FMetaSoundFrontendDocumentBuilder::SetGraphInputInheritsDefault(FName InName, bool bInputInheritsDefault, bool bForceUpdate)
 {
 	FMetasoundFrontendGraphClassPresetOptions& PresetOptions = GetDocumentChecked().RootGraph.PresetOptions;
 	if (bInputInheritsDefault)
 	{
-		if (PresetOptions.bIsPreset)
+		if (PresetOptions.bIsPreset || bForceUpdate)
 		{
 			return PresetOptions.InputsInheritingDefault.Add(InName).IsValidId();
 		}
 	}
 	else
 	{
-		if (PresetOptions.bIsPreset)
+		if (PresetOptions.bIsPreset || bForceUpdate)
 		{
 			return PresetOptions.InputsInheritingDefault.Remove(InName) > 0;
 		}
@@ -4390,16 +5043,14 @@ bool FMetaSoundFrontendDocumentBuilder::SetGraphVariableDefault(FName VariableNa
 {
 	using namespace Metasound;
 
-	const FGuid PageID = InPageID ? *InPageID : BuildPageID;
-	FMetasoundFrontendDocument& Document = GetDocumentChecked();
-	FMetasoundFrontendGraph& Graph = Document.RootGraph.FindGraphChecked(PageID);
-
-	auto NameMatchesVariable = [&VariableName](const FMetasoundFrontendVariable& Variable) { return Variable.Name == VariableName; };
-	if (FMetasoundFrontendVariable* Variable = Graph.Variables.FindByPredicate(NameMatchesVariable))
+	if (FMetasoundFrontendVariable* Variable = FindGraphVariableInternal(VariableName, InPageID))
 	{
 		if (Frontend::IDataTypeRegistry::Get().IsLiteralTypeSupported(Variable->TypeName, InDefaultLiteral.GetType()))
 		{
 			Variable->Literal = MoveTemp(InDefaultLiteral);
+#if WITH_EDITORONLY_DATA
+			GetDocumentChecked().Metadata.ModifyContext.AddMemberIDModified(Variable->ID);
+#endif // WITH_EDITORONLY_DATA
 			return true;
 		}
 	}
@@ -4407,6 +5058,52 @@ bool FMetaSoundFrontendDocumentBuilder::SetGraphVariableDefault(FName VariableNa
 	return false;
 }
 
+#if WITH_EDITORONLY_DATA
+bool FMetaSoundFrontendDocumentBuilder::SetGraphVariableDescription(FName VariableName, FText Description, const FGuid* InPageID)
+{
+	using namespace Metasound;
+
+	if (FMetasoundFrontendVariable* Variable = FindGraphVariableInternal(VariableName, InPageID))
+	{
+		Variable->Description = Description;
+		GetDocumentChecked().Metadata.ModifyContext.AddMemberIDModified(Variable->ID);
+		return true;
+	}
+
+	return false;
+}
+
+bool FMetaSoundFrontendDocumentBuilder::SetGraphVariableDisplayName(FName VariableName, FText DisplayName, const FGuid* InPageID)
+{
+	using namespace Metasound;
+
+	if (FMetasoundFrontendVariable* Variable = FindGraphVariableInternal(VariableName, InPageID))
+	{
+		Variable->DisplayName = DisplayName;
+		GetDocumentChecked().Metadata.ModifyContext.AddMemberIDModified(Variable->ID);
+		return true;
+	}
+
+	return false;
+}
+#endif // WITH_EDITORONLY_DATA
+
+bool FMetaSoundFrontendDocumentBuilder::SetGraphVariableName(FName VariableName, FName NewName, const FGuid* InPageID)
+{
+	using namespace Metasound;
+
+	if (FMetasoundFrontendVariable* Variable = FindGraphVariableInternal(VariableName, InPageID))
+	{
+		Variable->Name = NewName;
+#if WITH_EDITORONLY_DATA
+		GetDocumentChecked().Metadata.ModifyContext.AddMemberIDModified(Variable->ID);
+#endif // WITH_EDITORONLY_DATA
+		return true;
+	}
+
+	return false;
+}
+
 #if WITH_EDITOR
 void FMetaSoundFrontendDocumentBuilder::SetDisplayName(const FText& InDisplayName)
 {
@@ -4554,31 +5251,33 @@ bool FMetaSoundFrontendDocumentBuilder::SpliceVariableNodeFromStack(const FGuid&
 	using namespace Metasound;
 	using namespace Metasound::Frontend;
 
+	bool bSpliced = false;
 	FMetasoundFrontendDocument& Document = GetDocumentChecked();
 	FMetasoundFrontendGraph& Graph = Document.RootGraph.FindGraphChecked(InPageID);
 	const IDocumentGraphEdgeCache& EdgeCache = DocumentCache->GetEdgeCache(InPageID);
 	FMetasoundFrontendVertexHandle FromVariableVertexHandle;
 	{
-		const FMetasoundFrontendVertex* InputVertex = FindNodeInput(InNodeID, VariableNames::InputVariableName, &InPageID);
-		check(InputVertex);
-		if (const int32* InputEdgeIndex = EdgeCache.FindEdgeIndexToNodeInput(InNodeID, InputVertex->VertexID))
+		// InputVertex may be null if provided ID corresponds to the base variable node (which is always at head of stack and has no inputs)
+		if (const FMetasoundFrontendVertex* InputVertex = FindNodeInput(InNodeID, VariableNames::InputVariableName, &InPageID))
 		{
-			FromVariableVertexHandle = Graph.Edges[*InputEdgeIndex].GetFromVertexHandle();
-			const bool bRemovedEdge = RemoveEdgeToNodeInput(InNodeID, InputVertex->VertexID, &InPageID);
-			check(bRemovedEdge);
+			if (const int32* InputEdgeIndex = EdgeCache.FindEdgeIndexToNodeInput(InNodeID, InputVertex->VertexID))
+			{
+				FromVariableVertexHandle = Graph.Edges[*InputEdgeIndex].GetFromVertexHandle();
+				bSpliced = RemoveEdgeToNodeInput(InNodeID, InputVertex->VertexID, &InPageID);
+			}
 		}
 	}
 
-	if (FromVariableVertexHandle.IsSet())
+	if (const FMetasoundFrontendVertex* OutputVertex = FindNodeOutput(InNodeID, VariableNames::OutputVariableName, &InPageID))
 	{
-		if (const FMetasoundFrontendVertex* OutputVertex = FindNodeOutput(InNodeID, VariableNames::OutputVariableName, &InPageID))
-		{
-			TArray<FMetasoundFrontendVertexHandle> ToVertexHandles;
-			const TArrayView<const int32> OutputEdgeIndices = EdgeCache.FindEdgeIndicesFromNodeOutput(InNodeID, OutputVertex->VertexID);
-			Algo::Transform(OutputEdgeIndices, ToVertexHandles, [&Graph](const int32& VertIndex) { return Graph.Edges[VertIndex].GetToVertexHandle(); });
+		TArray<FMetasoundFrontendVertexHandle> ToVertexHandles;
+		const TArrayView<const int32> OutputEdgeIndices = EdgeCache.FindEdgeIndicesFromNodeOutput(InNodeID, OutputVertex->VertexID);
+		Algo::Transform(OutputEdgeIndices, ToVertexHandles, [&Graph](const int32& VertIndex) { return Graph.Edges[VertIndex].GetToVertexHandle(); });
 
-			RemoveEdgesFromNodeOutput(InNodeID, OutputVertex->VertexID, &InPageID);
+		bSpliced |= RemoveEdgesFromNodeOutput(InNodeID, OutputVertex->VertexID, &InPageID);
 
+		if (FromVariableVertexHandle.IsSet())
+		{
 			for (const FMetasoundFrontendVertexHandle& ToHandle : ToVertexHandles)
 			{
 				AddEdge(FMetasoundFrontendEdge
@@ -4589,11 +5288,10 @@ bool FMetaSoundFrontendDocumentBuilder::SpliceVariableNodeFromStack(const FGuid&
 					ToHandle.VertexID
 				}, &InPageID);
 			}
-			return true;
 		}
 	}
 
-	return false;
+	return bSpliced;
 }
 
 bool FMetaSoundFrontendDocumentBuilder::SwapGraphInput(const FMetasoundFrontendClassVertex& InExistingInputVertex, const FMetasoundFrontendClassVertex& InNewInputVertex)
@@ -4650,15 +5348,19 @@ bool FMetaSoundFrontendDocumentBuilder::SwapGraphInput(const FMetasoundFrontendC
 	}
 
 	// 4. Add new graph vertex
-	FMetasoundFrontendClassInput NewInput = InNewInputVertex;
-	NewInput.NodeID = NodeID;
+	FGuid VertexID;
+	{
+		FMetasoundFrontendClassInput NewInput = InNewInputVertex;
+		NewInput.NodeID = NodeID;
 #if WITH_EDITOR
-	NewInput.Metadata.SetSerializeText(InExistingInputVertex.Metadata.GetSerializeText());
+		NewInput.Metadata.SetSerializeText(InExistingInputVertex.Metadata.GetSerializeText());
 #endif // WITH_EDITOR
 
-	const FMetasoundFrontendNode* NewInputNode = AddGraphInput(NewInput);
-	checkf(NewInputNode, TEXT("Failed to add new Input node when swapping graph inputs"));
-	checkf(NewInputNode->GetID() == NewInput.NodeID, TEXT("Expected new node added to build graph to have same ID as provided input"));
+		const FMetasoundFrontendNode* NewInputNode = AddGraphInput(MoveTemp(NewInput));
+		checkf(NewInputNode, TEXT("Failed to add new Input node when swapping graph inputs"));
+		checkf(NewInputNode->GetID() == NodeID, TEXT("Expected new node added to build graph to have same ID as provided input"));
+		VertexID = NewInputNode->Interface.Outputs.Last().VertexID;
+	}
 
 	Document.RootGraph.IterateGraphPages([&](FMetasoundFrontendGraph& Graph)
 	{
@@ -4680,8 +5382,8 @@ bool FMetaSoundFrontendDocumentBuilder::SwapGraphInput(const FMetasoundFrontendC
 		for (const FMetasoundFrontendEdge& RemovedEdge : RemovedEdges)
 		{
 			FMetasoundFrontendEdge NewEdge = RemovedEdge;
-			NewEdge.FromNodeID = NewInputNode->GetID();
-			NewEdge.FromVertexID = NewInputNode->Interface.Outputs.Last().VertexID;
+			NewEdge.FromNodeID = NodeID;
+			NewEdge.FromVertexID = VertexID;
 			AddEdge(MoveTemp(NewEdge), &Graph.PageID);
 		}
 	});
@@ -4743,15 +5445,19 @@ bool FMetaSoundFrontendDocumentBuilder::SwapGraphOutput(const FMetasoundFrontend
 	}
 	
 	// 4. Add new graph vertex
-	FMetasoundFrontendClassOutput NewOutput = InNewOutputVertex;
-	NewOutput.NodeID = NodeID;
+	FGuid VertexID;
+	{
+		FMetasoundFrontendClassOutput NewOutput = InNewOutputVertex;
+		NewOutput.NodeID = NodeID;
 #if WITH_EDITOR
-	NewOutput.Metadata.SetSerializeText(InExistingOutputVertex.Metadata.GetSerializeText());
+		NewOutput.Metadata.SetSerializeText(InExistingOutputVertex.Metadata.GetSerializeText());
 #endif // WITH_EDITOR
 
-	const FMetasoundFrontendNode* NewOutputNode = AddGraphOutput(NewOutput);
-	checkf(NewOutputNode, TEXT("Failed to add new output node when swapping graph outputs"));
-	checkf(NewOutputNode->GetID() == NewOutput.NodeID, TEXT("Expected new node added to build graph to have same ID as provided output"));
+		const FMetasoundFrontendNode* NewOutputNode = AddGraphOutput(MoveTemp(NewOutput));
+		checkf(NewOutputNode, TEXT("Failed to add new output node when swapping graph outputs"));
+		checkf(NewOutputNode->GetID() == NodeID, TEXT("Expected new node added to build graph to have same ID as provided output"));
+		VertexID = NewOutputNode->Interface.Inputs.Last().VertexID;
+	}
 
 	Document.RootGraph.IterateGraphPages([&](FMetasoundFrontendGraph& Graph)
 	{
@@ -4773,8 +5479,8 @@ bool FMetaSoundFrontendDocumentBuilder::SwapGraphOutput(const FMetasoundFrontend
 		for (const FMetasoundFrontendEdge& RemovedEdge : RemovedEdges)
 		{
 			FMetasoundFrontendEdge NewEdge = RemovedEdge;
-			NewEdge.ToNodeID = NewOutputNode->GetID();
-			NewEdge.ToVertexID = NewOutputNode->Interface.Inputs.Last().VertexID;
+			NewEdge.ToNodeID = NodeID;
+			NewEdge.ToVertexID = VertexID;
 			AddEdge(MoveTemp(NewEdge), &Graph.PageID);
 		}
 	});
diff --git a/Engine/Plugins/Runtime/Metasound/Source/MetasoundFrontend/Public/MetasoundFrontendDocumentBuilder.h b/Engine/Plugins/Runtime/Metasound/Source/MetasoundFrontend/Public/MetasoundFrontendDocumentBuilder.h
--- a/Engine/Plugins/Runtime/Metasound/Source/MetasoundFrontend/Public/MetasoundFrontendDocumentBuilder.h
+++ b/Engine/Plugins/Runtime/Metasound/Source/MetasoundFrontend/Public/MetasoundFrontendDocumentBuilder.h
@@ -147,6 +147,8 @@ struct METASOUNDFRONTEND_API FMetaSoundFrontendDocumentBuilder
 	GENERATED_BODY()
 
 public:
+	using FDepthFirstVisitFunction = TFunctionRef<TSet<FGuid>(const FGuid& /* InNodeID */)>;
+
 	// Default ctor should typically never be used directly as builder interface (and optionally delegates) should be specified on construction (Default exists only to make UObject reflection happy).
 	FMetaSoundFrontendDocumentBuilder(TScriptInterface<IMetaSoundDocumentInterface> InDocumentInterface = { }, TSharedPtr<Metasound::Frontend::FDocumentModifyDelegates> InDocumentDelegates = { }, bool bPrimeCache = false);
 	virtual ~FMetaSoundFrontendDocumentBuilder();
@@ -154,7 +156,9 @@ public:
 	// Call when the builder will no longer modify the IMetaSoundDocumentInterface
 	void FinishBuilding();
 
-	const FMetasoundFrontendClass* AddDependency(const FMetasoundFrontendClass& InClass);
+	// Adds new dependency to MetaSound. Typically not necessary to call directly as dependencies are added automatically via registry when nodes are added, and can be removed when no longer referenced (see 'RemoveUnusedDependencies`).
+	const FMetasoundFrontendClass* AddDependency(FMetasoundFrontendClass NewDependency);
+
 	void AddEdge(FMetasoundFrontendEdge&& InNewEdge, const FGuid* InPageID = nullptr);
 	bool AddNamedEdges(const TSet<Metasound::Frontend::FNamedEdge>& ConnectionsToMake, TArray<const FMetasoundFrontendEdge*>* OutEdgesCreated = nullptr, bool bReplaceExistingConnections = true, const FGuid* InPageID = nullptr);
 	bool AddEdgesByNodeClassInterfaceBindings(const FGuid& InFromNodeID, const FGuid& InToNodeID, bool bReplaceExistingConnections = true, const FGuid* InPageID = nullptr);
@@ -163,7 +167,13 @@ public:
 
 	// Adds Graph Input to document, which in turn adds a referencing input node to ALL pages.  If valid PageID is provided, returns associated page's node pointer.
 	// If none provided, returns node pointer to node for the builder's currently set build page ID (see 'GetBuildPageID').
-	const FMetasoundFrontendNode* AddGraphInput(const FMetasoundFrontendClassInput& InClassInput, const FGuid* InPageID = nullptr);
+	const FMetasoundFrontendNode* AddGraphInput(FMetasoundFrontendClassInput ClassInput, const FGuid* InPageID = nullptr);
+
+	const FMetasoundFrontendVariable* AddGraphVariable(FName VariableName, FName DataType, const FMetasoundFrontendLiteral* Literal = nullptr, const FText* DisplayName = nullptr, const FText* Description = nullptr, const FGuid* InPageID = nullptr);
+	const FMetasoundFrontendNode* AddGraphVariableNode(FName VariableName, EMetasoundFrontendClassType ClassType, FGuid InNodeID = FGuid::NewGuid(), const FGuid* InPageID = nullptr);
+	const FMetasoundFrontendNode* AddGraphVariableMutatorNode(FName VariableName, FGuid InNodeID = FGuid::NewGuid(), const FGuid* InPageID = nullptr);
+	const FMetasoundFrontendNode* AddGraphVariableAccessorNode(FName VariableName, FGuid InNodeID = FGuid::NewGuid(), const FGuid* InPageID = nullptr);
+	const FMetasoundFrontendNode* AddGraphVariableDeferredAccessorNode(FName VariableName, FGuid InNodeID = FGuid::NewGuid(), const FGuid* InPageID = nullptr);
 
 	// Adds node to document to the page associated with the given PageID.  If no valid PageID is provided, adds and returns node pointer to node for the builder's
 	// currently set build page ID (see 'GetBuildPageID').
@@ -171,7 +181,7 @@ public:
 
 	// Adds Graph Output to document, which in turn adds a referencing output node to ALL pages.  If valid PageID is provided, returns associated page's node pointer.
 	// If none provided, returns node pointer to node for the builder's currently set build page ID (see 'GetBuildPageID').
-	const FMetasoundFrontendNode* AddGraphOutput(const FMetasoundFrontendClassOutput& InClassOutput, const FGuid* InPageID = nullptr);
+	const FMetasoundFrontendNode* AddGraphOutput(FMetasoundFrontendClassOutput ClassOutput, const FGuid* InPageID = nullptr);
 
 	bool AddInterface(FName InterfaceName);
 
@@ -207,9 +217,17 @@ public:
 	bool ConvertFromPreset();
 	bool ConvertToPreset(const FMetasoundFrontendDocument& InReferencedDocument, TSharedPtr<Metasound::Frontend::FDocumentModifyDelegates> ModifyDelegates = { });
 
+	const FMetasoundFrontendClassInput* DuplicateGraphInput(FName ExistingName, FName NewName);
+	const FMetasoundFrontendClassOutput* DuplicateGraphOutput(FName ExistingName, FName NewName);
+	const FMetasoundFrontendVariable* DuplicateGraphVariable(FName ExistingName, FName NewName, const FGuid* InPageID = nullptr);
+
+	UE_DEPRECATED(5.6, "Use the duplicate overload that supplies existing name and new name and returns input")
 	const FMetasoundFrontendNode* DuplicateGraphInput(const FMetasoundFrontendClassInput& InClassInput, const FName InName, const FGuid* InPageID = nullptr);
+
+	UE_DEPRECATED(5.6, "Use the duplicate overload that supplies existing name and new name and returns output")
 	const FMetasoundFrontendNode* DuplicateGraphOutput(const FMetasoundFrontendClassOutput& InClassOutput, const FName InName, const FGuid* InPageID = nullptr);
 
+
 #if WITH_EDITORONLY_DATA
 	const FMetasoundFrontendEdgeStyle* FindConstEdgeStyle(const FGuid& InNodeID, FName OutputName, const FGuid* InPageID = nullptr) const;
 	FMetasoundFrontendEdgeStyle* FindEdgeStyle(const FGuid& InNodeID, FName OutputName, const FGuid* InPageID = nullptr);
@@ -232,7 +250,9 @@ public:
 	const FMetasoundFrontendClassOutput* FindGraphOutput(FName OutputName) const;
 	const FMetasoundFrontendNode* FindGraphOutputNode(FName OutputName, const FGuid* InPageID = nullptr) const;
 
+	const FMetasoundFrontendVariable* FindGraphVariable(const FGuid& InVariableID, const FGuid* InPageID = nullptr) const;
 	const FMetasoundFrontendVariable* FindGraphVariable(FName InVariableName, const FGuid* InPageID = nullptr) const;
+	const FMetasoundFrontendVariable* FindGraphVariableByNodeID(const FGuid& InNodeID, const FGuid* InPageID = nullptr) const;
 
 	bool FindInterfaceInputNodes(FName InterfaceName, TArray<const FMetasoundFrontendNode*>& OutInputs, const FGuid* InPageID = nullptr) const;
 	bool FindInterfaceOutputNodes(FName InterfaceName, TArray<const FMetasoundFrontendNode*>& OutOutputs, const FGuid* InPageID = nullptr) const;
@@ -304,6 +324,9 @@ public:
 
 	TArray<const FMetasoundFrontendNode*> GetGraphInputTemplateNodes(FName InInputName, const FGuid* InPageID = nullptr);
 
+	// If graph is set to be a preset, returns set of graph input names inheriting default data from the referenced graph. If not a preset, returns null.
+	const TSet<FName>* GetGraphInputsInheritingDefault() const;
+
 	EMetasoundFrontendVertexAccessType GetNodeInputAccessType(const FGuid& InNodeID, const FGuid& InVertexID, const FGuid* InPageID = nullptr) const;
 
 	UE_DEPRECATED(5.5, "Use FindNodeInputClass overloads instead and use GetDefaults() on result (now supports page values)")
@@ -318,8 +341,12 @@ public:
 	const bool GetIsAdvancedDisplay(const FName MemberName, const EMetasoundFrontendClassType Type) const;
 #endif // WITH_EDITORONLY_DATA
 
+	// Returns the default value set for the input with the given name on the given page.
 	const FMetasoundFrontendLiteral* GetGraphInputDefault(FName InputName, const FGuid* InPageID = nullptr) const;
-	
+
+	// Returns the default value set for the variable with the given name on the given page.
+	const FMetasoundFrontendLiteral* GetGraphVariableDefault(FName InputName, const FGuid* InPageID = nullptr) const;
+
 	// Initializes the builder's document, using the (optional) provided document template, (optional) class name, and (optionally) whether or not to reset the existing class version.
 	void InitDocument(const FMetasoundFrontendDocument* InDocumentTemplate = nullptr, const FMetasoundFrontendClassName* InNewClassName = nullptr, bool bResetVersion = true);
 
@@ -347,6 +374,10 @@ public:
 	// Note that this does not return whether or not the given edge exists, but rather if it could be legally applied to the given edge vertices.
 	Metasound::Frontend::EInvalidEdgeReason IsValidEdge(const FMetasoundFrontendEdge& InEdge, const FGuid* InPageID = nullptr) const;
 
+	// Iterates nodes directly downstream from that with the given ID. If item(s) of interest not found,
+	// provided VisitFunction should return IDs of children to continue non-recursive, depth-first search.
+	void IterateNodesDownstreamDepthFirst(const FGuid& InNodeID, FDepthFirstVisitFunction VisitFunction);
+
 	// Iterates nodes that are filtered by only subscribing to a class with the given type (asserts if provided invalid class type).
 	void IterateNodesByClassType(Metasound::Frontend::FConstClassAndNodeFunctionRef Func, EMetasoundFrontendClassType ClassType, const FGuid* InPageID = nullptr) const;
 
@@ -357,8 +388,12 @@ public:
 		"To enforce this, new recommended pattern is to construct a new builder instead")
 	void ReloadCache();
 
+	// Removes all dependencies with the given ClassID. Removes any nodes (and corresponding edges) remaining in any MetaSound paged graphs.
 	bool RemoveDependency(const FGuid& InClassID);
+
+	// Removes all dependencies with the given Class Type, Name, & Version Number. Removes any nodes (and corresponding edges) remaining in any MetaSound paged graphs.
 	bool RemoveDependency(EMetasoundFrontendClassType ClassType, const FMetasoundFrontendClassName& InClassName, const FMetasoundFrontendVersionNumber& InClassVersionNumber);
+
 	bool RemoveEdge(const FMetasoundFrontendEdge& EdgeToRemove, const FGuid* InPageID = nullptr);
 
 	// Removes all edges connected to an input or output vertex associated with the node of the given ID.
@@ -373,14 +408,16 @@ public:
 	bool RemoveGraphComment(const FGuid& InCommentID, const FGuid* InPageID = nullptr);
 #endif // WITH_EDITORONLY_DATA
 
-	bool RemoveGraphInput(FName InInputName);
-	bool RemoveGraphOutput(FName InOutputName);
+	bool RemoveGraphInput(FName InputName);
+	bool RemoveGraphOutput(FName OutputName);
 
 #if WITH_EDITORONLY_DATA
 	bool RemoveGraphPage(const FGuid& InPageID);
 #endif // WITH_EDITORONLY_DATA
 
-	bool RemoveInterface(FName InName);
+	bool RemoveGraphVariable(FName VariableName, const FGuid* InPageID = nullptr);
+
+	bool RemoveInterface(FName Name);
 	bool RemoveNamedEdges(const TSet<Metasound::Frontend::FNamedEdge>& InNamedEdgesToRemove, TArray<FMetasoundFrontendEdge>* OutRemovedEdges = nullptr, const FGuid* InPageID = nullptr);
 	bool RemoveNode(const FGuid& InNodeID, const FGuid* InPageID = nullptr);
 
@@ -440,6 +477,11 @@ public:
 	bool SetGraphInputDefault(FName InputName, FMetasoundFrontendLiteral InDefaultLiteral, const FGuid* InPageID = nullptr);
 	bool SetGraphInputDefaults(FName InputName, TArray<FMetasoundFrontendClassInputDefault> Defaults);
 
+	// Sets whether or not graph input inherits default.  By default, updates only if graph is marked as a preset.
+	// Optionally, if bForceUpdate is set, updates inheritance even if not a preset (primarily for clearing flag
+	// if non-preset has unnecessary data).
+	bool SetGraphInputInheritsDefault(FName InName, bool bInputInheritsDefault, bool bForceUpdate = false);
+
 	// Sets a given graph input's name to a new name. Succeeds if the graph output exists and the new name is set (or is the same as the old name).
 	bool SetGraphInputName(FName InputName, FName InName);
 
@@ -465,6 +507,17 @@ public:
 	// Sets the given graph variable's default.
 	bool SetGraphVariableDefault(FName VariableName, FMetasoundFrontendLiteral InDefaultLiteral, const FGuid* InPageID = nullptr);
 
+#if WITH_EDITORONLY_DATA
+	// Sets the given graph variable's description.
+	bool SetGraphVariableDescription(FName VariableName, FText Description, const FGuid* InPageID = nullptr);
+
+	// Sets the given graph variable's display name.
+	bool SetGraphVariableDisplayName(FName VariableName, FText DisplayName, const FGuid* InPageID = nullptr);
+#endif // WITH_EDITORONLY_DATA
+
+	// Sets the given graph variable's description.
+	bool SetGraphVariableName(FName VariableName, FName NewName, const FGuid* InPageID = nullptr);
+
 #if WITH_EDITOR
 	void SetDisplayName(const FText& InDisplayName);
 
@@ -528,6 +581,7 @@ private:
 	using FFinalizeNodeFunctionRef = TFunctionRef<void(FMetasoundFrontendNode&, const Metasound::Frontend::FNodeRegistryKey&)>;
 
 	FMetasoundFrontendNode* AddNodeInternal(const FMetasoundFrontendClassMetadata& InClassMetadata, Metasound::Frontend::FFinalizeNodeFunctionRef FinalizeNode, const FGuid& InPageID, FGuid InNodeID = FGuid::NewGuid(), int32* NewNodeIndex = nullptr);
+	FMetasoundFrontendNode* AddNodeInternal(const Metasound::Frontend::FNodeRegistryKey& InClassKey, FFinalizeNodeFunctionRef FinalizeNode, const FGuid& InPageID, FGuid InNodeID = FGuid::NewGuid(), int32* NewNodeIndex = nullptr);
 	void BeginBuilding(TSharedPtr<Metasound::Frontend::FDocumentModifyDelegates> Delegates = {}, bool bPrimeCache = false);
 
 	// Conforms GraphOutput node's ClassID, Access & Data Type with the GraphOutput.
@@ -544,9 +598,14 @@ private:
 
 	FMetasoundFrontendGraph& FindBuildGraphChecked() const;
 
+	FMetasoundFrontendVariable* FindGraphVariableInternal(FName InVariableName, const FGuid* InPageID = nullptr);
+
 	bool FindNodeClassInterfaces(const FGuid& InNodeID, TSet<FMetasoundFrontendVersion>& OutInterfaces, const FGuid& InPageID) const;
 	FMetasoundFrontendNode* FindNodeInternal(const FGuid& InNodeID, const FGuid* InPageID = nullptr);
 
+	const FMetasoundFrontendNode* FindHeadNodeInVariableStack(FName VariableName, const FGuid* InPageID = nullptr) const;
+	const FMetasoundFrontendNode* FindTailNodeInVariableStack(FName VariableName, const FGuid* InPageID = nullptr) const;
+
 	void IterateNodesConnectedWithVertex(const FMetasoundFrontendVertexHandle& Vertex, TFunctionRef<void(const FMetasoundFrontendEdge&, FMetasoundFrontendNode&)> NodeIndexIterFunc, const FGuid& InPageID);
 
 	const FTopLevelAssetPath GetBuilderClassPath() const;
@@ -555,8 +614,6 @@ private:
 
 	void RemoveSwapDependencyInternal(int32 Index);
 
-	bool SetGraphInputInheritsDefault(FName InName, bool bInputInheritsDefault);
-
 	bool SpliceVariableNodeFromStack(const FGuid& InNodeID, const FGuid& InPageID);
 	bool UnlinkVariableNode(const FGuid& InNodeID, const FGuid& InPageID);
 
