Fix some bad file swaps from prior CL
#rnx

diff --git a/Engine/Plugins/Runtime/Metasound/Source/MetasoundEditor/Private/SMetasoundGraphPin.h b/Engine/Plugins/Runtime/Metasound/Source/MetasoundEditor/Private/SMetasoundGraphPin.h
--- a/Engine/Plugins/Runtime/Metasound/Source/MetasoundEditor/Private/SMetasoundGraphPin.h
+++ b/Engine/Plugins/Runtime/Metasound/Source/MetasoundEditor/Private/SMetasoundGraphPin.h
@@ -263,7 +263,9 @@ namespace Metasound
 					{
 						if (UMetasoundEditorGraphNode* MetasoundEditorNode = Cast<UMetasoundEditorGraphNode>(Pin->GetOwningNode()))
 						{
+PRAGMA_DISABLE_DEPRECATION_WARNINGS
 							FNodeHandle NodeHandle = MetasoundEditorNode->GetNodeHandle();
+PRAGMA_ENABLE_DEPRECATION_WARNINGS
 							return NodeHandle->GetInputWithVertexName(Pin->GetFName());
 						}
 					}
diff --git a/Engine/Plugins/Runtime/Metasound/Source/MetasoundEditor/Public/MetasoundEditorGraphNode.h b/Engine/Plugins/Runtime/Metasound/Source/MetasoundEditor/Public/MetasoundEditorGraphNode.h
--- a/Engine/Plugins/Runtime/Metasound/Source/MetasoundEditor/Public/MetasoundEditorGraphNode.h
+++ b/Engine/Plugins/Runtime/Metasound/Source/MetasoundEditor/Public/MetasoundEditorGraphNode.h
@@ -172,7 +172,10 @@ public:
 
 	Metasound::Frontend::FGraphHandle GetRootGraphHandle() const;
 	Metasound::Frontend::FConstGraphHandle GetConstRootGraphHandle() const;
+
+	UE_DEPRECATED(5.6, "Node Handles are actively being deprecated, use the MetaSound Frontend Document Builder API")
 	Metasound::Frontend::FNodeHandle GetNodeHandle() const;
+
 	Metasound::Frontend::FConstNodeHandle GetConstNodeHandle() const;
 	Metasound::Frontend::FDataTypeRegistryInfo GetPinDataTypeInfo(const UEdGraphPin& InPin) const;
 
@@ -393,10 +396,10 @@ public:
 	virtual UMetasoundEditorGraphMember* GetMember() const override;
 	virtual bool EnableInteractWidgets() const override;
 
-	
 	virtual FMetasoundFrontendClassName GetClassName() const override;
 	virtual FGuid GetNodeID() const override;
 	virtual FName GetCornerIcon() const override;
+	virtual void GetPinHoverText(const UEdGraphPin& Pin, FString& OutHoverText) const override;
 	virtual void PinDefaultValueChanged(UEdGraphPin* Pin) override;
 
 	virtual EMetasoundFrontendClassType GetClassType() const;
diff --git a/Engine/Plugins/Runtime/Metasound/Source/MetasoundFrontend/MetasoundFrontend.Build.cs b/Engine/Plugins/Runtime/Metasound/Source/MetasoundFrontend/MetasoundFrontend.Build.cs
--- a/Engine/Plugins/Runtime/Metasound/Source/MetasoundFrontend/MetasoundFrontend.Build.cs
+++ b/Engine/Plugins/Runtime/Metasound/Source/MetasoundFrontend/MetasoundFrontend.Build.cs
@@ -9,8 +9,6 @@ namespace UnrealBuildTool.Rules
 	{
 		public MetasoundFrontend(ReadOnlyTargetRules Target) : base(Target)
 		{
-			OptimizeCode = CodeOptimization.Never;
-
 			NumIncludedBytesPerUnityCPPOverride = 294912; // best unity size found from using UBT ProfileUnitySizes mode
 
 			PublicDependencyModuleNames.AddRange
