Implement style/better error reporting feedback for prior variable API migration change
#rnx

diff --git a/Engine/Plugins/Runtime/Metasound/Source/MetasoundEngine/Public/MetasoundBuilderBase.cpp b/Engine/Plugins/Runtime/Metasound/Source/MetasoundEngine/Public/MetasoundBuilderBase.cpp
--- a/Engine/Plugins/Runtime/Metasound/Source/MetasoundEngine/Public/MetasoundBuilderBase.cpp
+++ b/Engine/Plugins/Runtime/Metasound/Source/MetasoundEngine/Public/MetasoundBuilderBase.cpp
@@ -173,45 +173,18 @@ void UMetaSoundBuilderBase::AddGraphPage(FName PageName, bool bDuplicateLastGrap
 
 void UMetaSoundBuilderBase::AddGraphVariable(FName Name, FName DataType, FMetasoundFrontendLiteral DefaultValue, EMetaSoundBuilderResult& OutResult)
 {
-	using namespace Metasound;
-
-	OutResult = EMetaSoundBuilderResult::Failed;
-	if (Frontend::IDataTypeRegistry::Get().FindDataTypeRegistryEntry(DataType) == nullptr)
-	{
-		UE_LOG(LogMetaSound, Error, TEXT("AddGraphVariable Failed: Attempted creation of variable '%s' with unregistered DataType '%s'"), *Name.ToString(), *DataType.ToString());
-	}
-	else
-	{
-		const FMetasoundFrontendVariable* Variable = Builder.FindGraphVariable(Name);
-		if (Variable)
-		{
-			UE_LOG(LogMetaSound, Warning, TEXT("AddGraphVariable Failed: Variable already exists with name '%s' (existing DataType '%s', requested DataType '%s')"), *Name.ToString(), *Variable->TypeName.ToString(), *DataType.ToString());
-		}
-		else
-		{
-			Variable = Builder.AddGraphVariable(Name, DataType, &DefaultValue);
-			OutResult = Variable ? EMetaSoundBuilderResult::Succeeded : EMetaSoundBuilderResult::Failed;
-		}
-	}
+	const FMetasoundFrontendVariable* Variable = Builder.AddGraphVariable(Name, DataType, &DefaultValue);
+	OutResult = Variable ? EMetaSoundBuilderResult::Succeeded : EMetaSoundBuilderResult::Failed;
 }
 
 FMetaSoundNodeHandle UMetaSoundBuilderBase::AddGraphVariableGetNode(FName Name, EMetaSoundBuilderResult& OutResult)
 {
-	using namespace Metasound;
-
 	OutResult = EMetaSoundBuilderResult::Failed;
 	FMetaSoundNodeHandle NodeHandle;
-	if (const FMetasoundFrontendVariable* Variable = Builder.FindGraphVariable(Name))
-	{
-		if (const FMetasoundFrontendNode* AccessorNode = Builder.AddGraphVariableAccessorNode(Name))
-		{
-			NodeHandle.NodeID = AccessorNode->GetID();
-			OutResult = EMetaSoundBuilderResult::Succeeded;
-		}
-	}
-	else
+	if (const FMetasoundFrontendNode* AccessorNode = Builder.AddGraphVariableAccessorNode(Name))
 	{
-		UE_LOG(LogMetaSound, Warning, TEXT("AddGraphVariableGetNode Failed: Variable does not exists with name '%s'"), *Name.ToString());
+		NodeHandle.NodeID = AccessorNode->GetID();
+		OutResult = EMetaSoundBuilderResult::Succeeded;
 	}
 
 	return NodeHandle;
@@ -223,17 +196,10 @@ FMetaSoundNodeHandle UMetaSoundBuilderBase::AddGraphVariableGetDelayedNode(FName
 
 	OutResult = EMetaSoundBuilderResult::Failed;
 	FMetaSoundNodeHandle NodeHandle;
-	if (const FMetasoundFrontendVariable* Variable = Builder.FindGraphVariable(Name))
-	{
-		if (const FMetasoundFrontendNode* AccessorNode = Builder.AddGraphVariableDeferredAccessorNode(Name))
-		{
-			NodeHandle.NodeID = AccessorNode->GetID();
-			OutResult = EMetaSoundBuilderResult::Succeeded;
-		}
-	}
-	else
+	if (const FMetasoundFrontendNode* AccessorNode = Builder.AddGraphVariableDeferredAccessorNode(Name))
 	{
-		UE_LOG(LogMetaSound, Warning, TEXT("AddGraphVariableGetDelayedNode Failed: Variable does not exists with name '%s'"), *Name.ToString());
+		NodeHandle.NodeID = AccessorNode->GetID();
+		OutResult = EMetaSoundBuilderResult::Succeeded;
 	}
 
 	return NodeHandle;
@@ -245,23 +211,10 @@ FMetaSoundNodeHandle UMetaSoundBuilderBase::AddGraphVariableSetNode(FName Name,
 
 	OutResult = EMetaSoundBuilderResult::Failed;
 	FMetaSoundNodeHandle NodeHandle;
-	if (const FMetasoundFrontendVariable* Variable = Builder.FindGraphVariable(Name))
-	{
-		const bool bAlreadyContainsMutator = Variable->MutatorNodeID.IsValid();
-		if (const FMetasoundFrontendNode* MutatorNode = Builder.AddGraphVariableMutatorNode(Name))
-		{
-			NodeHandle.NodeID = MutatorNode->GetID();
-
-			// Frontend Document builder already logs warning if
-			// attempting to add addition mutator and returns the
-			// existing one, so just handle the OutResult but no
-			// log necessary.
-			OutResult = bAlreadyContainsMutator ? EMetaSoundBuilderResult::Failed : EMetaSoundBuilderResult::Succeeded;
-		}
-	}
-	else
+	if (const FMetasoundFrontendNode* MutatorNode = Builder.AddGraphVariableMutatorNode(Name))
 	{
-		UE_LOG(LogMetaSound, Warning, TEXT("AddGraphVariableSetNode Failed: Variable does not exists with name '%s'"), *Name.ToString());
+		NodeHandle.NodeID = MutatorNode->GetID();
+		OutResult = EMetaSoundBuilderResult::Succeeded;
 	}
 
 	return NodeHandle;
diff --git a/Engine/Plugins/Runtime/Metasound/Source/MetasoundFrontend/Private/MetasoundFrontendDocumentBuilder.cpp b/Engine/Plugins/Runtime/Metasound/Source/MetasoundFrontend/Private/MetasoundFrontendDocumentBuilder.cpp
--- a/Engine/Plugins/Runtime/Metasound/Source/MetasoundFrontend/Private/MetasoundFrontendDocumentBuilder.cpp
+++ b/Engine/Plugins/Runtime/Metasound/Source/MetasoundFrontend/Private/MetasoundFrontendDocumentBuilder.cpp
@@ -1291,91 +1291,99 @@ const FMetasoundFrontendVariable* FMetaSoundFrontendDocumentBuilder::AddGraphVar
 {
 	using namespace Metasound::Frontend;
 
-	checkf(FindGraphVariable(VariableName, InPageID) == nullptr, TEXT("Attempted to add variable '%s' that already exists"), *VariableName.ToString());
+	if (const FMetasoundFrontendVariable* ExistingVariable = FindGraphVariable(VariableName))
+	{
+		UE_LOG(LogMetaSound, Warning, TEXT("AddGraphVariable Failed: Variable already exists with name '%s' (existing DataType '%s', requested DataType '%s')"),
+			*VariableName.ToString(), *ExistingVariable->TypeName.ToString(), *DataType.ToString());
+		return nullptr;
+	}
 
 	const IDataTypeRegistry& Registry = IDataTypeRegistry::Get();
 	FDataTypeRegistryInfo Info;
-	if (ensure(Registry.GetDataTypeInfo(DataType, Info)))
+	if (!Registry.GetDataTypeInfo(DataType, Info))
 	{
-		FMetasoundFrontendVariable Variable
-		{
-			.Name = VariableName,
-			.ID = FGuid::NewGuid()
-		};
+		UE_LOG(LogMetaSound, Error, TEXT("AddGraphVariable Failed: Attempted creation of variable '%s' with unregistered DataType '%s'"), *VariableName.ToString(), *DataType.ToString());
+		return nullptr;
+	}
 
-		Variable.TypeName = Info.DataTypeName;
-		if (Literal)
-		{
-			Variable.Literal = *Literal;
-		}
-		else
-		{
-			Variable.Literal.SetFromLiteral(Registry.CreateDefaultLiteral(DataType));
-		}
+	FMetasoundFrontendVariable Variable
+	{
+		.Name = VariableName,
+		.ID = FGuid::NewGuid()
+	};
+
+	Variable.TypeName = Info.DataTypeName;
+	if (Literal)
+	{
+		Variable.Literal = *Literal;
+	}
+	else
+	{
+		Variable.Literal.SetFromLiteral(Registry.CreateDefaultLiteral(DataType));
+	}
 
 #if WITH_EDITORONLY_DATA
-		if (DisplayName)
-		{
-			Variable.DisplayName = *DisplayName;
-		}
+	if (DisplayName)
+	{
+		Variable.DisplayName = *DisplayName;
+	}
 
-		if (Description)
-		{
-			Variable.Description = *Description;
-		}
+	if (Description)
+	{
+		Variable.Description = *Description;
+	}
 #endif // WITH_EDITORONLY_DATA
 
 #if WITH_EDITOR
-		GetDocumentChecked().Metadata.ModifyContext.AddMemberIDModified(Variable.ID);
+	GetDocumentChecked().Metadata.ModifyContext.AddMemberIDModified(Variable.ID);
 #endif // WITH_EDITOR
 
-		FNodeRegistryKey VariableNodeClassKey;
+	FNodeRegistryKey VariableNodeClassKey;
+	{
+		FMetasoundFrontendClass VariableNodeClass;
+		if (!IDataTypeRegistry::Get().GetFrontendVariableClass(Variable.TypeName, VariableNodeClass))
 		{
-			FMetasoundFrontendClass VariableNodeClass;
-			if (!IDataTypeRegistry::Get().GetFrontendVariableClass(Variable.TypeName, VariableNodeClass))
-			{
-				return nullptr;
-			}
-
-			VariableNodeClassKey = FNodeRegistryKey(VariableNodeClass.Metadata);
-			const FMetasoundFrontendClass* Dependency = FindDependency(VariableNodeClass.Metadata);
-			if (!Dependency)
-			{
-				Dependency = AddDependency(MoveTemp(VariableNodeClass));
-			}
-			check(Dependency);
+			return nullptr;
 		}
 
-		auto FinalizeNode = [](FMetasoundFrontendNode& InOutNode, const Metasound::Frontend::FNodeRegistryKey& ClassKey)
+		VariableNodeClassKey = FNodeRegistryKey(VariableNodeClass.Metadata);
+		const FMetasoundFrontendClass* Dependency = FindDependency(VariableNodeClass.Metadata);
+		if (!Dependency)
 		{
+			Dependency = AddDependency(MoveTemp(VariableNodeClass));
+		}
+		check(Dependency);
+	}
+
+	auto FinalizeNode = [](FMetasoundFrontendNode& InOutNode, const Metasound::Frontend::FNodeRegistryKey& ClassKey)
+	{
 #if WITH_EDITOR
-			using namespace Metasound::Frontend;
+		using namespace Metasound::Frontend;
 
-			// Cache the asset name on the node if it node is reference to asset-defined graph.
-			const FTopLevelAssetPath Path = IMetaSoundAssetManager::GetChecked().FindAssetPath(FAssetKey(ClassKey.ClassName, ClassKey.Version));
-			if (Path.IsValid())
-			{
-				InOutNode.Name = Path.GetAssetName();
-				return;
-			}
+		// Cache the asset name on the node if it node is reference to asset-defined graph.
+		const FTopLevelAssetPath Path = IMetaSoundAssetManager::GetChecked().FindAssetPath(FAssetKey(ClassKey.ClassName, ClassKey.Version));
+		if (Path.IsValid())
+		{
+			InOutNode.Name = Path.GetAssetName();
+			return;
+		}
 
-			InOutNode.Name = ClassKey.ClassName.GetFullName();
+		InOutNode.Name = ClassKey.ClassName.GetFullName();
 #endif // WITH_EDITOR
-		};
+	};
 
-		const FGuid& PageID = InPageID ? *InPageID : BuildPageID;
-		if (FMetasoundFrontendNode* VariableNode = AddNodeInternal(VariableNodeClassKey, FinalizeNode, PageID))
-		{
-			Variable.VariableNodeID = VariableNode->GetID();
-			FMetasoundFrontendGraph& Graph = GetDocumentChecked().RootGraph.FindGraphChecked(PageID);
-			return &Graph.Variables.Add_GetRef(MoveTemp(Variable));
-		}
+	const FGuid& PageID = InPageID ? *InPageID : BuildPageID;
+	if (FMetasoundFrontendNode* VariableNode = AddNodeInternal(VariableNodeClassKey, FinalizeNode, PageID))
+	{
+		Variable.VariableNodeID = VariableNode->GetID();
+		FMetasoundFrontendGraph& Graph = GetDocumentChecked().RootGraph.FindGraphChecked(PageID);
+		return &Graph.Variables.Add_GetRef(MoveTemp(Variable));
 	}
 
 	return nullptr;
 }
 
-const FMetasoundFrontendNode* FMetaSoundFrontendDocumentBuilder::AddGraphVariableNode(FName VariableName, EMetasoundFrontendClassType ClassType, FGuid InNodeID,const FGuid* InPageID)
+const FMetasoundFrontendNode* FMetaSoundFrontendDocumentBuilder::AddGraphVariableNode(FName VariableName, EMetasoundFrontendClassType ClassType, FGuid InNodeID, const FGuid* InPageID)
 {
 	using namespace Metasound::Frontend;
 
@@ -1389,15 +1397,7 @@ const FMetasoundFrontendNode* FMetaSoundFrontendDocumentBuilder::AddGraphVariabl
 
 		case EMetasoundFrontendClassType::VariableMutator:
 		{
-			const FMetasoundFrontendVariable* Variable = FindGraphVariable(VariableName, InPageID);
-			if (Variable && Variable->MutatorNodeID.IsValid())
-			{
-				UE_LOG(LogMetaSound, Error, TEXT("Cannot add node: variable mutator already exists for variable \"%s\""), *Variable->Name.ToString());
-			}
-			else
-			{
-				return AddGraphVariableMutatorNode(VariableName, InNodeID, InPageID);
-			}
+			return AddGraphVariableMutatorNode(VariableName, InNodeID, InPageID);
 		}
 		break;
 
@@ -1410,140 +1410,148 @@ const FMetasoundFrontendNode* FMetaSoundFrontendDocumentBuilder::AddGraphVariabl
 	return nullptr;
 }
 
-const FMetasoundFrontendNode* FMetaSoundFrontendDocumentBuilder::AddGraphVariableAccessorNode(FName VariableName, FGuid InNodeID,const FGuid* InPageID)
+const FMetasoundFrontendNode* FMetaSoundFrontendDocumentBuilder::AddGraphVariableAccessorNode(FName VariableName, FGuid InNodeID, const FGuid* InPageID)
 {
 	using namespace Metasound::Frontend;
 	using namespace Metasound::VariableNames;
 
-	if (FMetasoundFrontendVariable* Variable = FindGraphVariableInternal(VariableName, InPageID))
+	FMetasoundFrontendVariable* Variable = FindGraphVariableInternal(VariableName, InPageID);
+	if (!Variable)
+	{
+		UE_LOG(LogMetaSound, Error, TEXT("AddGraphVariableAccessorNode Failed: Variable does not exists with name '%s'"), *VariableName.ToString());
+		return nullptr;
+	} 
+	
+	FNodeRegistryKey VariableNodeClassKey;
 	{
-		FNodeRegistryKey VariableNodeClassKey;
+		FMetasoundFrontendClass NodeClass;
+		if (!IDataTypeRegistry::Get().GetFrontendVariableAccessorClass(Variable->TypeName, NodeClass))
 		{
-			FMetasoundFrontendClass NodeClass;
-			if (!IDataTypeRegistry::Get().GetFrontendVariableAccessorClass(Variable->TypeName, NodeClass))
-			{
-				UE_LOG(LogMetaSound, Error, TEXT("Could not find registered \"get variable\" node class for data type \"%s\""), *Variable->TypeName.ToString());
-				return nullptr;
-			}
-
-			VariableNodeClassKey = FNodeRegistryKey(NodeClass.Metadata);
-			const FMetasoundFrontendClass* Dependency = FindDependency(NodeClass.Metadata);
-			if (!Dependency)
-			{
-				Dependency = AddDependency(MoveTemp(NodeClass));
-			}
-			check(Dependency);
+			UE_LOG(LogMetaSound, Error, TEXT("Could not find registered \"get variable\" node class for data type \"%s\""), *Variable->TypeName.ToString());
+			return nullptr;
 		}
 
-		auto FinalizeNodeFunction = [](const FMetasoundFrontendNode&, const Metasound::Frontend::FNodeRegistryKey&) { };
-		const FGuid& PageID = InPageID ? *InPageID : BuildPageID;
-		if (const FMetasoundFrontendNode* NewNode = AddNodeInternal(VariableNodeClassKey, FinalizeNodeFunction, PageID, InNodeID))
+		VariableNodeClassKey = FNodeRegistryKey(NodeClass.Metadata);
+		const FMetasoundFrontendClass* Dependency = FindDependency(NodeClass.Metadata);
+		if (!Dependency)
 		{
-			// Connect new node.
-			const FMetasoundFrontendVertex* NewInput = FindNodeInput(NewNode->GetID(), METASOUND_GET_PARAM_NAME(InputVariable), InPageID);
-			check(NewInput);
+			Dependency = AddDependency(MoveTemp(NodeClass));
+		}
+		check(Dependency);
+	}
 
-			const FMetasoundFrontendNode* TailNode = FindTailNodeInVariableStack(VariableName, InPageID);
-			if (!TailNode)
-			{
-				// variable stack is empty. Connect to init variable node.
-				TailNode = FindNode(Variable->VariableNodeID, InPageID);
-			}
+	auto FinalizeNodeFunction = [](const FMetasoundFrontendNode&, const Metasound::Frontend::FNodeRegistryKey&) { };
+	const FGuid& PageID = InPageID ? *InPageID : BuildPageID;
+	if (const FMetasoundFrontendNode* NewNode = AddNodeInternal(VariableNodeClassKey, FinalizeNodeFunction, PageID, InNodeID))
+	{
+		// Connect new node.
+		const FMetasoundFrontendVertex* NewInput = FindNodeInput(NewNode->GetID(), METASOUND_GET_PARAM_NAME(InputVariable), InPageID);
+		check(NewInput);
 
-			if (ensure(TailNode))
-			{
-				// connect new node to the last "get" node.
-				const FMetasoundFrontendVertex* TailNodeOutput = FindNodeOutput(TailNode->GetID(), METASOUND_GET_PARAM_NAME(OutputVariable), InPageID);
-				check(TailNodeOutput);
-
-				FMetasoundFrontendEdge NewEdge;
-				NewEdge.FromNodeID = TailNode->GetID();
-				NewEdge.FromVertexID = TailNodeOutput->VertexID;
-				NewEdge.ToNodeID = NewNode->GetID();
-				NewEdge.ToVertexID = NewInput->VertexID;
-				AddEdge(MoveTemp(NewEdge), InPageID);
-			}
+		const FMetasoundFrontendNode* TailNode = FindTailNodeInVariableStack(VariableName, InPageID);
+		if (!TailNode)
+		{
+			// variable stack is empty. Connect to init variable node.
+			TailNode = FindNode(Variable->VariableNodeID, InPageID);
+		}
+
+		if (ensure(TailNode))
+		{
+			// connect new node to the last "get" node.
+			const FMetasoundFrontendVertex* TailNodeOutput = FindNodeOutput(TailNode->GetID(), METASOUND_GET_PARAM_NAME(OutputVariable), InPageID);
+			check(TailNodeOutput);
 
-			// Add node ID to variable after connecting since the array
-			// order of node ids is used to determine whether a node 
-			// is the tail node.
-			Variable->AccessorNodeIDs.Add(NewNode->GetID());
-			return NewNode;
+			FMetasoundFrontendEdge NewEdge;
+			NewEdge.FromNodeID = TailNode->GetID();
+			NewEdge.FromVertexID = TailNodeOutput->VertexID;
+			NewEdge.ToNodeID = NewNode->GetID();
+			NewEdge.ToVertexID = NewInput->VertexID;
+			AddEdge(MoveTemp(NewEdge), InPageID);
 		}
+
+		// Add node ID to variable after connecting since the array
+		// order of node ids is used to determine whether a node 
+		// is the tail node.
+		Variable->AccessorNodeIDs.Add(NewNode->GetID());
+		return NewNode;
 	}
 
 	return nullptr;
 }
 
-const FMetasoundFrontendNode* FMetaSoundFrontendDocumentBuilder::AddGraphVariableDeferredAccessorNode(FName VariableName, FGuid InNodeID,const FGuid* InPageID)
+const FMetasoundFrontendNode* FMetaSoundFrontendDocumentBuilder::AddGraphVariableDeferredAccessorNode(FName VariableName, FGuid InNodeID, const FGuid* InPageID)
 {
 	using namespace Metasound::Frontend;
 	using namespace Metasound::VariableNames;
 
-	if (FMetasoundFrontendVariable* Variable = FindGraphVariableInternal(VariableName, InPageID))
+	FMetasoundFrontendVariable* Variable = FindGraphVariableInternal(VariableName, InPageID);
+	if (!Variable)
 	{
-		FNodeRegistryKey ClassKey;
-		{
-			FMetasoundFrontendClass NodeClass;
-			if (!IDataTypeRegistry::Get().GetFrontendVariableDeferredAccessorClass(Variable->TypeName, NodeClass))
-			{
-				UE_LOG(LogMetaSound, Error, TEXT("Could not find registered \"get variable\" node class for data type \"%s\""), *Variable->TypeName.ToString());
-				return nullptr;
-			}
+		UE_LOG(LogMetaSound, Error, TEXT("AddGraphVariableGetDelayedNode Failed: Variable does not exists with name '%s'"), *VariableName.ToString());
+		return nullptr;
+	}
 
-			ClassKey = FNodeRegistryKey(NodeClass.Metadata);
-			const FMetasoundFrontendClass* Dependency = FindDependency(NodeClass.Metadata);
-			if (!Dependency)
-			{
-				Dependency = AddDependency(MoveTemp(NodeClass));
-			}
-			check(Dependency);
+	FNodeRegistryKey ClassKey;
+	{
+		FMetasoundFrontendClass NodeClass;
+		if (!IDataTypeRegistry::Get().GetFrontendVariableDeferredAccessorClass(Variable->TypeName, NodeClass))
+		{
+			UE_LOG(LogMetaSound, Error, TEXT("AddGraphVariableGetDelayedNode Failed: Could not find registered \"get variable\" node class for data type \"%s\""), *Variable->TypeName.ToString());
+			return nullptr;
 		}
 
-		auto FinalizeNodeFunction = [](const FMetasoundFrontendNode&, const Metasound::Frontend::FNodeRegistryKey&) { };
-		const FGuid& PageID = InPageID ? *InPageID : BuildPageID;
-		if (const FMetasoundFrontendNode* NewNode = AddNodeInternal(ClassKey, FinalizeNodeFunction, PageID, InNodeID))
+		ClassKey = FNodeRegistryKey(NodeClass.Metadata);
+		const FMetasoundFrontendClass* Dependency = FindDependency(NodeClass.Metadata);
+		if (!Dependency)
 		{
-			// Connect new node.
-			const FMetasoundFrontendVertex* NewNodeOutput = FindNodeOutput(NewNode->GetID(), METASOUND_GET_PARAM_NAME(OutputVariable), InPageID);
-			const FMetasoundFrontendNode* HeadNode = FindHeadNodeInVariableStack(VariableName, InPageID);
-			if (HeadNode)
-			{
-				const FMetasoundFrontendVertex* HeadNodeInput = FindNodeInput(HeadNode->GetID(), METASOUND_GET_PARAM_NAME(InputVariable), InPageID);
-				check(HeadNodeInput);
+			Dependency = AddDependency(MoveTemp(NodeClass));
+		}
+		check(Dependency);
+	}
 
-				RemoveEdgeToNodeInput(HeadNode->GetID(), HeadNodeInput->VertexID, InPageID);
+	auto FinalizeNodeFunction = [](const FMetasoundFrontendNode&, const Metasound::Frontend::FNodeRegistryKey&) { };
+	const FGuid& PageID = InPageID ? *InPageID : BuildPageID;
+	if (const FMetasoundFrontendNode* NewNode = AddNodeInternal(ClassKey, FinalizeNodeFunction, PageID, InNodeID))
+	{
+		// Connect new node.
+		const FMetasoundFrontendVertex* NewNodeOutput = FindNodeOutput(NewNode->GetID(), METASOUND_GET_PARAM_NAME(OutputVariable), InPageID);
+		const FMetasoundFrontendNode* HeadNode = FindHeadNodeInVariableStack(VariableName, InPageID);
+		if (HeadNode)
+		{
+			const FMetasoundFrontendVertex* HeadNodeInput = FindNodeInput(HeadNode->GetID(), METASOUND_GET_PARAM_NAME(InputVariable), InPageID);
+			check(HeadNodeInput);
 
-				FMetasoundFrontendEdge NewEdge;
-				NewEdge.FromNodeID = NewNode->GetID();
-				NewEdge.FromVertexID = NewNodeOutput->VertexID;
-				NewEdge.ToNodeID = HeadNode->GetID();
-				NewEdge.ToVertexID = HeadNodeInput->VertexID;
-				AddEdge(MoveTemp(NewEdge), InPageID);
-			}
+			RemoveEdgeToNodeInput(HeadNode->GetID(), HeadNodeInput->VertexID, InPageID);
 
-			const FMetasoundFrontendVertex* NewNodeInput = FindNodeInput(NewNode->GetID(), METASOUND_GET_PARAM_NAME(InputVariable), InPageID);
-			check(NewNodeInput);
+			FMetasoundFrontendEdge NewEdge;
+			NewEdge.FromNodeID = NewNode->GetID();
+			NewEdge.FromVertexID = NewNodeOutput->VertexID;
+			NewEdge.ToNodeID = HeadNode->GetID();
+			NewEdge.ToVertexID = HeadNodeInput->VertexID;
+			AddEdge(MoveTemp(NewEdge), InPageID);
+		}
 
-			const FMetasoundFrontendNode* VariableNode = FindNode(Variable->VariableNodeID, InPageID);
-			check(VariableNode);
+		const FMetasoundFrontendVertex* NewNodeInput = FindNodeInput(NewNode->GetID(), METASOUND_GET_PARAM_NAME(InputVariable), InPageID);
+		check(NewNodeInput);
 
-			const FMetasoundFrontendVertex* VariableNodeOutput = FindNodeOutput(VariableNode->GetID(), METASOUND_GET_PARAM_NAME(OutputVariable), InPageID);
-			check(VariableNodeOutput);
+		const FMetasoundFrontendNode* VariableNode = FindNode(Variable->VariableNodeID, InPageID);
+		check(VariableNode);
 
-			FMetasoundFrontendEdge NewEdge;
-			NewEdge.FromNodeID = VariableNode->GetID();
-			NewEdge.FromVertexID = VariableNodeOutput->VertexID;
-			NewEdge.ToNodeID = NewNode->GetID();
-			NewEdge.ToVertexID = NewNodeInput->VertexID;
-			AddEdge(MoveTemp(NewEdge), InPageID);
+		const FMetasoundFrontendVertex* VariableNodeOutput = FindNodeOutput(VariableNode->GetID(), METASOUND_GET_PARAM_NAME(OutputVariable), InPageID);
+		check(VariableNodeOutput);
 
-			// Add node ID to variable after connecting since the array
-			// order of node ids is used to determine whether a node 
-			// is the tail node.
-			Variable->DeferredAccessorNodeIDs.Add(NewNode->GetID());
-			return NewNode;
-		}
+		FMetasoundFrontendEdge NewEdge;
+		NewEdge.FromNodeID = VariableNode->GetID();
+		NewEdge.FromVertexID = VariableNodeOutput->VertexID;
+		NewEdge.ToNodeID = NewNode->GetID();
+		NewEdge.ToVertexID = NewNodeInput->VertexID;
+		AddEdge(MoveTemp(NewEdge), InPageID);
+
+		// Add node ID to variable after connecting since the array
+		// order of node ids is used to determine whether a node 
+		// is the tail node.
+		Variable->DeferredAccessorNodeIDs.Add(NewNode->GetID());
+		return NewNode;
 	}
 
 	return nullptr;
@@ -1554,94 +1562,96 @@ const FMetasoundFrontendNode* FMetaSoundFrontendDocumentBuilder::AddGraphVariabl
 	using namespace Metasound::Frontend;
 	using namespace Metasound::VariableNames;
 
-	if (FMetasoundFrontendVariable* Variable = FindGraphVariableInternal(VariableName, InPageID))
+	FMetasoundFrontendVariable* Variable = FindGraphVariableInternal(VariableName, InPageID);
+	if (!Variable)
+	{
+		UE_LOG(LogMetaSound, Error, TEXT("AddGraphVariableMutatorNode Failed: Variable does not exists with name '%s'"), *VariableName.ToString());
+		return nullptr;
+	}
+
+	if (const FMetasoundFrontendNode* ExistingMutatorNode = FindNode(Variable->MutatorNodeID, InPageID))
+	{
+		UE_LOG(LogMetaSound, Error, TEXT("Cannot add mutator node as one already exists for variable '%s'."), *VariableName.ToString());
+		return nullptr;
+	}
+
+	FNodeRegistryKey ClassKey;
 	{
-		const FMetasoundFrontendNode* MutatorNode = FindNode(Variable->MutatorNodeID, InPageID);
-		if (MutatorNode)
+		FMetasoundFrontendClass MutatorNodeClass;
+		if (!IDataTypeRegistry::Get().GetFrontendVariableMutatorClass(Variable->TypeName, MutatorNodeClass))
 		{
-			UE_LOG(LogMetaSound, Warning, TEXT("Cannot add mutator node as one already exists for variable '%s'. Returning existing mutator"), *VariableName.ToString());
-			return MutatorNode;
+			UE_LOG(LogMetaSound, Error, TEXT("Could not find registered \"set variable\" node class for data type \"%s\""), *Variable->TypeName.ToString());
+			return nullptr;
 		}
-		else
-		{
-			FNodeRegistryKey ClassKey;
-			{
-				FMetasoundFrontendClass MutatorNodeClass;
-				if (!IDataTypeRegistry::Get().GetFrontendVariableMutatorClass(Variable->TypeName, MutatorNodeClass))
-				{
-					UE_LOG(LogMetaSound, Error, TEXT("Could not find registered \"set variable\" node class for data type \"%s\""), *Variable->TypeName.ToString());
-					return nullptr;
-				}
 
-				ClassKey = FNodeRegistryKey(MutatorNodeClass.Metadata);
-				const FMetasoundFrontendClass* Dependency = FindDependency(MutatorNodeClass.Metadata);
-				if (!Dependency)
-				{
-					Dependency = AddDependency(MoveTemp(MutatorNodeClass));
-				}
-				check(Dependency);
-			}
+		ClassKey = FNodeRegistryKey(MutatorNodeClass.Metadata);
+		const FMetasoundFrontendClass* Dependency = FindDependency(MutatorNodeClass.Metadata);
+		if (!Dependency)
+		{
+			Dependency = AddDependency(MoveTemp(MutatorNodeClass));
+		}
+		check(Dependency);
+	}
 
-			auto FinalizeNodeFunction = [](const FMetasoundFrontendNode&, const Metasound::Frontend::FNodeRegistryKey&) { };
-			const FGuid& PageID = InPageID ? *InPageID : BuildPageID;
-			MutatorNode = AddNodeInternal(ClassKey, FinalizeNodeFunction, PageID, InNodeID);
-			if (MutatorNode)
-			{
-				// Initialize mutator default literal value to that of the variable
-				const FMetasoundFrontendVertex* MutatorDataInput = FindNodeInput(MutatorNode->GetID(), METASOUND_GET_PARAM_NAME(InputData), InPageID);
-				check(MutatorDataInput);
-				SetNodeInputDefault(MutatorNode->GetID(), MutatorDataInput->VertexID, Variable->Literal, InPageID);
+	auto FinalizeNodeFunction = [](const FMetasoundFrontendNode&, const Metasound::Frontend::FNodeRegistryKey&) { };
+	const FGuid& PageID = InPageID ? *InPageID : BuildPageID;
+	const FMetasoundFrontendNode* MutatorNode = AddNodeInternal(ClassKey, FinalizeNodeFunction, PageID, InNodeID);
+	if (MutatorNode)
+	{
+		// Initialize mutator default literal value to that of the variable
+		const FMetasoundFrontendVertex* MutatorDataInput = FindNodeInput(MutatorNode->GetID(), METASOUND_GET_PARAM_NAME(InputData), InPageID);
+		check(MutatorDataInput);
+		SetNodeInputDefault(MutatorNode->GetID(), MutatorDataInput->VertexID, Variable->Literal, InPageID);
 
-				Variable->MutatorNodeID = MutatorNode->GetID();
-				FGuid SourceVariableNodeID = Variable->VariableNodeID;
+		Variable->MutatorNodeID = MutatorNode->GetID();
+		FGuid SourceVariableNodeID = Variable->VariableNodeID;
 
-				// Connect last delayed getter in variable stack.
-				if (!Variable->DeferredAccessorNodeIDs.IsEmpty())
-				{
-					SourceVariableNodeID = Variable->DeferredAccessorNodeIDs.Last();
-				}
+		// Connect last delayed getter in variable stack.
+		if (!Variable->DeferredAccessorNodeIDs.IsEmpty())
+		{
+			SourceVariableNodeID = Variable->DeferredAccessorNodeIDs.Last();
+		}
 
-				const FMetasoundFrontendNode* SourceVariableNode = FindNode(SourceVariableNodeID, InPageID);
-				if (ensure(SourceVariableNode))
-				{
-					const FMetasoundFrontendVertex* MutatorNodeInput = FindNodeInput(MutatorNode->GetID(), METASOUND_GET_PARAM_NAME(InputVariable), InPageID);
-					check(MutatorNodeInput);
+		const FMetasoundFrontendNode* SourceVariableNode = FindNode(SourceVariableNodeID, InPageID);
+		if (ensure(SourceVariableNode))
+		{
+			const FMetasoundFrontendVertex* MutatorNodeInput = FindNodeInput(MutatorNode->GetID(), METASOUND_GET_PARAM_NAME(InputVariable), InPageID);
+			check(MutatorNodeInput);
 
-					const FMetasoundFrontendNode* VariableSourceNode = FindNode(SourceVariableNodeID, InPageID);
-					check(VariableSourceNode);
+			const FMetasoundFrontendNode* VariableSourceNode = FindNode(SourceVariableNodeID, InPageID);
+			check(VariableSourceNode);
 
-					const FMetasoundFrontendVertex* SourceVariableNodeOutput = FindNodeOutput(VariableSourceNode->GetID(), METASOUND_GET_PARAM_NAME(OutputVariable), InPageID);
-					check(SourceVariableNodeOutput);
+			const FMetasoundFrontendVertex* SourceVariableNodeOutput = FindNodeOutput(VariableSourceNode->GetID(), METASOUND_GET_PARAM_NAME(OutputVariable), InPageID);
+			check(SourceVariableNodeOutput);
 
-					FMetasoundFrontendEdge NewEdge;
-					NewEdge.FromNodeID = SourceVariableNodeID;
-					NewEdge.FromVertexID = SourceVariableNodeOutput->VertexID;
-					NewEdge.ToNodeID = MutatorNode->GetID();
-					NewEdge.ToVertexID = MutatorNodeInput->VertexID;
-					AddEdge(MoveTemp(NewEdge), InPageID);
-				}
+			FMetasoundFrontendEdge NewEdge;
+			NewEdge.FromNodeID = SourceVariableNodeID;
+			NewEdge.FromVertexID = SourceVariableNodeOutput->VertexID;
+			NewEdge.ToNodeID = MutatorNode->GetID();
+			NewEdge.ToVertexID = MutatorNodeInput->VertexID;
+			AddEdge(MoveTemp(NewEdge), InPageID);
+		}
 
-				// Connect to first inline getter in variable stack
-				if (!Variable->AccessorNodeIDs.IsEmpty())
-				{
-					const FGuid& HeadAccessorNodeID = Variable->AccessorNodeIDs[0];
-					const FMetasoundFrontendVertex* MutatorNodeOutput = FindNodeOutput(MutatorNode->GetID(), METASOUND_GET_PARAM_NAME(OutputVariable), InPageID);
-					check(MutatorNodeOutput);
+		// Connect to first inline getter in variable stack
+		if (!Variable->AccessorNodeIDs.IsEmpty())
+		{
+			const FGuid& HeadAccessorNodeID = Variable->AccessorNodeIDs[0];
+			const FMetasoundFrontendVertex* MutatorNodeOutput = FindNodeOutput(MutatorNode->GetID(), METASOUND_GET_PARAM_NAME(OutputVariable), InPageID);
+			check(MutatorNodeOutput);
 
-					const FMetasoundFrontendVertex* AccessorNodeInput = FindNodeInput(HeadAccessorNodeID, METASOUND_GET_PARAM_NAME(InputVariable), InPageID);
-					check(AccessorNodeInput);
+			const FMetasoundFrontendVertex* AccessorNodeInput = FindNodeInput(HeadAccessorNodeID, METASOUND_GET_PARAM_NAME(InputVariable), InPageID);
+			check(AccessorNodeInput);
 
-					RemoveEdgeToNodeInput(HeadAccessorNodeID, AccessorNodeInput->VertexID, InPageID);
+			RemoveEdgeToNodeInput(HeadAccessorNodeID, AccessorNodeInput->VertexID, InPageID);
 
-					FMetasoundFrontendEdge NewEdge;
-					NewEdge.FromNodeID = MutatorNode->GetID();
-					NewEdge.FromVertexID = MutatorNodeOutput->VertexID;
-					NewEdge.ToNodeID = HeadAccessorNodeID;
-					NewEdge.ToVertexID = AccessorNodeInput->VertexID;
-					AddEdge(MoveTemp(NewEdge), InPageID);
-				}
-			}
+			FMetasoundFrontendEdge NewEdge;
+			NewEdge.FromNodeID = MutatorNode->GetID();
+			NewEdge.FromVertexID = MutatorNodeOutput->VertexID;
+			NewEdge.ToNodeID = HeadAccessorNodeID;
+			NewEdge.ToVertexID = AccessorNodeInput->VertexID;
+			AddEdge(MoveTemp(NewEdge), InPageID);
 		}
+
 		return MutatorNode;
 	}
 
@@ -3474,26 +3484,6 @@ Metasound::Frontend::EInvalidEdgeReason FMetaSoundFrontendDocumentBuilder::IsVal
 	return EInvalidEdgeReason::None;
 }
 
-void FMetaSoundFrontendDocumentBuilder::IterateNodesDownstreamDepthFirst(const FGuid& InNodeID, FDepthFirstVisitFunction VisitFunction)
-{
-	TArray<FGuid> Stack({ InNodeID });
-	TSet<FGuid> Visited;
-	while (!Stack.IsEmpty())
-	{
-		// Do not revisit a node that has already been visited.
-		const FGuid CurrentNodeID = Stack.Pop();
-		if (!Visited.Contains(CurrentNodeID))
-		{
-			if (const FMetasoundFrontendNode* CurrentNode = FindNode(CurrentNodeID))
-			{
-				TArray<FGuid> Children = VisitFunction(CurrentNode->GetID()).Array();
-				Stack.Append(Children);
-				Visited.Add(CurrentNodeID);
-			}
-		}
-	}
-}
-
 void FMetaSoundFrontendDocumentBuilder::IterateNodesConnectedWithVertex(const FMetasoundFrontendVertexHandle& Vertex, TFunctionRef<void(const FMetasoundFrontendEdge&, FMetasoundFrontendNode&)> NodeIndexIterFunc, const FGuid& InPageID)
 {
 	using namespace Metasound::Frontend;
diff --git a/Engine/Plugins/Runtime/Metasound/Source/MetasoundFrontend/Public/MetasoundFrontendDocumentBuilder.h b/Engine/Plugins/Runtime/Metasound/Source/MetasoundFrontend/Public/MetasoundFrontendDocumentBuilder.h
--- a/Engine/Plugins/Runtime/Metasound/Source/MetasoundFrontend/Public/MetasoundFrontendDocumentBuilder.h
+++ b/Engine/Plugins/Runtime/Metasound/Source/MetasoundFrontend/Public/MetasoundFrontendDocumentBuilder.h
@@ -147,8 +147,6 @@ struct METASOUNDFRONTEND_API FMetaSoundFrontendDocumentBuilder
 	GENERATED_BODY()
 
 public:
-	using FDepthFirstVisitFunction = TFunctionRef<TSet<FGuid>(const FGuid& /* InNodeID */)>;
-
 	// Default ctor should typically never be used directly as builder interface (and optionally delegates) should be specified on construction (Default exists only to make UObject reflection happy).
 	FMetaSoundFrontendDocumentBuilder(TScriptInterface<IMetaSoundDocumentInterface> InDocumentInterface = { }, TSharedPtr<Metasound::Frontend::FDocumentModifyDelegates> InDocumentDelegates = { }, bool bPrimeCache = false);
 	virtual ~FMetaSoundFrontendDocumentBuilder();
@@ -374,10 +372,6 @@ public:
 	// Note that this does not return whether or not the given edge exists, but rather if it could be legally applied to the given edge vertices.
 	Metasound::Frontend::EInvalidEdgeReason IsValidEdge(const FMetasoundFrontendEdge& InEdge, const FGuid* InPageID = nullptr) const;
 
-	// Iterates nodes directly downstream from that with the given ID. If item(s) of interest not found,
-	// provided VisitFunction should return IDs of children to continue non-recursive, depth-first search.
-	void IterateNodesDownstreamDepthFirst(const FGuid& InNodeID, FDepthFirstVisitFunction VisitFunction);
-
 	// Iterates nodes that are filtered by only subscribing to a class with the given type (asserts if provided invalid class type).
 	void IterateNodesByClassType(Metasound::Frontend::FConstClassAndNodeFunctionRef Func, EMetasoundFrontendClassType ClassType, const FGuid* InPageID = nullptr) const;
 
