Find in MetaSound updates: Display paged input values, node keyword/category search, enum display, visual improvements
- Show all paged input values, include node keywords and category in search string, display enums properly as strings
- Hide node comment from list view and only show on tooltip, change value text alignment
- Replace controller calls with builder calls
- Fix incorrect input names being shown as a result of input template node refactor

diff --git a/Engine/Plugins/Runtime/Metasound/Source/MetasoundEditor/Private/FindInMetasound.cpp b/Engine/Plugins/Runtime/Metasound/Source/MetasoundEditor/Private/FindInMetasound.cpp
--- a/Engine/Plugins/Runtime/Metasound/Source/MetasoundEditor/Private/FindInMetasound.cpp
+++ b/Engine/Plugins/Runtime/Metasound/Source/MetasoundEditor/Private/FindInMetasound.cpp
@@ -8,6 +8,7 @@
 #include "MetasoundEditorGraphInputNode.h"
 #include "MetasoundEditorModule.h"
 #include "MetasoundFrontendDocumentBuilder.h"
+#include "MetasoundSettings.h"
 #include "Styling/AppStyle.h"
 #include "Widgets/Images/SImage.h"
 #include "Widgets/Input/SSearchBox.h"
@@ -16,6 +17,142 @@
 
 namespace Metasound::Editor
 {
+	namespace FindInMetaSoundPrivate
+	{
+		/** Get the string for the pin's enum value, empty string if not enum type */
+		FString GetEnumPinValueString(const UEdGraphPin& EnumPin)
+		{
+			using namespace Metasound::Frontend;
+			
+			if (const UMetasoundEditorGraphNode* MetaSoundNode = Cast<UMetasoundEditorGraphNode>(EnumPin.GetOwningNode()))
+			{
+				const FMetaSoundFrontendDocumentBuilder& DocBuilder = MetaSoundNode->GetBuilderChecked().GetConstBuilder();
+				if (const FMetasoundFrontendVertex* Vertex = FGraphBuilder::GetPinVertex(DocBuilder, &EnumPin))
+				{
+					FDataTypeRegistryInfo Info;
+					if (IDataTypeRegistry::Get().GetDataTypeInfo(Vertex->TypeName, Info))
+					{
+						if (Info.bIsEnum)
+						{
+							TSharedPtr<const Frontend::IEnumDataTypeInterface> EnumInterface = IDataTypeRegistry::Get().GetEnumInterfaceForDataType(Vertex->TypeName);
+							if (EnumInterface)
+							{
+								FMetasoundFrontendLiteral DefaultLiteral;
+								FGraphBuilder::GetPinLiteral(EnumPin, DefaultLiteral);
+								if (DefaultLiteral.IsValid())
+								{
+									const int32 EnumValue = FCString::Atoi(*DefaultLiteral.ToString());
+									const TOptional<FName> EnumName = EnumInterface->ToName(EnumValue);
+									if (EnumName.IsSet())
+									{
+										return EnumName->ToString();
+									}
+								}
+							}
+						}
+					}
+				}
+			}
+			return FString();
+		}
+
+		/** Get the string describing the member's literal value. */
+		FString GetMemberLiteralString(const UMetasoundEditorGraphMember* Member, bool bUseFullValueString)
+		{
+			using namespace Metasound::Frontend;
+
+			FString LiteralString;
+			const bool bIsDefaultPaged = Member->IsDefaultPaged();
+
+			const UMetaSoundSettings* Settings = GetDefault<UMetaSoundSettings>();
+			check(Settings);
+			auto GetPageName = [&Settings](const FGuid& PageID)
+				{
+					const FMetaSoundPageSettings* Page = Settings->FindPageSettings(PageID);
+					if (Page)
+					{
+						return Page->Name.ToString();
+					}
+					return LOCTEXT("FindMetasound_InvalidPage", "Invalid Page").ToString();
+				};
+
+			// Enum interface for converting int value to enum string
+			TSharedPtr<const Frontend::IEnumDataTypeInterface> EnumInterface;
+			FDataTypeRegistryInfo Info;
+			if (IDataTypeRegistry::Get().GetDataTypeInfo(Member->GetDataType(), Info))
+			{
+				if (Info.bIsEnum)
+				{
+					EnumInterface = IDataTypeRegistry::Get().GetEnumInterfaceForDataType(Member->GetDataType());
+				}
+			}
+
+			Member->GetLiteral()->IterateDefaults([&LiteralString, bIsDefaultPaged, &GetPageName, bUseFullValueString, &EnumInterface](const FGuid& PageId, FMetasoundFrontendLiteral Literal)
+				{
+					if (bIsDefaultPaged)
+					{
+						LiteralString += FText::Format(LOCTEXT("FindMetasound_PageFormat", "{0} Page: "), FText::FromString(GetPageName(PageId))).ToString();
+					}
+
+					if (bUseFullValueString)
+					{
+						if (EnumInterface)
+						{
+							int32 Value;
+							Literal.TryGet(Value);
+							LiteralString += EnumInterface->ToName(Value)->ToString();
+						}
+						else
+						{
+							LiteralString += Literal.ToString();
+						}
+					}
+					// Shortened strings for certain types
+					else
+					{
+						if (Literal.IsArray())
+						{
+							// Don't list array contents (only show on tooltip to save space)
+							LiteralString += FText::Format(LOCTEXT("FindMetasound_ArrayNumDescriptionFormat", "({0} element array)"), Literal.GetArrayNum()).ToString();
+						}
+						else if (Literal.GetType() == EMetasoundFrontendLiteralType::UObject)
+						{
+							UObject* Object = nullptr;
+							Literal.TryGet(Object);
+							if (Object)
+							{
+								// Get a shorter name than ToString (which is LexToString)
+								FString ObjectName;
+								Object->GetName(ObjectName);
+								LiteralString += ObjectName;
+							}
+							else
+							{
+								LiteralString += LOCTEXT("FindMetasound_NoneObject", "(None Object)").ToString();
+							}
+						}
+						else if (EnumInterface)
+						{
+							int32 Value;
+							Literal.TryGet(Value);
+							LiteralString += EnumInterface->ToName(Value)->ToString();
+						}
+						else
+						{
+							LiteralString += Literal.ToString();
+						}
+					}
+
+					LiteralString += "\n";
+				});
+
+			// Remove last "\n" if needed
+			LiteralString.RemoveFromEnd("\n");
+
+			return LiteralString;
+		}
+	}
+
 	FFindInMetasoundResult::FFindInMetasoundResult(const FString& InResultName)
 		:Value(InResultName), DuplicationIndex(0), Class(nullptr), Pin(), GraphNode(nullptr)
 	{
@@ -90,15 +227,23 @@ namespace Metasound::Editor
 				{
 					if (ResolvedPin->Direction == EGPD_Input)
 					{
-						Frontend::FConstInputHandle InputHandle = FGraphBuilder::GetConstInputHandleFromPin(ResolvedPin);
-						DataTypeName = InputHandle->GetDataType();
-						bIsConstructorPin = InputHandle->GetVertexAccessType() == EMetasoundFrontendVertexAccessType::Value;
+						const FMetaSoundFrontendDocumentBuilder& DocBuilder = FGraphBuilder::GetBuilderFromPinChecked(*ResolvedPin).GetConstBuilder();
+						FMetasoundFrontendVertexHandle InputVertexHandle = FGraphBuilder::GetPinVertexHandle(DocBuilder, ResolvedPin);
+						check(InputVertexHandle.IsSet());
+						const FMetasoundFrontendVertex* InputVertex = DocBuilder.FindNodeInput(InputVertexHandle.NodeID, InputVertexHandle.VertexID);
+						check(InputVertex);
+						DataTypeName = InputVertex->TypeName;
+						bIsConstructorPin = DocBuilder.GetNodeInputAccessType(InputVertexHandle.NodeID, InputVertexHandle.VertexID) == EMetasoundFrontendVertexAccessType::Value;
 					}
 					else if (ResolvedPin->Direction == EGPD_Output)
 					{
-						Frontend::FConstOutputHandle OutputHandle = FGraphBuilder::GetConstOutputHandleFromPin(ResolvedPin);
-						DataTypeName = OutputHandle->GetDataType();
-						bIsConstructorPin = OutputHandle->GetVertexAccessType() == EMetasoundFrontendVertexAccessType::Value;
+						const FMetaSoundFrontendDocumentBuilder& DocBuilder = FGraphBuilder::GetBuilderFromPinChecked(*ResolvedPin).GetConstBuilder();
+						FMetasoundFrontendVertexHandle OutputVertexHandle = FGraphBuilder::GetPinVertexHandle(DocBuilder, ResolvedPin);
+						check(OutputVertexHandle.IsSet());
+						const FMetasoundFrontendVertex* OutputVertex = DocBuilder.FindNodeOutput(OutputVertexHandle.NodeID, OutputVertexHandle.VertexID);
+						check(OutputVertex);
+						DataTypeName = OutputVertex->TypeName;
+						bIsConstructorPin = DocBuilder.GetNodeOutputAccessType(OutputVertexHandle.NodeID, OutputVertexHandle.VertexID) == EMetasoundFrontendVertexAccessType::Value;
 					}
 				}
 
@@ -182,50 +327,34 @@ namespace Metasound::Editor
 			if (Member->GetDataType() != GetMetasoundDataTypeName<FTrigger>())
 			{
 				// Only list number of items for arrays to avoid string issues 
-				FMetaSoundFrontendDocumentBuilder& Builder = Member->GetFrontendBuilderChecked();
-				FMetasoundFrontendLiteral DefaultLiteral;
-				Member->GetLiteral()->TryFindDefault(DefaultLiteral, &Builder.GetBuildPageID());
-				ValueText = FText::FromString(DefaultLiteral.ToString());
-				if (DefaultLiteral.IsArray())
-				{
-					ValueText = FText::Format(LOCTEXT("FindMetasound_ArrayNumDescriptionFormat", "({0} element array)"), DefaultLiteral.GetArrayNum());
-				}
-				// Get a shorter name than the ToString one (which is LexToString)
-				else if (DefaultLiteral.GetType() == EMetasoundFrontendLiteralType::UObject)
-				{
-					UObject* Object = nullptr;
-					DefaultLiteral.TryGet(Object);
-					if (Object)
-					{
-						FString ObjectName;
-						Object->GetName(ObjectName);
-						ValueText = FText::FromString(ObjectName);
-					}
-				}
+				FMetaSoundFrontendDocumentBuilder& DocBuilder = Member->GetFrontendBuilderChecked();
+				
+				// Don't show full object names and array contents
+				ValueText = FText::FromString(FindInMetaSoundPrivate::GetMemberLiteralString(Member, /*bUseFullValueString=*/false));
 
 				// Check if input is an overridden default value (for presets) 
 				if (const UMetasoundEditorGraphInput* Input = Cast<UMetasoundEditorGraphInput>(Member))
 				{
-					Frontend::FConstGraphHandle RootGraph = Input->GetOwningGraph()->GetDocumentHandle()->GetRootGraph();
-					if (RootGraph->GetGraphPresetOptions().bIsPreset)
+					if (DocBuilder.IsPreset())
 					{
-						const TSet<FName>& InputsInheritingDefault = RootGraph->GetInputsInheritingDefault();
-						FName NodeName = Input->GetConstNodeHandle()->GetNodeName();
-						if (!InputsInheritingDefault.Contains(NodeName))
+						const TSet<FName>* InputsInheritingDefault = DocBuilder.GetGraphInputsInheritingDefault();
+						FName NodeName = Input->GetMemberName();
+						if (InputsInheritingDefault && !InputsInheritingDefault->Contains(NodeName))
 						{
 							bIsOverriddenLiteral = true;
 						}
 					}
 				}
+
 				// Don't show default value if variable mutator node is connected 
 				if (const UMetasoundEditorGraphVariable* Variable = Cast<UMetasoundEditorGraphVariable>(Member))
 				{
-					Frontend::FConstNodeHandle MutatorNodeHandle = Variable->GetConstVariableHandle()->FindMutatorNode();
-					if (MutatorNodeHandle->IsValid())
+					const FMetasoundFrontendVariable* FrontendVariable = DocBuilder.FindGraphVariable(Variable->GetMemberName());
+					if (FrontendVariable)
 					{
-						auto IsMutatorNode = [&MutatorNodeHandle](const UMetasoundEditorGraphMemberNode* Node)
+						auto IsMutatorNode = [&FrontendVariable](const UMetasoundEditorGraphMemberNode* Node)
 						{
-							return Node->GetNodeID() == MutatorNodeHandle->GetID();
+							return Node->GetNodeID() == FrontendVariable->MutatorNodeID;
 						};
 						TArray<UMetasoundEditorGraphMemberNode*> Nodes = Variable->GetNodes();
 						if ( UMetasoundEditorGraphMemberNode** MutatorNode = Nodes.FindByPredicate(IsMutatorNode))
@@ -242,23 +371,14 @@ namespace Metasound::Editor
 		// Get value information directly from pin 
 		else if (ResolvedPin)
 		{
-			const bool bHasCachedValue = !ResolvedPin->DefaultValue.IsEmpty() || !ResolvedPin->AutogeneratedDefaultValue.IsEmpty() || ResolvedPin->DefaultObject || !ResolvedPin->DefaultTextValue.IsEmpty();
 			const bool bIsUnconnectedInputPin = ResolvedPin->Direction == EGPD_Input && ResolvedPin->LinkedTo.Num() == 0;
-			if (bHasCachedValue && bIsUnconnectedInputPin)
+			if (bIsUnconnectedInputPin)
 			{
 				if (ResolvedPin->DefaultObject)
 				{
 					const AActor* DefaultActor = Cast<AActor>(ResolvedPin->DefaultObject);
 					ValueText = FText::FromString(DefaultActor ? *DefaultActor->GetActorLabel() : *ResolvedPin->DefaultObject->GetName());
 				}
-				else if (!ResolvedPin->DefaultValue.IsEmpty())
-				{
-					// Don't show default for trigger type
-					if (ResolvedPin->PinType.PinCategory != GetMetasoundDataTypeName<FTrigger>())
-					{
-						ValueText = FText::FromString(ResolvedPin->DefaultValue);
-					}
-				}
 				else if (!ResolvedPin->AutogeneratedDefaultValue.IsEmpty())
 				{
 					ValueText = FText::FromString(ResolvedPin->AutogeneratedDefaultValue);
@@ -267,27 +387,48 @@ namespace Metasound::Editor
 				{
 					ValueText = FText::FromString(ResolvedPin->DefaultTextValue.ToString());
 				}
+				// Cached value on ed graph pin or enum as special case to be converted from int32 to enum string
+				else if (ResolvedPin->PinType.PinCategory.ToString().Equals(TEXT("Int32")))
+				{
+					// If non enum int32, case below with default value may apply
+					ValueText = FText::FromString(FindInMetaSoundPrivate::GetEnumPinValueString(*ResolvedPin));
+				}
+
+				if (ValueText.IsEmpty() && !ResolvedPin->DefaultValue.IsEmpty())
+				{
+					// Don't show default for trigger type
+					if (ResolvedPin->PinType.PinCategory != GetMetasoundDataTypeName<FTrigger>())
+					{
+						ValueText = FText::FromString(ResolvedPin->DefaultValue);
+					}
+				}
 
 				// Get value from input literal 
 				if (const UMetasoundEditorGraphExternalNode* ExternalNode = Cast<UMetasoundEditorGraphExternalNode>(ResolvedPin->GetOwningNode()))
 				{
 					if (ResolvedPin->Direction == EGPD_Input)
 					{
-						Frontend::FConstInputHandle InputHandle = FGraphBuilder::GetConstInputHandleFromPin(ResolvedPin);
-						if (const FMetasoundFrontendLiteral* DefaultLiteral = InputHandle->GetLiteral())
+						const FMetaSoundFrontendDocumentBuilder& DocBuilder = FGraphBuilder::GetBuilderFromPinChecked(*ResolvedPin).GetConstBuilder();
+						FMetasoundFrontendVertexHandle InputVertexHandle = FGraphBuilder::GetPinVertexHandle(DocBuilder, ResolvedPin);
+						check(InputVertexHandle.IsSet());
+						const FMetasoundFrontendVertex* InputVertex = DocBuilder.FindNodeInput(InputVertexHandle.NodeID, InputVertexHandle.VertexID);
+						check(InputVertex);
+						const FMetasoundFrontendVertexLiteral* DefaultLiteral = DocBuilder.FindNodeInputDefault(InputVertexHandle.NodeID, InputVertex->Name);
+						if (DefaultLiteral)
 						{
+							FMetasoundFrontendLiteral Literal = DefaultLiteral->Value;
 							// Get input literal class default override status 
-							const bool bIsDefaultConstructed = DefaultLiteral->GetType() == EMetasoundFrontendLiteralType::None;
+							const bool bIsDefaultConstructed = Literal.GetType() == EMetasoundFrontendLiteralType::None;
 							if (!bIsDefaultConstructed)
 							{
 								bIsOverriddenLiteral = true;
 							}
 
 							// Get shorter object names 
-							if (DefaultLiteral->GetType() == EMetasoundFrontendLiteralType::UObject)
+							if (Literal.GetType() == EMetasoundFrontendLiteralType::UObject)
 							{
 								UObject* Object = nullptr;
-								DefaultLiteral->TryGet(Object);
+								Literal.TryGet(Object);
 								if (Object)
 								{
 									FString ObjectName;
@@ -303,7 +444,7 @@ namespace Metasound::Editor
 			
 		if (bIsOverriddenLiteral)
 		{
-			ValueText = FText::Format(LOCTEXT("FindMetasound_OverriddenValueDescriptionFormat", "(overridden) {0}"), ValueText);
+			ValueText = FText::Format(LOCTEXT("FindMetasound_OverriddenValueDescriptionFormat", "{0} (overridden)"), ValueText);
 		}
 		return ValueText;
 	}
@@ -324,10 +465,7 @@ namespace Metasound::Editor
 		if (const UMetasoundEditorGraphMember* Member = GetMetaSoundGraphMember(MemberNode))
 		{
 			// Get full object names and array contents (GetValueText shortens these)
-			FMetaSoundFrontendDocumentBuilder& Builder = Member->GetFrontendBuilderChecked();
-			FMetasoundFrontendLiteral DefaultLiteral;
-			Member->GetLiteral()->TryFindDefault(DefaultLiteral, &Builder.GetBuildPageID());
-			return FText::FromString(DefaultLiteral.ToString());
+			return FText::FromString(FindInMetaSoundPrivate::GetMemberLiteralString(Member, /*bUseFullValueString=*/true));
 		}
 		// Get value from pin's external node 
 		else if (ResolvedPin)
@@ -337,10 +475,15 @@ namespace Metasound::Editor
 			{
 				if (ResolvedPin->Direction == EGPD_Input)
 				{
-					Frontend::FConstInputHandle InputHandle = FGraphBuilder::GetConstInputHandleFromPin(ResolvedPin);
-					if (const FMetasoundFrontendLiteral* DefaultLiteral = InputHandle->GetLiteral())
+					const FMetaSoundFrontendDocumentBuilder& DocBuilder = FGraphBuilder::GetBuilderFromPinChecked(*ResolvedPin).GetConstBuilder();
+					FMetasoundFrontendVertexHandle InputVertexHandle = FGraphBuilder::GetPinVertexHandle(DocBuilder, ResolvedPin);
+					check(InputVertexHandle.IsSet());
+					const FMetasoundFrontendVertex* InputVertex = DocBuilder.FindNodeInput(InputVertexHandle.NodeID, InputVertexHandle.VertexID);
+					check(InputVertex);
+
+					if (const FMetasoundFrontendVertexLiteral* DefaultLiteral = DocBuilder.FindNodeInputDefault(InputVertexHandle.NodeID, InputVertex->Name))
 					{
-						return FText::FromString(DefaultLiteral->ToString());
+						return FText::FromString(DefaultLiteral->Value.ToString());
 					}
 				}
 			}
@@ -504,6 +647,7 @@ namespace Metasound::Editor
 
 	void SFindInMetasound::MatchTokensInGraph(const UEdGraph* Graph, const TArray<FString>& Tokens)
 	{
+		using namespace Metasound::Frontend;
 		if (!Graph)
 		{
 			return;
@@ -520,9 +664,11 @@ namespace Metasound::Editor
 			FString NodeName = Node->GetNodeTitle(ENodeTitleType::FullTitle).ToString();
 			FString NodeType = Node->GetNodeTitle(ENodeTitleType::ListView).ToString();
 			FString DataTypeName;
+			FString MetadataString;
 			bool bIsMemberNode = false;
 			if (const UMetasoundEditorGraphNode* MetaSoundNode = Cast<UMetasoundEditorGraphNode>(Node))
 			{
+				const FMetaSoundFrontendDocumentBuilder& DocBuilder = MetaSoundNode->GetBuilderChecked().GetConstBuilder();
 				DisplayName = MetaSoundNode->GetDisplayName().ToString();
 
 				// Additional information for member nodes
@@ -530,6 +676,7 @@ namespace Metasound::Editor
 				{
 					bIsMemberNode = true;
 					DataTypeName = MetaSoundMember->GetDataType().ToString();
+					NodeName = MetaSoundMember->GetMemberName().ToString();
 					// Get specific node type and display name for variables
 					if (const UMetasoundEditorGraphVariableNode* VariableNode = Cast<UMetasoundEditorGraphVariableNode>(MetaSoundNode))
 					{
@@ -561,6 +708,24 @@ namespace Metasound::Editor
 						}
 					}
 				}
+				// Add external node keyword and category to search string
+				else if (const UMetasoundEditorGraphExternalNode* ExternalNode = Cast<UMetasoundEditorGraphExternalNode>(Node))
+				{
+					if (const FMetasoundFrontendNode* FrontendNode = DocBuilder.FindNode(ExternalNode->GetNodeID()))
+					{
+						if (const FMetasoundFrontendClass* Class = DocBuilder.FindDependency(FrontendNode->ClassID))
+						{
+							for (const FText& Keyword : Class->Metadata.GetKeywords())
+							{
+								MetadataString += Keyword.ToString();
+							}
+							for (const FText& CategoryHierarchy : Class->Metadata.GetCategoryHierarchy())
+							{
+								MetadataString += CategoryHierarchy.ToString();
+							}
+						}
+					}
+				}
 			}
 
 			FString NodeResultName;
@@ -583,7 +748,7 @@ namespace Metasound::Editor
 
 			FSearchResult NodeResult;
 
-			FString NodeSearchString = NodeName + NodeType + Node->NodeComment + DisplayName + DataTypeName;
+			FString NodeSearchString = NodeName + NodeType + Node->NodeComment + DisplayName + DataTypeName + MetadataString;
 			NodeSearchString = NodeSearchString.Replace(TEXT(" "), TEXT(""));
 
 			bool bNodeMatchesSearch = StringMatchesSearchTokens(Tokens, NodeSearchString);
@@ -599,13 +764,29 @@ namespace Metasound::Editor
 				if (Pin && Pin->PinFriendlyName.CompareTo(FText::FromString(TEXT(" "))) != 0)
 				{
 					FText PinDisplayName = Pin->GetSchema()->GetPinDisplayName(Pin);
+					FString PinCategory = Pin->PinType.PinCategory.ToString();
+					// String for values that may not be directly retrievable from the pin default value
+					FString PinValueString;
+
+					// Member data type name and value
 					FString PinDataTypeName;
 					if (const UMetasoundEditorGraphMember* MetaSoundMember = FFindInMetasoundResult::GetMetaSoundGraphMember(Node))
 					{
 						PinDataTypeName = MetaSoundMember->GetDataType().ToString();
+						PinValueString = FindInMetaSoundPrivate::GetMemberLiteralString(MetaSoundMember, /*bUseFullValueString=*/true);
 					}
+					else if (PinCategory.Equals(TEXT("Int32")))
+					{
+						// Enum value string (needs conversion from pin DefaultValue int)
+						const FString EnumValueString = FindInMetaSoundPrivate::GetEnumPinValueString(*Pin);
+						if (!EnumValueString.IsEmpty())
+						{
+							PinValueString = EnumValueString;
+						}
+					}
+					
+					FString PinSearchString = PinDisplayName.ToString() + Pin->PinName.ToString() + Pin->PinFriendlyName.ToString() + Pin->DefaultValue + PinCategory + Pin->PinType.PinSubCategory.ToString() + PinDataTypeName + (Pin->PinType.PinSubCategoryObject.IsValid() ? Pin->PinType.PinSubCategoryObject.Get()->GetFullName() : TEXT("")) + PinValueString;
 
-					FString PinSearchString = PinDisplayName.ToString() + Pin->PinName.ToString() + Pin->PinFriendlyName.ToString() + Pin->DefaultValue + Pin->PinType.PinCategory.ToString() + Pin->PinType.PinSubCategory.ToString() + PinDataTypeName + (Pin->PinType.PinSubCategoryObject.IsValid() ? Pin->PinType.PinSubCategoryObject.Get()->GetFullName() : TEXT(""));
 					PinSearchString = PinSearchString.Replace(TEXT(" "), TEXT(""));
 					if (StringMatchesSearchTokens(Tokens, PinSearchString))
 					{
@@ -635,52 +816,49 @@ namespace Metasound::Editor
 
 	TSharedRef<ITableRow> SFindInMetasound::OnGenerateRow(FSearchResult InItem, const TSharedRef<STableViewBase>& OwnerTable)
 	{
-		FString CommentText = InItem->GetCommentText();
-
-		return SNew(STableRow<TSharedPtr<FFindInMetasoundResult>>, OwnerTable)
+		TSharedRef< SHorizontalBox> TableRowBox = SNew(SHorizontalBox)
+			+ SHorizontalBox::Slot()
+			.AutoWidth()
+			.VAlign(VAlign_Center)
+			.Padding(4.f)
 			[
-				SNew(SHorizontalBox)
-				+ SHorizontalBox::Slot()
-				.AutoWidth()
+				InItem->CreateIcon()
+			]
+			+ SHorizontalBox::Slot()
+			.FillWidth(1)
+			.HAlign(HAlign_Left)
+			.VAlign(VAlign_Center)
+			.Padding(2)
+			[
+				SNew(STextBlock)
+					.Text(FText::FromString(InItem->Value))
+					.HighlightText(HighlightText)
+					.ToolTipText(FText::Format(LOCTEXT("FindMetasound_ResultSearchToolTipFmt", "{0}: {1}"), InItem->GetCategory(), FText::FromString(InItem->Value)))
+			];
+		
+		// To avoid showing value on both pin and node,
+		// only show value on pin result or on node result if no pin  
+		if (InItem->Pin.Get() || InItem->Children.Num() == 0)
+		{
+			const FText ValueTooltip = InItem->GetCommentText().IsEmpty() ? 
+				InItem->GetValueTooltipText() : 
+				FText::Format(LOCTEXT("FindMetasound_NodeValueWithCommentFmt", "[Comment: {0}]\n{1}"), FText::FromString(InItem->GetCommentText()), InItem->GetValueTooltipText());
+			TableRowBox->AddSlot()
+				.FillWidth(1)
+				.HAlign(HAlign_Left)
 				.VAlign(VAlign_Center)
-				.Padding(4.f)
+				.Padding(2)
 				[
-					InItem->CreateIcon()
-				]
-				+ SHorizontalBox::Slot()
-					.FillWidth(1)
-					.HAlign(HAlign_Left)
-					.VAlign(VAlign_Center)
-					.Padding(2, 0)
-					[
-						SNew(STextBlock)
-						.Text(FText::FromString(InItem->Value))
-						.HighlightText(HighlightText)
-						.ToolTipText(FText::Format(LOCTEXT("FindMetasound_ResultSearchToolTipFmt", "{0}: {1}"), InItem->GetCategory(), FText::FromString(InItem->Value)))
-					]
-				+ SHorizontalBox::Slot()
-					.FillWidth(1)
-					.HAlign(HAlign_Right)
-					.VAlign(VAlign_Center)
-					.Padding(2, 0)
-					[
-						SNew(STextBlock)
-						.Text(InItem->GetValueText())
-						.HighlightText(HighlightText)
-						.ToolTipText(InItem->GetValueTooltipText())
-					]
-				+ SHorizontalBox::Slot()
-					.AutoWidth()
-					.HAlign(HAlign_Right)
-					.VAlign(VAlign_Center)
-					.Padding(2, 0)
-					[
-						SNew(STextBlock)
-						.Text(CommentText.IsEmpty() ? FText::GetEmpty() : FText::Format(LOCTEXT("FindMetasound_NodeCommentFmt", "[Comment: {0}]"), FText::FromString(CommentText)))
-						.ColorAndOpacity(FSlateColor::UseForeground())
-						.HighlightText(HighlightText)
-					]
-			];
+					SNew(STextBlock)
+					.Text(InItem->GetValueText())
+					.HighlightText(HighlightText)
+					.ToolTipText(ValueTooltip)
+				];
+		}
+		return SNew(STableRow<TSharedPtr<FFindInMetasoundResult>>, OwnerTable)
+		[
+			TableRowBox
+		];
 	}
 
 	void SFindInMetasound::OnGetChildren(FSearchResult InItem, TArray< FSearchResult >& OutChildren)
diff --git a/Engine/Plugins/Runtime/Metasound/Source/MetasoundEditor/Private/MetasoundEditorGraphInputNode.cpp b/Engine/Plugins/Runtime/Metasound/Source/MetasoundEditor/Private/MetasoundEditorGraphInputNode.cpp
--- a/Engine/Plugins/Runtime/Metasound/Source/MetasoundEditor/Private/MetasoundEditorGraphInputNode.cpp
+++ b/Engine/Plugins/Runtime/Metasound/Source/MetasoundEditor/Private/MetasoundEditorGraphInputNode.cpp
@@ -114,6 +114,16 @@ void UMetasoundEditorGraphInputNode::UpdatePreviewInstance(const Metasound::FVer
 	}
 }
 
+FText UMetasoundEditorGraphInputNode::GetDisplayName() const
+{
+	const FMetaSoundFrontendDocumentBuilder& Builder = Input->GetFrontendBuilderChecked();
+	if (const FMetasoundFrontendClassInput* ClassInput = Builder.FindGraphInput(Input->GetMemberName()))
+	{
+		return ClassInput->Metadata.GetDisplayName();
+	}
+	return FText::GetEmpty();
+}
+
 FGuid UMetasoundEditorGraphInputNode::GetNodeID() const
 {
 	return NodeID;
diff --git a/Engine/Plugins/Runtime/Metasound/Source/MetasoundEditor/Public/MetasoundEditorGraphInputNode.h b/Engine/Plugins/Runtime/Metasound/Source/MetasoundEditor/Public/MetasoundEditorGraphInputNode.h
--- a/Engine/Plugins/Runtime/Metasound/Source/MetasoundEditor/Public/MetasoundEditorGraphInputNode.h
+++ b/Engine/Plugins/Runtime/Metasound/Source/MetasoundEditor/Public/MetasoundEditorGraphInputNode.h
@@ -40,6 +40,7 @@ public:
 	virtual void UpdatePreviewInstance(const Metasound::FVertexName& InParameterName, TScriptInterface<IAudioParameterControllerInterface>& InParameterInterface) const;
 	virtual FMetasoundFrontendClassName GetClassName() const;
 
+	FText GetDisplayName() const override;
 	virtual FGuid GetNodeID() const override;
 	virtual FLinearColor GetNodeTitleColor() const override;
 	virtual FSlateIcon GetNodeTitleIcon() const override;
