Fix for bug where removal of a MetaSound variable or changing its type may incorrectly remove other variables

#rnx

diff --git a/Engine/Plugins/Runtime/Metasound/Source/MetasoundEditor/Private/MetasoundEditor.cpp b/Engine/Plugins/Runtime/Metasound/Source/MetasoundEditor/Private/MetasoundEditor.cpp
--- a/Engine/Plugins/Runtime/Metasound/Source/MetasoundEditor/Private/MetasoundEditor.cpp
+++ b/Engine/Plugins/Runtime/Metasound/Source/MetasoundEditor/Private/MetasoundEditor.cpp
@@ -3073,14 +3073,9 @@ namespace Metasound
 						{
 							Builder->RemoveGraphOutput(MemberName, Result);
 						}
-						// TODO: Move to builder API
 						else if (UMetasoundEditorGraphVariable* Variable = Cast<UMetasoundEditorGraphVariable>(GraphMember))
 						{
-							const FGuid VariableID = Variable->GetVariableID();
-							if (VariableID.IsValid())
-							{
-								Builder->RemoveGraphVariable(MemberName, Result);
-							}
+							Builder->RemoveGraphVariable(MemberName, Result);
 						}
 						ensure(Result == EMetaSoundBuilderResult::Succeeded);
 					}
diff --git a/Engine/Plugins/Runtime/Metasound/Source/MetasoundEditor/Private/MetasoundEditorGraph.cpp b/Engine/Plugins/Runtime/Metasound/Source/MetasoundEditor/Private/MetasoundEditorGraph.cpp
--- a/Engine/Plugins/Runtime/Metasound/Source/MetasoundEditor/Private/MetasoundEditorGraph.cpp
+++ b/Engine/Plugins/Runtime/Metasound/Source/MetasoundEditor/Private/MetasoundEditorGraph.cpp
@@ -1996,13 +1996,13 @@ void UMetasoundEditorGraphVariable::UpdateFrontendDefaultLiteral(bool bPostTrans
 	using namespace Metasound::Frontend;
 	using namespace Metasound::VariableNames;
 
-	if (!ensure(Literal))
+	if (!Literal)
 	{
 		return;
 	}
 
 	const FMetasoundFrontendVariable* Variable = GetFrontendVariable();
-	if (!ensure(Variable))
+	if (!Variable)
 	{
 		return;
 	}
diff --git a/Engine/Plugins/Runtime/Metasound/Source/MetasoundFrontend/Private/MetasoundFrontendDocumentBuilder.cpp b/Engine/Plugins/Runtime/Metasound/Source/MetasoundFrontend/Private/MetasoundFrontendDocumentBuilder.cpp
--- a/Engine/Plugins/Runtime/Metasound/Source/MetasoundFrontend/Private/MetasoundFrontendDocumentBuilder.cpp
+++ b/Engine/Plugins/Runtime/Metasound/Source/MetasoundFrontend/Private/MetasoundFrontendDocumentBuilder.cpp
@@ -4262,14 +4262,14 @@ bool FMetaSoundFrontendDocumentBuilder::RemoveGraphVariable(FName VariableName,
 		RemoveNode(Variable->VariableNodeID);
 		RemoveNode(Variable->MutatorNodeID);
 
-		// Copy ids as node removal will update variable node IDs
+		// Copy ids as node removal will update AccessorNodeIDs array on FrontendVariable internally to RemoveNode call
 		TArray<FGuid> AccessorNodeIDs = Variable->AccessorNodeIDs;
 		for (const FGuid& NodeID : AccessorNodeIDs)
 		{
 			RemoveNode(NodeID, InPageID);
 		}
 
-		// Copy ids as node removal will update variable node IDs
+		// Copy ids as node removal will update DeferredAccessorNodeIDs array on FrontendVariable internally to RemoveNode call
 		TArray<FGuid> DeferredAccessorNodeIDs = Variable->DeferredAccessorNodeIDs;
 		for (const FGuid& NodeID : DeferredAccessorNodeIDs)
 		{
@@ -4279,18 +4279,21 @@ bool FMetaSoundFrontendDocumentBuilder::RemoveGraphVariable(FName VariableName,
 		FMetasoundFrontendDocument& Document = GetDocumentChecked();
 		FMetasoundFrontendGraph& Graph = Document.RootGraph.FindGraphChecked(InPageID ? *InPageID : BuildPageID);
 
-#if WITH_EDITOR
-		Document.Metadata.ModifyContext.AddMemberIDModified(Variable->ID);
-#endif // WITH_EDITOR
-
-		auto IsVariableWithID = [&Variable](const FMetasoundFrontendVariable& InVariable)
+		// VariableID must be cached to avoid being mutated when variable
+		// is ultimately removed in the RemoveAllSwap below.
+		const FGuid VariableID = Variable->ID;
+		auto IsVariableWithID = [VariableID](const FMetasoundFrontendVariable& InVariable)
 		{
-			return InVariable.ID == Variable->ID;
+			return InVariable.ID == VariableID;
 		};
 		Graph.Variables.RemoveAllSwap(IsVariableWithID);
 
 		// Clean-up/remove variable dependencies that may or may no longer be referenced.
 		RemoveUnusedDependencies();
+
+#if WITH_EDITOR
+		Document.Metadata.ModifyContext.AddMemberIDModified(VariableID);
+#endif // WITH_EDITOR
 		return true;
 	}
 
