Add folder action for bulk versioning MetaSounds

[RN] Added MetaSound bulk folder action to enable versioning specified MetaSounds quickly and easily

diff --git a/Engine/Plugins/Runtime/Metasound/Metasound.uplugin b/Engine/Plugins/Runtime/Metasound/Metasound.uplugin
--- a/Engine/Plugins/Runtime/Metasound/Metasound.uplugin
+++ b/Engine/Plugins/Runtime/Metasound/Metasound.uplugin
@@ -60,6 +60,10 @@
 			"Name": "AudioSynesthesia",
 			"Enabled": true
 		},
+		{
+			"Name": "ContentBrowserAssetDataSource",
+			"Enabled":  true
+		},
 		{
 			"Name": "WaveTable",
 			"Enabled": true
diff --git a/Engine/Plugins/Runtime/Metasound/Source/MetasoundEditor/MetasoundEditor.Build.cs b/Engine/Plugins/Runtime/Metasound/Source/MetasoundEditor/MetasoundEditor.Build.cs
--- a/Engine/Plugins/Runtime/Metasound/Source/MetasoundEditor/MetasoundEditor.Build.cs
+++ b/Engine/Plugins/Runtime/Metasound/Source/MetasoundEditor/MetasoundEditor.Build.cs
@@ -17,6 +17,9 @@ namespace UnrealBuildTool.Rules
 					"AppFramework",
 					"AudioExtensions",
 					"AudioSynesthesia",
+					"ContentBrowser",
+					"ContentBrowserAssetDataSource",
+					"ContentBrowserData",
 					"EditorWidgets",
 					"Kismet",
 					"KismetWidgets",
diff --git a/Engine/Plugins/Runtime/Metasound/Source/MetasoundEditor/Private/MetasoundEditorModule.cpp b/Engine/Plugins/Runtime/Metasound/Source/MetasoundEditor/Private/MetasoundEditorModule.cpp
--- a/Engine/Plugins/Runtime/Metasound/Source/MetasoundEditor/Private/MetasoundEditorModule.cpp
+++ b/Engine/Plugins/Runtime/Metasound/Source/MetasoundEditor/Private/MetasoundEditorModule.cpp
@@ -1,6 +1,9 @@
 // Copyright Epic Games, Inc. All Rights Reserved.
 #include "MetasoundEditorModule.h"
 
+#include "Algo/Transform.h"
+#include "AssetFolderContextMenu.h"
+#include "AssetRegistry/AssetRegistryModule.h"
 #include "AssetTypeActions_Base.h"
 #include "AudioMaterialSlate/AudioMaterialSlateTypes.h"
 #include "AudioMeter.h"
@@ -8,10 +11,17 @@
 #include "AudioSpectrumPlotStyle.h"
 #include "AudioVectorscopePanelStyle.h"
 #include "Brushes/SlateImageBrush.h"
+#include "ContentBrowserAssetDataCore.h"
+#include "ContentBrowserAssetDataPayload.h"
+#include "ContentBrowserDataSource.h"
+#include "ContentBrowserDataMenuContexts.h"
+#include "ContentBrowserMenuContexts.h"
 #include "EdGraph/EdGraphNode.h"
 #include "EdGraph/EdGraphPin.h"
 #include "EdGraphUtilities.h"
 #include "Editor.h"
+#include "FileHelpers.h"
+#include "Framework/Commands/UIAction.h"
 #include "Framework/Notifications/NotificationManager.h"
 #include "HAL/IConsoleManager.h"
 #include "IDetailCustomization.h"
@@ -43,6 +53,7 @@
 #include "MetasoundTrigger.h"
 #include "MetasoundUObjectRegistry.h"
 #include "Modules/ModuleManager.h"
+#include "NodeTemplates/MetasoundFrontendNodeTemplateInput.h"
 #include "PackageMigrationContext.h"
 #include "PropertyEditorDelegates.h"
 #include "PropertyEditorModule.h"
@@ -55,6 +66,14 @@
 #include "Styling/SlateTypes.h"
 #include "Styling/StyleColors.h"
 #include "Templates/SharedPointer.h"
+#include "Textures/SlateIcon.h"
+#include "ToolMenu.h"
+#include "ToolMenuContext.h"
+#include "ToolMenuDelegates.h"
+#include "ToolMenuEntry.h"
+#include "ToolMenus.h"
+#include "ToolMenuSection.h"
+#include "UObject/NoExportTypes.h"
 #include "UObject/UObjectGlobals.h"
 #include "UObject/UObjectIterator.h"
 #include "Widgets/Notifications/SNotificationList.h"
@@ -64,927 +83,1011 @@ DEFINE_LOG_CATEGORY(LogMetasoundEditor);
 
 #define LOCTEXT_NAMESPACE "MetaSounds"
 
-namespace Metasound
+namespace Metasound::Editor
 {
-	namespace Editor
+	namespace ModulePrivate
 	{
-		using FMetasoundGraphPanelPinFactory = FGraphPanelPinFactory;
+		static const FLazyName AssetToolName { "AssetTools" };
 
-		static const FName AssetToolName { "AssetTools" };
-
-		template <typename T>
-		void AddAssetAction(IAssetTools& AssetTools, TArray<TSharedPtr<FAssetTypeActions_Base>>& AssetArray)
-		{
-			TSharedPtr<T> AssetAction = MakeShared<T>();
-			TSharedPtr<FAssetTypeActions_Base> AssetActionBase = StaticCastSharedPtr<FAssetTypeActions_Base>(AssetAction);
-			AssetTools.RegisterAssetTypeActions(AssetAction.ToSharedRef());
-			AssetArray.Add(AssetActionBase);
-		}
-
-		class FSlateStyle final : public FSlateStyleSet
+		void AddVersionDocumentFolderMenuEntry(UToolMenu* Menu)
 		{
-		public:
-			FSlateStyle()
-				: FSlateStyleSet("MetaSoundStyle")
+			if (!Menu)
 			{
-				SetParentStyleName(FAudioWidgetsStyle::Get().GetStyleSetName());
+				return;
+			}
 
-				SetContentRoot(FPaths::EnginePluginsDir() / TEXT("Runtime/Metasound/Content/Editor/Slate"));
-				SetCoreContentRoot(FPaths::EngineContentDir() / TEXT("Slate"));
+			const UContentBrowserFolderContext* Context = Menu->FindContext<UContentBrowserFolderContext>();
+			if (!Context)
+			{
+				return;
+			}
 
-				static const FVector2D Icon20x20(20.0f, 20.0f);
-				static const FVector2D Icon40x40(40.0f, 40.0f);
+			Menu->AddMenuEntry("MetaSounds", FToolMenuEntry::InitMenuEntry(
+				"VersionMetaSounds",
+				LOCTEXT("ResaveMetaSoundAssetsMenuEntry", "Version MetaSounds"),
+				LOCTEXT("ResaveMetaSoundAssetsMenuEntryTooltip", "Load, version MetaSound asset document (if necessary), and resave all MetaSounds in the given folder context. Does not replace deprecated node classes with newer major versions."),
+				FSlateIcon(FAppStyle::GetAppStyleSetName(), "Icons.Save"),
+				FToolUIActionChoice(FExecuteAction::CreateLambda([FolderPaths = Context->GetSelectedPackagePaths()]()
+				{
+					using namespace Frontend;
+
+					FNotificationInfo Info(LOCTEXT("VersioningMetaSoundsNotifyTitle", "Versioning MetaSounds..."));
+					Info.bFireAndForget = false;
+					Info.ExpireDuration = 0.0f;
+					Info.bUseThrobber = true;
+					TSharedPtr<SNotificationItem> Notify = FSlateNotificationManager::Get().AddNotification(Info);
+					if (Notify.IsValid())
+					{
+						Notify->SetCompletionState(SNotificationItem::CS_Pending);
+					}
 
-				static const FVector2D Icon16 = FVector2D(16.0f, 16.0f);
-				static const FVector2D Icon64 = FVector2D(64.0f, 64.0f);
+					IMetaSoundAssetManager::FVersionDocumentResults Results;
+					{
+						FGCScopeGuard Guard;
+						UE::Tasks::FTask ProcessDocuments = UE::Tasks::Launch(UE_SOURCE_LOCATION, [&Results, FolderPaths]()
+						{
+							Results = IMetaSoundAssetManager::GetChecked().VersionDocumentsInFolders(FolderPaths);
+						});
 
-				const FVector2D Icon15x11(15.0f, 11.0f);
+						// TODO: Report progress and tick instead of immediate wait.
+						// Make this a more generic utility for future folder/object utility calls.
+						ProcessDocuments.Wait();
+					}
 
-				// Metasound Editor
-				{
-					Set("MetaSoundPatch.Color", FColor(31, 133, 31));
-					Set("MetaSoundSource.Color", FColor(103, 214, 66));
-
-					// Actions
-					Set("MetasoundEditor.Play", new IMAGE_BRUSH_SVG(TEXT("Icons/play"), Icon40x40));
-					Set("MetasoundEditor.Play.Small", new IMAGE_BRUSH_SVG(TEXT("Icons/play"), Icon20x20));
-					Set("MetasoundEditor.Play.Thumbnail", new IMAGE_BRUSH_SVG(TEXT("Icons/play_thumbnail"), Icon64));
-					Set("MetasoundEditor.Play.Thumbnail.Hovered", new IMAGE_BRUSH_SVG(TEXT("Icons/play_thumbnail_hover"), Icon64));
-
-					Set("MetasoundEditor.Play.Active.Valid", new IMAGE_BRUSH_SVG(TEXT("Icons/play_active_valid"), Icon40x40));
-					Set("MetasoundEditor.Play.Active.Warning", new IMAGE_BRUSH_SVG(TEXT("Icons/play_active_warning"), Icon40x40));
-					Set("MetasoundEditor.Play.Inactive.Valid", new IMAGE_BRUSH_SVG(TEXT("Icons/play_inactive_valid"), Icon40x40));
-					Set("MetasoundEditor.Play.Inactive.Warning", new IMAGE_BRUSH_SVG(TEXT("Icons/play_inactive_warning"), Icon40x40));
-					Set("MetasoundEditor.Play.Error", new IMAGE_BRUSH_SVG(TEXT("Icons/play_error"), Icon40x40));
-
-					Set("MetasoundEditor.Stop", new IMAGE_BRUSH_SVG(TEXT("Icons/stop"), Icon40x40));
-
-					Set("MetasoundEditor.Stop.Disabled", new IMAGE_BRUSH_SVG(TEXT("Icons/stop_disabled"), Icon40x40));
-					Set("MetasoundEditor.Stop.Active", new IMAGE_BRUSH_SVG(TEXT("Icons/stop_active"), Icon40x40));
-					Set("MetasoundEditor.Stop.Inactive", new IMAGE_BRUSH_SVG(TEXT("Icons/stop_inactive"), Icon40x40));
-					Set("MetasoundEditor.Stop.Thumbnail", new IMAGE_BRUSH_SVG(TEXT("Icons/stop_thumbnail"), Icon64));
-					Set("MetasoundEditor.Stop.Thumbnail.Hovered", new IMAGE_BRUSH_SVG(TEXT("Icons/stop_thumbnail_hover"), Icon64));
-
-					Set("MetasoundEditor.Import", new FSlateImageBrush(RootToContentDir(TEXT("/Icons/build_40x.png")), Icon40x40));
-					Set("MetasoundEditor.Import.Small", new FSlateImageBrush(RootToContentDir(TEXT("/Icons/build_40x.png")), Icon20x20));
-					Set("MetasoundEditor.Export", new FSlateImageBrush(RootToContentDir(TEXT("/Icons/build_40x.png")), Icon40x40));
-					Set("MetasoundEditor.Export.Small", new FSlateImageBrush(RootToContentDir(TEXT("/Icons/build_40x.png")), Icon20x20));
-					Set("MetasoundEditor.ExportError", new FSlateImageBrush(RootToContentDir(TEXT("/Icons/build_error_40x.png")), Icon40x40));
-					Set("MetasoundEditor.ExportError.Small", new FSlateImageBrush(RootToContentDir(TEXT("/Icons/build_error_40x.png")), Icon20x20));
-					Set("MetasoundEditor.Settings", new FSlateImageBrush(RootToContentDir(TEXT("/Icons/settings_40x.png")), Icon20x20));
-
-					// Graph Editor
-					Set("MetasoundEditor.Graph.Node.Body.Input", new FSlateImageBrush(RootToContentDir(TEXT("/Graph/node_input_body_64x.png")), FVector2D(114.0f, 64.0f)));
-					Set("MetasoundEditor.Graph.Node.Body.Default", new FSlateImageBrush(RootToContentDir(TEXT("/Graph/node_default_body_64x.png")), FVector2D(64.0f, 64.0f)));
-
-					Set("MetasoundEditor.Graph.TriggerPin.Connected", new IMAGE_BRUSH(TEXT("Graph/pin_trigger_connected"), Icon15x11));
-					Set("MetasoundEditor.Graph.TriggerPin.Disconnected", new IMAGE_BRUSH(TEXT("Graph/pin_trigger_disconnected"), Icon15x11));
-
-					Set("MetasoundEditor.Graph.Node.Class.Native", new IMAGE_BRUSH_SVG(TEXT("Icons/native_node"), FVector2D(8.0f, 16.0f)));
-					Set("MetasoundEditor.Graph.Node.Class.Graph", new IMAGE_BRUSH_SVG(TEXT("Icons/graph_node"), Icon16));
-					Set("MetasoundEditor.Graph.Node.Class.Input", new IMAGE_BRUSH_SVG(TEXT("Icons/input_node"), FVector2D(16.0f, 13.0f)));
-					Set("MetasoundEditor.Graph.Node.Class.Output", new IMAGE_BRUSH_SVG(TEXT("Icons/output_node"), FVector2D(16.0f, 13.0f)));
-					Set("MetasoundEditor.Graph.Node.Class.Reroute", new IMAGE_BRUSH_SVG(TEXT("Icons/reroute_node"), Icon16));
-					Set("MetasoundEditor.Graph.Node.Class.Variable", new IMAGE_BRUSH_SVG(TEXT("Icons/variable_node"), FVector2D(16.0f, 13.0f)));
-
-					Set("MetasoundEditor.Graph.Node.Math.Add", new FSlateImageBrush(RootToContentDir(TEXT("/Graph/node_math_add_40x.png")), Icon40x40));
-					Set("MetasoundEditor.Graph.Node.Math.Divide", new FSlateImageBrush(RootToContentDir(TEXT("/Graph/node_math_divide_40x.png")), Icon40x40));
-					Set("MetasoundEditor.Graph.Node.Math.Modulo", new FSlateImageBrush(RootToContentDir(TEXT("/Graph/node_math_modulo_40x.png")), Icon40x40));
-					Set("MetasoundEditor.Graph.Node.Math.Multiply", new FSlateImageBrush(RootToContentDir(TEXT("/Graph/node_math_multiply_40x.png")), Icon40x40));
-					Set("MetasoundEditor.Graph.Node.Math.Subtract", new FSlateImageBrush(RootToContentDir(TEXT("/Graph/node_math_subtract_40x.png")), Icon40x40));
-					Set("MetasoundEditor.Graph.Node.Math.Power", new FSlateImageBrush(RootToContentDir(TEXT("/Graph/node_math_power_40x.png")), Icon40x40));
-					Set("MetasoundEditor.Graph.Node.Math.Logarithm", new FSlateImageBrush(RootToContentDir(TEXT("/Graph/node_math_logarithm_40x.png")), Icon40x40));
-					Set("MetasoundEditor.Graph.Node.Conversion", new FSlateImageBrush(RootToContentDir(TEXT("/Graph/node_conversion_40x.png")), Icon40x40));
-
-					Set("MetasoundEditor.Graph.InvalidReroute", new IMAGE_BRUSH_SVG(TEXT("Icons/invalid_reroute"), Icon16));
-					Set("MetasoundEditor.Graph.ConstructorPinArray", new FSlateImageBrush(RootToContentDir(TEXT("/Icons/array_pin_rotated.png")), Icon16));
-					Set("MetasoundEditor.Graph.ConstructorPinArrayDisconnected", new FSlateImageBrush(RootToContentDir(TEXT("/Icons/array_pin_rotated_disconnected.png")), Icon16));
-					Set("MetasoundEditor.Graph.ArrayPin", new FSlateImageBrush(RootToContentDir(TEXT("/Icons/array_pin.png")), Icon16));
-					Set("MetasoundEditor.Graph.ConstructorPin", new FSlateImageBrush(RootToContentDir(TEXT("/Icons/square_pin_rotated.png")), Icon16));
-					Set("MetasoundEditor.Graph.ConstructorPinDisconnected", new FSlateImageBrush(RootToContentDir(TEXT("/Icons/square_pin_rotated_disconnected.png")), Icon16));
-
-					// Analyzers
-					Set("MetasoundEditor.Analyzers.BackgroundColor", FLinearColor(0.0075f, 0.0075f, 0.0075, 1.0f));
-					Set("MetasoundEditor.Analyzers.ForegroundColor", FLinearColor(0.025719f, 0.208333f, 0.069907f, 1.0f)); // "Audio" Green
-
-					// Misc
-					Set("MetasoundEditor.Audition", new IMAGE_BRUSH_SVG(TEXT("Icons/metasound_page"), Icon16));
-					Set("MetasoundEditor.Metasound.Icon", new IMAGE_BRUSH_SVG(TEXT("Icons/metasound_icon"), Icon16));
-					Set("MetasoundEditor.Speaker", new FSlateImageBrush(RootToContentDir(TEXT("/Icons/speaker_144x.png")), FVector2D(144.0f, 144.0f)));
-
-					// Pages
-					Set("MetasoundEditor.Page.Executing.ForegroundColor", FStyleColors::AccentGreen.GetSpecifiedColor());
-					Set("MetasoundEditor.Page.Executing", new IMAGE_BRUSH_SVG(TEXT("Icons/metasound_page_exec"), Icon16));
-
-					// Class Icons
-					auto SetClassIcon = [this, InIcon16 = Icon16, InIcon64 = Icon64](const FString& ClassName)
+					if (!Results.DocumentsFoundInPackages())
 					{
-						const FString IconFileName = FString::Printf(TEXT("Icons/%s"), *ClassName.ToLower());
-						const FSlateColor DefaultForeground(FStyleColors::Foreground);
-
-						Set(*FString::Printf(TEXT("ClassIcon.%s"), *ClassName), new IMAGE_BRUSH_SVG(IconFileName, InIcon16));
-						Set(*FString::Printf(TEXT("ClassThumbnail.%s"), *ClassName), new IMAGE_BRUSH_SVG(IconFileName, InIcon64));
-					};
-
-					SetClassIcon(TEXT("MetasoundPatch"));
-					SetClassIcon(TEXT("MetasoundSource"));
-
-					Set("MetasoundEditor.MetasoundPatch.Icon", new IMAGE_BRUSH_SVG(TEXT("Icons/metasoundpatch_icon"), Icon20x20));
-					Set("MetasoundEditor.MetasoundPatch.Preset.Icon", new IMAGE_BRUSH_SVG(TEXT("Icons/metasoundpatchpreset_icon"), Icon20x20));
-					Set("MetasoundEditor.MetasoundSource.Icon", new IMAGE_BRUSH_SVG(TEXT("Icons/metasoundsource_icon"), Icon20x20));
-					Set("MetasoundEditor.MetasoundSource.Preset.Icon", new IMAGE_BRUSH_SVG(TEXT("Icons/metasoundsourcepreset_icon"), Icon20x20));
-					Set("MetasoundEditor.MetasoundPatch.Thumbnail", new IMAGE_BRUSH_SVG(TEXT("Icons/metasoundpatch_thumbnail"), Icon20x20));
-					Set("MetasoundEditor.MetasoundPatch.Preset.Thumbnail", new IMAGE_BRUSH_SVG(TEXT("Icons/metasoundpatchpreset_thumbnail"), Icon20x20));
-					Set("MetasoundEditor.MetasoundSource.Thumbnail", new IMAGE_BRUSH_SVG(TEXT("Icons/metasoundsource_thumbnail"), Icon20x20));
-					Set("MetasoundEditor.MetasoundSource.Preset.Thumbnail", new IMAGE_BRUSH_SVG(TEXT("Icons/metasoundsourcepreset_thumbnail"), Icon20x20));
-				}
+						if (Notify.IsValid())
+						{
+							Notify->SetCompletionState(SNotificationItem::CS_Fail);
+						}
+						UE_LOG(LogMetasoundEditor, Display, TEXT("No MetaSound(s) Versioned: Folder/child folder(s) contain(s) no MetaSound asset(s)."));
+					}
+					else if (Results.VersionedPackages.IsEmpty())
+					{
+						if (Notify.IsValid())
+						{
+							Notify->SetCompletionState(SNotificationItem::CS_Fail);
+						}
+						UE_LOG(LogMetasoundEditor, Display, TEXT("No MetaSound(s) Versioned: Folder/child folder(s) contain(s) no MetaSound asset(s) requiring versioning."));
+					}
+					else
+					{
+						if (Notify.IsValid())
+						{
+							Notify->SetCompletionState(SNotificationItem::CS_Success);
+						}
+						UE_LOG(LogMetasoundEditor, Display, TEXT("%i MetaSound(s) Successfully Versioned."), Results.VersionedPackages.Num());
+						FEditorFileUtils::PromptToCheckoutPackages(false, Results.VersionedPackages);
+					}
+					
+					if (Notify.IsValid())
+					{
+						Notify->ExpireAndFadeout();
+					}
+				}))
+			));
+		}
+	} // namespace ModulePrivate
 
-				// Audio Widgets
-				{
-					const FLinearColor& AnalyzerBackgroundColor = GetColor("MetasoundEditor.Analyzers.BackgroundColor");
-					const FLinearColor& AnalyzerForegroundColor = GetColor("MetasoundEditor.Analyzers.ForegroundColor");
+	class FSlateStyle final : public FSlateStyleSet
+	{
+	public:
+		FSlateStyle()
+			: FSlateStyleSet("MetaSoundStyle")
+		{
+			SetParentStyleName(FAudioWidgetsStyle::Get().GetStyleSetName());
 
-					// Add static overrides for these widget styles:
-					FAudioMeterDefaultColorStyle MeterStyle;
-					MeterStyle.MeterValueColor = AnalyzerForegroundColor;
-					Set("AudioMeter.DefaultColorStyle", MeterStyle);
+			SetContentRoot(FPaths::EnginePluginsDir() / TEXT("Runtime/Metasound/Content/Editor/Slate"));
+			SetCoreContentRoot(FPaths::EngineContentDir() / TEXT("Slate"));
 
-					Set("AudioOscilloscope.PanelStyle", FAudioOscilloscopePanelStyle()
-						.SetWaveViewerStyle(FSampledSequenceViewerStyle()
-							.SetSequenceColor(AnalyzerForegroundColor)));
+			static const FVector2D Icon20x20(20.0f, 20.0f);
+			static const FVector2D Icon40x40(40.0f, 40.0f);
 
-					Set("AudioSpectrumPlot.Style", FAudioSpectrumPlotStyle()
-						.SetCrosshairColor(FSlateColor(AnalyzerForegroundColor).UseSubduedForeground())
-						.SetSpectrumColor(AnalyzerForegroundColor));
+			static const FVector2D Icon16 = FVector2D(16.0f, 16.0f);
+			static const FVector2D Icon64 = FVector2D(64.0f, 64.0f);
 
-					Set("AudioVectorscope.PanelStyle", FAudioVectorscopePanelStyle()
-						.SetVectorViewerStyle(FSampledSequenceVectorViewerStyle()
-							.SetLineColor(AnalyzerForegroundColor)));
+			const FVector2D Icon15x11(15.0f, 11.0f);
 
-					Set("AudioAnalyzerRack.BackgroundColor", AnalyzerBackgroundColor);
+			// Metasound Editor
+			{
+				Set("MetaSoundPatch.Color", FColor(31, 133, 31));
+				Set("MetaSoundSource.Color", FColor(103, 214, 66));
+
+				// Actions
+				Set("MetasoundEditor.Play", new IMAGE_BRUSH_SVG(TEXT("Icons/play"), Icon40x40));
+				Set("MetasoundEditor.Play.Small", new IMAGE_BRUSH_SVG(TEXT("Icons/play"), Icon20x20));
+				Set("MetasoundEditor.Play.Thumbnail", new IMAGE_BRUSH_SVG(TEXT("Icons/play_thumbnail"), Icon64));
+				Set("MetasoundEditor.Play.Thumbnail.Hovered", new IMAGE_BRUSH_SVG(TEXT("Icons/play_thumbnail_hover"), Icon64));
+
+				Set("MetasoundEditor.Play.Active.Valid", new IMAGE_BRUSH_SVG(TEXT("Icons/play_active_valid"), Icon40x40));
+				Set("MetasoundEditor.Play.Active.Warning", new IMAGE_BRUSH_SVG(TEXT("Icons/play_active_warning"), Icon40x40));
+				Set("MetasoundEditor.Play.Inactive.Valid", new IMAGE_BRUSH_SVG(TEXT("Icons/play_inactive_valid"), Icon40x40));
+				Set("MetasoundEditor.Play.Inactive.Warning", new IMAGE_BRUSH_SVG(TEXT("Icons/play_inactive_warning"), Icon40x40));
+				Set("MetasoundEditor.Play.Error", new IMAGE_BRUSH_SVG(TEXT("Icons/play_error"), Icon40x40));
+
+				Set("MetasoundEditor.Stop", new IMAGE_BRUSH_SVG(TEXT("Icons/stop"), Icon40x40));
+
+				Set("MetasoundEditor.Stop.Disabled", new IMAGE_BRUSH_SVG(TEXT("Icons/stop_disabled"), Icon40x40));
+				Set("MetasoundEditor.Stop.Active", new IMAGE_BRUSH_SVG(TEXT("Icons/stop_active"), Icon40x40));
+				Set("MetasoundEditor.Stop.Inactive", new IMAGE_BRUSH_SVG(TEXT("Icons/stop_inactive"), Icon40x40));
+				Set("MetasoundEditor.Stop.Thumbnail", new IMAGE_BRUSH_SVG(TEXT("Icons/stop_thumbnail"), Icon64));
+				Set("MetasoundEditor.Stop.Thumbnail.Hovered", new IMAGE_BRUSH_SVG(TEXT("Icons/stop_thumbnail_hover"), Icon64));
+
+				Set("MetasoundEditor.Import", new FSlateImageBrush(RootToContentDir(TEXT("/Icons/build_40x.png")), Icon40x40));
+				Set("MetasoundEditor.Import.Small", new FSlateImageBrush(RootToContentDir(TEXT("/Icons/build_40x.png")), Icon20x20));
+				Set("MetasoundEditor.Export", new FSlateImageBrush(RootToContentDir(TEXT("/Icons/build_40x.png")), Icon40x40));
+				Set("MetasoundEditor.Export.Small", new FSlateImageBrush(RootToContentDir(TEXT("/Icons/build_40x.png")), Icon20x20));
+				Set("MetasoundEditor.ExportError", new FSlateImageBrush(RootToContentDir(TEXT("/Icons/build_error_40x.png")), Icon40x40));
+				Set("MetasoundEditor.ExportError.Small", new FSlateImageBrush(RootToContentDir(TEXT("/Icons/build_error_40x.png")), Icon20x20));
+				Set("MetasoundEditor.Settings", new FSlateImageBrush(RootToContentDir(TEXT("/Icons/settings_40x.png")), Icon20x20));
+
+				// Graph Editor
+				Set("MetasoundEditor.Graph.Node.Body.Input", new FSlateImageBrush(RootToContentDir(TEXT("/Graph/node_input_body_64x.png")), FVector2D(114.0f, 64.0f)));
+				Set("MetasoundEditor.Graph.Node.Body.Default", new FSlateImageBrush(RootToContentDir(TEXT("/Graph/node_default_body_64x.png")), FVector2D(64.0f, 64.0f)));
+
+				Set("MetasoundEditor.Graph.TriggerPin.Connected", new IMAGE_BRUSH(TEXT("Graph/pin_trigger_connected"), Icon15x11));
+				Set("MetasoundEditor.Graph.TriggerPin.Disconnected", new IMAGE_BRUSH(TEXT("Graph/pin_trigger_disconnected"), Icon15x11));
+
+				Set("MetasoundEditor.Graph.Node.Class.Native", new IMAGE_BRUSH_SVG(TEXT("Icons/native_node"), FVector2D(8.0f, 16.0f)));
+				Set("MetasoundEditor.Graph.Node.Class.Graph", new IMAGE_BRUSH_SVG(TEXT("Icons/graph_node"), Icon16));
+				Set("MetasoundEditor.Graph.Node.Class.Input", new IMAGE_BRUSH_SVG(TEXT("Icons/input_node"), FVector2D(16.0f, 13.0f)));
+				Set("MetasoundEditor.Graph.Node.Class.Output", new IMAGE_BRUSH_SVG(TEXT("Icons/output_node"), FVector2D(16.0f, 13.0f)));
+				Set("MetasoundEditor.Graph.Node.Class.Reroute", new IMAGE_BRUSH_SVG(TEXT("Icons/reroute_node"), Icon16));
+				Set("MetasoundEditor.Graph.Node.Class.Variable", new IMAGE_BRUSH_SVG(TEXT("Icons/variable_node"), FVector2D(16.0f, 13.0f)));
+
+				Set("MetasoundEditor.Graph.Node.Math.Add", new FSlateImageBrush(RootToContentDir(TEXT("/Graph/node_math_add_40x.png")), Icon40x40));
+				Set("MetasoundEditor.Graph.Node.Math.Divide", new FSlateImageBrush(RootToContentDir(TEXT("/Graph/node_math_divide_40x.png")), Icon40x40));
+				Set("MetasoundEditor.Graph.Node.Math.Modulo", new FSlateImageBrush(RootToContentDir(TEXT("/Graph/node_math_modulo_40x.png")), Icon40x40));
+				Set("MetasoundEditor.Graph.Node.Math.Multiply", new FSlateImageBrush(RootToContentDir(TEXT("/Graph/node_math_multiply_40x.png")), Icon40x40));
+				Set("MetasoundEditor.Graph.Node.Math.Subtract", new FSlateImageBrush(RootToContentDir(TEXT("/Graph/node_math_subtract_40x.png")), Icon40x40));
+				Set("MetasoundEditor.Graph.Node.Math.Power", new FSlateImageBrush(RootToContentDir(TEXT("/Graph/node_math_power_40x.png")), Icon40x40));
+				Set("MetasoundEditor.Graph.Node.Math.Logarithm", new FSlateImageBrush(RootToContentDir(TEXT("/Graph/node_math_logarithm_40x.png")), Icon40x40));
+				Set("MetasoundEditor.Graph.Node.Conversion", new FSlateImageBrush(RootToContentDir(TEXT("/Graph/node_conversion_40x.png")), Icon40x40));
+
+				Set("MetasoundEditor.Graph.InvalidReroute", new IMAGE_BRUSH_SVG(TEXT("Icons/invalid_reroute"), Icon16));
+				Set("MetasoundEditor.Graph.ConstructorPinArray", new FSlateImageBrush(RootToContentDir(TEXT("/Icons/array_pin_rotated.png")), Icon16));
+				Set("MetasoundEditor.Graph.ConstructorPinArrayDisconnected", new FSlateImageBrush(RootToContentDir(TEXT("/Icons/array_pin_rotated_disconnected.png")), Icon16));
+				Set("MetasoundEditor.Graph.ArrayPin", new FSlateImageBrush(RootToContentDir(TEXT("/Icons/array_pin.png")), Icon16));
+				Set("MetasoundEditor.Graph.ConstructorPin", new FSlateImageBrush(RootToContentDir(TEXT("/Icons/square_pin_rotated.png")), Icon16));
+				Set("MetasoundEditor.Graph.ConstructorPinDisconnected", new FSlateImageBrush(RootToContentDir(TEXT("/Icons/square_pin_rotated_disconnected.png")), Icon16));
+
+				// Analyzers
+				Set("MetasoundEditor.Analyzers.BackgroundColor", FLinearColor(0.0075f, 0.0075f, 0.0075, 1.0f));
+				Set("MetasoundEditor.Analyzers.ForegroundColor", FLinearColor(0.025719f, 0.208333f, 0.069907f, 1.0f)); // "Audio" Green
+
+				// Misc
+				Set("MetasoundEditor.Audition", new IMAGE_BRUSH_SVG(TEXT("Icons/metasound_page"), Icon16));
+				Set("MetasoundEditor.Metasound.Icon", new IMAGE_BRUSH_SVG(TEXT("Icons/metasound_icon"), Icon16));
+				Set("MetasoundEditor.Speaker", new FSlateImageBrush(RootToContentDir(TEXT("/Icons/speaker_144x.png")), FVector2D(144.0f, 144.0f)));
+
+				// Pages
+				Set("MetasoundEditor.Page.Executing.ForegroundColor", FStyleColors::AccentGreen.GetSpecifiedColor());
+				Set("MetasoundEditor.Page.Executing", new IMAGE_BRUSH_SVG(TEXT("Icons/metasound_page_exec"), Icon16));
+
+				// Class Icons
+				auto SetClassIcon = [this, InIcon16 = Icon16, InIcon64 = Icon64](const FString& ClassName)
+				{
+					const FString IconFileName = FString::Printf(TEXT("Icons/%s"), *ClassName.ToLower());
+					const FSlateColor DefaultForeground(FStyleColors::Foreground);
 
-					// Add dynamic overrides for the widget style types with settings overrides:
-					AddDynamicLoadedWidgetStyle<FAudioMaterialKnobStyle>([]() { return GetDefault<UMetasoundEditorSettings>()->KnobStyleOverride; });
-					AddDynamicLoadedWidgetStyle<FAudioMaterialSliderStyle>([]() { return GetDefault<UMetasoundEditorSettings>()->SliderStyleOverride; });
-					AddDynamicLoadedWidgetStyle<FAudioMaterialButtonStyle>([]() { return GetDefault<UMetasoundEditorSettings>()->ButtonStyleOverride; });
-					AddDynamicLoadedWidgetStyle<FAudioMaterialMeterStyle>([]() { return GetDefault<UMetasoundEditorSettings>()->MeterStyleOverride; });
-				}
+					Set(*FString::Printf(TEXT("ClassIcon.%s"), *ClassName), new IMAGE_BRUSH_SVG(IconFileName, InIcon16));
+					Set(*FString::Printf(TEXT("ClassThumbnail.%s"), *ClassName), new IMAGE_BRUSH_SVG(IconFileName, InIcon64));
+				};
 
-				FSlateStyleRegistry::RegisterSlateStyle(*this);
+				SetClassIcon(TEXT("MetasoundPatch"));
+				SetClassIcon(TEXT("MetasoundSource"));
+
+				Set("MetasoundEditor.MetasoundPatch.Icon", new IMAGE_BRUSH_SVG(TEXT("Icons/metasoundpatch_icon"), Icon20x20));
+				Set("MetasoundEditor.MetasoundPatch.Preset.Icon", new IMAGE_BRUSH_SVG(TEXT("Icons/metasoundpatchpreset_icon"), Icon20x20));
+				Set("MetasoundEditor.MetasoundSource.Icon", new IMAGE_BRUSH_SVG(TEXT("Icons/metasoundsource_icon"), Icon20x20));
+				Set("MetasoundEditor.MetasoundSource.Preset.Icon", new IMAGE_BRUSH_SVG(TEXT("Icons/metasoundsourcepreset_icon"), Icon20x20));
+				Set("MetasoundEditor.MetasoundPatch.Thumbnail", new IMAGE_BRUSH_SVG(TEXT("Icons/metasoundpatch_thumbnail"), Icon20x20));
+				Set("MetasoundEditor.MetasoundPatch.Preset.Thumbnail", new IMAGE_BRUSH_SVG(TEXT("Icons/metasoundpatchpreset_thumbnail"), Icon20x20));
+				Set("MetasoundEditor.MetasoundSource.Thumbnail", new IMAGE_BRUSH_SVG(TEXT("Icons/metasoundsource_thumbnail"), Icon20x20));
+				Set("MetasoundEditor.MetasoundSource.Preset.Thumbnail", new IMAGE_BRUSH_SVG(TEXT("Icons/metasoundsourcepreset_thumbnail"), Icon20x20));
 			}
 
-		private:
-			virtual const FSlateWidgetStyle* GetWidgetStyleInternal(const FName DesiredTypeName, const FName StyleName, const FSlateWidgetStyle* DefaultStyle, bool bWarnIfNotFound) const override
+			// Audio Widgets
 			{
-				if (const auto* WidgetStyleLoader = AudioMaterialWidgetStyleLoaders.Find(DesiredTypeName))
-				{
-					const UMetasoundEditorSettings* EditorSettings = GetDefault<UMetasoundEditorSettings>();
-					if (EditorSettings == nullptr || !EditorSettings->bUseAudioMaterialWidgets)
-					{
-						ensure(!bWarnIfNotFound);
+				const FLinearColor& AnalyzerBackgroundColor = GetColor("MetasoundEditor.Analyzers.BackgroundColor");
+				const FLinearColor& AnalyzerForegroundColor = GetColor("MetasoundEditor.Analyzers.ForegroundColor");
 
-						// AudioMaterialWidgets are disabled, so explicitly return null rather than fallback to the base style.
-						return nullptr;
-					}
-					
-					if (const FSlateWidgetStyle* CustomWidgetStyle = (*WidgetStyleLoader)->LoadWidgetStyle())
-					{
-						return CustomWidgetStyle;
-					}
-				}
-				
-				const FSlateWidgetStyle* WidgetStyle = FSlateStyleSet::GetWidgetStyleInternal(DesiredTypeName, StyleName, DefaultStyle, bWarnIfNotFound);
-				ensure(!bWarnIfNotFound || WidgetStyle != nullptr);
-				return WidgetStyle;
-			}
+				// Add static overrides for these widget styles:
+				FAudioMeterDefaultColorStyle MeterStyle;
+				MeterStyle.MeterValueColor = AnalyzerForegroundColor;
+				Set("AudioMeter.DefaultColorStyle", MeterStyle);
 
-			using FGetAssetPathFunc = TFunction<FSoftObjectPath()>;
+				Set("AudioOscilloscope.PanelStyle", FAudioOscilloscopePanelStyle()
+					.SetWaveViewerStyle(FSampledSequenceViewerStyle()
+						.SetSequenceColor(AnalyzerForegroundColor)));
 
-			template<std::derived_from<FSlateWidgetStyle> WidgetStyleType>
-			void AddDynamicLoadedWidgetStyle(FGetAssetPathFunc GetAssetPath)
-			{
-				static_assert(std::is_base_of_v<FAudioMaterialWidgetStyle, WidgetStyleType>);
-				AudioMaterialWidgetStyleLoaders.Add(WidgetStyleType::TypeName, MakeShared<TWidgetStyleLoader<WidgetStyleType>>(GetAssetPath));
+				Set("AudioSpectrumPlot.Style", FAudioSpectrumPlotStyle()
+					.SetCrosshairColor(FSlateColor(AnalyzerForegroundColor).UseSubduedForeground())
+					.SetSpectrumColor(AnalyzerForegroundColor));
+
+				Set("AudioVectorscope.PanelStyle", FAudioVectorscopePanelStyle()
+					.SetVectorViewerStyle(FSampledSequenceVectorViewerStyle()
+						.SetLineColor(AnalyzerForegroundColor)));
+
+				Set("AudioAnalyzerRack.BackgroundColor", AnalyzerBackgroundColor);
+
+				// Add dynamic overrides for the widget style types with settings overrides:
+				AddDynamicLoadedWidgetStyle<FAudioMaterialKnobStyle>([]() { return GetDefault<UMetasoundEditorSettings>()->KnobStyleOverride; });
+				AddDynamicLoadedWidgetStyle<FAudioMaterialSliderStyle>([]() { return GetDefault<UMetasoundEditorSettings>()->SliderStyleOverride; });
+				AddDynamicLoadedWidgetStyle<FAudioMaterialButtonStyle>([]() { return GetDefault<UMetasoundEditorSettings>()->ButtonStyleOverride; });
+				AddDynamicLoadedWidgetStyle<FAudioMaterialMeterStyle>([]() { return GetDefault<UMetasoundEditorSettings>()->MeterStyleOverride; });
 			}
 
-			class IWidgetStyleLoader
-			{
-			public:
-				virtual ~IWidgetStyleLoader() {};
-				virtual const FSlateWidgetStyle* LoadWidgetStyle() = 0;
-			};
+			FSlateStyleRegistry::RegisterSlateStyle(*this);
+		}
 
-			template<std::derived_from<FSlateWidgetStyle> WidgetStyleType>
-			class TWidgetStyleLoader : public IWidgetStyleLoader
+	private:
+		virtual const FSlateWidgetStyle* GetWidgetStyleInternal(const FName DesiredTypeName, const FName StyleName, const FSlateWidgetStyle* DefaultStyle, bool bWarnIfNotFound) const override
+		{
+			if (const auto* WidgetStyleLoader = AudioMaterialWidgetStyleLoaders.Find(DesiredTypeName))
 			{
-			public:
-				TWidgetStyleLoader(FGetAssetPathFunc InGetAssetPath)
-					: GetAssetPath(InGetAssetPath)
-				{
-					//
-				}
-
-				virtual const FSlateWidgetStyle* LoadWidgetStyle() override
+				const UMetasoundEditorSettings* EditorSettings = GetDefault<UMetasoundEditorSettings>();
+				if (EditorSettings == nullptr || !EditorSettings->bUseAudioMaterialWidgets)
 				{
-					TSoftObjectPtr<USlateWidgetStyleAsset> SoftObjectPtr(GetAssetPath());
-					if (const USlateWidgetStyleAsset* WidgetStyleAsset = SoftObjectPtr.LoadSynchronous())
-					{
-						if (const WidgetStyleType* WidgetStyle = WidgetStyleAsset->GetStyle<WidgetStyleType>())
-						{
-							// Copy style:
-							CachedStyle = *WidgetStyle;
-
-							// Return pointer to our copy:
-							return &CachedStyle;
-						}
-					}
+					ensure(!bWarnIfNotFound);
 
+					// AudioMaterialWidgets are disabled, so explicitly return null rather than fallback to the base style.
 					return nullptr;
 				}
+					
+				if (const FSlateWidgetStyle* CustomWidgetStyle = (*WidgetStyleLoader)->LoadWidgetStyle())
+				{
+					return CustomWidgetStyle;
+				}
+			}
+				
+			const FSlateWidgetStyle* WidgetStyle = FSlateStyleSet::GetWidgetStyleInternal(DesiredTypeName, StyleName, DefaultStyle, bWarnIfNotFound);
+			ensure(!bWarnIfNotFound || WidgetStyle != nullptr);
+			return WidgetStyle;
+		}
 
-			private:
-				FGetAssetPathFunc GetAssetPath;
-				WidgetStyleType CachedStyle;
-			};
+		using FGetAssetPathFunc = TFunction<FSoftObjectPath()>;
 
-			TMap<FName, TSharedRef<IWidgetStyleLoader>> AudioMaterialWidgetStyleLoaders;
+		template<std::derived_from<FSlateWidgetStyle> WidgetStyleType>
+		void AddDynamicLoadedWidgetStyle(FGetAssetPathFunc GetAssetPath)
+		{
+			static_assert(std::is_base_of_v<FAudioMaterialWidgetStyle, WidgetStyleType>);
+			AudioMaterialWidgetStyleLoaders.Add(WidgetStyleType::TypeName, MakeShared<TWidgetStyleLoader<WidgetStyleType>>(GetAssetPath));
+		}
+
+		class IWidgetStyleLoader
+		{
+		public:
+			virtual ~IWidgetStyleLoader() {};
+			virtual const FSlateWidgetStyle* LoadWidgetStyle() = 0;
 		};
 
-		namespace Style
+		template<std::derived_from<FSlateWidgetStyle> WidgetStyleType>
+		class TWidgetStyleLoader : public IWidgetStyleLoader
 		{
-			FSlateIcon CreateSlateIcon(FName InName)
+		public:
+			TWidgetStyleLoader(FGetAssetPathFunc InGetAssetPath)
+				: GetAssetPath(InGetAssetPath)
 			{
-				return { "MetaSoundStyle", InName};
+				//
 			}
 
-			const FSlateBrush& GetSlateBrushSafe(FName InName)
+			virtual const FSlateWidgetStyle* LoadWidgetStyle() override
 			{
-				const ISlateStyle* MetaSoundStyle = FSlateStyleRegistry::FindSlateStyle("MetaSoundStyle");
-				if (ensureMsgf(MetaSoundStyle, TEXT("Missing slate style 'MetaSoundStyle'")))
+				TSoftObjectPtr<USlateWidgetStyleAsset> SoftObjectPtr(GetAssetPath());
+				if (const USlateWidgetStyleAsset* WidgetStyleAsset = SoftObjectPtr.LoadSynchronous())
 				{
-					const FSlateBrush* Brush = MetaSoundStyle->GetBrush(InName);
-					if (ensureMsgf(Brush, TEXT("Missing brush '%s'"), *InName.ToString()))
+					if (const WidgetStyleType* WidgetStyle = WidgetStyleAsset->GetStyle<WidgetStyleType>())
 					{
-						return *Brush;
+						// Copy style:
+						CachedStyle = *WidgetStyle;
+
+						// Return pointer to our copy:
+						return &CachedStyle;
 					}
 				}
 
-				if (const FSlateBrush* NoBrush = FAppStyle::GetBrush("NoBrush"))
+				return nullptr;
+			}
+
+		private:
+			FGetAssetPathFunc GetAssetPath;
+			WidgetStyleType CachedStyle;
+		};
+
+		TMap<FName, TSharedRef<IWidgetStyleLoader>> AudioMaterialWidgetStyleLoaders;
+	};
+
+	namespace Style
+	{
+		FSlateIcon CreateSlateIcon(FName InName)
+		{
+			return { "MetaSoundStyle", InName};
+		}
+
+		const FSlateBrush& GetSlateBrushSafe(FName InName)
+		{
+			const ISlateStyle* MetaSoundStyle = FSlateStyleRegistry::FindSlateStyle("MetaSoundStyle");
+			if (ensureMsgf(MetaSoundStyle, TEXT("Missing slate style 'MetaSoundStyle'")))
+			{
+				const FSlateBrush* Brush = MetaSoundStyle->GetBrush(InName);
+				if (ensureMsgf(Brush, TEXT("Missing brush '%s'"), *InName.ToString()))
 				{
-					return *NoBrush;
+					return *Brush;
 				}
+			}
 
-				static const FSlateBrush NullBrush;
-				return NullBrush;
+			if (const FSlateBrush* NoBrush = FAppStyle::GetBrush("NoBrush"))
+			{
+				return *NoBrush;
 			}
 
-			const FSlateColor& GetPageExecutingColor()
+			static const FSlateBrush NullBrush;
+			return NullBrush;
+		}
+
+		const FSlateColor& GetPageExecutingColor()
+		{
+			auto MakeColor = []() -> FSlateColor
 			{
-				auto MakeColor = []() -> FSlateColor
+				if (const ISlateStyle* MetaSoundStyle = FSlateStyleRegistry::FindSlateStyle("MetaSoundStyle"))
 				{
-					if (const ISlateStyle* MetaSoundStyle = FSlateStyleRegistry::FindSlateStyle("MetaSoundStyle"))
-					{
-						return MetaSoundStyle->GetColor("MetasoundEditor.Page.Executing.ForegroundColor");
-					}
+					return MetaSoundStyle->GetColor("MetasoundEditor.Page.Executing.ForegroundColor");
+				}
 
-					return FStyleColors::AccentWhite;
-				};
-				static const FSlateColor AnalyzerColor = MakeColor();
-				return AnalyzerColor;
-			}
-		} // namespace Style
+				return FStyleColors::AccentWhite;
+			};
+			static const FSlateColor AnalyzerColor = MakeColor();
+			return AnalyzerColor;
+		}
+	} // namespace Style
 
-		// A structure that contains information about registered custom pin types. 
-		struct FGraphPinConfiguration
-		{
-			FEdGraphPinType PinType;
-			const FSlateBrush* PinConnectedIcon = nullptr;
-			const FSlateBrush* PinDisconnectedIcon = nullptr;
-		};
+	// A structure that contains information about registered custom pin types. 
+	struct FGraphPinConfiguration
+	{
+		FEdGraphPinType PinType;
+		const FSlateBrush* PinConnectedIcon = nullptr;
+		const FSlateBrush* PinDisconnectedIcon = nullptr;
+	};
 
 
-		class FModule : public IMetasoundEditorModule
+	class FModule : public IMetasoundEditorModule
+	{
+		void RegisterInputDefaultClasses()
 		{
-			void RegisterInputDefaultClasses()
+			TSubclassOf<UMetasoundEditorGraphMemberDefaultLiteral> NodeClass;
+			for (TObjectIterator<UClass> ClassIt; ClassIt; ++ClassIt)
 			{
-				TSubclassOf<UMetasoundEditorGraphMemberDefaultLiteral> NodeClass;
-				for (TObjectIterator<UClass> ClassIt; ClassIt; ++ClassIt)
+				UClass* Class = *ClassIt;
+				if (!Class->IsNative())
 				{
-					UClass* Class = *ClassIt;
-					if (!Class->IsNative())
-					{
-						continue;
-					}
+					continue;
+				}
 
-					if (Class->HasAnyClassFlags(CLASS_Deprecated | CLASS_NewerVersionExists))
-					{
-						continue;
-					}
+				if (Class->HasAnyClassFlags(CLASS_Deprecated | CLASS_NewerVersionExists))
+				{
+					continue;
+				}
 
-					if (!ClassIt->IsChildOf(UMetasoundEditorGraphMemberDefaultLiteral::StaticClass()))
-					{
-						continue;
-					}
+				if (!ClassIt->IsChildOf(UMetasoundEditorGraphMemberDefaultLiteral::StaticClass()))
+				{
+					continue;
+				}
 
-					if (const UMetasoundEditorGraphMemberDefaultLiteral* DefaultLiteralCDO = Class->GetDefaultObject<UMetasoundEditorGraphMemberDefaultLiteral>())
-					{
-						InputDefaultLiteralClassRegistry.Add(DefaultLiteralCDO->GetLiteralType(), DefaultLiteralCDO->GetClass());
-					}
+				if (const UMetasoundEditorGraphMemberDefaultLiteral* DefaultLiteralCDO = Class->GetDefaultObject<UMetasoundEditorGraphMemberDefaultLiteral>())
+				{
+					InputDefaultLiteralClassRegistry.Add(DefaultLiteralCDO->GetLiteralType(), DefaultLiteralCDO->GetClass());
 				}
 			}
+		}
 
-			void RegisterCorePinTypes()
-			{
-				using namespace Metasound::Frontend;
+		void RegisterCorePinTypes()
+		{
+			using namespace Metasound::Frontend;
 
-				const IDataTypeRegistry& DataTypeRegistry = IDataTypeRegistry::Get();
+			const IDataTypeRegistry& DataTypeRegistry = IDataTypeRegistry::Get();
 
-				TArray<FName> DataTypeNames;
-				DataTypeRegistry.GetRegisteredDataTypeNames(DataTypeNames);
+			TArray<FName> DataTypeNames;
+			DataTypeRegistry.GetRegisteredDataTypeNames(DataTypeNames);
 
-				for (FName DataTypeName : DataTypeNames)
+			for (FName DataTypeName : DataTypeNames)
+			{
+				FDataTypeRegistryInfo RegistryInfo;
+				if (ensure(DataTypeRegistry.GetDataTypeInfo(DataTypeName, RegistryInfo)))
 				{
-					FDataTypeRegistryInfo RegistryInfo;
-					if (ensure(DataTypeRegistry.GetDataTypeInfo(DataTypeName, RegistryInfo)))
-					{
-						FName PinCategory = DataTypeName;
-						FName PinSubCategory;
+					FName PinCategory = DataTypeName;
+					FName PinSubCategory;
 
-						// Types like triggers & AudioBuffer are specialized, so ignore their preferred
-						// literal types to classify the category.
-						if (!FGraphBuilder::IsPinCategoryMetaSoundCustomDataType(PinCategory) && !CustomPinCategories.Contains(PinCategory))
+					// Types like triggers & AudioBuffer are specialized, so ignore their preferred
+					// literal types to classify the category.
+					if (!FGraphBuilder::IsPinCategoryMetaSoundCustomDataType(PinCategory) && !CustomPinCategories.Contains(PinCategory))
+					{
+						// Primitives
+						switch (RegistryInfo.PreferredLiteralType)
 						{
-							// Primitives
-							switch (RegistryInfo.PreferredLiteralType)
+							case ELiteralType::Boolean:
+							case ELiteralType::BooleanArray:
 							{
-								case ELiteralType::Boolean:
-								case ELiteralType::BooleanArray:
-								{
-									PinCategory = FGraphBuilder::PinCategoryBoolean;
-								}
-								break;
+								PinCategory = FGraphBuilder::PinCategoryBoolean;
+							}
+							break;
 
-								case ELiteralType::Float:
+							case ELiteralType::Float:
+							{
+								PinCategory = FGraphBuilder::PinCategoryFloat;
+							}
+							break;
+
+							case ELiteralType::FloatArray:
+							{
+								if (RegistryInfo.bIsArrayType)
 								{
 									PinCategory = FGraphBuilder::PinCategoryFloat;
 								}
-								break;
+							}
+							break;
 
-								case ELiteralType::FloatArray:
-								{
-									if (RegistryInfo.bIsArrayType)
-									{
-										PinCategory = FGraphBuilder::PinCategoryFloat;
-									}
-								}
-								break;
+							case ELiteralType::Integer:
+							{
+								PinCategory = FGraphBuilder::PinCategoryInt32;
+							}
+							break;
 
-								case ELiteralType::Integer:
+							case ELiteralType::IntegerArray:
+							{
+								if (RegistryInfo.bIsArrayType)
 								{
 									PinCategory = FGraphBuilder::PinCategoryInt32;
 								}
-								break;
+							}
+							break;
 
-								case ELiteralType::IntegerArray:
-								{
-									if (RegistryInfo.bIsArrayType)
-									{
-										PinCategory = FGraphBuilder::PinCategoryInt32;
-									}
-								}
-								break;
+							case ELiteralType::String:
+							{
+								PinCategory = FGraphBuilder::PinCategoryString;
+							}
+							break;
 
-								case ELiteralType::String:
+							case ELiteralType::StringArray:
+							{
+								if (RegistryInfo.bIsArrayType)
 								{
 									PinCategory = FGraphBuilder::PinCategoryString;
 								}
-								break;
-
-								case ELiteralType::StringArray:
-								{
-									if (RegistryInfo.bIsArrayType)
-									{
-										PinCategory = FGraphBuilder::PinCategoryString;
-									}
-								}
-								break;
+							}
+							break;
 
-								case ELiteralType::UObjectProxy:
-								case ELiteralType::UObjectProxyArray:
-								{
-									PinCategory = FGraphBuilder::PinCategoryObject;
-								}
-								break;
+							case ELiteralType::UObjectProxy:
+							case ELiteralType::UObjectProxyArray:
+							{
+								PinCategory = FGraphBuilder::PinCategoryObject;
+							}
+							break;
 
-								case ELiteralType::None:
-								case ELiteralType::NoneArray:
-								case ELiteralType::Invalid:
-								default:
-								{
-									static_assert(static_cast<int32>(ELiteralType::Invalid) == 12, "Possible missing binding of pin category to primitive type");
-								}
-								break;
+							case ELiteralType::None:
+							case ELiteralType::NoneArray:
+							case ELiteralType::Invalid:
+							default:
+							{
+								static_assert(static_cast<int32>(ELiteralType::Invalid) == 12, "Possible missing binding of pin category to primitive type");
 							}
+							break;
 						}
-
-						RegisterPinType(DataTypeName, PinCategory, PinSubCategory);
 					}
+
+					RegisterPinType(DataTypeName, PinCategory, PinSubCategory);
 				}
 			}
+		}
 
-			virtual void RegisterPinType(FName InDataTypeName, FName InPinCategory, FName InPinSubCategory, const FSlateBrush* InPinConnectedIcon = nullptr, const FSlateBrush* InPinDisconnectedIcon = nullptr) override
+		void RegisterFolderBulkOperations()
+		{
+			if (UToolMenu* Menu = UToolMenus::Get()->ExtendMenu("ContentBrowser.FolderContextMenu"))
 			{
-				using namespace Frontend;
-
-				FDataTypeRegistryInfo DataTypeInfo;
-				IDataTypeRegistry::Get().GetDataTypeInfo(InDataTypeName, DataTypeInfo);
-
-				// Default to object as most calls to this outside of the MetaSound Editor will be for custom UObject types
-				const FName PinCategory = InPinCategory.IsNone() ? FGraphBuilder::PinCategoryObject : InPinCategory;
-
-				const EPinContainerType ContainerType = DataTypeInfo.bIsArrayType ? EPinContainerType::Array : EPinContainerType::None;
-				FGraphPinConfiguration PinConfiguration;
-				PinConfiguration.PinType.PinCategory = PinCategory;
-				PinConfiguration.PinType.PinSubCategory = InPinSubCategory;
-				PinConfiguration.PinType.ContainerType = ContainerType;
-				UClass* ClassToUse = IDataTypeRegistry::Get().GetUClassForDataType(InDataTypeName);
-				PinConfiguration.PinType.PinSubCategoryObject = Cast<UObject>(ClassToUse);
-				PinConfiguration.PinConnectedIcon = InPinConnectedIcon;
-				PinConfiguration.PinDisconnectedIcon = InPinDisconnectedIcon;
-				PinTypes.Emplace(InDataTypeName, MoveTemp(PinConfiguration));
+				FToolMenuSection& BulkOpsSection = Menu->FindOrAddSection("PathContextBulkOperations");
+				constexpr bool bOpenSubMenuOnClick = false;
+				BulkOpsSection.AddSubMenu(
+					"MetaSounds",
+					LOCTEXT("MetaSoundBulkOpsLabel", "MetaSounds"),
+					LOCTEXT("MetaSoundBulkOps_ToolTip", "Bulk operations pertaining to MetaSound UObject types"),
+					FNewToolMenuDelegate::CreateStatic(&ModulePrivate::AddVersionDocumentFolderMenuEntry),
+					bOpenSubMenuOnClick,
+					Style::CreateSlateIcon("ClassIcon.MetasoundPatch")
+				);
 			}
+		}
+
+		virtual void RegisterPinType(FName InDataTypeName, FName InPinCategory, FName InPinSubCategory, const FSlateBrush* InPinConnectedIcon = nullptr, const FSlateBrush* InPinDisconnectedIcon = nullptr) override
+		{
+			using namespace Frontend;
+
+			FDataTypeRegistryInfo DataTypeInfo;
+			IDataTypeRegistry::Get().GetDataTypeInfo(InDataTypeName, DataTypeInfo);
+
+			// Default to object as most calls to this outside of the MetaSound Editor will be for custom UObject types
+			const FName PinCategory = InPinCategory.IsNone() ? FGraphBuilder::PinCategoryObject : InPinCategory;
+
+			const EPinContainerType ContainerType = DataTypeInfo.bIsArrayType ? EPinContainerType::Array : EPinContainerType::None;
+			FGraphPinConfiguration PinConfiguration;
+			PinConfiguration.PinType.PinCategory = PinCategory;
+			PinConfiguration.PinType.PinSubCategory = InPinSubCategory;
+			PinConfiguration.PinType.ContainerType = ContainerType;
+			UClass* ClassToUse = IDataTypeRegistry::Get().GetUClassForDataType(InDataTypeName);
+			PinConfiguration.PinType.PinSubCategoryObject = Cast<UObject>(ClassToUse);
+			PinConfiguration.PinConnectedIcon = InPinConnectedIcon;
+			PinConfiguration.PinDisconnectedIcon = InPinDisconnectedIcon;
+			PinTypes.Emplace(InDataTypeName, MoveTemp(PinConfiguration));
+		}
 
-			virtual void RegisterCustomPinType(FName InDataTypeName, const FGraphPinParams& Params) override
+		virtual void RegisterCustomPinType(FName InDataTypeName, const FGraphPinParams& Params) override
+		{
+			RegisterPinType(InDataTypeName, Params.PinCategory, Params.PinSubcategory, Params.PinConnectedIcon, Params.PinDisconnectedIcon);
+			if (Params.PinCategory.IsNone())
 			{
-				RegisterPinType(InDataTypeName, Params.PinCategory, Params.PinSubcategory, Params.PinConnectedIcon, Params.PinDisconnectedIcon);
-				if (Params.PinCategory.IsNone())
-				{
-					return;
-				}
+				return;
+			}
 				
-				if (FGraphBuilder::IsPinCategoryMetaSoundCustomDataType(InDataTypeName))
-				{
-					UE_LOG(LogMetasoundEditor, Warning, TEXT("Attempted to register a \"Custom Pin Type\": \"%s\", but this is already a Metasound Custom Data Type"), *InDataTypeName.ToString());
-					return;
-				}
+			if (FGraphBuilder::IsPinCategoryMetaSoundCustomDataType(InDataTypeName))
+			{
+				UE_LOG(LogMetasoundEditor, Warning, TEXT("Attempted to register a \"Custom Pin Type\": \"%s\", but this is already a Metasound Custom Data Type"), *InDataTypeName.ToString());
+				return;
+			}
 				
-				CustomPinCategories.Add(Params.PinCategory);
-				UMetasoundEditorSettings* Settings = GetMutableDefault<UMetasoundEditorSettings>();
-				Settings->CustomPinTypeColors.Add(Params.PinCategory, Params.PinColor ? *Params.PinColor : Settings->DefaultPinTypeColor);
+			CustomPinCategories.Add(Params.PinCategory);
+			UMetasoundEditorSettings* Settings = GetMutableDefault<UMetasoundEditorSettings>();
+			Settings->CustomPinTypeColors.Add(Params.PinCategory, Params.PinColor ? *Params.PinColor : Settings->DefaultPinTypeColor);
 
-			}
+		}
 
-			virtual void RegisterGraphNodeVisualization(FName InNodeClassName, FOnCreateGraphNodeVisualizationWidget OnCreateGraphNodeVisualizationWidget) override
-			{
-				FGraphNodeVisualizationRegistry::Get().RegisterVisualization(InNodeClassName, OnCreateGraphNodeVisualizationWidget);
-			}
+		virtual void RegisterGraphNodeVisualization(FName InNodeClassName, FOnCreateGraphNodeVisualizationWidget OnCreateGraphNodeVisualizationWidget) override
+		{
+			FGraphNodeVisualizationRegistry::Get().RegisterVisualization(InNodeClassName, OnCreateGraphNodeVisualizationWidget);
+		}
 
-			virtual bool IsRestrictedMode() const override
-			{
-				return bIsRestrictedMode;
-			}
+		virtual bool IsRestrictedMode() const override
+		{
+			return bIsRestrictedMode;
+		}
+
+		virtual void SetRestrictedMode(bool bInRestrictedMode) override
+		{
+			bIsRestrictedMode = bInRestrictedMode;
+			const bool bEnableLogging = !bIsRestrictedMode;
+			Frontend::DocumentTransform::SetVersioningLoggingEnabled(bEnableLogging);
+		}
 
-			virtual void SetRestrictedMode(bool bInRestrictedMode) override
+		void RegisterSettingsDelegates()
+		{
+			using namespace Engine;
+
+			// All the following delegates are used for UX notification, audition
+			// and PIE which are not desired/necessary when cooking.
+			if (IsRunningCookCommandlet())
 			{
-				bIsRestrictedMode = bInRestrictedMode;
-				const bool bEnableLogging = !bIsRestrictedMode;
-				Frontend::DocumentTransform::SetVersioningLoggingEnabled(bEnableLogging);
+				return;
 			}
 
-			void RegisterSettingsDelegates()
+			if (UMetaSoundSettings* Settings = GetMutableDefault<UMetaSoundSettings>())
 			{
-				using namespace Engine;
-
-				// All the following delegates are used for UX notification, audition
-				// and PIE which are not desired/necessary when cooking.
-				if (IsRunningCookCommandlet())
+				Settings->GetOnDefaultRenamedDelegate().AddLambda([]()
 				{
-					return;
-				}
+					FNotificationInfo Info(LOCTEXT("MetaSoundSettings_CannotNameDefaultPage", "Cannot name 'Default': reserved MetaSound page name"));
+					Info.bFireAndForget = true;
+					Info.ExpireDuration = 2.0f;
+					Info.bUseThrobber = true;
+					FSlateNotificationManager::Get().AddNotification(Info);
+				});
+			}
 
-				if (UMetaSoundSettings* Settings = GetMutableDefault<UMetaSoundSettings>())
-				{
-					Settings->GetOnDefaultRenamedDelegate().AddLambda([]()
-					{
-						FNotificationInfo Info(LOCTEXT("MetaSoundSettings_CannotNameDefaultPage", "Cannot name 'Default': reserved MetaSound page name"));
-						Info.bFireAndForget = true;
-						Info.ExpireDuration = 2.0f;
-						Info.bUseThrobber = true;
-						FSlateNotificationManager::Get().AddNotification(Info);
-					});
-				}
+			const UMetasoundEditorSettings* EditorSettings = GetDefault<UMetasoundEditorSettings>();
+			FDocumentBuilderRegistry::GetChecked().GetOnResolveAuditionPageDelegate().BindUObject(EditorSettings, &UMetasoundEditorSettings::ResolveAuditionPage);
 
-				const UMetasoundEditorSettings* EditorSettings = GetDefault<UMetasoundEditorSettings>();
-				FDocumentBuilderRegistry::GetChecked().GetOnResolveAuditionPageDelegate().BindUObject(EditorSettings, &UMetasoundEditorSettings::ResolveAuditionPage);
+			FEditorDelegates::PreBeginPIE.AddWeakLambda(EditorSettings, [this](const bool /* bSimulating */)
+			{
+				using namespace Metasound::Frontend;
 
-				FEditorDelegates::PreBeginPIE.AddWeakLambda(EditorSettings, [this](const bool /* bSimulating */)
+				if (const UMetasoundEditorSettings* EdSettings = GetDefault<UMetasoundEditorSettings>())
 				{
-					using namespace Metasound::Frontend;
-
-					if (const UMetasoundEditorSettings* EdSettings = GetDefault<UMetasoundEditorSettings>())
+					FOnResolveEditorPage& OnResolveAuditionPage = FDocumentBuilderRegistry::GetChecked().GetOnResolveAuditionPageDelegate();
+					if (OnResolveAuditionPage.IsBoundToObject(EdSettings))
 					{
-						FOnResolveEditorPage& OnResolveAuditionPage = FDocumentBuilderRegistry::GetChecked().GetOnResolveAuditionPageDelegate();
-						if (OnResolveAuditionPage.IsBoundToObject(EdSettings))
+						if (!EdSettings->bApplyAuditionSettingsInPIE)
 						{
-							if (!EdSettings->bApplyAuditionSettingsInPIE)
-							{
-								OnResolveAuditionPage.Unbind();
-							}
+							OnResolveAuditionPage.Unbind();
 						}
-
-						IMetaSoundAssetManager::GetChecked().ReloadMetaSoundAssets();
 					}
-				});
-				FEditorDelegates::EndPIE.AddWeakLambda(EditorSettings, [](const bool /* bSimulating */)
+
+					IMetaSoundAssetManager::GetChecked().ReloadMetaSoundAssets();
+				}
+			});
+			FEditorDelegates::EndPIE.AddWeakLambda(EditorSettings, [](const bool /* bSimulating */)
+			{
+				if (const UMetasoundEditorSettings* EdSettings = GetDefault<UMetasoundEditorSettings>())
 				{
-					if (const UMetasoundEditorSettings* EdSettings = GetDefault<UMetasoundEditorSettings>())
+					FOnResolveEditorPage& OnResolveAuditionPage = FDocumentBuilderRegistry::GetChecked().GetOnResolveAuditionPageDelegate();
+					if (!OnResolveAuditionPage.IsBoundToObject(EdSettings))
 					{
-						FOnResolveEditorPage& OnResolveAuditionPage = FDocumentBuilderRegistry::GetChecked().GetOnResolveAuditionPageDelegate();
-						if (!OnResolveAuditionPage.IsBoundToObject(EdSettings))
-						{
-							OnResolveAuditionPage.BindUObject(EdSettings, &UMetasoundEditorSettings::ResolveAuditionPage);
-						}
+						OnResolveAuditionPage.BindUObject(EdSettings, &UMetasoundEditorSettings::ResolveAuditionPage);
 					}
-				});
+				}
+			});
 
-			}
+		}
 
-			virtual void PrimeAssetRegistryAsync() override {}
+		virtual void PrimeAssetRegistryAsync() override {}
 
-			virtual EAssetPrimeStatus GetAssetRegistryPrimeStatus() const override
-			{
-				return EAssetPrimeStatus::NotRequested;
-			}
+		virtual EAssetPrimeStatus GetAssetRegistryPrimeStatus() const override
+		{
+			return EAssetPrimeStatus::NotRequested;
+		}
 
-			virtual EAssetScanStatus GetAssetRegistryScanStatus() const override
-			{
-				return EAssetScanStatus::NotRequested;
-			}
+		virtual EAssetScanStatus GetAssetRegistryScanStatus() const override
+		{
+			return EAssetScanStatus::NotRequested;
+		}
 
-			virtual void RegisterExplicitProxyClass(const UClass& InClass) override
-			{
-				using namespace Metasound::Frontend;
+		virtual void RegisterExplicitProxyClass(const UClass& InClass) override
+		{
+			using namespace Metasound::Frontend;
 
-				const IDataTypeRegistry& DataTypeRegistry = IDataTypeRegistry::Get();
-				FDataTypeRegistryInfo RegistryInfo;
-				ensureAlways(DataTypeRegistry.IsUObjectProxyFactory(InClass.GetDefaultObject()));
+			const IDataTypeRegistry& DataTypeRegistry = IDataTypeRegistry::Get();
+			FDataTypeRegistryInfo RegistryInfo;
+			ensureAlways(DataTypeRegistry.IsUObjectProxyFactory(InClass.GetDefaultObject()));
 
-				ExplicitProxyClasses.Add(&InClass);
-			}
+			ExplicitProxyClasses.Add(&InClass);
+		}
 
-			virtual bool IsExplicitProxyClass(const UClass& InClass) const override
+		virtual bool IsExplicitProxyClass(const UClass& InClass) const override
+		{
+			return ExplicitProxyClasses.Contains(&InClass);
+		}
+
+		virtual TUniquePtr<FMetasoundDefaultLiteralCustomizationBase> CreateMemberDefaultLiteralCustomization(UClass& InClass, IDetailCategoryBuilder& InDefaultCategoryBuilder) const override
+		{
+			const TUniquePtr<IMemberDefaultLiteralCustomizationFactory>* CustomizationFactory = LiteralCustomizationFactories.Find(&InClass);
+			if (CustomizationFactory && CustomizationFactory->IsValid())
 			{
-				return ExplicitProxyClasses.Contains(&InClass);
+				return (*CustomizationFactory)->CreateLiteralCustomization(InDefaultCategoryBuilder);
 			}
 
-			virtual TUniquePtr<FMetasoundDefaultLiteralCustomizationBase> CreateMemberDefaultLiteralCustomization(UClass& InClass, IDetailCategoryBuilder& InDefaultCategoryBuilder) const override
-			{
-				const TUniquePtr<IMemberDefaultLiteralCustomizationFactory>* CustomizationFactory = LiteralCustomizationFactories.Find(&InClass);
-				if (CustomizationFactory && CustomizationFactory->IsValid())
-				{
-					return (*CustomizationFactory)->CreateLiteralCustomization(InDefaultCategoryBuilder);
-				}
+			return nullptr;
+		}
 
-				return nullptr;
-			}
+		virtual const TSubclassOf<UMetasoundEditorGraphMemberDefaultLiteral> FindDefaultLiteralClass(EMetasoundFrontendLiteralType InLiteralType) const override
+		{
+			return InputDefaultLiteralClassRegistry.FindRef(InLiteralType);
+		}
 
-			virtual const TSubclassOf<UMetasoundEditorGraphMemberDefaultLiteral> FindDefaultLiteralClass(EMetasoundFrontendLiteralType InLiteralType) const override
+		virtual const FSlateBrush* GetIconBrush(FName InDataType, const bool bIsConstructorType) const override
+		{
+			Frontend::FDataTypeRegistryInfo Info;
+			Frontend::IDataTypeRegistry::Get().GetDataTypeInfo(InDataType, Info);
+
+			if (Info.bIsArrayType)
 			{
-				return InputDefaultLiteralClassRegistry.FindRef(InLiteralType);
+				return bIsConstructorType ? &Style::GetSlateBrushSafe("MetasoundEditor.Graph.ConstructorPinArray") : &Style::GetSlateBrushSafe("MetasoundEditor.Graph.ArrayPin");
 			}
-
-			virtual const FSlateBrush* GetIconBrush(FName InDataType, const bool bIsConstructorType) const override
+			else
 			{
-				Frontend::FDataTypeRegistryInfo Info;
-				Frontend::IDataTypeRegistry::Get().GetDataTypeInfo(InDataType, Info);
-
-				if (Info.bIsArrayType)
-				{
-					return bIsConstructorType ? &Style::GetSlateBrushSafe("MetasoundEditor.Graph.ConstructorPinArray") : &Style::GetSlateBrushSafe("MetasoundEditor.Graph.ArrayPin");
-				}
-				else
-				{
-					return bIsConstructorType ? &Style::GetSlateBrushSafe("MetasoundEditor.Graph.ConstructorPin") : FAppStyle::GetBrush("Icons.BulletPoint");
-				}
+				return bIsConstructorType ? &Style::GetSlateBrushSafe("MetasoundEditor.Graph.ConstructorPin") : FAppStyle::GetBrush("Icons.BulletPoint");
 			}
+		}
 			
-			virtual bool GetCustomPinIcons(UEdGraphPin* InPin, const FSlateBrush*& PinConnectedIcon, const FSlateBrush*& PinDisconnectedIcon) const override
+		virtual bool GetCustomPinIcons(UEdGraphPin* InPin, const FSlateBrush*& PinConnectedIcon, const FSlateBrush*& PinDisconnectedIcon) const override
+		{
+			if (const UEdGraphNode* Node = InPin->GetOwningNode())
 			{
-				if (const UEdGraphNode* Node = InPin->GetOwningNode())
+				if (const UMetasoundEditorGraphNode* MetaSoundNode = Cast<UMetasoundEditorGraphNode>(InPin->GetOwningNode()))
 				{
-					if (const UMetasoundEditorGraphNode* MetaSoundNode = Cast<UMetasoundEditorGraphNode>(InPin->GetOwningNode()))
-					{
-						Metasound::Frontend::FDataTypeRegistryInfo RegistryInfo = MetaSoundNode->GetPinDataTypeInfo(*InPin);
-						return GetCustomPinIcons(RegistryInfo.DataTypeName, PinConnectedIcon, PinDisconnectedIcon);
-					}
+					Metasound::Frontend::FDataTypeRegistryInfo RegistryInfo = MetaSoundNode->GetPinDataTypeInfo(*InPin);
+					return GetCustomPinIcons(RegistryInfo.DataTypeName, PinConnectedIcon, PinDisconnectedIcon);
 				}
-				return false;
 			}
+			return false;
+		}
 
-			virtual bool GetCustomPinIcons(FName InDataType, const FSlateBrush*& PinConnectedIcon, const FSlateBrush*& PinDisconnectedIcon) const override
+		virtual bool GetCustomPinIcons(FName InDataType, const FSlateBrush*& PinConnectedIcon, const FSlateBrush*& PinDisconnectedIcon) const override
+		{
+			const FGraphPinConfiguration* PinConfiguration = PinTypes.Find(InDataType);
+			if (!PinConfiguration || (!PinConfiguration->PinConnectedIcon && !PinConfiguration->PinDisconnectedIcon))
 			{
-				const FGraphPinConfiguration* PinConfiguration = PinTypes.Find(InDataType);
-				if (!PinConfiguration || (!PinConfiguration->PinConnectedIcon && !PinConfiguration->PinDisconnectedIcon))
-				{
-					return false;
-				}
-				PinConnectedIcon = PinConfiguration->PinConnectedIcon;
-				PinDisconnectedIcon = PinConfiguration->PinDisconnectedIcon ? PinConfiguration->PinDisconnectedIcon : PinConfiguration->PinConnectedIcon;
-				return true;
+				return false;
 			}
+			PinConnectedIcon = PinConfiguration->PinConnectedIcon;
+			PinDisconnectedIcon = PinConfiguration->PinDisconnectedIcon ? PinConfiguration->PinDisconnectedIcon : PinConfiguration->PinConnectedIcon;
+			return true;
+		}
 
-			virtual const FEdGraphPinType* FindPinType(FName InDataTypeName) const
+		virtual const FEdGraphPinType* FindPinType(FName InDataTypeName) const
+		{
+			const FGraphPinConfiguration* PinConfiguration = PinTypes.Find(InDataTypeName);
+			if (PinConfiguration)
 			{
-				const FGraphPinConfiguration* PinConfiguration = PinTypes.Find(InDataTypeName);
-				if (PinConfiguration)
-				{
-					return &PinConfiguration->PinType;
-				}
-				return nullptr;
+				return &PinConfiguration->PinType;
 			}
+			return nullptr;
+		}
 
-			virtual bool IsMetaSoundAssetClass(const FTopLevelAssetPath& InClassName) const override
+		virtual bool IsMetaSoundAssetClass(const FTopLevelAssetPath& InClassName) const override
+		{
+			if (const UClass* ClassObject = FindObject<const UClass>(InClassName))
 			{
-				if (const UClass* ClassObject = FindObject<const UClass>(InClassName))
-				{
-					return IMetasoundUObjectRegistry::Get().IsRegisteredClass(*ClassObject);
-				}
-
-				return false;
+				return IMetasoundUObjectRegistry::Get().IsRegisteredClass(*ClassObject);
 			}
 
-			virtual void StartupModule() override
-			{
-				using namespace Metasound::Engine;
-				METASOUND_LLM_SCOPE;
-				// Register Metasound asset type actions
-				IAssetTools& AssetTools = FModuleManager::LoadModuleChecked<FAssetToolsModule>(AssetToolName).Get();
+			return false;
+		}
 
-				FPropertyEditorModule& PropertyModule = FModuleManager::LoadModuleChecked<FPropertyEditorModule>("PropertyEditor");
-				
-				PropertyModule.RegisterCustomClassLayout(
-					UMetaSoundPatch::StaticClass()->GetFName(),
-					FOnGetDetailCustomizationInstance::CreateLambda([]() { return MakeShared<FMetasoundDetailCustomization>(UMetaSoundPatch::GetDocumentPropertyName()); }));
-
-				PropertyModule.RegisterCustomClassLayout(
-					UMetaSoundSource::StaticClass()->GetFName(),
-					FOnGetDetailCustomizationInstance::CreateLambda([]() { return MakeShared<FMetasoundDetailCustomization>(UMetaSoundSource::GetDocumentPropertyName()); }));
-
-				PropertyModule.RegisterCustomClassLayout(
-					UMetasoundInterfacesView::StaticClass()->GetFName(),
-					FOnGetDetailCustomizationInstance::CreateLambda([]() { return MakeShared<FMetasoundInterfacesDetailCustomization>(); }));
-
-				PropertyModule.RegisterCustomClassLayout(
-					UMetasoundPagesView::StaticClass()->GetFName(),
-					FOnGetDetailCustomizationInstance::CreateLambda([]() { return MakeShared<FMetasoundPagesDetailCustomization>(); }));
-
-				PropertyModule.RegisterCustomClassLayout(
-					UMetasoundEditorGraphInput::StaticClass()->GetFName(),
-					FOnGetDetailCustomizationInstance::CreateLambda([]() { return MakeShared<FMetasoundInputDetailCustomization>(); }));
-
-				PropertyModule.RegisterCustomClassLayout(
-					UMetasoundEditorGraphOutput::StaticClass()->GetFName(),
-					FOnGetDetailCustomizationInstance::CreateLambda([]() { return MakeShared<FMetasoundOutputDetailCustomization>(); }));
-
-				PropertyModule.RegisterCustomClassLayout(
-					UMetasoundEditorGraphVariable::StaticClass()->GetFName(),
-					FOnGetDetailCustomizationInstance::CreateLambda([]() { return MakeShared<FMetasoundVariableDetailCustomization>(); }));
-
-				PropertyModule.RegisterCustomPropertyTypeLayout(
-					"MetasoundEditorGraphMemberDefaultBoolRef",
-					FOnGetPropertyTypeCustomizationInstance::CreateLambda([]() { return MakeShared<FMetasoundMemberDefaultBoolDetailCustomization>(); }));
-
-				PropertyModule.RegisterCustomPropertyTypeLayout(
-					"MetasoundEditorGraphMemberDefaultIntRef",
-					FOnGetPropertyTypeCustomizationInstance::CreateLambda([]() { return MakeShared<FMetasoundMemberDefaultIntDetailCustomization>(); }));
-
-				PropertyModule.RegisterCustomPropertyTypeLayout(
-					"MetasoundEditorGraphMemberDefaultObjectRef",
-					FOnGetPropertyTypeCustomizationInstance::CreateLambda([]() { return MakeShared<FMetasoundMemberDefaultObjectDetailCustomization>(); }));
-
-				LiteralCustomizationFactories.Add(UMetasoundEditorGraphMemberDefaultLiteral::StaticClass(), MakeUnique<FMetasoundDefaultLiteralCustomizationFactory>());
-				LiteralCustomizationFactories.Add(UMetasoundEditorGraphMemberDefaultBool::StaticClass(), MakeUnique<FMetasoundBoolLiteralCustomizationFactory>());
-				LiteralCustomizationFactories.Add(UMetasoundEditorGraphMemberDefaultBoolArray::StaticClass(), MakeUnique<FMetasoundDefaultLiteralCustomizationFactory>());
-				LiteralCustomizationFactories.Add(UMetasoundEditorGraphMemberDefaultFloat::StaticClass(), MakeUnique<FMetasoundFloatLiteralCustomizationFactory>());
-				LiteralCustomizationFactories.Add(UMetasoundEditorGraphMemberDefaultFloatArray::StaticClass(), MakeUnique<FMetasoundDefaultLiteralCustomizationFactory>());
-				LiteralCustomizationFactories.Add(UMetasoundEditorGraphMemberDefaultInt::StaticClass(), MakeUnique<FMetasoundDefaultLiteralCustomizationFactory>());
-				LiteralCustomizationFactories.Add(UMetasoundEditorGraphMemberDefaultIntArray::StaticClass(), MakeUnique<FMetasoundDefaultLiteralCustomizationFactory>());
-				LiteralCustomizationFactories.Add(UMetasoundEditorGraphMemberDefaultObject::StaticClass(), MakeUnique<FMetasoundDefaultLiteralCustomizationFactory>());
-				LiteralCustomizationFactories.Add(UMetasoundEditorGraphMemberDefaultObjectArray::StaticClass(), MakeUnique<FMetasoundObjectArrayLiteralCustomizationFactory>());
-				LiteralCustomizationFactories.Add(UMetasoundEditorGraphMemberDefaultString::StaticClass(), MakeUnique<FMetasoundDefaultLiteralCustomizationFactory>());
-				LiteralCustomizationFactories.Add(UMetasoundEditorGraphMemberDefaultStringArray::StaticClass(), MakeUnique<FMetasoundDefaultLiteralCustomizationFactory>());
-
-				StyleSet = MakeShared<FSlateStyle>();
-
-				RegisterCorePinTypes();
-				RegisterInputDefaultClasses();
-
-				GraphConnectionFactory = MakeShared<FGraphConnectionDrawingPolicyFactory>();
-				FEdGraphUtilities::RegisterVisualPinConnectionFactory(GraphConnectionFactory);
-
-				GraphNodeFactory = MakeShared<FMetasoundGraphNodeFactory>();
-				FEdGraphUtilities::RegisterVisualNodeFactory(GraphNodeFactory);
-
-				GraphPanelPinFactory = MakeShared<FMetasoundGraphPanelPinFactory>();
-				FEdGraphUtilities::RegisterVisualPinFactory(GraphPanelPinFactory);
-
-				RegisterGraphNodeVisualization(
-					"UE.Biquad Filter.Audio",
-					FOnCreateGraphNodeVisualizationWidget::CreateStatic(&CreateMetaSoundBiquadFilterGraphNodeVisualizationWidget));
-
-				RegisterGraphNodeVisualization(
-					"UE.Ladder Filter.Audio",
-					FOnCreateGraphNodeVisualizationWidget::CreateStatic(&CreateMetaSoundLadderFilterGraphNodeVisualizationWidget));
-
-				RegisterGraphNodeVisualization(
-					"UE.One-Pole High Pass Filter.Audio",
-					FOnCreateGraphNodeVisualizationWidget::CreateStatic(&CreateMetaSoundOnePoleHighPassFilterGraphNodeVisualizationWidget));
-
-				RegisterGraphNodeVisualization(
-					"UE.One-Pole Low Pass Filter.Audio",
-					FOnCreateGraphNodeVisualizationWidget::CreateStatic(&CreateMetaSoundOnePoleLowPassFilterGraphNodeVisualizationWidget));
-
-				RegisterGraphNodeVisualization(
-					"UE.State Variable Filter.Audio",
-					FOnCreateGraphNodeVisualizationWidget::CreateStatic(&CreateMetaSoundStateVariableFilterGraphNodeVisualizationWidget));
-
-				ISettingsModule& SettingsModule = FModuleManager::LoadModuleChecked<ISettingsModule>("Settings");
-
-				SettingsModule.RegisterSettings("Editor", "ContentEditors", "MetaSound Editor",
-					LOCTEXT("MetaSoundEditorSettingsName", "MetaSound Editor"),
-					LOCTEXT("MetaSoundEditorSettingsDescription", "Customize MetaSound Editor."),
-					GetMutableDefault<UMetasoundEditorSettings>()
-				);
+		virtual void StartupModule() override
+		{
+			using namespace Engine;
+			using namespace ModulePrivate;
 
-				// Metasound Engine registers USoundWave as a proxy class in the
-				// Metasound Frontend. The frontend registration must occur before
-				// the Metasound Editor registration of a USoundWave.
-				IMetasoundEngineModule& MetaSoundEngineModule = FModuleManager::LoadModuleChecked<IMetasoundEngineModule>("MetasoundEngine");
+			METASOUND_LLM_SCOPE;
+
+			FPropertyEditorModule& PropertyModule = FModuleManager::LoadModuleChecked<FPropertyEditorModule>("PropertyEditor");
 				
-				// Bind delegates for MetaSound registration in the asset registry
-				MetaSoundEngineModule.GetOnGraphRegisteredDelegate().BindLambda([](UObject& InMetaSound, Engine::ERegistrationAssetContext AssetContext)
-				{
-					using namespace Engine;
+			PropertyModule.RegisterCustomClassLayout(
+				UMetaSoundPatch::StaticClass()->GetFName(),
+				FOnGetDetailCustomizationInstance::CreateLambda([]() { return MakeShared<FMetasoundDetailCustomization>(UMetaSoundPatch::GetDocumentPropertyName()); }));
+
+			PropertyModule.RegisterCustomClassLayout(
+				UMetaSoundSource::StaticClass()->GetFName(),
+				FOnGetDetailCustomizationInstance::CreateLambda([]() { return MakeShared<FMetasoundDetailCustomization>(UMetaSoundSource::GetDocumentPropertyName()); }));
+
+			PropertyModule.RegisterCustomClassLayout(
+				UMetasoundInterfacesView::StaticClass()->GetFName(),
+				FOnGetDetailCustomizationInstance::CreateLambda([]() { return MakeShared<FMetasoundInterfacesDetailCustomization>(); }));
+
+			PropertyModule.RegisterCustomClassLayout(
+				UMetasoundPagesView::StaticClass()->GetFName(),
+				FOnGetDetailCustomizationInstance::CreateLambda([]() { return MakeShared<FMetasoundPagesDetailCustomization>(); }));
+
+			PropertyModule.RegisterCustomClassLayout(
+				UMetasoundEditorGraphInput::StaticClass()->GetFName(),
+				FOnGetDetailCustomizationInstance::CreateLambda([]() { return MakeShared<FMetasoundInputDetailCustomization>(); }));
+
+			PropertyModule.RegisterCustomClassLayout(
+				UMetasoundEditorGraphOutput::StaticClass()->GetFName(),
+				FOnGetDetailCustomizationInstance::CreateLambda([]() { return MakeShared<FMetasoundOutputDetailCustomization>(); }));
+
+			PropertyModule.RegisterCustomClassLayout(
+				UMetasoundEditorGraphVariable::StaticClass()->GetFName(),
+				FOnGetDetailCustomizationInstance::CreateLambda([]() { return MakeShared<FMetasoundVariableDetailCustomization>(); }));
+
+			PropertyModule.RegisterCustomPropertyTypeLayout(
+				"MetasoundEditorGraphMemberDefaultBoolRef",
+				FOnGetPropertyTypeCustomizationInstance::CreateLambda([]() { return MakeShared<FMetasoundMemberDefaultBoolDetailCustomization>(); }));
+
+			PropertyModule.RegisterCustomPropertyTypeLayout(
+				"MetasoundEditorGraphMemberDefaultIntRef",
+				FOnGetPropertyTypeCustomizationInstance::CreateLambda([]() { return MakeShared<FMetasoundMemberDefaultIntDetailCustomization>(); }));
+
+			PropertyModule.RegisterCustomPropertyTypeLayout(
+				"MetasoundEditorGraphMemberDefaultObjectRef",
+				FOnGetPropertyTypeCustomizationInstance::CreateLambda([]() { return MakeShared<FMetasoundMemberDefaultObjectDetailCustomization>(); }));
+
+			LiteralCustomizationFactories.Add(UMetasoundEditorGraphMemberDefaultLiteral::StaticClass(), MakeUnique<FMetasoundDefaultLiteralCustomizationFactory>());
+			LiteralCustomizationFactories.Add(UMetasoundEditorGraphMemberDefaultBool::StaticClass(), MakeUnique<FMetasoundBoolLiteralCustomizationFactory>());
+			LiteralCustomizationFactories.Add(UMetasoundEditorGraphMemberDefaultBoolArray::StaticClass(), MakeUnique<FMetasoundDefaultLiteralCustomizationFactory>());
+			LiteralCustomizationFactories.Add(UMetasoundEditorGraphMemberDefaultFloat::StaticClass(), MakeUnique<FMetasoundFloatLiteralCustomizationFactory>());
+			LiteralCustomizationFactories.Add(UMetasoundEditorGraphMemberDefaultFloatArray::StaticClass(), MakeUnique<FMetasoundDefaultLiteralCustomizationFactory>());
+			LiteralCustomizationFactories.Add(UMetasoundEditorGraphMemberDefaultInt::StaticClass(), MakeUnique<FMetasoundDefaultLiteralCustomizationFactory>());
+			LiteralCustomizationFactories.Add(UMetasoundEditorGraphMemberDefaultIntArray::StaticClass(), MakeUnique<FMetasoundDefaultLiteralCustomizationFactory>());
+			LiteralCustomizationFactories.Add(UMetasoundEditorGraphMemberDefaultObject::StaticClass(), MakeUnique<FMetasoundDefaultLiteralCustomizationFactory>());
+			LiteralCustomizationFactories.Add(UMetasoundEditorGraphMemberDefaultObjectArray::StaticClass(), MakeUnique<FMetasoundObjectArrayLiteralCustomizationFactory>());
+			LiteralCustomizationFactories.Add(UMetasoundEditorGraphMemberDefaultString::StaticClass(), MakeUnique<FMetasoundDefaultLiteralCustomizationFactory>());
+			LiteralCustomizationFactories.Add(UMetasoundEditorGraphMemberDefaultStringArray::StaticClass(), MakeUnique<FMetasoundDefaultLiteralCustomizationFactory>());
+
+			StyleSet = MakeShared<FSlateStyle>();
+
+			RegisterCorePinTypes();
+			RegisterInputDefaultClasses();
+
+			GraphConnectionFactory = MakeShared<FGraphConnectionDrawingPolicyFactory>();
+			FEdGraphUtilities::RegisterVisualPinConnectionFactory(GraphConnectionFactory);
+
+			GraphNodeFactory = MakeShared<FMetasoundGraphNodeFactory>();
+			FEdGraphUtilities::RegisterVisualNodeFactory(GraphNodeFactory);
+
+			GraphPanelPinFactory = MakeShared<FGraphPanelPinFactory>();
+			FEdGraphUtilities::RegisterVisualPinFactory(GraphPanelPinFactory);
+
+			RegisterGraphNodeVisualization(
+				"UE.Biquad Filter.Audio",
+				FOnCreateGraphNodeVisualizationWidget::CreateStatic(&CreateMetaSoundBiquadFilterGraphNodeVisualizationWidget));
+
+			RegisterGraphNodeVisualization(
+				"UE.Ladder Filter.Audio",
+				FOnCreateGraphNodeVisualizationWidget::CreateStatic(&CreateMetaSoundLadderFilterGraphNodeVisualizationWidget));
+
+			RegisterGraphNodeVisualization(
+				"UE.One-Pole High Pass Filter.Audio",
+				FOnCreateGraphNodeVisualizationWidget::CreateStatic(&CreateMetaSoundOnePoleHighPassFilterGraphNodeVisualizationWidget));
+
+			RegisterGraphNodeVisualization(
+				"UE.One-Pole Low Pass Filter.Audio",
+				FOnCreateGraphNodeVisualizationWidget::CreateStatic(&CreateMetaSoundOnePoleLowPassFilterGraphNodeVisualizationWidget));
+
+			RegisterGraphNodeVisualization(
+				"UE.State Variable Filter.Audio",
+				FOnCreateGraphNodeVisualizationWidget::CreateStatic(&CreateMetaSoundStateVariableFilterGraphNodeVisualizationWidget));
+
+			ISettingsModule& SettingsModule = FModuleManager::LoadModuleChecked<ISettingsModule>("Settings");
+
+			SettingsModule.RegisterSettings("Editor", "ContentEditors", "MetaSound Editor",
+				LOCTEXT("MetaSoundEditorSettingsName", "MetaSound Editor"),
+				LOCTEXT("MetaSoundEditorSettingsDescription", "Customize MetaSound Editor."),
+				GetMutableDefault<UMetasoundEditorSettings>()
+			);
+
+			// Metasound Engine registers USoundWave as a proxy class in the
+			// Metasound Frontend. The frontend registration must occur before
+			// the Metasound Editor registration of a USoundWave.
+			IMetasoundEngineModule& MetaSoundEngineModule = FModuleManager::LoadModuleChecked<IMetasoundEngineModule>("MetasoundEngine");
+				
+			// Bind delegates for MetaSound registration in the asset registry
+			MetaSoundEngineModule.GetOnGraphRegisteredDelegate().BindLambda([](UObject& InMetaSound, Engine::ERegistrationAssetContext AssetContext)
+			{
+				using namespace Engine;
 
-					// Use the editor version of RegisterWithFrontend so it refreshes any open MetaSound editors.
-					const bool bForceViewSynchronization = AssetContext == ERegistrationAssetContext::Renaming;
-					FGraphBuilder::RegisterGraphWithFrontend(InMetaSound, bForceViewSynchronization);
-				});
-				MetaSoundEngineModule.GetOnGraphUnregisteredDelegate().BindLambda([](UObject& InMetaSound, Engine::ERegistrationAssetContext AssetContext)
-				{
-					using namespace Engine;
+				// Use the editor version of RegisterWithFrontend so it refreshes any open MetaSound editors.
+				const bool bForceViewSynchronization = AssetContext == ERegistrationAssetContext::Renaming;
+				FGraphBuilder::RegisterGraphWithFrontend(InMetaSound, bForceViewSynchronization);
+			});
+			MetaSoundEngineModule.GetOnGraphUnregisteredDelegate().BindLambda([](UObject& InMetaSound, Engine::ERegistrationAssetContext AssetContext)
+			{
+				using namespace Engine;
 
-					switch(AssetContext)
+				switch(AssetContext)
+				{
+					case ERegistrationAssetContext::Reloading:
+					case ERegistrationAssetContext::Removing:
+					case ERegistrationAssetContext::Renaming:
 					{
-						case ERegistrationAssetContext::Reloading:
-						case ERegistrationAssetContext::Removing:
-						case ERegistrationAssetContext::Renaming:
-						{
-							UAssetEditorSubsystem* AssetEditorSubsystem = GEditor->GetEditorSubsystem<UAssetEditorSubsystem>();
-							if (AssetEditorSubsystem)
-							{
-								// Close the editors so the internal reference to the builder doesn't suddenly 
-								// upon GC enter an invalid state (pointing to a null MetaSound asset)
-								AssetEditorSubsystem->CloseAllEditorsForAsset(&InMetaSound);
-							}
-							break;
-						}
-
-						case ERegistrationAssetContext::None:
-						default:
+						UAssetEditorSubsystem* AssetEditorSubsystem = GEditor->GetEditorSubsystem<UAssetEditorSubsystem>();
+						if (AssetEditorSubsystem)
 						{
-							break;
+							// Close the editors so the internal reference to the builder doesn't suddenly 
+							// upon GC enter an invalid state (pointing to a null MetaSound asset)
+							AssetEditorSubsystem->CloseAllEditorsForAsset(&InMetaSound);
 						}
+						break;
 					}
 
-					Metasound::IMetasoundUObjectRegistry& UObjectRegistry = Metasound::IMetasoundUObjectRegistry::Get();
-					if (FMetasoundAssetBase* AssetBase = UObjectRegistry.GetObjectAsAssetBase(&InMetaSound))
+					case ERegistrationAssetContext::None:
+					default:
 					{
-						AssetBase->UnregisterGraphWithFrontend();
+						break;
 					}
-				});
+				}
 
-				// Required to ensure logic to order nodes for presets exclusive to
-				// editor is propagated to transform instances while editing in editor.
-				Frontend::DocumentTransform::RegisterNodeDisplayNameProjection([](const Frontend::FNodeHandle& NodeHandle)
+				Metasound::IMetasoundUObjectRegistry& UObjectRegistry = Metasound::IMetasoundUObjectRegistry::Get();
+				if (FMetasoundAssetBase* AssetBase = UObjectRegistry.GetObjectAsAssetBase(&InMetaSound))
 				{
-					constexpr bool bIncludeNamespace = false;
-					return FGraphBuilder::GetDisplayName(*NodeHandle, bIncludeNamespace);
-				});
-
-				AssetTools.GetOnPackageMigration().AddRaw(this, &FModule::OnPackageMigration);
-				RegisterSettingsDelegates();
-			}
+					AssetBase->UnregisterGraphWithFrontend();
+				}
+			});
 
-			virtual void ShutdownModule() override
+			// Required to ensure logic to order nodes for presets exclusive to
+			// editor is propagated to transform instances while editing in editor.
+			Frontend::DocumentTransform::RegisterNodeDisplayNameProjection([](const Frontend::FNodeHandle& NodeHandle)
 			{
-				METASOUND_LLM_SCOPE;
+				constexpr bool bIncludeNamespace = false;
+				return FGraphBuilder::GetDisplayName(*NodeHandle, bIncludeNamespace);
+			});
 
-				if (ISettingsModule* SettingsModule = FModuleManager::GetModulePtr<ISettingsModule>("Settings"))
-				{
-					SettingsModule->UnregisterSettings("Editor", "Audio", "MetaSound Editor");
-				}
+			IAssetTools& AssetTools = FModuleManager::LoadModuleChecked<FAssetToolsModule>(AssetToolName).Get();
+			AssetTools.GetOnPackageMigration().AddRaw(this, &FModule::OnPackageMigration);
 
-				if (FModuleManager::Get().IsModuleLoaded(AssetToolName))
-				{
-					IAssetTools& AssetTools = FModuleManager::GetModuleChecked<FAssetToolsModule>(AssetToolName).Get();
-					for (TSharedPtr<FAssetTypeActions_Base>& AssetAction : AssetActions)
-					{
-						AssetTools.UnregisterAssetTypeActions(AssetAction.ToSharedRef());
-					}
-					AssetTools.GetOnPackageMigration().RemoveAll(this);
-				}
+			RegisterSettingsDelegates();
+			RegisterFolderBulkOperations();
+		}
 
-				if (GraphConnectionFactory.IsValid())
-				{
-					FEdGraphUtilities::UnregisterVisualPinConnectionFactory(GraphConnectionFactory);
-				}
+		virtual void ShutdownModule() override
+		{
+			using namespace ModulePrivate;
 
-				if (GraphNodeFactory.IsValid())
-				{
-					FEdGraphUtilities::UnregisterVisualNodeFactory(GraphNodeFactory);
-					GraphNodeFactory.Reset();
-				}
+			METASOUND_LLM_SCOPE;
 
-				if (GraphPanelPinFactory.IsValid())
-				{
-					FEdGraphUtilities::UnregisterVisualPinFactory(GraphPanelPinFactory);
-					GraphPanelPinFactory.Reset();
-				}
+			if (ISettingsModule* SettingsModule = FModuleManager::GetModulePtr<ISettingsModule>("Settings"))
+			{
+				SettingsModule->UnregisterSettings("Editor", "Audio", "MetaSound Editor");
+			}
 
-				AssetActions.Reset();
-				PinTypes.Reset();
+			if (FModuleManager::Get().IsModuleLoaded(AssetToolName))
+			{
+				IAssetTools& AssetTools = FModuleManager::LoadModuleChecked<FAssetToolsModule>(AssetToolName).Get();
+				AssetTools.GetOnPackageMigration().RemoveAll(this);
+			}
 
-				LiteralCustomizationFactories.Reset();
+			if (GraphConnectionFactory.IsValid())
+			{
+				FEdGraphUtilities::UnregisterVisualPinConnectionFactory(GraphConnectionFactory);
+			}
 
-				FGraphNodeVisualizationRegistry::TearDown();
+			if (GraphNodeFactory.IsValid())
+			{
+				FEdGraphUtilities::UnregisterVisualNodeFactory(GraphNodeFactory);
+				GraphNodeFactory.Reset();
 			}
 
-			void OnPackageMigration(UE::AssetTools::FPackageMigrationContext& MigrationContext)
+			if (GraphPanelPinFactory.IsValid())
 			{
-				using namespace Metasound::Frontend;
+				FEdGraphUtilities::UnregisterVisualPinFactory(GraphPanelPinFactory);
+				GraphPanelPinFactory.Reset();
+			}
+
+			PinTypes.Reset();
+
+			LiteralCustomizationFactories.Reset();
 
-				// Migration can create temporary new packages that use the same name 
-				// (and therefore node registry key) as the asset migrated. 
-				// So generate new class names to avoid registry key collisions. 
-				if (MigrationContext.GetCurrentStep() == UE::AssetTools::FPackageMigrationContext::EPackageMigrationStep::InstancedPackagesLoaded)
+			FGraphNodeVisualizationRegistry::TearDown();
+		}
+
+		void OnPackageMigration(UE::AssetTools::FPackageMigrationContext& MigrationContext)
+		{
+			using namespace Metasound::Frontend;
+
+			// Migration can create temporary new packages that use the same name 
+			// (and therefore node registry key) as the asset migrated. 
+			// So generate new class names to avoid registry key collisions. 
+			if (MigrationContext.GetCurrentStep() == UE::AssetTools::FPackageMigrationContext::EPackageMigrationStep::InstancedPackagesLoaded)
+			{
+				// Gather the new MetaSound assets
+				TArray<FMetaSoundFrontendDocumentBuilder> NewMetaSoundAssetBuilders;
+				for (const UE::AssetTools::FPackageMigrationContext::FMigrationPackageData& MigrationPackageData : MigrationContext.GetMigrationPackagesData())
 				{
-					// Gather the new MetaSound assets
-					TArray<FMetaSoundFrontendDocumentBuilder> NewMetaSoundAssetBuilders;
-					for (const UE::AssetTools::FPackageMigrationContext::FMigrationPackageData& MigrationPackageData : MigrationContext.GetMigrationPackagesData())
+					UPackage* Package = MigrationPackageData.GetInstancedPackage();
+					if (Package)
 					{
-						UPackage* Package = MigrationPackageData.GetInstancedPackage();
-						if (Package)
+						UObject* MainAsset = Package->FindAssetInPackage();
+						// Only apply to MetaSound assets 
+						if (IMetasoundUObjectRegistry::Get().IsRegisteredClass(MainAsset))
 						{
-							UObject* MainAsset = Package->FindAssetInPackage();
-							// Only apply to MetaSound assets 
-							if (IMetasoundUObjectRegistry::Get().IsRegisteredClass(MainAsset))
-							{
-								NewMetaSoundAssetBuilders.Add(FMetaSoundFrontendDocumentBuilder(MainAsset));
-							}
+							NewMetaSoundAssetBuilders.Add(FMetaSoundFrontendDocumentBuilder(MainAsset));
 						}
 					}
+				}
 
-					// Assign new class names and cache mapping with old one
-					IMetaSoundAssetManager& AssetManager = IMetaSoundAssetManager::GetChecked();
-					TMap<FNodeRegistryKey, FNodeRegistryKey> OldToNewReferenceKeys;
-					for (FMetaSoundFrontendDocumentBuilder& MetaSoundBuilder : NewMetaSoundAssetBuilders)
-					{
-						FNodeRegistryKey OldRegistryKey(MetaSoundBuilder.GetConstDocumentChecked().RootGraph);
-						FNodeRegistryKey NewRegistryKey(EMetasoundFrontendClassType::External, MetaSoundBuilder.GenerateNewClassName(), OldRegistryKey.Version);
-						OldToNewReferenceKeys.FindOrAdd(MoveTemp(OldRegistryKey)) = MoveTemp(NewRegistryKey);
+				// Assign new class names and cache mapping with old one
+				IMetaSoundAssetManager& AssetManager = IMetaSoundAssetManager::GetChecked();
+				TMap<FNodeRegistryKey, FNodeRegistryKey> OldToNewReferenceKeys;
+				for (FMetaSoundFrontendDocumentBuilder& MetaSoundBuilder : NewMetaSoundAssetBuilders)
+				{
+					FNodeRegistryKey OldRegistryKey(MetaSoundBuilder.GetConstDocumentChecked().RootGraph);
+					FNodeRegistryKey NewRegistryKey(EMetasoundFrontendClassType::External, MetaSoundBuilder.GenerateNewClassName(), OldRegistryKey.Version);
+					OldToNewReferenceKeys.FindOrAdd(MoveTemp(OldRegistryKey)) = MoveTemp(NewRegistryKey);
 
-						UObject& MetaSoundObject = MetaSoundBuilder.CastDocumentObjectChecked<UObject>();
-						AssetManager.AddOrUpdateAsset(MetaSoundObject);
-					}
+					UObject& MetaSoundObject = MetaSoundBuilder.CastDocumentObjectChecked<UObject>();
+					AssetManager.AddOrUpdateAsset(MetaSoundObject);
+				}
 
-					// Fix up dependencies
-					for (FMetaSoundFrontendDocumentBuilder& MetaSoundBuilder : NewMetaSoundAssetBuilders)
-					{
-						MetaSoundBuilder.UpdateDependencyRegistryData(OldToNewReferenceKeys);
-					}
+				// Fix up dependencies
+				for (FMetaSoundFrontendDocumentBuilder& MetaSoundBuilder : NewMetaSoundAssetBuilders)
+				{
+					MetaSoundBuilder.UpdateDependencyRegistryData(OldToNewReferenceKeys);
 				}
 			}
-			
-			TArray<TSharedPtr<FAssetTypeActions_Base>> AssetActions;
-			TMap<EMetasoundFrontendLiteralType, const TSubclassOf<UMetasoundEditorGraphMemberDefaultLiteral>> InputDefaultLiteralClassRegistry;
-			TMap<FName, FGraphPinConfiguration> PinTypes;
-			TSet<FName> CustomPinCategories;
+		}
 
-			TMap<UClass*, TUniquePtr<IMemberDefaultLiteralCustomizationFactory>> LiteralCustomizationFactories;
+		TMap<EMetasoundFrontendLiteralType, const TSubclassOf<UMetasoundEditorGraphMemberDefaultLiteral>> InputDefaultLiteralClassRegistry;
+		TMap<FName, FGraphPinConfiguration> PinTypes;
+		TSet<FName> CustomPinCategories;
 
-			TSharedPtr<FMetasoundGraphNodeFactory> GraphNodeFactory;
-			TSharedPtr<FGraphPanelPinConnectionFactory> GraphConnectionFactory;
-			TSharedPtr<FMetasoundGraphPanelPinFactory> GraphPanelPinFactory;
-			TSharedPtr<FSlateStyleSet> StyleSet;
+		TMap<UClass*, TUniquePtr<IMemberDefaultLiteralCustomizationFactory>> LiteralCustomizationFactories;
 
-			TSet<const UClass*> ExplicitProxyClasses;
+		TSharedPtr<FMetasoundGraphNodeFactory> GraphNodeFactory;
+		TSharedPtr<FGraphPanelPinConnectionFactory> GraphConnectionFactory;
+		TSharedPtr<FGraphPanelPinFactory> GraphPanelPinFactory;
+		TSharedPtr<FSlateStyleSet> StyleSet;
 
-			// Whether or not the editor is in restricted mode: can only make new presets and not modify graphs
-			bool bIsRestrictedMode = false;
-		};
-	} // namespace Editor
-} // namespace Metasound
+		TSet<const UClass*> ExplicitProxyClasses;
+
+		// Whether or not the editor is in restricted mode: can only make new presets and not modify graphs
+		bool bIsRestrictedMode = false;
+	};
+} // namespace Metasound::Editor
 
 IMPLEMENT_MODULE(Metasound::Editor::FModule, MetasoundEditor);
 
diff --git a/Engine/Plugins/Runtime/Metasound/Source/MetasoundEngine/Private/MetasoundAssetSubsystem.cpp b/Engine/Plugins/Runtime/Metasound/Source/MetasoundEngine/Private/MetasoundAssetSubsystem.cpp
--- a/Engine/Plugins/Runtime/Metasound/Source/MetasoundEngine/Private/MetasoundAssetSubsystem.cpp
+++ b/Engine/Plugins/Runtime/Metasound/Source/MetasoundEngine/Private/MetasoundAssetSubsystem.cpp
@@ -144,6 +144,11 @@ namespace Metasound::Engine
 	public:
 		using FAssetInfo = Frontend::IMetaSoundAssetManager::FAssetInfo;
 		using FAssetKey = Frontend::FAssetKey;
+		using FNodeRegistryKey = Frontend::FNodeRegistryKey;
+
+#if WITH_EDITORONLY_DATA
+		using FVersionDocumentResults = Frontend::IMetaSoundAssetManager::FVersionDocumentResults;
+#endif // WITH_EDITORONLY_DATA
 
 		FMetaSoundAssetManager() = default;
 		virtual ~FMetaSoundAssetManager();
@@ -158,7 +163,7 @@ namespace Metasound::Engine
 		void RegisterAssetClassesInDirectories(const TArray<FMetaSoundAssetDirectory>& Directories);
 
 #if WITH_EDITOR
-		bool ReplaceReferencesInDirectory(const TArray<FMetaSoundAssetDirectory>& InDirectories, const Metasound::Frontend::FNodeRegistryKey& OldClassKey, const Metasound::Frontend::FNodeRegistryKey& NewClassKey) const;
+		bool ReplaceReferencesInDirectory(const TArray<FMetaSoundAssetDirectory>& InDirectories, const FNodeRegistryKey& OldClassKey, const FNodeRegistryKey& NewClassKey) const;
 #endif // WITH_EDITOR
 		void RequestAsyncLoadReferencedAssets(FMetasoundAssetBase& InAssetBase);
 		void OnAssetScanComplete();
@@ -175,7 +180,7 @@ namespace Metasound::Engine
 		virtual bool CanAutoUpdate(const FMetasoundFrontendClassName& InClassName) const override;
 		virtual bool ContainsKey(const FAssetKey& InKey) const override;
 		virtual FMetasoundAssetBase* FindAsset(const FAssetKey& InKey) const override;
-		virtual TScriptInterface<IMetaSoundDocumentInterface> FindAssetAsDocumentInterface(const Frontend::FAssetKey& InKey) const override;
+		virtual TScriptInterface<IMetaSoundDocumentInterface> FindAssetAsDocumentInterface(const FAssetKey& InKey) const override;
 		virtual FTopLevelAssetPath FindAssetPath(const FAssetKey& InKey) const override;
 		virtual TArray<FTopLevelAssetPath> FindAssetPaths(const FAssetKey& InKey) const override;
 		virtual FMetasoundAssetBase* GetAsAsset(UObject& InObject) const override;
@@ -193,6 +198,10 @@ namespace Metasound::Engine
 		virtual FMetasoundAssetBase* TryLoadAssetFromKey(const FAssetKey& InKey) const override;
 		virtual bool TryGetAssetIDFromClassName(const FMetasoundFrontendClassName& InClassName, FGuid& OutGuid) const override;
 		virtual bool TryLoadReferencedAssets(const FMetasoundAssetBase& InAssetBase, TArray<FMetasoundAssetBase*>& OutReferencedAssets) const override;
+#if WITH_EDITORONLY_DATA
+		virtual FVersionDocumentResults VersionDocumentsInFolders(const TArray<FString>& FolderPaths) const override;
+#endif // WITH_EDITORONLY_DATA
+
 		virtual void WaitUntilAsyncLoadReferencedAssetsComplete(FMetasoundAssetBase& InAssetBase) override;
 
 		/* FGCObject */
@@ -410,7 +419,7 @@ namespace Metasound::Engine
 		return !AutoUpdateDenyListCache.Contains(InClassName.GetFullName());
 	}
 
-	bool FMetaSoundAssetManager::ContainsKey(const Metasound::Frontend::FAssetKey& InKey) const
+	bool FMetaSoundAssetManager::ContainsKey(const Frontend::FAssetKey& InKey) const
 	{
 		FScopeLock Lock(&PathMapCriticalSection);
 		return PathMap.Contains(InKey);
@@ -436,7 +445,7 @@ namespace Metasound::Engine
 		return LoadingDependencies.FindByPredicate(IsEqualID);
 	}
 
-	FMetasoundAssetBase* FMetaSoundAssetManager::FindAsset(const Metasound::Frontend::FAssetKey& InKey) const
+	FMetasoundAssetBase* FMetaSoundAssetManager::FindAsset(const Frontend::FAssetKey& InKey) const
 	{
 		FTopLevelAssetPath AssetPath = FindAssetPath(InKey);
 		if (AssetPath.IsValid())
@@ -450,7 +459,7 @@ namespace Metasound::Engine
 		return nullptr;
 	}
 
-	TScriptInterface<IMetaSoundDocumentInterface> FMetaSoundAssetManager::FindAssetAsDocumentInterface(const Metasound::Frontend::FAssetKey& InKey) const
+	TScriptInterface<IMetaSoundDocumentInterface> FMetaSoundAssetManager::FindAssetAsDocumentInterface(const Frontend::FAssetKey& InKey) const
 	{
 		const FTopLevelAssetPath AssetPath = FindAssetPath(InKey);
 		if (AssetPath.IsValid())
@@ -464,7 +473,7 @@ namespace Metasound::Engine
 		return nullptr;
 	}
 
-	FTopLevelAssetPath FMetaSoundAssetManager::FindAssetPath(const Metasound::Frontend::FAssetKey& InKey) const
+	FTopLevelAssetPath FMetaSoundAssetManager::FindAssetPath(const Frontend::FAssetKey& InKey) const
 	{
 		FScopeLock Lock(&PathMapCriticalSection);
 		if (const TArray<FTopLevelAssetPath>* Paths = PathMap.Find(InKey))
@@ -478,7 +487,7 @@ namespace Metasound::Engine
 		return nullptr;
 	}
 
-	TArray<FTopLevelAssetPath> FMetaSoundAssetManager::FindAssetPaths(const Metasound::Frontend::FAssetKey& InKey) const
+	TArray<FTopLevelAssetPath> FMetaSoundAssetManager::FindAssetPaths(const Frontend::FAssetKey& InKey) const
 	{
 		FScopeLock Lock(&PathMapCriticalSection);
 		if (const TArray<FTopLevelAssetPath>* Paths = PathMap.Find(InKey))
@@ -550,7 +559,7 @@ namespace Metasound::Engine
 	TSet<UMetaSoundAssetSubsystem::FAssetInfo> FMetaSoundAssetManager::GetReferencedAssetClasses(const FMetasoundAssetBase& InAssetBase) const
 	{
 		METASOUND_TRACE_CPUPROFILER_EVENT_SCOPE(FMetaSoundAssetManager::GetReferencedAssetClasses);
-		using namespace Metasound::Frontend;
+		using namespace Frontend;
 
 		TSet<FAssetInfo> OutAssetInfos;
 		const FMetasoundFrontendDocument& Document = InAssetBase.GetConstDocumentChecked();
@@ -665,7 +674,7 @@ namespace Metasound::Engine
 
 	void FMetaSoundAssetManager::RebuildDenyListCache(const UAssetManager& InAssetManager)
 	{
-		using namespace Metasound::Frontend;
+		using namespace Frontend;
 
 		const UMetaSoundSettings* Settings = GetDefault<UMetaSoundSettings>();
 		if (Settings->DenyListCacheChangeID == AutoUpdateDenyListChangeID)
@@ -710,7 +719,7 @@ namespace Metasound::Engine
 			FMetasoundAssetBase* MetaSoundAsset = Metasound::IMetasoundUObjectRegistry::Get().GetObjectAsAssetBase(AssetData.GetAsset());
 			check(MetaSoundAsset);
 
-			Metasound::Frontend::FMetaSoundAssetRegistrationOptions RegOptions;
+			Frontend::FMetaSoundAssetRegistrationOptions RegOptions;
 			if (const UMetaSoundSettings* Settings = GetDefault<UMetaSoundSettings>())
 			{
 				RegOptions.bAutoUpdateLogWarningOnDroppedConnection = Settings->bAutoUpdateLogWarningOnDroppedConnection;
@@ -792,7 +801,7 @@ namespace Metasound::Engine
 	}
 
 #if WITH_EDITOR
-	bool FMetaSoundAssetManager::ReplaceReferencesInDirectory(const TArray<FMetaSoundAssetDirectory>& InDirectories, const Metasound::Frontend::FNodeRegistryKey& OldClassKey, const Metasound::Frontend::FNodeRegistryKey& NewClassKey) const
+	bool FMetaSoundAssetManager::ReplaceReferencesInDirectory(const TArray<FMetaSoundAssetDirectory>& InDirectories, const Frontend::FNodeRegistryKey& OldClassKey, const Frontend::FNodeRegistryKey& NewClassKey) const
 	{
 		using namespace Frontend;
 
@@ -943,7 +952,7 @@ namespace Metasound::Engine
 		bLogActiveAssetsOnShutdown = bInLogActiveAssetsOnShutdown;
 	}
 
-	FMetasoundAssetBase* FMetaSoundAssetManager::TryLoadAssetFromKey(const Metasound::Frontend::FAssetKey& InAssetKey) const
+	FMetasoundAssetBase* FMetaSoundAssetManager::TryLoadAssetFromKey(const Frontend::FAssetKey& InAssetKey) const
 	{
 		FTopLevelAssetPath ObjectPath = FindAssetPath(InAssetKey);
 		if (ObjectPath.IsValid())
@@ -967,7 +976,7 @@ namespace Metasound::Engine
 
 	bool FMetaSoundAssetManager::TryLoadReferencedAssets(const FMetasoundAssetBase& InAssetBase, TArray<FMetasoundAssetBase*>& OutReferencedAssets) const
 	{
-		using namespace Metasound::Frontend;
+		using namespace Frontend;
 
 		bool bSucceeded = true;
 		OutReferencedAssets.Reset();
@@ -1000,7 +1009,7 @@ namespace Metasound::Engine
 		SearchAndIterateDirectoryAssets(Directories, [this](const FAssetData& AssetData)
 		{
 			using namespace Metasound;
-			using namespace Metasound::Frontend;
+			using namespace Frontend;
 
 			if (AssetData.IsAssetLoaded())
 			{
@@ -1029,6 +1038,63 @@ namespace Metasound::Engine
 		});
 	}
 
+#if WITH_EDITORONLY_DATA
+	FMetaSoundAssetManager::FVersionDocumentResults FMetaSoundAssetManager::VersionDocumentsInFolders(const TArray<FString>& FolderPaths) const
+	{
+		TArray<FTopLevelAssetPath> ClassNames;
+		IMetasoundUObjectRegistry::Get().IterateRegisteredUClasses([&ClassNames](UClass& InClass)
+		{
+			ClassNames.Add(InClass.GetClassPathName());
+		});
+
+		FARFilter Filter;
+		Filter.ClassPaths = ClassNames;
+		Filter.bRecursivePaths = true;
+		Filter.bRecursiveClasses = true;
+		Algo::Transform(FolderPaths, Filter.PackagePaths, [](const FString& PackagePath)
+		{
+			return FName(*PackagePath);
+		});
+
+		FVersionDocumentResults Results;
+		FAssetRegistryModule& AssetRegistryModule = FModuleManager::LoadModuleChecked<FAssetRegistryModule>("AssetRegistry");
+		AssetRegistryModule.Get().EnumerateAssets(Filter, [&Results](const FAssetData& AssetData)
+		{
+			using namespace Frontend;
+
+			if (UObject* MetaSound = AssetData.GetAsset())
+			{
+				if (FMetasoundAssetBase* MetaSoundAsset = IMetasoundUObjectRegistry::Get().GetObjectAsAssetBase(MetaSound))
+				{
+					UMetaSoundBuilderBase& Builder = FDocumentBuilderRegistry::GetChecked().FindOrBeginBuilding(*MetaSound);
+					FMetaSoundFrontendDocumentBuilder& DocBuilder = Builder.GetBuilder();
+					const FString MetaSoundName = DocBuilder.GetDebugName();
+					const bool bVersionedOnLoad = MetaSoundAsset->GetVersionedOnLoad();
+					const FMetasoundFrontendDocumentMetadata& DocMetadata = DocBuilder.GetConstDocumentChecked().Metadata;
+					if (bVersionedOnLoad || MetaSoundAsset->VersionAsset(DocBuilder))
+					{
+						const FString DocVersion = DocMetadata.Version.Number.ToString();
+						constexpr bool bAlwaysMarkDirty = true;
+						MetaSound->Modify(bAlwaysMarkDirty);
+						MetaSoundAsset->SetVersionedOnLoad();
+						Results.VersionedPackages.Add(MetaSound->GetPackage());
+						UE_LOG(LogMetaSound, Display, TEXT("MetaSound asset '%s' document versioned to %s."), *MetaSoundName, *DocVersion);
+					}
+					else
+					{
+						const FString DocVersion = DocMetadata.Version.Number.ToString();
+						Results.UnversionedPackages.Add(MetaSound->GetPackage());
+						UE_LOG(LogMetaSound, Display, TEXT("MetaSound asset '%s' document not versioned (already version %s)."), *MetaSoundName, *DocVersion);
+					}
+				}
+			}
+			return true;
+		});
+
+		return Results;
+	}
+#endif // WITH_EDITORONLY_DATA
+
 	void FMetaSoundAssetManager::WaitUntilAsyncLoadReferencedAssetsComplete(FMetasoundAssetBase& InAssetBase)
 	{
 		TSet<FMetasoundAssetBase*> TransitiveReferences;
diff --git a/Engine/Plugins/Runtime/Metasound/Source/MetasoundFrontend/Public/MetasoundAssetManager.h b/Engine/Plugins/Runtime/Metasound/Source/MetasoundFrontend/Public/MetasoundAssetManager.h
--- a/Engine/Plugins/Runtime/Metasound/Source/MetasoundFrontend/Public/MetasoundAssetManager.h
+++ b/Engine/Plugins/Runtime/Metasound/Source/MetasoundFrontend/Public/MetasoundAssetManager.h
@@ -4,6 +4,7 @@
 #include "MetasoundFrontendDocument.h"
 #include "MetasoundFrontendRegistryKey.h"
 #include "MetasoundVertex.h"
+#include "UObject/Package.h"
 #include "UObject/SoftObjectPath.h"
 #include "UObject/WeakObjectPtrTemplates.h"
 
@@ -147,6 +148,18 @@ namespace Metasound::Frontend
 		};
 
 #if WITH_EDITORONLY_DATA
+		struct FVersionDocumentResults
+		{
+			TArray<UPackage*> VersionedPackages;
+			TArray<UPackage*> UnversionedPackages;
+
+			// Returns whether or not documents were found and loaded/attempted to version
+			bool DocumentsFoundInPackages() const
+			{
+				return !VersionedPackages.IsEmpty() || !UnversionedPackages.IsEmpty();
+			}
+		};
+
 		// Adds missing assets using the provided asset's local reference class cache. Used
 		// to prime system from asset attempting to register prior to asset scan being complete.
 		// Returns true if references were added, false if it they are already found.
@@ -229,6 +242,13 @@ namespace Metasound::Frontend
 		// Updates the given MetaSound's asset record with the new name and optionally reregisters it with the Frontend Node Class Registry.
 		virtual void RenameAsset(const FAssetData& InAssetData, const FString& InOldObjectPath) = 0;
 
+#if WITH_EDITORONLY_DATA
+		// Versions all MetaSound documents found recursively within the given folder paths. Populates the provided array with MetaSound
+		// packages whose document versions were out-of-date and consequently updated. Returns true if MetaSounds were found in the given
+		// directories, false if not.
+		virtual FVersionDocumentResults VersionDocumentsInFolders(const TArray<FString>& FolderPaths) const = 0;
+#endif // WITH_EDITORONLY_DATA
+
 		// Waits until all async load requests related to this asset are complete.
  		virtual void WaitUntilAsyncLoadReferencedAssetsComplete(FMetasoundAssetBase& InAssetBase) = 0;
 	};
