Add new tests for MetaSound variables added/removed with builder (tests utilize pages as well)

#rnx

diff --git a/Engine/Plugins/Runtime/Metasound/Source/MetasoundEngineTest/Private/EngineTestMetasoundSource.cpp b/Engine/Plugins/Runtime/Metasound/Source/MetasoundEngineTest/Private/EngineTestMetasoundSource.cpp
--- a/Engine/Plugins/Runtime/Metasound/Source/MetasoundEngineTest/Private/EngineTestMetasoundSource.cpp
+++ b/Engine/Plugins/Runtime/Metasound/Source/MetasoundEngineTest/Private/EngineTestMetasoundSource.cpp
@@ -1,5 +1,7 @@
 // Copyright Epic Games, Inc. All Rights Reserved.
 
+
+#if WITH_DEV_AUTOMATION_TESTS
 #include "AudioDevice.h"
 #include "Components/AudioComponent.h"
 #include "EngineTestMetaSoundBuilder.h"
@@ -7,7 +9,10 @@
 #include "Interfaces/IPluginManager.h"
 #include "Interfaces/MetasoundOutputFormatInterfaces.h"
 #include "Interfaces/MetasoundFrontendSourceInterface.h"
+#include "Math/UnrealMathUtility.h"
+#include "MetasoundDataReference.h"
 #include "MetasoundLog.h"
+#include "MetasoundSettings.h"
 #include "MetasoundSource.h"
 #include "Misc/AutomationTest.h"
 #include "Misc/Paths.h"
@@ -16,7 +21,6 @@
 #include "Sound/SoundAttenuation.h"
 
 
-#if WITH_DEV_AUTOMATION_TESTS
 namespace EngineTestMetaSoundSourcePrivate
 {
 	struct FInitTestBuilderSourceOutput
@@ -609,7 +613,7 @@ bool FAudioMetasoundSourceBuilderMutateInterface::RunTest(const FString& Paramet
 	constexpr EMetaSoundOutputAudioFormat OutputFormat = EMetaSoundOutputAudioFormat::Mono;
 	constexpr bool bIsOneShot = false;
 	FInitTestBuilderSourceOutput Output;
-	UMetaSoundSourceBuilder& Builder = CreateSourceBuilder(*this, EMetaSoundOutputAudioFormat::Mono, bIsOneShot, Output);
+	UMetaSoundSourceBuilder& Builder = CreateSourceBuilder(*this, OutputFormat, bIsOneShot, Output);
 	Builder.AddToRoot();
 
 	EMetaSoundBuilderResult Result;
@@ -656,40 +660,210 @@ bool FAudioMetasoundSourceBuilderMutateInterface::RunTest(const FString& Paramet
 	return true;
 }
 
-// This test creates a MetaSound source from a SourceBuilder, then adds and connects saw oscillator nodes setting their input freq to a new chromatic tone which cascades up, and finally removes all of the added nodes.
-IMPLEMENT_SIMPLE_AUTOMATION_TEST(FAudioMetasoundSourceBuilderAddRemoveNodes, "Audio.Metasound.Builder.AddRemoveNodes", EAutomationTestFlags::EditorContext | EAutomationTestFlags::EngineFilter)
-bool FAudioMetasoundSourceBuilderAddRemoveNodes::RunTest(const FString& Parameters)
+// This test creates a MetaSound source from a SourceBuilder, then adds and removes various variable nodes to and from pages (if editor-only data is loaded and the project defines at least one page setting).
+IMPLEMENT_SIMPLE_AUTOMATION_TEST(FAudioMetasoundSourceBuilderAddRemovePagedVariableNodes, "Audio.Metasound.Builder.AddRemovePagedVariableNodes", EAutomationTestFlags::EditorContext | EAutomationTestFlags::EngineFilter)
+bool FAudioMetasoundSourceBuilderAddRemovePagedVariableNodes::RunTest(const FString& Parameters)
 {
 	using namespace Audio;
 	using namespace EngineTestMetaSoundPatchBuilderPrivate;
 	using namespace EngineTestMetaSoundSourcePrivate;
+	using namespace Metasound;
 	using namespace Metasound::Frontend;
 
 	constexpr EMetaSoundOutputAudioFormat OutputFormat = EMetaSoundOutputAudioFormat::Mono;
 	constexpr bool bIsOneShot = false;
 	FInitTestBuilderSourceOutput Output;
-	UMetaSoundSourceBuilder& Builder = CreateSourceBuilder(*this, EMetaSoundOutputAudioFormat::Mono, bIsOneShot, Output);
+	UMetaSoundSourceBuilder& Builder = CreateSourceBuilder(*this, OutputFormat, bIsOneShot, Output);
 	Builder.AddToRoot();
 
+	constexpr float SuperNotSpecialNumber = UE_EULERS_NUMBER * UE_GOLDEN_RATIO * UE_PI;
+	FMetasoundFrontendLiteral DefaultValue;
+	DefaultValue.Set(SuperNotSpecialNumber);
 	EMetaSoundBuilderResult Result;
 
-		if (UAudioComponent* AudioComponent = CreateTestComponent(*this))
+	static const FLazyName VariableName = "MyThing";
+
+	FMetaSoundFrontendDocumentBuilder& DocBuilder = Builder.GetBuilder();
+
+	Builder.AddGraphVariable(VariableName, GetMetasoundDataTypeName<float>(), DefaultValue, Result);
+	AddErrorIfFalse(Result == EMetaSoundBuilderResult::Succeeded, TEXT("Failed to add graph variable"));
+
+	{
+		const FMetasoundFrontendVariable* Variable = DocBuilder.FindGraphVariable(VariableName);
+		AddErrorIfFalse(Variable != nullptr, TEXT("Failed to retrieve frontend variable"));
+	}
+
+	const FMetasoundFrontendLiteral RetrievedDefault = Builder.GetGraphVariableDefault(VariableName, Result);
+	AddErrorIfFalse(Result == EMetaSoundBuilderResult::Succeeded, TEXT("Failed to retrieve graph variable default"));
+
+	AddErrorIfFalse(RetrievedDefault.IsEqual(DefaultValue), TEXT("Failed to retrieve graph variable default"));
+
+	Builder.GetGraphVariableDefault(VariableName, Result);
+	AddErrorIfFalse(Result == EMetaSoundBuilderResult::Succeeded, TEXT("Failed to retrieve graph variable default"));
+
+	{
+		const FMetaSoundNodeHandle SetNodeHandle = Builder.AddGraphVariableSetNode(VariableName, Result);
+		AddErrorIfFalse(Result == EMetaSoundBuilderResult::Succeeded, TEXT("Failed to add graph variable setter node"));
+
+		const FMetasoundFrontendVariable* FrontendVariable = DocBuilder.FindGraphVariableByNodeID(SetNodeHandle.NodeID);
+		AddErrorIfFalse(FrontendVariable != nullptr, TEXT("Failed to find FrontendVariable by Setter NodeID"));
+	}
+
+	{
+		const FMetaSoundNodeHandle GetNodeHandle = Builder.AddGraphVariableGetNode(VariableName, Result);
+		AddErrorIfFalse(Result == EMetaSoundBuilderResult::Succeeded, TEXT("Failed to add graph variable getter node"));
+
+		const FMetasoundFrontendVariable* FrontendVariable = DocBuilder.FindGraphVariableByNodeID(GetNodeHandle.NodeID);
+		AddErrorIfFalse(FrontendVariable != nullptr, TEXT("Failed to find FrontendVariable by Getter NodeID"));
+	}
+
+	{
+		const FMetaSoundNodeHandle GetNodeHandle = Builder.AddGraphVariableGetDelayedNode(VariableName, Result);
+		AddErrorIfFalse(Result == EMetaSoundBuilderResult::Succeeded, TEXT("Failed to add graph variable delayed getter node"));
+
+		const FMetasoundFrontendVariable* FrontendVariable = DocBuilder.FindGraphVariableByNodeID(GetNodeHandle.NodeID);
+		AddErrorIfFalse(FrontendVariable != nullptr, TEXT("Failed to find FrontendVariable by Delayed Getter NodeID"));
+	}
+
+	bool bRanPageTests = false;
+	if (const UMetaSoundSettings* Settings = GetDefault<UMetaSoundSettings>())
+	{
+		const TArray<FMetaSoundPageSettings>& PageSettings = Settings->GetProjectPageSettings();
+		constexpr bool bHasEditorOnlyData = (bool)(WITH_EDITORONLY_DATA);
+
+		if (!bHasEditorOnlyData || PageSettings.IsEmpty())
 		{
-			constexpr bool bEnableLiveUpdate = true;
-			ADD_LATENT_AUTOMATION_COMMAND(FMetaSoundSourceBuilderAuditionLatentCommand(&Builder, AudioComponent, bEnableLiveUpdate));
+			UE_LOG(LogMetaSound, Display, TEXT("Skipping MetaSound Page Variable Tests: Project has no project-specific page settings implemented or editor data is not loaded"));
+		}
+		else
+		{
+			bRanPageTests = true;
+
+			{
+#if WITH_EDITORONLY_DATA
+				constexpr bool bDuplicateLastGraph = false;
+				constexpr bool bSetAsBuildGraph = true;
+				const FGuid PageID = PageSettings.Last().UniqueId;
+				AddErrorIfFalse(PageID != Frontend::DefaultPageID, TEXT("Added PageID cannot be the default page ID"));
+				DocBuilder.AddGraphPage(PageID, bDuplicateLastGraph, bSetAsBuildGraph);
+#endif // WITH_EDITORONLY_DATA
+			}
+
+			const FMetasoundFrontendVariable* PagedVariable = DocBuilder.AddGraphVariable(VariableName, GetMetasoundDataTypeName<int32>());
+			int32 IntValue = (int32)(SuperNotSpecialNumber);
+			AddErrorIfFalse(PagedVariable && PagedVariable->Literal.TryGet(IntValue) && IntValue == 0, TEXT("Failed to add graph variable or initialize paged variable value to 0"));
 
-			const FName OscType = "Sine";
-			const TArray<float> ChromaticFreqs
+			for (int32 Index = 0; Index < 10; ++Index)
 			{
-				293.66,
-				311.13,
-				329.63,
-				349.23,
-				369.99,
-				392.00,
-				415.30,
-				440.00
-			};
+				DocBuilder.AddGraphVariableDeferredAccessorNode(VariableName);
+				DocBuilder.AddGraphVariableDeferredAccessorNode(VariableName, FGuid::NewGuid(), &Frontend::DefaultPageID);
+
+				DocBuilder.AddGraphVariableAccessorNode(VariableName);
+				DocBuilder.AddGraphVariableAccessorNode(VariableName, FGuid::NewGuid(), &Frontend::DefaultPageID);
+			}
+
+			AddErrorIfFalse(PagedVariable != nullptr, TEXT("Failed to add paged integer variable with same name as default float variable: Skipping additional page variable tests..."));
+
+			if (PagedVariable)
+			{
+				const FMetasoundFrontendVariable* Variable = DocBuilder.FindGraphVariable(VariableName, &Frontend::DefaultPageID);
+				AddErrorIfFalse(Variable != nullptr, TEXT("Failed to find graph variable getter node from default page (when build page ID is set to other ID"));
+
+				if (Variable)
+				{
+					auto CheckNodeIDs = [this, &DocBuilder](const TArray<FGuid>& PagedNodeIDs, const TArray<FGuid>& DefaultNodeIDs, EMetasoundFrontendClassType ClassType)
+					{
+						TSet<FGuid> DefaultNodeIDSet;
+						DefaultNodeIDSet.Append(DefaultNodeIDs);
+						const TSet<FGuid> Intersection = DefaultNodeIDSet.Intersect(TSet<FGuid> (PagedNodeIDs));
+						AddErrorIfFalse(Intersection.IsEmpty(), TEXT("Paged variable added variable node(s) to incorrect graph"));
+
+						for (const FGuid& NodeID : PagedNodeIDs)
+						{
+							const FMetasoundFrontendClass* PagedVariableDependency = nullptr;
+							if (const FMetasoundFrontendNode* PagedVariableNode = DocBuilder.FindNode(NodeID))
+							{
+								PagedVariableDependency = DocBuilder.FindDependency(PagedVariableNode->ClassID);
+							}
+						}
+
+						for (const FGuid& NodeID : DefaultNodeIDs)
+						{
+							const FMetasoundFrontendClass* VariableDependency = nullptr;
+							if (const FMetasoundFrontendNode* VariableNode = DocBuilder.FindNode(NodeID, &Frontend::DefaultPageID))
+							{
+								VariableDependency = DocBuilder.FindDependency(VariableNode->ClassID);
+							}
+
+							AddErrorIfFalse(VariableDependency && VariableDependency->Metadata.GetType() == ClassType, FString::Printf(
+								TEXT("MetaSound Class missing/mismatch type (%s) for added variable node"),
+								*FString(LexToString(ClassType))
+							));
+						}
+					};
+
+					CheckNodeIDs(PagedVariable->AccessorNodeIDs, Variable->AccessorNodeIDs, EMetasoundFrontendClassType::VariableAccessor);
+					CheckNodeIDs(PagedVariable->DeferredAccessorNodeIDs, Variable->DeferredAccessorNodeIDs, EMetasoundFrontendClassType::VariableDeferredAccessor);
+					CheckNodeIDs({ PagedVariable->VariableNodeID }, { Variable->VariableNodeID }, EMetasoundFrontendClassType::Variable);
+
+					AddErrorIfFalse(!PagedVariable->MutatorNodeID.IsValid(), TEXT("Paged variable should have invalid mutator node ID as no mutator was added"));
+					CheckNodeIDs({ /* PagedVariable->MutatorNodeID // Purposefully never added mutator as negation test */ }, { Variable->MutatorNodeID }, EMetasoundFrontendClassType::VariableMutator);
+				}
+			}
+		}
+	}
+
+	Builder.RemoveGraphVariable(VariableName, Result);
+	AddErrorIfFalse(Result == EMetaSoundBuilderResult::Succeeded, TEXT("Failed to remove graph variable"));
+
+#if WITH_EDITORONLY_DATA
+	if (bRanPageTests)
+	{
+		DocBuilder.SetBuildPageID(Frontend::DefaultPageID);
+
+		const FMetasoundFrontendVariable* PageVariable = DocBuilder.FindGraphVariable(VariableName);
+		AddErrorIfFalse(PageVariable != nullptr, TEXT("Remove graph request should not have removed default variable as well (as build page ID was set)"));
+	}
+#endif // WITH_EDITORONLY_DATA
+
+	Builder.RemoveFromRoot();
+	return true;
+}
+
+// This test creates a MetaSound source from a SourceBuilder, then adds, removes and connects a variable and its respective nodes (get/get delayed/set) to multiple pages. Does not playback or live-audition.
+IMPLEMENT_SIMPLE_AUTOMATION_TEST(FAudioMetasoundSourceBuilderAddRemoveNodes, "Audio.Metasound.Builder.AddRemoveNodes", EAutomationTestFlags::EditorContext | EAutomationTestFlags::EngineFilter)
+bool FAudioMetasoundSourceBuilderAddRemoveNodes::RunTest(const FString& Parameters)
+{
+	using namespace Audio;
+	using namespace EngineTestMetaSoundPatchBuilderPrivate;
+	using namespace EngineTestMetaSoundSourcePrivate;
+	using namespace Metasound::Frontend;
+
+	constexpr EMetaSoundOutputAudioFormat OutputFormat = EMetaSoundOutputAudioFormat::Mono;
+	constexpr bool bIsOneShot = false;
+	FInitTestBuilderSourceOutput Output;
+	UMetaSoundSourceBuilder& Builder = CreateSourceBuilder(*this, OutputFormat, bIsOneShot, Output);
+	Builder.AddToRoot();
+
+	EMetaSoundBuilderResult Result;
+
+	if (UAudioComponent* AudioComponent = CreateTestComponent(*this))
+	{
+		constexpr bool bEnableLiveUpdate = true;
+		ADD_LATENT_AUTOMATION_COMMAND(FMetaSoundSourceBuilderAuditionLatentCommand(&Builder, AudioComponent, bEnableLiveUpdate));
+
+		const FName OscType = "Sine";
+		const TArray<float> ChromaticFreqs
+		{
+			293.66,
+			311.13,
+			329.63,
+			349.23,
+			369.99,
+			392.00,
+			415.30,
+			440.00
+		};
 
 		TArray<FMetaSoundNodeHandle> GenNodes;
 		for (int32 i = 0; i < 8; ++i)
