Fix for test build warnings
#rnx

diff --git a/Engine/Plugins/Runtime/Metasound/Source/MetasoundEngineTest/Private/EngineTestMetasoundSource.cpp b/Engine/Plugins/Runtime/Metasound/Source/MetasoundEngineTest/Private/EngineTestMetasoundSource.cpp
--- a/Engine/Plugins/Runtime/Metasound/Source/MetasoundEngineTest/Private/EngineTestMetasoundSource.cpp
+++ b/Engine/Plugins/Runtime/Metasound/Source/MetasoundEngineTest/Private/EngineTestMetasoundSource.cpp
@@ -729,7 +729,12 @@ bool FAudioMetasoundSourceBuilderAddRemovePagedVariableNodes::RunTest(const FStr
 	if (const UMetaSoundSettings* Settings = GetDefault<UMetaSoundSettings>())
 	{
 		const TArray<FMetaSoundPageSettings>& PageSettings = Settings->GetProjectPageSettings();
-		constexpr bool bHasEditorOnlyData = (bool)(WITH_EDITORONLY_DATA);
+
+#if WITH_EDITORONLY_DATA
+		constexpr bool bHasEditorOnlyData = true;
+#else
+		constexpr bool bHasEditorOnlyData = false;
+#endif // WITH_EDITORONLY_DATA
 
 		if (!bHasEditorOnlyData || PageSettings.IsEmpty())
 		{
