Round 2 fix for static analysis compaints

#rnx

diff --git a/Engine/Plugins/Runtime/Metasound/Source/MetasoundEngineTest/Private/EngineTestMetasoundSource.cpp b/Engine/Plugins/Runtime/Metasound/Source/MetasoundEngineTest/Private/EngineTestMetasoundSource.cpp
--- a/Engine/Plugins/Runtime/Metasound/Source/MetasoundEngineTest/Private/EngineTestMetasoundSource.cpp
+++ b/Engine/Plugins/Runtime/Metasound/Source/MetasoundEngineTest/Private/EngineTestMetasoundSource.cpp
@@ -725,20 +725,18 @@ bool FAudioMetasoundSourceBuilderAddRemovePagedVariableNodes::RunTest(const FStr
 		AddErrorIfFalse(FrontendVariable != nullptr, TEXT("Failed to find FrontendVariable by Delayed Getter NodeID"));
 	}
 
+#if WITH_EDITORONLY_DATA
 	bool bRanPageTests = false;
+#endif // WITH_EDITORONLY_DATA
+
 	if (const UMetaSoundSettings* Settings = GetDefault<UMetaSoundSettings>())
 	{
 		const TArray<FMetaSoundPageSettings>& PageSettings = Settings->GetProjectPageSettings();
 
 #if WITH_EDITORONLY_DATA
-		constexpr bool bHasEditorOnlyData = true;
-#else
-		constexpr bool bHasEditorOnlyData = false;
-#endif // WITH_EDITORONLY_DATA
-
-		if (!bHasEditorOnlyData || PageSettings.IsEmpty())
+		if (PageSettings.IsEmpty())
 		{
-			UE_LOG(LogMetaSound, Display, TEXT("Skipping MetaSound Page Variable Tests: Project has no project-specific page settings implemented or editor data is not loaded"));
+			UE_LOG(LogMetaSound, Display, TEXT("Skipping MetaSound Page Variable Tests: Project has no project-specific page settings implemented"));
 		}
 		else
 		{
@@ -816,6 +814,10 @@ bool FAudioMetasoundSourceBuilderAddRemovePagedVariableNodes::RunTest(const FStr
 				}
 			}
 		}
+
+#else // !WITH_EDITORONLY_DATA
+		UE_LOG(LogMetaSound, Display, TEXT("Skipping MetaSound Page Variable Tests: Test is running with editor data not loaded"));
+#endif // !WITH_EDITORONLY_DATA
 	}
 
 	Builder.RemoveGraphVariable(VariableName, Result);
