Minor test clean-up (unnecessary nested duplicate editor-only macro)
#rnx

diff --git a/Engine/Plugins/Runtime/Metasound/Source/MetasoundEngineTest/Private/EngineTestMetasoundSource.cpp b/Engine/Plugins/Runtime/Metasound/Source/MetasoundEngineTest/Private/EngineTestMetasoundSource.cpp
--- a/Engine/Plugins/Runtime/Metasound/Source/MetasoundEngineTest/Private/EngineTestMetasoundSource.cpp
+++ b/Engine/Plugins/Runtime/Metasound/Source/MetasoundEngineTest/Private/EngineTestMetasoundSource.cpp
@@ -742,15 +742,11 @@ bool FAudioMetasoundSourceBuilderAddRemovePagedVariableNodes::RunTest(const FStr
 		{
 			bRanPageTests = true;
 
-			{
-#if WITH_EDITORONLY_DATA
-				constexpr bool bDuplicateLastGraph = false;
-				constexpr bool bSetAsBuildGraph = true;
-				const FGuid PageID = PageSettings.Last().UniqueId;
-				AddErrorIfFalse(PageID != Frontend::DefaultPageID, TEXT("Added PageID cannot be the default page ID"));
-				DocBuilder.AddGraphPage(PageID, bDuplicateLastGraph, bSetAsBuildGraph);
-#endif // WITH_EDITORONLY_DATA
-			}
+			constexpr bool bDuplicateLastGraph = false;
+			constexpr bool bSetAsBuildGraph = true;
+			const FGuid PageID = PageSettings.Last().UniqueId;
+			AddErrorIfFalse(PageID != Frontend::DefaultPageID, TEXT("Added PageID cannot be the default page ID"));
+			DocBuilder.AddGraphPage(PageID, bDuplicateLastGraph, bSetAsBuildGraph);
 
 			const FMetasoundFrontendVariable* PagedVariable = DocBuilder.AddGraphVariable(VariableName, GetMetasoundDataTypeName<int32>());
 			int32 IntValue = (int32)(SuperNotSpecialNumber);
