Namespace generic-named local functions

#rnx

diff --git a/Engine/Plugins/Runtime/Metasound/Source/MetasoundEditor/Private/MetasoundAssetDefinitions.cpp b/Engine/Plugins/Runtime/Metasound/Source/MetasoundEditor/Private/MetasoundAssetDefinitions.cpp
--- a/Engine/Plugins/Runtime/Metasound/Source/MetasoundEditor/Private/MetasoundAssetDefinitions.cpp
+++ b/Engine/Plugins/Runtime/Metasound/Source/MetasoundEditor/Private/MetasoundAssetDefinitions.cpp
@@ -39,50 +39,54 @@
 #define LOCTEXT_NAMESPACE "MetaSoundEditor"
 
 
-static bool GetIsPreset(const FSoftObjectPath& InSourcePath)
+namespace Metasound::Editor
 {
-	using namespace Metasound;
-	using namespace Metasound::Editor;
-	using namespace Metasound::Frontend;
-
-	bool bIsPreset = false;
-	// If object in memory, try to resolve but not load
-	if (UObject* Object = InSourcePath.ResolveObject())
-	{
-		const FMetasoundAssetBase* MetaSoundAsset = IMetasoundUObjectRegistry::Get().GetObjectAsAssetBase(Object);
-		if (MetaSoundAsset)
-		{
-			bIsPreset = MetaSoundAsset->GetConstDocumentChecked().RootGraph.PresetOptions.bIsPreset;
-		}
-	}
-	// Otherwise, try to pull from asset registry, but avoid load as this call
-	// would then be slow and is called from the ContentBrowser many times
-	else
+	namespace AssetDefinitionsPrivate
 	{
-		FAssetData IconAssetData;
-		const FAssetRegistryModule& AssetRegistryModule = FModuleManager::GetModuleChecked<FAssetRegistryModule>("AssetRegistry");
-		UE::AssetRegistry::EExists Exists = AssetRegistryModule.TryGetAssetByObjectPath(InSourcePath, IconAssetData);
-		if (Exists == UE::AssetRegistry::EExists::Exists)
+		bool GetIsPreset(const FSoftObjectPath& InSourcePath)
 		{
-			IconAssetData.GetTagValue(AssetTags::IsPreset, bIsPreset);
-		}
-	}
+			using namespace Frontend;
 
-	return bIsPreset;
-}
+			bool bIsPreset = false;
+			// If object in memory, try to resolve but not load
+			if (UObject* Object = InSourcePath.ResolveObject())
+			{
+				const FMetasoundAssetBase* MetaSoundAsset = IMetasoundUObjectRegistry::Get().GetObjectAsAssetBase(Object);
+				if (MetaSoundAsset)
+				{
+					bIsPreset = MetaSoundAsset->GetConstDocumentChecked().RootGraph.PresetOptions.bIsPreset;
+				}
+			}
+			// Otherwise, try to pull from asset registry, but avoid load as this call
+			// would then be slow and is called from the ContentBrowser many times
+			else
+			{
+				FAssetData IconAssetData;
+				const FAssetRegistryModule& AssetRegistryModule = FModuleManager::GetModuleChecked<FAssetRegistryModule>("AssetRegistry");
+				UE::AssetRegistry::EExists Exists = AssetRegistryModule.TryGetAssetByObjectPath(InSourcePath, IconAssetData);
+				if (Exists == UE::AssetRegistry::EExists::Exists)
+				{
+					IconAssetData.GetTagValue(AssetTags::IsPreset, bIsPreset);
+				}
+			}
 
-static const FSlateBrush* GetClassBrush(const FAssetData& InAssetData, FName InClassName, bool bIsThumbnail = false)
-{
-	const bool bIsPreset = GetIsPreset(InAssetData.ToSoftObjectPath());
-	FString BrushName = FString::Printf(TEXT("MetasoundEditor.%s"), *InClassName.ToString());
-	if (bIsPreset)
-	{
-		BrushName += TEXT(".Preset");
-	}
-	BrushName += bIsThumbnail ? TEXT(".Thumbnail") : TEXT(".Icon");
+			return bIsPreset;
+		}
 
-	return &Metasound::Editor::Style::GetSlateBrushSafe(FName(*BrushName));
-}
+		const FSlateBrush* GetClassBrush(const FAssetData& InAssetData, FName InClassName, bool bIsThumbnail = false)
+		{
+			const bool bIsPreset = GetIsPreset(InAssetData.ToSoftObjectPath());
+			FString BrushName = FString::Printf(TEXT("MetasoundEditor.%s"), *InClassName.ToString());
+			if (bIsPreset)
+			{
+				BrushName += TEXT(".Preset");
+			}
+			BrushName += bIsThumbnail ? TEXT(".Thumbnail") : TEXT(".Icon");
+
+			return &Metasound::Editor::Style::GetSlateBrushSafe(FName(*BrushName));
+		}
+	} // namespace AssetDefinitionsPrivate
+} // namespace Metasound::Editor
 
 
 FLinearColor UAssetDefinition_MetaSoundPatch::GetAssetColor() const
@@ -129,13 +133,15 @@ EAssetCommandResult UAssetDefinition_MetaSoundPatch::OpenAssets(const FAssetOpen
 
 const FSlateBrush* UAssetDefinition_MetaSoundPatch::GetThumbnailBrush(const FAssetData& InAssetData, const FName InClassName) const
 {
+	using namespace Metasound::Editor;
 	constexpr bool bIsThumbnail = true;
-	return GetClassBrush(InAssetData, InClassName, bIsThumbnail);
+	return AssetDefinitionsPrivate::GetClassBrush(InAssetData, InClassName, bIsThumbnail);
 }
 
 const FSlateBrush* UAssetDefinition_MetaSoundPatch::GetIconBrush(const FAssetData& InAssetData, const FName InClassName) const
 {
-	return GetClassBrush(InAssetData, InClassName);
+	using namespace Metasound::Editor;
+	return AssetDefinitionsPrivate::GetClassBrush(InAssetData, InClassName);
 }
 
 FLinearColor UAssetDefinition_MetaSoundSource::GetAssetColor() const
@@ -184,13 +190,15 @@ EAssetCommandResult UAssetDefinition_MetaSoundSource::OpenAssets(const FAssetOpe
 
 const FSlateBrush* UAssetDefinition_MetaSoundSource::GetThumbnailBrush(const FAssetData& InAssetData, const FName InClassName) const
 {
+	using namespace Metasound::Editor;
 	constexpr bool bIsThumbnail = true;
-	return GetClassBrush(InAssetData, InClassName, bIsThumbnail);
+	return AssetDefinitionsPrivate::GetClassBrush(InAssetData, InClassName, bIsThumbnail);
 }
 
 const FSlateBrush* UAssetDefinition_MetaSoundSource::GetIconBrush(const FAssetData& InAssetData, const FName InClassName) const
 {
-	return GetClassBrush(InAssetData, InClassName);
+	using namespace Metasound::Editor;
+	return AssetDefinitionsPrivate::GetClassBrush(InAssetData, InClassName);
 }
 
 void UAssetDefinition_MetaSoundSource::ExecutePlaySound(const FToolMenuContext& InContext)
