[MetaSound][UEAudio]
- Envelope ADSR Node fix to trigger OnSustain when Decay is 0 seconds in length.

[RN] Fixed issue where the ADSR Envelope MetaSound node would only trigger OnSustain if the Decay stage was more than 0 seconds long. Now it triggers when the Sustain stage begins regardless of the Decay stage length.

diff --git a/Engine/Plugins/Runtime/Metasound/Source/MetasoundStandardNodes/Private/MetasoundADSREnvelopeNode.cpp b/Engine/Plugins/Runtime/Metasound/Source/MetasoundStandardNodes/Private/MetasoundADSREnvelopeNode.cpp
--- a/Engine/Plugins/Runtime/Metasound/Source/MetasoundStandardNodes/Private/MetasoundADSREnvelopeNode.cpp
+++ b/Engine/Plugins/Runtime/Metasound/Source/MetasoundStandardNodes/Private/MetasoundADSREnvelopeNode.cpp
@@ -160,6 +160,12 @@ namespace Metasound
 					if (InState.CurrentSampleIndex == InState.AttackSampleCount)
 					{
 						OutOnDecayFrames.Add(0);
+
+						// Decay stage is skipped
+						if (InState.DecaySampleCount == 0)
+						{
+							OutOnSustainFrames.Add(0);
+						}
 					}
 				}
 				// We are in decay
@@ -274,6 +280,12 @@ namespace Metasound
 					if (InState.CurrentSampleIndex == InState.AttackSampleCount)
 					{
 						OutOnDecayFrames.Add(EndAttackFrame);
+
+						// Decay Stage is Skipped
+						if (InState.DecaySampleCount == 0)
+						{
+							OutOnSustainFrames.Add(EndAttackFrame);
+						}
 					}
 				}
 
