Fix MetaSound enum pins not showing reset to default arrow when applicable

diff --git a/Engine/Plugins/Runtime/Metasound/Source/MetasoundEditor/Private/SMetasoundGraphEnumPin.cpp b/Engine/Plugins/Runtime/Metasound/Source/MetasoundEditor/Private/SMetasoundGraphEnumPin.cpp
--- a/Engine/Plugins/Runtime/Metasound/Source/MetasoundEditor/Private/SMetasoundGraphEnumPin.cpp
+++ b/Engine/Plugins/Runtime/Metasound/Source/MetasoundEditor/Private/SMetasoundGraphEnumPin.cpp
@@ -24,14 +24,26 @@ namespace Metasound
 			GenerateComboBoxIndexes(ComboItems);
 
 			//Create widget
-			return SAssignNew(ComboBox, SPinComboBox)
-				.ComboItemList(ComboItems)
-				.VisibleText(this, &SMetasoundGraphEnumPin::OnGetText)
-				.OnSelectionChanged(this, &SMetasoundGraphEnumPin::ComboBoxSelectionChanged)
-				.IsEnabled(this, &SGraphPin::GetDefaultValueIsEditable)
-				.Visibility(this, &SGraphPin::GetDefaultValueVisibility)
-				.OnGetDisplayName(this, &SMetasoundGraphEnumPin::OnGetFriendlyName)
-				.OnGetTooltip(this, &SMetasoundGraphEnumPin::OnGetTooltip);
+			
+			return SNew(SHorizontalBox)
+				+ SHorizontalBox::Slot()
+				.VAlign(VAlign_Center)
+				[
+					SAssignNew(ComboBox, SPinComboBox)
+						.ComboItemList(ComboItems)
+						.VisibleText(this, &SMetasoundGraphEnumPin::OnGetText)
+						.OnSelectionChanged(this, &SMetasoundGraphEnumPin::ComboBoxSelectionChanged)
+						.IsEnabled(this, &SGraphPin::GetDefaultValueIsEditable)
+						.Visibility(this, &SGraphPin::GetDefaultValueVisibility)
+						.OnGetDisplayName(this, &SMetasoundGraphEnumPin::OnGetFriendlyName)
+						.OnGetTooltip(this, &SMetasoundGraphEnumPin::OnGetTooltip)
+				]
+				+ SHorizontalBox::Slot()
+				.AutoWidth()
+				.VAlign(VAlign_Center)
+				[
+					CreateResetToDefaultWidget()
+				];
 		}
 
 		TSharedPtr<const Frontend::IEnumDataTypeInterface>
diff --git a/Engine/Plugins/Runtime/Metasound/Source/MetasoundEditor/Private/SMetasoundGraphPin.h b/Engine/Plugins/Runtime/Metasound/Source/MetasoundEditor/Private/SMetasoundGraphPin.h
--- a/Engine/Plugins/Runtime/Metasound/Source/MetasoundEditor/Private/SMetasoundGraphPin.h
+++ b/Engine/Plugins/Runtime/Metasound/Source/MetasoundEditor/Private/SMetasoundGraphPin.h
@@ -376,6 +376,94 @@ PRAGMA_ENABLE_DEPRECATION_WARNINGS
 				return DefaultFloat->WidgetType == EMetasoundMemberDefaultWidget::None;
 			}
 
+			TSharedRef<SWidget> CreateResetToDefaultWidget()
+			{
+				return SNew(SButton)
+					.ToolTipText(LOCTEXT("ResetToClassDefaultToolTip", "Reset to class default"))
+					.ButtonStyle(FAppStyle::Get(), TEXT("NoBorder"))
+					.ContentPadding(0.0f)
+					.Visibility(TAttribute<EVisibility>::Create([this]
+					{
+						using namespace Frontend;
+						if (!ParentPinType::IsConnected())
+						{
+							FConstInputHandle InputHandle = GetConstInputHandle();
+							if (const FMetasoundFrontendLiteral* Literal = InputHandle->GetLiteral())
+							{
+								const bool bIsDefaultConstructed = Literal->GetType() == EMetasoundFrontendLiteralType::None;
+								const bool bIsTriggerDataType = InputHandle->GetDataType() == GetMetasoundDataTypeName<FTrigger>();
+								bool bIsRerouteNode = false;
+								if (UEdGraphPin* Pin = ParentPinType::GetPinObj())
+								{
+									if (UMetasoundEditorGraphExternalNode* Node = Cast<UMetasoundEditorGraphExternalNode>(Pin->GetOwningNode()))
+									{
+										if (Node->GetBreadcrumb().ClassName == FRerouteNodeTemplate::ClassName)
+										{
+											bIsRerouteNode = true;
+										}
+									}
+								}
+
+								if (!bIsDefaultConstructed && !bIsTriggerDataType && !bIsRerouteNode)
+								{
+									return EVisibility::Visible;
+								}
+							}
+						}
+
+						return EVisibility::Collapsed;
+					}))
+					.OnClicked(FOnClicked::CreateLambda([this]()
+					{
+						using namespace Editor;
+						using namespace Frontend;
+
+						if (UEdGraphPin* Pin = ParentPinType::GetPinObj())
+						{
+							if (UMetasoundEditorGraphNode* Node = Cast<UMetasoundEditorGraphNode>(Pin->GetOwningNode()))
+							{
+								if (UMetasoundEditorGraph* MetaSoundGraph = CastChecked<UMetasoundEditorGraph>(Node->GetGraph()))
+								{
+									UObject& MetaSound = MetaSoundGraph->GetMetasoundChecked();
+									FMetasoundFrontendDocumentModifyContext& ModifyContext = FGraphBuilder::GetOutermostMetaSoundChecked(MetaSound).GetModifyContext();
+									{
+										const FScopedTransaction Transaction(LOCTEXT("MetaSoundEditorResetToClassDefault", "Reset to Class Default"));
+										MetaSound.Modify();
+										MetaSoundGraph->Modify();
+
+										if (UMetasoundEditorGraphMemberNode* MemberNode = Cast<UMetasoundEditorGraphMemberNode>(Node))
+										{
+											UMetasoundEditorGraphMember* Member = MemberNode->GetMember();
+											if (ensure(Member))
+											{
+												Member->ResetToClassDefault();
+												ModifyContext.AddMemberIDsModified({ Member->GetMemberID() });
+											}
+											else
+											{
+												ModifyContext.SetDocumentModified();
+											}
+										}
+										else
+										{
+											FInputHandle InputHandle = GetInputHandle();
+											InputHandle->ClearLiteral();
+											ModifyContext.SetDocumentModified();
+										}
+									}
+								}
+							}
+						}
+
+						return FReply::Handled();
+					}))
+					.Content()
+					[
+						SNew(SImage)
+							.Image(FAppStyle::GetBrush("PropertyWindow.DiffersFromDefault"))
+					];
+			}
+
 			virtual TSharedRef<SWidget> GetDefaultValueWidget() override
 			{
 				using namespace Frontend;
@@ -405,90 +493,7 @@ PRAGMA_ENABLE_DEPRECATION_WARNINGS
 					.AutoWidth()
 					.VAlign(VAlign_Center)
 					[
-						SNew(SButton)
-						.ToolTipText(LOCTEXT("ResetToClassDefaultToolTip", "Reset to class default"))
-						.ButtonStyle(FAppStyle::Get(), TEXT("NoBorder"))
-						.ContentPadding(0.0f)
-						.Visibility(TAttribute<EVisibility>::Create([this]
-						{
-							using namespace Frontend;
-							if (!ParentPinType::IsConnected())
-							{
-								FConstInputHandle InputHandle = GetConstInputHandle();
-								if (const FMetasoundFrontendLiteral* Literal = InputHandle->GetLiteral())
-								{
-									const bool bIsDefaultConstructed = Literal->GetType() == EMetasoundFrontendLiteralType::None;
-									const bool bIsTriggerDataType = InputHandle->GetDataType() == GetMetasoundDataTypeName<FTrigger>();
-									bool bIsRerouteNode = false;
-									if (UEdGraphPin* Pin = ParentPinType::GetPinObj())
-									{
-										if (UMetasoundEditorGraphExternalNode* Node = Cast<UMetasoundEditorGraphExternalNode>(Pin->GetOwningNode()))
-										{
-											if (Node->GetBreadcrumb().ClassName == FRerouteNodeTemplate::ClassName)
-											{
-												bIsRerouteNode = true;
-											}
-										}
-									}
-									
-									if (!bIsDefaultConstructed && !bIsTriggerDataType && !bIsRerouteNode)
-									{
-										return EVisibility::Visible;
-									}
-								}
-							}
-
-							return EVisibility::Collapsed;
-						}))
-						.OnClicked(FOnClicked::CreateLambda([this]()
-						{
-							using namespace Editor;
-							using namespace Frontend;
-
-							if (UEdGraphPin* Pin = ParentPinType::GetPinObj())
-							{
-								if (UMetasoundEditorGraphNode* Node = Cast<UMetasoundEditorGraphNode>(Pin->GetOwningNode()))
-								{
-									if (UMetasoundEditorGraph* MetaSoundGraph = CastChecked<UMetasoundEditorGraph>(Node->GetGraph()))
-									{
-										UObject& MetaSound = MetaSoundGraph->GetMetasoundChecked();
-										FMetasoundFrontendDocumentModifyContext& ModifyContext = FGraphBuilder::GetOutermostMetaSoundChecked(MetaSound).GetModifyContext();
-										{
-											const FScopedTransaction Transaction(LOCTEXT("MetaSoundEditorResetToClassDefault", "Reset to Class Default"));
-											MetaSound.Modify();
-											MetaSoundGraph->Modify();
-
-											if (UMetasoundEditorGraphMemberNode* MemberNode = Cast<UMetasoundEditorGraphMemberNode>(Node))
-											{
-												UMetasoundEditorGraphMember* Member = MemberNode->GetMember();
-												if (ensure(Member))
-												{
-													Member->ResetToClassDefault();
-													ModifyContext.AddMemberIDsModified({ Member->GetMemberID() });
-												}
-												else
-												{
-													ModifyContext.SetDocumentModified();
-												}
-											}
-											else
-											{
-												FInputHandle InputHandle = GetInputHandle();
-												InputHandle->ClearLiteral();
-												ModifyContext.SetDocumentModified();
-											}
-										}
-									}
-								}
-							}
-
-							return FReply::Handled();
-						}))
-						.Content()
-						[
-							SNew(SImage)
-							.Image(FAppStyle::GetBrush("PropertyWindow.DiffersFromDefault"))
-						]
+						CreateResetToDefaultWidget()
 					];
 			}
 
