- Fix for MetaSound member breadcrumbs not getting updated and reporting wrong description when changed in editor.
- More controller callside replacement with builder API
#rnx

diff --git a/Engine/Plugins/Runtime/Metasound/Source/MetasoundEditor/Private/MetasoundEditorGraph.cpp b/Engine/Plugins/Runtime/Metasound/Source/MetasoundEditor/Private/MetasoundEditorGraph.cpp
--- a/Engine/Plugins/Runtime/Metasound/Source/MetasoundEditor/Private/MetasoundEditorGraph.cpp
+++ b/Engine/Plugins/Runtime/Metasound/Source/MetasoundEditor/Private/MetasoundEditorGraph.cpp
@@ -44,6 +44,12 @@ namespace Metasound::Editor
 {
 	namespace GraphPrivate
 	{
+		static const FText SetMemberAccessTypeTransactionLabelFormat = LOCTEXT("RenameGraphMemberAccessTypeFormat", "Set MetaSound {0} '{1}' AccessType");
+		static const FText SetMemberDefaultTransactionLabelFormat = LOCTEXT("SetGraphMemberDefaultFormat", "Set MetaSound {0} '{1}' Default(s)");
+		static const FText SetMemberDescriptionTransactionLabelFormat = LOCTEXT("SetGraphMemberTooltipFormat", "Set MetaSound {0} '{1}' ToolTip");
+		static const FText SetMemberDisplayNameTransactionLabelFormat = LOCTEXT("RenameGraphMemberDisplayNameFormat", "Set MetaSound {0} '{1}' DisplayName to '{2}'");
+		static const FText SetMemberNameTransactionLabelFormat = LOCTEXT("RenameGraphVertexMemberNameFormat", "Set MetaSound {0} Namespace and Name from '{1}' to '{2}'");
+
 		FName GetUniqueTransientMemberName()
 		{
 			// Use unique instance ID to avoid copy/paste logic resolving invalid relationship between graphs.
@@ -266,28 +272,7 @@ TArray<UMetasoundEditorGraphMemberNode*> UMetasoundEditorGraphVertex::GetNodes()
 
 void UMetasoundEditorGraphVertex::SetDescription(const FText& InDescription, bool bPostTransaction)
 {
-	using namespace Metasound::Editor;
-	using namespace Metasound::Frontend;
-
-	{
-		const FText TransactionLabel = FText::Format(LOCTEXT("SetGraphVertexTooltipFormat", "Set MetaSound {0}'s ToolTip"), GetGraphMemberLabel());
-		const FScopedTransaction Transaction(TransactionLabel, bPostTransaction);
-
-		if (UMetasoundEditorGraph* Graph = GetOwningGraph())
-		{
-			Graph->Modify();
-			UObject& MetaSound = Graph->GetMetasoundChecked();
-			MetaSound.Modify();
-
-PRAGMA_DISABLE_DEPRECATION_WARNINGS
-			FNodeHandle NodeHandle = GetNodeHandle();
-PRAGMA_ENABLE_DEPRECATION_WARNINGS
-			NodeHandle->SetDescription(InDescription);
-
-			Graph->RegisterGraphWithFrontend();
-			FGraphBuilder::GetOutermostMetaSoundChecked(*Graph).GetModifyContext().AddMemberIDsModified({ GetMemberID() });
-		}
-	}
+	Breadcrumb.Description = InDescription;
 }
 
 FGuid UMetasoundEditorGraphVertex::GetMemberID() const 
@@ -335,13 +320,14 @@ void UMetasoundEditorGraphVertex::SetMemberName(const FName& InNewName, bool bPo
 		OldName = Node->Name;
 	}
 
-	const FText TransactionLabel = FText::Format(LOCTEXT("RenameGraphVertexMemberNameFormat", "Set Metasound {0} MemberName"), GetGraphMemberLabel());
+	const FText TransactionLabel = FText::Format(GraphPrivate::SetMemberNameTransactionLabelFormat, GetGraphMemberLabel(), FText::FromName(OldName), FText::FromName(InNewName));
 	const FScopedTransaction Transaction(TransactionLabel, bPostTransaction);
 
 	Graph->Modify();
 	Graph->GetMetasoundChecked().Modify();
 
 	RenameFrontendMemberInternal(DocBuilder, OldName, InNewName);
+	Breadcrumb.MemberName = InNewName;
 
 	const TArray<UMetasoundEditorGraphMemberNode*> Nodes = GetNodes();
 	for (UMetasoundEditorGraphMemberNode* Node : Nodes)
@@ -365,49 +351,6 @@ FText UMetasoundEditorGraphVertex::GetDisplayName() const
 	return Metasound::Editor::FGraphBuilder::GetDisplayName(*GetConstNodeHandle(), bIncludeNamespace);
 }
 
-void UMetasoundEditorGraphVertex::SetDisplayName(const FText& InNewName, bool bPostTransaction)
-{
-	using namespace Metasound::Editor;
-	using namespace Metasound::Frontend;
-
-	UMetasoundEditorGraph* Graph = GetOwningGraph();
-	if (!ensure(Graph))
-	{
-		return;
-	}
-
-PRAGMA_DISABLE_DEPRECATION_WARNINGS
-	FNodeHandle NodeHandle = GetNodeHandle();
-PRAGMA_ENABLE_DEPRECATION_WARNINGS
-
-	if (NodeHandle->GetDisplayName().EqualTo(InNewName))
-	{
-		return;
-	}
-
-	const FText TransactionLabel = FText::Format(LOCTEXT("RenameGraphVertexDisplayNameFormat", "Set Metasound {0} DisplayName"), GetGraphMemberLabel());
-	const FScopedTransaction Transaction(TransactionLabel, bPostTransaction);
-
-	Graph->Modify();
-	Graph->GetMetasoundChecked().Modify();
-
-	NodeHandle->SetDisplayName(InNewName);
-
-	const TArray<UMetasoundEditorGraphMemberNode*> Nodes = GetNodes();
-	for (UMetasoundEditorGraphMemberNode* Node : Nodes)
-	{
-		const TArray<UEdGraphPin*>& Pins = Node->GetAllPins();
-		ensure(Pins.Num() == 1);
-
-		for (UEdGraphPin* Pin : Pins)
-		{
-			Pin->PinFriendlyName = InNewName;
-		}
-	}
-
-	Graph->RegisterGraphWithFrontend();
-	FGraphBuilder::GetOutermostMetaSoundChecked(*Graph).GetModifyContext().AddMemberIDsModified({ GetMemberID() });
-}
 #if WITH_EDITORONLY_DATA
 bool UMetasoundEditorGraphVertex::SetIsAdvancedDisplay(const bool IsAdvancedDisplay)
 {
@@ -1015,6 +958,70 @@ void UMetasoundEditorGraphInput::SetDataType(FName InNewType, bool bPostTransact
 	}
 }
 
+void UMetasoundEditorGraphInput::SetDescription(const FText& InDescription, bool bPostTransaction)
+{
+	using namespace Metasound::Editor;
+	using namespace Metasound::Frontend;
+
+	UMetasoundEditorGraph* Graph = GetOwningGraph();
+	if (!ensure(Graph))
+	{
+		return;
+	}
+
+	const FText TransactionLabel = FText::Format(GraphPrivate::SetMemberDescriptionTransactionLabelFormat, GetGraphMemberLabel(), GetDisplayName());
+	const FScopedTransaction Transaction(TransactionLabel, bPostTransaction);
+
+	Graph->Modify();
+	Graph->GetMetasoundChecked().Modify();
+	GetFrontendBuilderChecked().SetGraphInputDescription(GetMemberName(), InDescription);
+	Breadcrumb.Description = InDescription;
+	Graph->RegisterGraphWithFrontend();
+
+	Super::SetDescription(InDescription, bPostTransaction);
+}
+
+void UMetasoundEditorGraphInput::SetDisplayName(const FText& InNewName, bool bPostTransaction)
+{
+	using namespace Metasound::Editor;
+	using namespace Metasound::Frontend;
+
+	UMetasoundEditorGraph* Graph = GetOwningGraph();
+	if (!ensure(Graph))
+	{
+		return;
+	}
+
+	FMetaSoundFrontendDocumentBuilder& DocBuilder = GetFrontendBuilderChecked();
+	const FMetasoundFrontendClassInput* Input = DocBuilder.FindGraphInput(GetMemberName());
+	if (!Input || Input->Metadata.GetDisplayName().EqualTo(InNewName))
+	{
+		return;
+	}
+
+	const FText TransactionLabel = FText::Format(GraphPrivate::SetMemberDisplayNameTransactionLabelFormat, GetGraphMemberLabel(), GetDisplayName(), InNewName);
+	const FScopedTransaction Transaction(TransactionLabel, bPostTransaction);
+
+	Graph->Modify();
+	Graph->GetMetasoundChecked().Modify();
+
+	DocBuilder.SetGraphInputDisplayName(Input->Name, InNewName);
+
+	const TArray<UMetasoundEditorGraphMemberNode*> Nodes = GetNodes();
+	for (UMetasoundEditorGraphMemberNode* Node : Nodes)
+	{
+		const TArray<UEdGraphPin*>& Pins = Node->GetAllPins();
+		ensure(Pins.Num() == 1);
+
+		for (UEdGraphPin* Pin : Pins)
+		{
+			Pin->PinFriendlyName = InNewName;
+		}
+	}
+
+	Graph->RegisterGraphWithFrontend();
+}
+
 void UMetasoundEditorGraphInput::SetVertexAccessType(EMetasoundFrontendVertexAccessType InNewAccessType, bool bPostTransaction)
 {
 	using namespace Metasound::Editor;
@@ -1024,7 +1031,7 @@ void UMetasoundEditorGraphInput::SetVertexAccessType(EMetasoundFrontendVertexAcc
 	{
 		if (UMetasoundEditorGraph* Graph = GetOwningGraph())
 		{
-			const FScopedTransaction Transaction(LOCTEXT("SetGraphInputVertexType", "Set MetaSound Graph Input Vertex Access Type"), bPostTransaction);
+			const FScopedTransaction Transaction(FText::Format(GraphPrivate::SetMemberAccessTypeTransactionLabelFormat, GetGraphMemberLabel(), GetDisplayName()), bPostTransaction);
 			Graph->GetMetasoundChecked().Modify();
 			Graph->Modify();
 			Modify();
@@ -1093,7 +1100,7 @@ void UMetasoundEditorGraphInput::UpdateFrontendDefaultLiteral(bool bPostTransact
 	{
 		FMetaSoundFrontendDocumentBuilder& Builder = GetFrontendBuilderChecked();
 
-		const FScopedTransaction Transaction(FText::Format(LOCTEXT("Set Input Defaults", "Set MetaSound Input '{0}' Default(s)"), GetDisplayName()), bPostTransaction);
+		const FScopedTransaction Transaction(FText::Format(GraphPrivate::SetMemberDefaultTransactionLabelFormat, GetGraphMemberLabel(), GetDisplayName()), bPostTransaction);
 		Builder.CastDocumentObjectChecked<UObject>().Modify();
 		Literal->Modify();
 
@@ -1308,6 +1315,69 @@ void UMetasoundEditorGraphOutput::SetDataType(FName InNewType, bool bPostTransac
 	}
 }
 
+void UMetasoundEditorGraphOutput::SetDescription(const FText& InDescription, bool bPostTransaction)
+{
+	using namespace Metasound::Editor;
+	using namespace Metasound::Frontend;
+
+	UMetasoundEditorGraph* Graph = GetOwningGraph();
+	if (!ensure(Graph))
+	{
+		return;
+	} 
+
+	const FText TransactionLabel = FText::Format(GraphPrivate::SetMemberDescriptionTransactionLabelFormat, GetGraphMemberLabel(), GetDisplayName());
+	const FScopedTransaction Transaction(TransactionLabel, bPostTransaction);
+
+	Graph->Modify();
+	Graph->GetMetasoundChecked().Modify();
+	GetFrontendBuilderChecked().SetGraphOutputDescription(GetMemberName(), InDescription);
+	Graph->RegisterGraphWithFrontend();
+
+	Super::SetDescription(InDescription, bPostTransaction);
+}
+
+void UMetasoundEditorGraphOutput::SetDisplayName(const FText& InNewName, bool bPostTransaction)
+{
+	using namespace Metasound::Editor;
+	using namespace Metasound::Frontend;
+
+	UMetasoundEditorGraph* Graph = GetOwningGraph();
+	if (!ensure(Graph))
+	{
+		return;
+	}
+
+	FMetaSoundFrontendDocumentBuilder& DocBuilder = GetFrontendBuilderChecked();
+	const FMetasoundFrontendClassOutput* Output = DocBuilder.FindGraphOutput(GetMemberName());
+	if (!Output || Output->Metadata.GetDisplayName().EqualTo(InNewName))
+	{
+		return;
+	}
+
+	const FText TransactionLabel = FText::Format(GraphPrivate::SetMemberDisplayNameTransactionLabelFormat, GetGraphMemberLabel(), GetDisplayName(), InNewName);
+	const FScopedTransaction Transaction(TransactionLabel, bPostTransaction);
+
+	Graph->Modify();
+	Graph->GetMetasoundChecked().Modify();
+
+	DocBuilder.SetGraphOutputDisplayName(Output->Name, InNewName);
+
+	const TArray<UMetasoundEditorGraphMemberNode*> Nodes = GetNodes();
+	for (UMetasoundEditorGraphMemberNode* Node : Nodes)
+	{
+		const TArray<UEdGraphPin*>& Pins = Node->GetAllPins();
+		ensure(Pins.Num() == 1);
+
+		for (UEdGraphPin* Pin : Pins)
+		{
+			Pin->PinFriendlyName = InNewName;
+		}
+	}
+
+	Graph->RegisterGraphWithFrontend();
+}
+
 void UMetasoundEditorGraphOutput::SetVertexAccessType(EMetasoundFrontendVertexAccessType InNewAccessType, bool bPostTransaction)
 {
 	using namespace Metasound::Editor;
@@ -1318,7 +1388,7 @@ void UMetasoundEditorGraphOutput::SetVertexAccessType(EMetasoundFrontendVertexAc
 		UMetasoundEditorGraph* Graph = GetOwningGraph();
 		if (ensure(Graph))
 		{
-			const FScopedTransaction Transaction(LOCTEXT("SetGraphOutputAccessType", "Set MetaSound Graph Output Access Type"), bPostTransaction);
+			const FScopedTransaction Transaction(FText::Format(GraphPrivate::SetMemberAccessTypeTransactionLabelFormat, GetGraphMemberLabel(), GetDisplayName()), bPostTransaction);
 			Graph->GetMetasoundChecked().Modify();
 			Graph->Modify();
 			Modify();
@@ -1360,7 +1430,7 @@ void UMetasoundEditorGraphOutput::UpdateFrontendDefaultLiteral(bool bPostTransac
 	FMetasoundFrontendLiteral DefaultLiteral;
 	if (ensure(Literal->TryFindDefault(DefaultLiteral)))
 	{
-		const FScopedTransaction Transaction(FText::Format(LOCTEXT("Set Output Default", "Set MetaSound Output '{0}' Default"), GetDisplayName()), bPostTransaction);
+		const FScopedTransaction Transaction(FText::Format(GraphPrivate::SetMemberDefaultTransactionLabelFormat, GetGraphMemberLabel(), GetDisplayName()), bPostTransaction);
 		Metasound->Modify();
 
 		if (const FMetasoundFrontendNode* FrontendNode = GetFrontendNode())
@@ -1513,14 +1583,15 @@ void UMetasoundEditorGraphVariable::SetMemberName(const FName& InNewName, bool b
 		return;
 	}
 
-	const FText TransactionLabel = FText::Format(LOCTEXT("RenameGraphVariableMemberNameFormat", "Set Metasound {0} Name"), GetGraphMemberLabel());
+	const FName OldName = GetMemberName();
+	const FText TransactionLabel = FText::Format(GraphPrivate::SetMemberNameTransactionLabelFormat, GetGraphMemberLabel(), FText::FromName(OldName), FText::FromName(InNewName));
 	const FScopedTransaction Transaction(TransactionLabel, bPostTransaction);
 
 	Graph->Modify();
 	UObject& MetaSound = Graph->GetMetasoundChecked();
 	MetaSound.Modify();
 
-	Graph->GetBuilderChecked().GetBuilder().SetGraphVariableName(GetMemberName(), InNewName);
+	Graph->GetBuilderChecked().GetBuilder().SetGraphVariableName(OldName, InNewName);
 }
 
 FGuid UMetasoundEditorGraphVariable::GetMemberID() const 
@@ -1565,7 +1636,9 @@ FText UMetasoundEditorGraphVariable::GetDescription() const
 
 void UMetasoundEditorGraphVariable::SetDescription(const FText& InDescription, bool bPostTransaction)
 {
-	const FText TransactionLabel = FText::Format(LOCTEXT("SetGraphVariableTooltipFormat", "Set MetaSound {0}'s ToolTip"), GetGraphMemberLabel());
+	using namespace Metasound::Editor;
+
+	const FText TransactionLabel = FText::Format(GraphPrivate::SetMemberDescriptionTransactionLabelFormat, GetGraphMemberLabel(), GetDisplayName());
 	const FScopedTransaction Transaction(TransactionLabel, bPostTransaction);
 
 	if (UMetasoundEditorGraph* Graph = GetOwningGraph())
@@ -1668,7 +1741,7 @@ void UMetasoundEditorGraphVariable::SetDisplayName(const FText& InNewName, bool
 		return;
 	}
 
-	const FText TransactionLabel = FText::Format(LOCTEXT("RenameGraphVariableDisplayNameFormat", "Set Metasound {0} DisplayName"), GetGraphMemberLabel(), bPostTransaction);
+	const FText TransactionLabel = FText::Format(GraphPrivate::SetMemberDisplayNameTransactionLabelFormat, GetGraphMemberLabel(), GetDisplayName(), InNewName);
 	const FScopedTransaction Transaction(TransactionLabel, bPostTransaction);
 	{
 		Graph->Modify();
@@ -2007,7 +2080,7 @@ void UMetasoundEditorGraphVariable::UpdateFrontendDefaultLiteral(bool bPostTrans
 		return;
 	}
 
-	const FScopedTransaction Transaction(LOCTEXT("Set Variable Default", "Set MetaSound Variable Default"), bPostTransaction);
+	const FScopedTransaction Transaction(FText::Format(GraphPrivate::SetMemberDefaultTransactionLabelFormat, GetGraphMemberLabel(), GetDisplayName()), bPostTransaction);
 	FMetaSoundFrontendDocumentBuilder& Builder = GetFrontendBuilderChecked();
 	Builder.CastDocumentObjectChecked<UObject>().Modify();
 
diff --git a/Engine/Plugins/Runtime/Metasound/Source/MetasoundEditor/Private/MetasoundNodeDetailCustomization.cpp b/Engine/Plugins/Runtime/Metasound/Source/MetasoundEditor/Private/MetasoundNodeDetailCustomization.cpp
--- a/Engine/Plugins/Runtime/Metasound/Source/MetasoundEditor/Private/MetasoundNodeDetailCustomization.cpp
+++ b/Engine/Plugins/Runtime/Metasound/Source/MetasoundEditor/Private/MetasoundNodeDetailCustomization.cpp
@@ -1353,7 +1353,8 @@ PRAGMA_ENABLE_DEPRECATION_WARNINGS
 					return;
 				}
 
-				const FText TransactionLabel = FText::Format(LOCTEXT("RenameGraphMember_Format", "Set MetaSound {0}'s Name"), GraphMember->GetGraphMemberLabel());
+				const FText SetMemberNameTransactionLabelFormat = LOCTEXT("Commit_RenameGraphVertexMemberNameFormat", "Set MetaSound {0} Namespace and Name from '{1}' to '{2}' (DisplayName cleared)");
+				const FText TransactionLabel = FText::Format(SetMemberNameTransactionLabelFormat, GraphMember->GetGraphMemberLabel(), FText::FromName(GraphMember->GetMemberName()), InNewName);
 				const FScopedTransaction Transaction(TransactionLabel);
 
 				constexpr bool bPostTransaction = false;
diff --git a/Engine/Plugins/Runtime/Metasound/Source/MetasoundEditor/Public/MetasoundEditorGraph.h b/Engine/Plugins/Runtime/Metasound/Source/MetasoundEditor/Public/MetasoundEditorGraph.h
--- a/Engine/Plugins/Runtime/Metasound/Source/MetasoundEditor/Public/MetasoundEditorGraph.h
+++ b/Engine/Plugins/Runtime/Metasound/Source/MetasoundEditor/Public/MetasoundEditorGraph.h
@@ -316,7 +316,6 @@ public:
 	virtual TArray<UMetasoundEditorGraphMemberNode*> GetNodes() const override;
 	virtual void SetDescription(const FText& InDescription, bool bPostTransaction) override;
 	virtual void SetMemberName(const FName& InNewName, bool bPostTransaction) override;
-	virtual void SetDisplayName(const FText& InNewName, bool bPostTransaction) override;
 	virtual bool Synchronize() override;
 	virtual void CacheBreadcrumb() override;
 	/* ~End UMetasoundEditorGraphMember interface */
@@ -387,6 +386,8 @@ public:
 	virtual bool IsInterfaceMember(FMetasoundFrontendInterface* OutInterface = nullptr) const override;
 	virtual void ResetToClassDefault() override;
 	virtual void SetDataType(FName InNewType, bool bPostTransaction) override;
+	virtual void SetDescription(const FText& InDescription, bool bPostTransaction) override;
+	virtual void SetDisplayName(const FText& InNewName, bool bPostTransaction) override;
 	virtual void SetMemberName(const FName& InNewName, bool bPostTransaction) override;
 	virtual void SetVertexAccessType(EMetasoundFrontendVertexAccessType InNewAccessType, bool bPostTransaction) override;
 	virtual bool Synchronize() override;
@@ -422,6 +423,8 @@ public:
 	virtual const FText& GetGraphMemberLabel() const override;
 	virtual void ResetToClassDefault() override;
 	virtual void SetDataType(FName InNewType, bool bPostTransaction) override;
+	virtual void SetDescription(const FText& InDescription, bool bPostTransaction) override;
+	virtual void SetDisplayName(const FText& InNewName, bool bPostTransaction) override;
 	virtual void SetVertexAccessType(EMetasoundFrontendVertexAccessType InNewAccessType, bool bPostTransaction) override;
 	virtual bool Synchronize() override;
 	virtual void UpdateFrontendDefaultLiteral(bool bPostTransaction, const FGuid* InPageID = nullptr) override;
diff --git a/Engine/Plugins/Runtime/Metasound/Source/MetasoundFrontend/Private/MetasoundFrontendDocumentBuilder.cpp b/Engine/Plugins/Runtime/Metasound/Source/MetasoundFrontend/Private/MetasoundFrontendDocumentBuilder.cpp
--- a/Engine/Plugins/Runtime/Metasound/Source/MetasoundFrontend/Private/MetasoundFrontendDocumentBuilder.cpp
+++ b/Engine/Plugins/Runtime/Metasound/Source/MetasoundFrontend/Private/MetasoundFrontendDocumentBuilder.cpp
@@ -4777,20 +4777,33 @@
 }
 
 #if WITH_EDITORONLY_DATA
-bool FMetaSoundFrontendDocumentBuilder::SetGraphOutputAdvancedDisplay(const FName OutputName, const bool InAdvancedDisplay)
+bool FMetaSoundFrontendDocumentBuilder::SetGraphInputDescription(FName InputName, FText Description)
 {
-	FMetasoundFrontendDocument& Document = GetDocumentChecked();
-	FMetasoundFrontendGraphClass& RootGraph = Document.RootGraph;
+	using namespace Metasound;
 
-	if (const int32* Index = DocumentCache->GetInterfaceCache().FindOutputIndex(OutputName))
+	if (const int32* Index = DocumentCache->GetInterfaceCache().FindInputIndex(InputName))
 	{
-		FMetasoundFrontendClassOutput& GraphOutput = Document.RootGraph.Interface.Outputs[*Index];
-		if (GraphOutput.Metadata.bIsAdvancedDisplay != InAdvancedDisplay)
-		{
-			GraphOutput.Metadata.SetIsAdvancedDisplay(InAdvancedDisplay);
-			Document.Metadata.ModifyContext.AddMemberIDModified(GraphOutput.VertexID);
-			return true;
-		}
+		FMetasoundFrontendDocument& Document = GetDocumentChecked();
+		FMetasoundFrontendClassInput& GraphInput = Document.RootGraph.Interface.Inputs[*Index];
+		GraphInput.Metadata.SetDescription(Description);
+		Document.Metadata.ModifyContext.AddMemberIDModified(GraphInput.NodeID);
+		return true;
 	}
 
 	return false;
+}
+
+bool FMetaSoundFrontendDocumentBuilder::SetGraphInputDisplayName(FName InputName, FText DisplayName)
+{
+	using namespace Metasound;
+
+	if (const int32* Index = DocumentCache->GetInterfaceCache().FindInputIndex(InputName))
+	{
+		FMetasoundFrontendDocument& Document = GetDocumentChecked();
+		FMetasoundFrontendClassInput& GraphInput = Document.RootGraph.Interface.Inputs[*Index];
+		GraphInput.Metadata.SetDisplayName(DisplayName);
+		Document.Metadata.ModifyContext.AddMemberIDModified(GraphInput.NodeID);
+		return true;
+	}
+
+	return false;
@@ -4866,54 +4879,26 @@
 	return true;
 }
 
-bool FMetaSoundFrontendDocumentBuilder::SetGraphOutputName(FName OutputName, FName NewName)
+#if WITH_EDITORONLY_DATA
+bool FMetaSoundFrontendDocumentBuilder::SetGraphOutputAdvancedDisplay(const FName OutputName, const bool InAdvancedDisplay)
 {
-	using namespace Metasound::Frontend;
-
-	if (OutputName == NewName)
-	{
-		return true;
-	}
-
-	const int32* Index = DocumentCache->GetInterfaceCache().FindOutputIndex(OutputName);
-	if (!Index)
-	{
-		return false;
-	}
-
 	FMetasoundFrontendDocument& Document = GetDocumentChecked();
-	FMetasoundFrontendGraphClass& GraphClass = Document.RootGraph;
-	FMetasoundFrontendClassInterface& Interface = GraphClass.Interface;
-	Interface.UpdateChangeID();
+	FMetasoundFrontendGraphClass& RootGraph = Document.RootGraph;
 
-	FMetasoundFrontendClassOutput& GraphOutput = Interface.Outputs[*Index];
-	GraphOutput.Name = NewName;
-	
-	GraphClass.IterateGraphPages([this, &GraphOutput, &NewName](FMetasoundFrontendGraph& Graph)
+	if (const int32* Index = DocumentCache->GetInterfaceCache().FindOutputIndex(OutputName))
 	{
-		const IDocumentGraphNodeCache& NodeCache = DocumentCache->GetNodeCache(Graph.PageID);
-		if (const int32* NodeIndex = NodeCache.FindNodeIndex(GraphOutput.NodeID))
+		FMetasoundFrontendClassOutput& GraphOutput = Document.RootGraph.Interface.Outputs[*Index];
+		if (GraphOutput.Metadata.bIsAdvancedDisplay != InAdvancedDisplay)
 		{
-			FMetasoundFrontendNode& Node = Graph.Nodes[*NodeIndex];
-			Node.Name = NewName;
-			for (FMetasoundFrontendVertex& Vertex : Node.Interface.Inputs)
-			{
-				Vertex.Name = NewName;
-			}
-			for (FMetasoundFrontendVertex& Vertex : Node.Interface.Outputs)
-			{
-				Vertex.Name = NewName;
-			}
+			GraphOutput.Metadata.SetIsAdvancedDisplay(InAdvancedDisplay);
+			Document.Metadata.ModifyContext.AddMemberIDModified(GraphOutput.VertexID);
+			return true;
 		}
-	});
-	DocumentDelegates->InterfaceDelegates.OnOutputNameChanged.Broadcast(OutputName, NewName);
-	
-#if WITH_EDITORONLY_DATA
-	Document.Metadata.ModifyContext.AddMemberIDModified(GraphOutput.NodeID);
+	}
+
+	return false;
+}
 #endif // WITH_EDITORONLY_DATA
 
-	return true;
-}
-
 bool FMetaSoundFrontendDocumentBuilder::SetGraphOutputAccessType(FName OutputName, EMetasoundFrontendVertexAccessType AccessType)
 {
@@ -5032,6 +5017,89 @@
 	return true;
 }
 
+#if WITH_EDITORONLY_DATA
+bool FMetaSoundFrontendDocumentBuilder::SetGraphOutputDescription(FName OutputName, FText Description)
+{
+	using namespace Metasound;
+
+	if (const int32* Index = DocumentCache->GetInterfaceCache().FindOutputIndex(OutputName))
+	{
+		FMetasoundFrontendDocument& Document = GetDocumentChecked();
+		FMetasoundFrontendClassOutput& GraphOutput = Document.RootGraph.Interface.Outputs[*Index];
+		GraphOutput.Metadata.SetDescription(Description);
+		Document.Metadata.ModifyContext.AddMemberIDModified(GraphOutput.NodeID);
+		return true;
+	}
+
+	return false;
+}
+
+bool FMetaSoundFrontendDocumentBuilder::SetGraphOutputDisplayName(FName OutputName, FText DisplayName)
+{
+	using namespace Metasound;
+
+	if (const int32* Index = DocumentCache->GetInterfaceCache().FindOutputIndex(OutputName))
+	{
+		FMetasoundFrontendDocument& Document = GetDocumentChecked();
+		FMetasoundFrontendClassOutput& GraphOutput = Document.RootGraph.Interface.Outputs[*Index];
+		GraphOutput.Metadata.SetDisplayName(DisplayName);
+		Document.Metadata.ModifyContext.AddMemberIDModified(GraphOutput.NodeID);
+		return true;
+	}
+
+	return false;
+}
+#endif // WITH_EDITORONLY_DATA
+
+bool FMetaSoundFrontendDocumentBuilder::SetGraphOutputName(FName OutputName, FName NewName)
+{
+	using namespace Metasound::Frontend;
+
+	if (OutputName == NewName)
+	{
+		return true;
+	}
+
+	const int32* Index = DocumentCache->GetInterfaceCache().FindOutputIndex(OutputName);
+	if (!Index)
+	{
+		return false;
+	}
+
+	FMetasoundFrontendDocument& Document = GetDocumentChecked();
+	FMetasoundFrontendGraphClass& GraphClass = Document.RootGraph;
+	FMetasoundFrontendClassInterface& Interface = GraphClass.Interface;
+	Interface.UpdateChangeID();
+
+	FMetasoundFrontendClassOutput& GraphOutput = Interface.Outputs[*Index];
+	GraphOutput.Name = NewName;
+	
+	GraphClass.IterateGraphPages([this, &GraphOutput, &NewName](FMetasoundFrontendGraph& Graph)
+	{
+		const IDocumentGraphNodeCache& NodeCache = DocumentCache->GetNodeCache(Graph.PageID);
+		if (const int32* NodeIndex = NodeCache.FindNodeIndex(GraphOutput.NodeID))
+		{
+			FMetasoundFrontendNode& Node = Graph.Nodes[*NodeIndex];
+			Node.Name = NewName;
+			for (FMetasoundFrontendVertex& Vertex : Node.Interface.Inputs)
+			{
+				Vertex.Name = NewName;
+			}
+			for (FMetasoundFrontendVertex& Vertex : Node.Interface.Outputs)
+			{
+				Vertex.Name = NewName;
+			}
+		}
+	});
+	DocumentDelegates->InterfaceDelegates.OnOutputNameChanged.Broadcast(OutputName, NewName);
+	
+#if WITH_EDITORONLY_DATA
+	Document.Metadata.ModifyContext.AddMemberIDModified(GraphOutput.NodeID);
+#endif // WITH_EDITORONLY_DATA
+
+	return true;
+}
+
 bool FMetaSoundFrontendDocumentBuilder::SetGraphVariableDefault(FName VariableName, FMetasoundFrontendLiteral InDefaultLiteral, const FGuid* InPageID)
 {
 	using namespace Metasound;
diff --git a/Engine/Plugins/Runtime/Metasound/Source/MetasoundFrontend/Public/MetasoundFrontendDocumentBuilder.h b/Engine/Plugins/Runtime/Metasound/Source/MetasoundFrontend/Public/MetasoundFrontendDocumentBuilder.h
--- a/Engine/Plugins/Runtime/Metasound/Source/MetasoundFrontend/Public/MetasoundFrontendDocumentBuilder.h
+++ b/Engine/Plugins/Runtime/Metasound/Source/MetasoundFrontend/Public/MetasoundFrontendDocumentBuilder.h
@@ -471,6 +471,14 @@ public:
 	bool SetGraphInputDefault(FName InputName, FMetasoundFrontendLiteral InDefaultLiteral, const FGuid* InPageID = nullptr);
 	bool SetGraphInputDefaults(FName InputName, TArray<FMetasoundFrontendClassInputDefault> Defaults);
 
+#if WITH_EDITORONLY_DATA
+	// Sets the graph input's description. Returns true if found and set, false if not.
+	bool SetGraphInputDescription(FName InputName, FText Description);
+
+	// Sets the graph input's display name. Returns true if found and set, false if not.
+	bool SetGraphInputDisplayName(FName InputName, FText DisplayName);
+#endif // WITH_EDITORONLY_DATA
+
 	// Sets whether or not graph input inherits default.  By default, updates only if graph is marked as a preset.
 	// Optionally, if bForceUpdate is set, updates inheritance even if not a preset (primarily for clearing flag
 	// if non-preset has unnecessary data).
@@ -495,6 +503,14 @@ public:
 	// already the given DataType.
 	bool SetGraphOutputDataType(FName OutputName, FName DataType);
 
+#if WITH_EDITORONLY_DATA
+	// Sets the graph output's description. Returns true if found and set, false if not.
+	bool SetGraphOutputDescription(FName OutputName, FText Description);
+
+	// Sets the graph input's display name. Returns true if found and set, false if not.
+	bool SetGraphOutputDisplayName(FName OutputName, FText DisplayName);
+#endif // WITH_EDITORONLY_DATA
+
 	// Sets a given graph output's name to a new name. Succeeds if the graph output exists and the new name is set (or is the same as the old name).
 	bool SetGraphOutputName(FName InputName, FName InName);
 
