Fix missing transaction scope when pasting members
#rnx

diff --git a/Engine/Plugins/Runtime/Metasound/Source/MetasoundEditor/Private/MetasoundEditor.cpp b/Engine/Plugins/Runtime/Metasound/Source/MetasoundEditor/Private/MetasoundEditor.cpp
--- a/Engine/Plugins/Runtime/Metasound/Source/MetasoundEditor/Private/MetasoundEditor.cpp
+++ b/Engine/Plugins/Runtime/Metasound/Source/MetasoundEditor/Private/MetasoundEditor.cpp
@@ -3812,7 +3812,7 @@ namespace Metasound
 				return false;
 			}
 
-			FString ClipboardText;			
+			FString ClipboardText;
 			FPlatformApplicationMisc::ClipboardPaste(ClipboardText);
 
 			if (ClipboardText.IsEmpty())
@@ -3830,8 +3830,7 @@ namespace Metasound
 
 		void FEditor::PasteSelectedMemberItems()
 		{
-			const UMetasoundEditorGraphMember* Member = FDocumentClipboardUtils::GetMemberFromClipboard();
-			if (Member)
+			if (const UMetasoundEditorGraphMember* Member = FDocumentClipboardUtils::GetMemberFromClipboard())
 			{
 				UObject* MetaSound = GetMetasoundObject();
 				if (!MetaSound)
@@ -3840,10 +3839,13 @@ namespace Metasound
 				}
 
 				UMetasoundEditorGraph& Graph = GetMetaSoundGraphChecked();
-				FMetaSoundFrontendDocumentBuilder& DocumentBuilder = Frontend::IDocumentBuilderRegistry::GetChecked().FindOrBeginBuilding(&Graph.GetMetasoundChecked());
+				const FText TransactionLabel = FText::Format(LOCTEXT("PasteSelectedMember_Vertex", "Paste {0} '{1}'"), Member->GetGraphMemberLabel(), FText::FromName(Member->GetMemberName()));
+				const FScopedTransaction Transaction(TransactionLabel);
 				MetaSound->Modify();
 				Graph.Modify();
 
+				FMetaSoundFrontendDocumentBuilder& DocumentBuilder = Frontend::IDocumentBuilderRegistry::GetChecked().FindOrBeginBuilding(&Graph.GetMetasoundChecked());
+
 				if (const UMetasoundEditorGraphVertex* SourceGraphVertex = Cast<UMetasoundEditorGraphVertex>(Member))
 				{
 					const FMetasoundEditorGraphVertexBreadcrumb Breadcrumb = SourceGraphVertex->GetBreadcrumb();
@@ -3878,6 +3880,10 @@ namespace Metasound
 
 						DocumentBuilder.AddGraphOutput(MoveTemp(ClassOutput));
 					}
+					else
+					{
+						checkNoEntry();
+					}
 				}
 				else if (const UMetasoundEditorGraphVariable* SourceGraphVariable = Cast<UMetasoundEditorGraphVariable>(Member))
 				{
@@ -3896,6 +3902,10 @@ namespace Metasound
 						Graph.FindOrAddVariable(FrontendVariable->Name);
 					}
 				}
+				else
+				{
+					checkNoEntry();
+				}
 
 				FGraphBuilder::RegisterGraphWithFrontend(*MetaSound, true);
 			}
