new variant of metasound delay node that accepts an audio buffer for delay time, to allow for sample-accurate manipulation of a delay line

diff --git a/Engine/Plugins/Runtime/Metasound/Source/MetasoundStandardNodes/Private/MetasoundDelayNode.cpp b/Engine/Plugins/Runtime/Metasound/Source/MetasoundStandardNodes/Private/MetasoundDelayNode.cpp
--- a/Engine/Plugins/Runtime/Metasound/Source/MetasoundStandardNodes/Private/MetasoundDelayNode.cpp
+++ b/Engine/Plugins/Runtime/Metasound/Source/MetasoundStandardNodes/Private/MetasoundDelayNode.cpp
@@ -12,12 +12,14 @@
 #include "DSP/Delay.h"
 #include "MetasoundStandardNodesCategories.h"
 #include "MetasoundFacade.h"
+#include "MetasoundParamHelper.h"
+
+#define LOCTEXT_NAMESPACE "MetasoundStandardNodes_Delay"
 
-#define LOCTEXT_NAMESPACE "MetasoundStandardNodes"
 
 namespace Metasound
 {
-	namespace Delay
+	namespace DelayNodePrivate
 	{
 		METASOUND_PARAM(InputResetDelay, "Reset", "Resets the delay buffer.");
 		METASOUND_PARAM(InParamAudioInput, "In", "Audio input.")
@@ -26,286 +28,358 @@ namespace Metasound
 		METASOUND_PARAM(InParamWetLevel, "Wet Level", "The wet level of the delay.")
 		METASOUND_PARAM(InParamFeedbackAmount, "Feedback", "Feedback amount.")
 		METASOUND_PARAM(InParamMaxDelayTime, "Max Delay Time", "The maximum amount of time to delay the audio.")
+
 		METASOUND_PARAM(OutParamAudio, "Out", "Audio output.")
 
 		static constexpr float MinMaxDelaySeconds = 0.001f;
 		static constexpr float MaxMaxDelaySeconds = 1000.f;
 		static constexpr float DefaultMaxDelaySeconds = 5.0f;
-	}
 
-	class FDelayOperator : public TExecutableOperator<FDelayOperator>
-	{
-	public:
-
-		static const FNodeClassMetadata& GetNodeInfo();
-		static const FVertexInterface& GetVertexInterface();
-		static TUniquePtr<IOperator> CreateOperator(const FBuildOperatorParams& InParams, FBuildResults& OutResults);
+		template<typename ValueType>
+		struct TDelay
+		{
+			bool bSupported = false;
+		};
 
-		FDelayOperator(const FBuildOperatorParams& InParams,
-			const FAudioBufferReadRef& InAudioInput, 
-			const FTimeReadRef& InDelayTime, 
-			const FFloatReadRef& InDryLevel, 
-			const FFloatReadRef& InWetLevel, 
-			const FFloatReadRef& InFeedback,
-			float InMaxDelayTimeSeconds,
-			const FTriggerReadRef& InTriggerReset);
+		template<>
+		struct TDelay<FTime>
+		{
+			static FNodeClassName GetClassName(const FName& InOperatorName)
+			{
+				// Note the older FTime variant was used StandardNodes::AudioVariant
+				return FNodeClassName{ StandardNodes::Namespace, InOperatorName, StandardNodes::AudioVariant };
+			}
 
-		virtual void BindInputs(FInputVertexInterfaceData& InOutVertexData) override;
-		virtual void BindOutputs(FOutputVertexInterfaceData& InOutVertexData) override;
-		virtual FDataReferenceCollection GetInputs() const override;
-		virtual FDataReferenceCollection GetOutputs() const override;
-		void Reset(const IOperator::FResetParams& InParams);
-		void Execute();
-		void ExecuteInternal(int32 StartFrame, int32 EndFrame);
+			static float GetDelayLengthSeconds(TDataReadReference<FTime>& InDelayParam, int32 InFrameIndex = 0)
+			{
+				return (float)InDelayParam->GetSeconds();
+			}
 
-	private:
-		float GetInputDelayTimeMsec() const;
+			static TDataReadReference<FTime> CreateInRef(const FBuildOperatorParams& InParams)
+			{
+				const FInputVertexInterfaceData& InputData = InParams.InputData;
+				return InputData.GetOrCreateDefaultDataReadReference<FTime>(METASOUND_GET_PARAM_NAME(InParamDelayTime), InParams.OperatorSettings);
+			}
 
-		// The input audio buffer
-		FAudioBufferReadRef AudioInput;
+			static bool IsAudioType() { return false; }
+		};
 
-		// The amount of delay time
-		FTimeReadRef DelayTime;
+		template<>
+		struct TDelay<FAudioBuffer>
+		{
+			static FNodeClassName GetClassName(const FName& InOperatorName)
+			{
+				// Even though this is technically the audio variant of the node, we need to pick something different than "audio variant", so pick something very specific
+				return FNodeClassName{ "Delay", InOperatorName, "AudioBufferDelayTime" };
+			}
 
-		// The dry level
-		FFloatReadRef DryLevel;
+			static float GetDelayLengthSeconds(TDataReadReference<FAudioBuffer>& InDelayParam, int32 InFrameIndex = 0)
+			{
+				const float* Data = InDelayParam->GetData();
+				// the first sample of the audio buffer is the initial delay length
+				return Data[InFrameIndex];
+			}
 
-		// The wet level
-		FFloatReadRef WetLevel;
+			static TDataReadReference<FAudioBuffer> CreateInRef(const FBuildOperatorParams& InParams)
+			{
+				const FInputVertexInterfaceData& InputData = InParams.InputData;
+				return InputData.GetOrConstructDataReadReference<FAudioBuffer>(METASOUND_GET_PARAM_NAME(InParamDelayTime), InParams.OperatorSettings);
+			}
 
-		// The feedback amount
-		FFloatReadRef Feedback;
+			static bool IsAudioType() { return true; }
+		};
+	}
 
-		// The audio output
-		FAudioBufferWriteRef AudioOutput;
+	template<typename ValueType>
+	class TDelayNodeOperator : public TExecutableOperator<TDelayNodeOperator<ValueType>>
+	{
+	public:
+		static const FVertexInterface& GetDefaultInterface()
+		{
+			using namespace DelayNodePrivate;
+
+			FDataVertexMetadata MaxDelayTimeMetadata = METASOUND_GET_PARAM_METADATA(InParamMaxDelayTime);
+			MaxDelayTimeMetadata.bIsAdvancedDisplay = true;
+
+			static const FVertexInterface DefaultInterface(
+				FInputVertexInterface(
+					TInputDataVertex<FAudioBuffer>(METASOUND_GET_PARAM_NAME_AND_METADATA(InParamAudioInput)),
+					TInputDataVertex<ValueType>(METASOUND_GET_PARAM_NAME_AND_METADATA(InParamDelayTime), 1.0f),
+					TInputDataVertex<float>(METASOUND_GET_PARAM_NAME_AND_METADATA(InParamDryLevel), 0.0f),
+					TInputDataVertex<float>(METASOUND_GET_PARAM_NAME_AND_METADATA(InParamWetLevel), 1.0f),
+					TInputDataVertex<float>(METASOUND_GET_PARAM_NAME_AND_METADATA(InParamFeedbackAmount), 0.0f),
+					TInputConstructorVertex<FTime>(METASOUND_GET_PARAM_NAME(InParamMaxDelayTime), MaxDelayTimeMetadata, DefaultMaxDelaySeconds),
+					TInputDataVertex<FTrigger>(METASOUND_GET_PARAM_NAME_AND_METADATA(InputResetDelay))
+				),
+				FOutputVertexInterface(
+					TOutputDataVertex<FAudioBuffer>(METASOUND_GET_PARAM_NAME_AND_METADATA(OutParamAudio))
+				)
+			);
+
+			return DefaultInterface;
+		}
 
-		// The internal delay buffer
-		Audio::FDelay DelayBuffer;
+		static const FNodeClassMetadata& GetNodeInfo()
+		{
+			auto CreateNodeClassMetadata = []() -> FNodeClassMetadata
+				{
+					FNodeClassMetadata Metadata
+					{
+						DelayNodePrivate::TDelay<ValueType>::GetClassName(TEXT("Delay")),
+						1, // Major Version
+						1, // Minor Version
+						METASOUND_LOCTEXT_FORMAT("DelayDisplayNamePattern", "Delay ({0})", GetMetasoundDataTypeDisplayText<ValueType>()),
+						METASOUND_LOCTEXT("DelayDesc", "Delays an audio buffer by the specified amount."),
+						PluginAuthor,
+						PluginNodeMissingPrompt,
+						GetDefaultInterface(),
+						{ NodeCategories::Delays },
+						{ },
+						FNodeDisplayStyle{}
+					};
+
+					return Metadata;
+				};
+
+			static const FNodeClassMetadata Metadata = CreateNodeClassMetadata();
+			return Metadata;
+		}
 
-		// The previous delay time
-		float PrevDelayTimeMsec = 0.f;
+		static TUniquePtr<IOperator> CreateOperator(const FBuildOperatorParams& InParams, FBuildResults& OutResults)
+		{
+			using namespace DelayNodePrivate;
 
-		// Feedback sample
-		float FeedbackSample = 0.f;
+			const FInputVertexInterfaceData& InputData = InParams.InputData;
 
-		// Maximum delay time
-		float MaxDelayTimeSeconds = Delay::DefaultMaxDelaySeconds;
+			FAudioBufferReadRef AudioIn = InputData.GetOrConstructDataReadReference<FAudioBuffer>(METASOUND_GET_PARAM_NAME(InParamAudioInput), InParams.OperatorSettings);
+			TDataReadReference<ValueType> DelayTime = TDelay<ValueType>::CreateInRef(InParams);
+			FFloatReadRef DryLevel = InputData.GetOrCreateDefaultDataReadReference<float>(METASOUND_GET_PARAM_NAME(InParamDryLevel), InParams.OperatorSettings);
+			FFloatReadRef WetLevel = InputData.GetOrCreateDefaultDataReadReference<float>(METASOUND_GET_PARAM_NAME(InParamWetLevel), InParams.OperatorSettings);
+			FFloatReadRef Feedback = InputData.GetOrCreateDefaultDataReadReference<float>(METASOUND_GET_PARAM_NAME(InParamFeedbackAmount), InParams.OperatorSettings);
+			FTime MaxDelayTime = InputData.GetOrCreateDefaultValue<FTime>(METASOUND_GET_PARAM_NAME(InParamMaxDelayTime), InParams.OperatorSettings);
+			FTriggerReadRef TriggerReset = InputData.GetOrConstructDataReadReference<FTrigger>(METASOUND_GET_PARAM_NAME(InputResetDelay), InParams.OperatorSettings);
 
-		// The reset trigger
-		FTriggerReadRef TriggerReset;
-	};
+			return MakeUnique<TDelayNodeOperator<ValueType>>(InParams, AudioIn, DelayTime, DryLevel, WetLevel, Feedback, MaxDelayTime.GetSeconds(), TriggerReset);
+		}
 
-	FDelayOperator::FDelayOperator(const FBuildOperatorParams& InParams,
-		const FAudioBufferReadRef& InAudioInput,
-		const FTimeReadRef& InDelayTime,
-		const FFloatReadRef& InDryLevel,
-		const FFloatReadRef& InWetLevel,
-		const FFloatReadRef& InFeedback,
-		float InMaxDelayTimeSeconds,
-		const FTriggerReadRef& InTriggerReset)
-
-		: AudioInput(InAudioInput)
-		, DelayTime(InDelayTime)
-		, DryLevel(InDryLevel)
-		, WetLevel(InWetLevel)
-		, Feedback(InFeedback)
-		, AudioOutput(FAudioBufferWriteRef::CreateNew(InParams.OperatorSettings))
-		, TriggerReset(InTriggerReset)
-	{
-		MaxDelayTimeSeconds = FMath::Clamp(InMaxDelayTimeSeconds, Delay::MinMaxDelaySeconds, Delay::MaxMaxDelaySeconds);
+		TDelayNodeOperator(const FBuildOperatorParams& InParams,
+			const FAudioBufferReadRef& InAudioInput,
+			const TDataReadReference<ValueType>& InDelayTime,
+			const FFloatReadRef& InDryLevel,
+			const FFloatReadRef& InWetLevel,
+			const FFloatReadRef& InFeedback,
+			float InMaxDelayTimeSeconds,
+			const FTriggerReadRef& InTriggerReset)
+
+			: AudioInput(InAudioInput)
+			, DelayTime(InDelayTime)
+			, DryLevel(InDryLevel)
+			, WetLevel(InWetLevel)
+			, Feedback(InFeedback)
+			, AudioOutput(FAudioBufferWriteRef::CreateNew(InParams.OperatorSettings))
+			, TriggerReset(InTriggerReset)
+		{
+			Reset(InParams);
+		}
 
-		Reset(InParams);
-	}
+		virtual ~TDelayNodeOperator() = default;
 
 
-	void FDelayOperator::BindInputs(FInputVertexInterfaceData& InOutVertexData)
-	{
-		using namespace Delay;
-
-		InOutVertexData.BindReadVertex(METASOUND_GET_PARAM_NAME(InParamAudioInput), AudioInput);
-		InOutVertexData.BindReadVertex(METASOUND_GET_PARAM_NAME(InParamDelayTime), DelayTime);
-		InOutVertexData.BindReadVertex(METASOUND_GET_PARAM_NAME(InParamDryLevel), DryLevel);
-		InOutVertexData.BindReadVertex(METASOUND_GET_PARAM_NAME(InParamWetLevel), WetLevel);
-		InOutVertexData.BindReadVertex(METASOUND_GET_PARAM_NAME(InParamFeedbackAmount), Feedback);
-		InOutVertexData.SetValue(METASOUND_GET_PARAM_NAME(InParamMaxDelayTime), FTime::FromSeconds(MaxDelayTimeSeconds));
-		InOutVertexData.BindReadVertex(METASOUND_GET_PARAM_NAME(InputResetDelay), TriggerReset);
-	}
+		virtual void BindInputs(FInputVertexInterfaceData& InOutVertexData) override
+		{
+			using namespace DelayNodePrivate;
+			InOutVertexData.BindReadVertex(METASOUND_GET_PARAM_NAME(InParamAudioInput), AudioInput);
+			InOutVertexData.BindReadVertex(METASOUND_GET_PARAM_NAME(InParamDelayTime), DelayTime);
+			InOutVertexData.BindReadVertex(METASOUND_GET_PARAM_NAME(InParamDryLevel), DryLevel);
+			InOutVertexData.BindReadVertex(METASOUND_GET_PARAM_NAME(InParamWetLevel), WetLevel);
+			InOutVertexData.BindReadVertex(METASOUND_GET_PARAM_NAME(InParamFeedbackAmount), Feedback);
+			InOutVertexData.SetValue(METASOUND_GET_PARAM_NAME(InParamMaxDelayTime), FTime::FromSeconds(MaxDelayTimeSeconds));
+			InOutVertexData.BindReadVertex(METASOUND_GET_PARAM_NAME(InputResetDelay), TriggerReset);
+		}
 
-	void FDelayOperator::BindOutputs(FOutputVertexInterfaceData& InOutVertexData)
-	{
-		using namespace Delay;
-		InOutVertexData.BindReadVertex(METASOUND_GET_PARAM_NAME(OutParamAudio), AudioOutput);
-	}
+		virtual void BindOutputs(FOutputVertexInterfaceData& InOutVertexData) override
+		{
+			using namespace DelayNodePrivate;
+			InOutVertexData.BindReadVertex(METASOUND_GET_PARAM_NAME(OutParamAudio), AudioOutput);
+		}
 
-	FDataReferenceCollection FDelayOperator::GetInputs() const
-	{
-		// This should never be called. Bind(...) is called instead. This method
-		// exists as a stop-gap until the API can be deprecated and removed.
-		checkNoEntry();
-		return {};
-	}
+		virtual FDataReferenceCollection GetInputs() const override
+		{
+			// This should never be called. Bind(...) is called instead. This method
+			// exists as a stop-gap until the API can be deprecated and removed.
+			checkNoEntry();
+			return {};
+		}
 
-	FDataReferenceCollection FDelayOperator::GetOutputs() const
-	{
-		// This should never be called. Bind(...) is called instead. This method
-		// exists as a stop-gap until the API can be deprecated and removed.
-		checkNoEntry();
-		return {};
-	}
+		virtual FDataReferenceCollection GetOutputs() const override
+		{
+			// This should never be called. Bind(...) is called instead. This method
+			// exists as a stop-gap until the API can be deprecated and removed.
+			checkNoEntry();
+			return {};
+		}
 
-	float FDelayOperator::GetInputDelayTimeMsec() const
-	{
-		// Clamp the delay time to the max delay allowed
-		return 1000.0f * FMath::Clamp((float)DelayTime->GetSeconds(), 0.0f, MaxDelayTimeSeconds);
-	}
+		float GetClampedDelayTimeMsec(float InDelayTime) const
+		{
+			// Clamp the delay time to the max delay allowed
+			return 1000.0f * FMath::Clamp(InDelayTime, 0.0f, MaxDelayTimeSeconds);
+		}
 
-	void FDelayOperator::Reset(const IOperator::FResetParams& InParams)
-	{
-		FeedbackSample = 0.f;
-		PrevDelayTimeMsec = GetInputDelayTimeMsec();
+		void Reset(const IOperator::FResetParams& InParams)
+		{
+			using namespace DelayNodePrivate;
 
-		DelayBuffer.Init(InParams.OperatorSettings.GetSampleRate(), MaxDelayTimeSeconds);
-		DelayBuffer.SetDelayMsec(PrevDelayTimeMsec);
+			FeedbackSample = 0.f;
 
-		AudioOutput->Zero();
-	}
+			float InitialDelayInSeconds = TDelay<ValueType>::GetDelayLengthSeconds(DelayTime);
+			PrevDelayTimeMsec = GetClampedDelayTimeMsec(InitialDelayInSeconds);
 
-	void FDelayOperator::Execute()
-	{
-		TriggerReset->ExecuteBlock(
-			[&](int32 StartFrame, int32 EndFrame)
-			{
-				ExecuteInternal(StartFrame, EndFrame); 
-			},
-			 [&, this](int32 StartFrame, int32 EndFrame)
-			{
-				FeedbackSample = 0.f;
-				DelayBuffer.ResetWithFade();
-			 	ExecuteInternal(StartFrame, EndFrame);
-			}
-		);
-	}
+			DelayBuffer.Init(InParams.OperatorSettings.GetSampleRate(), MaxDelayTimeSeconds);
+			DelayBuffer.SetDelayMsec(PrevDelayTimeMsec);
 
-	void FDelayOperator::ExecuteInternal(int32 StartFrame, int32 EndFrame)
-	{
-		// Get clamped delay time
-		float CurrentInputDelayTime = GetInputDelayTimeMsec();
+			AudioOutput->Zero();
+		}
 
-		// Check to see if our delay amount has changed
-		if (!FMath::IsNearlyEqual(PrevDelayTimeMsec, CurrentInputDelayTime))
+		void Execute()
 		{
-			PrevDelayTimeMsec = CurrentInputDelayTime;
-			DelayBuffer.SetEasedDelayMsec(PrevDelayTimeMsec);
+			TriggerReset->ExecuteBlock(
+				[&](int32 StartFrame, int32 EndFrame)
+				{
+					ExecuteInternal(StartFrame, EndFrame);
+				},
+				[&, this](int32 StartFrame, int32 EndFrame)
+				{
+					FeedbackSample = 0.f;
+					DelayBuffer.ResetWithFade();
+					ExecuteInternal(StartFrame, EndFrame);
+				}
+			);
 		}
 
-		const float* InputAudio = AudioInput->GetData();
+		void ExecuteInternal(int32 StartFrame, int32 EndFrame)
+		{
+			using namespace DelayNodePrivate;
 
-		float* OutputAudio = AudioOutput->GetData();
-		int32 NumFrames = AudioInput->Num();
+			const float* InputAudio = AudioInput->GetData();
 
-		// Clamp the feedback amount to make sure it's bounded. Clamp to a number slightly less than 1.0
-		float FeedbackAmount = FMath::Clamp(*Feedback, 0.0f, 1.0f - SMALL_NUMBER);
-		float CurrentDryLevel = FMath::Clamp(*DryLevel, 0.0f, 1.0f);
-		float CurrentWetLevel = FMath::Clamp(*WetLevel, 0.0f, 1.0f);
+			float* OutputAudio = AudioOutput->GetData();
+			int32 NumFrames = AudioInput->Num();
 
-		if (FMath::IsNearlyZero(FeedbackAmount))
-		{
-			FeedbackSample = 0.0f;
+			// Clamp the feedback amount to make sure it's bounded. Clamp to a number slightly less than 1.0
+			float CurrentFeedbackAmount = FMath::Clamp(*Feedback, 0.0f, 1.0f - SMALL_NUMBER);
+			float CurrentDryLevel = FMath::Clamp(*DryLevel, 0.0f, 1.0f);
+			float CurrentWetLevel = FMath::Clamp(*WetLevel, 0.0f, 1.0f);
 
-			for (int32 FrameIndex = StartFrame; FrameIndex < EndFrame; ++FrameIndex)
+			// check if we're doing an audio-rate delay time
+			if (TDelay<ValueType>::IsAudioType())
 			{
-				OutputAudio[FrameIndex] = CurrentWetLevel * DelayBuffer.ProcessAudioSample(InputAudio[FrameIndex]) + CurrentDryLevel * InputAudio[FrameIndex];
+				if (FMath::IsNearlyZero(CurrentFeedbackAmount))
+				{
+					FeedbackSample = 0.0f;
+					for (int32 FrameIndex = StartFrame; FrameIndex < EndFrame; ++FrameIndex)
+					{
+						DelayBuffer.SetDelayMsec(GetClampedDelayTimeMsec(TDelay<ValueType>::GetDelayLengthSeconds(DelayTime, FrameIndex)));
+						OutputAudio[FrameIndex] = CurrentWetLevel * DelayBuffer.ProcessAudioSample(InputAudio[FrameIndex]) + CurrentDryLevel * InputAudio[FrameIndex];
+					}
+				}
+				else
+				{
+					// There is some amount of feedback so we do the feedback mixing
+					for (int32 FrameIndex = StartFrame; FrameIndex < EndFrame; ++FrameIndex)
+					{
+						DelayBuffer.SetDelayMsec(GetClampedDelayTimeMsec(TDelay<ValueType>::GetDelayLengthSeconds(DelayTime, FrameIndex)));
+						OutputAudio[FrameIndex] = CurrentWetLevel * DelayBuffer.ProcessAudioSample(InputAudio[FrameIndex] + FeedbackSample * CurrentFeedbackAmount) + CurrentDryLevel * InputAudio[FrameIndex];
+						FeedbackSample = OutputAudio[FrameIndex];
+					}
+				}
 			}
-		}
-		else
-		{
-			// There is some amount of feedback so we do the feedback mixing
-			for (int32 FrameIndex = StartFrame; FrameIndex < EndFrame; ++FrameIndex)
+			else // non audio-rate delay
 			{
-				OutputAudio[FrameIndex] = CurrentWetLevel * DelayBuffer.ProcessAudioSample(InputAudio[FrameIndex] + FeedbackSample * FeedbackAmount) + CurrentDryLevel * InputAudio[FrameIndex];
-				FeedbackSample = OutputAudio[FrameIndex];
+				// Get clamped delay time
+				float CurrentInputDelayTime = GetClampedDelayTimeMsec(TDelay<ValueType>::GetDelayLengthSeconds(DelayTime));
+
+				// Check to see if our delay amount has changed
+				if (!FMath::IsNearlyEqual(PrevDelayTimeMsec, CurrentInputDelayTime))
+				{
+					PrevDelayTimeMsec = CurrentInputDelayTime;
+					DelayBuffer.SetEasedDelayMsec(PrevDelayTimeMsec);
+				}
+
+				if (FMath::IsNearlyZero(CurrentFeedbackAmount))
+				{
+					FeedbackSample = 0.0f;
+
+					for (int32 FrameIndex = StartFrame; FrameIndex < EndFrame; ++FrameIndex)
+					{
+						OutputAudio[FrameIndex] = CurrentWetLevel * DelayBuffer.ProcessAudioSample(InputAudio[FrameIndex]) + CurrentDryLevel * InputAudio[FrameIndex];
+					}
+				}
+				else
+				{
+					// There is some amount of feedback so we do the feedback mixing
+					for (int32 FrameIndex = StartFrame; FrameIndex < EndFrame; ++FrameIndex)
+					{
+						OutputAudio[FrameIndex] = CurrentWetLevel * DelayBuffer.ProcessAudioSample(InputAudio[FrameIndex] + FeedbackSample * CurrentFeedbackAmount) + CurrentDryLevel * InputAudio[FrameIndex];
+						FeedbackSample = OutputAudio[FrameIndex];
+					}
+				}
 			}
 		}
-	}
 
-	const FVertexInterface& FDelayOperator::GetVertexInterface()
-	{
-		using namespace Delay;
-
-		FDataVertexMetadata MaxDelayTimeMetadata = METASOUND_GET_PARAM_METADATA(InParamMaxDelayTime);
-		MaxDelayTimeMetadata.bIsAdvancedDisplay = true;
-
-		static const FVertexInterface Interface(
-			FInputVertexInterface(
-				TInputDataVertex<FAudioBuffer>(METASOUND_GET_PARAM_NAME_AND_METADATA(InParamAudioInput)),
-				TInputDataVertex<FTime>(METASOUND_GET_PARAM_NAME_AND_METADATA(InParamDelayTime), 1.0f),
-				TInputDataVertex<float>(METASOUND_GET_PARAM_NAME_AND_METADATA(InParamDryLevel), 0.0f),
-				TInputDataVertex<float>(METASOUND_GET_PARAM_NAME_AND_METADATA(InParamWetLevel), 1.0f),
-				TInputDataVertex<float>(METASOUND_GET_PARAM_NAME_AND_METADATA(InParamFeedbackAmount), 0.0f),
-				TInputConstructorVertex<FTime>(METASOUND_GET_PARAM_NAME(InParamMaxDelayTime), MaxDelayTimeMetadata, DefaultMaxDelaySeconds),
-				TInputDataVertex<FTrigger>(METASOUND_GET_PARAM_NAME_AND_METADATA(InputResetDelay))
-			),
-			FOutputVertexInterface(
-				TOutputDataVertex<FAudioBuffer>(METASOUND_GET_PARAM_NAME_AND_METADATA(OutParamAudio))
-			)
-		);
-
-		return Interface;
-	}
 
-	const FNodeClassMetadata& FDelayOperator::GetNodeInfo()
-	{
-		auto InitNodeInfo = []() -> FNodeClassMetadata
-		{
-			FNodeClassMetadata Info;
-			Info.ClassName = { StandardNodes::Namespace, "Delay", StandardNodes::AudioVariant };
-			Info.MajorVersion = 1;
-			Info.MinorVersion = 1;
-			Info.DisplayName = METASOUND_LOCTEXT("DelayNode_DisplayName", "Delay");
-			Info.Description = METASOUND_LOCTEXT("DelayNode_Description", "Delays an audio buffer by the specified amount.");
-			Info.Author = PluginAuthor;
-			Info.PromptIfMissing = PluginNodeMissingPrompt;
-			Info.DefaultInterface = GetVertexInterface();
-			Info.CategoryHierarchy.Emplace(NodeCategories::Delays);
-			return Info;
-		};
+	private:
 
-		static const FNodeClassMetadata Info = InitNodeInfo();
+		// The input audio buffer
+		FAudioBufferReadRef AudioInput;
 
-		return Info;
-	}
+		// The amount of delay time
+		TDataReadReference<ValueType> DelayTime;
 
-	TUniquePtr<IOperator> FDelayOperator::CreateOperator(const FBuildOperatorParams& InParams, FBuildResults& OutResults)
-	{
-		using namespace Delay; 
+		// The dry level
+		FFloatReadRef DryLevel;
 
-		const FInputVertexInterfaceData& InputData = InParams.InputData;
+		// The wet level
+		FFloatReadRef WetLevel;
 
-		FAudioBufferReadRef AudioIn = InputData.GetOrConstructDataReadReference<FAudioBuffer>(METASOUND_GET_PARAM_NAME(InParamAudioInput), InParams.OperatorSettings);
-		FTimeReadRef DelayTime = InputData.GetOrCreateDefaultDataReadReference<FTime>(METASOUND_GET_PARAM_NAME(InParamDelayTime), InParams.OperatorSettings);
-		FFloatReadRef DryLevel = InputData.GetOrCreateDefaultDataReadReference<float>(METASOUND_GET_PARAM_NAME(InParamDryLevel), InParams.OperatorSettings);
-		FFloatReadRef WetLevel = InputData.GetOrCreateDefaultDataReadReference<float>(METASOUND_GET_PARAM_NAME(InParamWetLevel), InParams.OperatorSettings);
-		FFloatReadRef Feedback = InputData.GetOrCreateDefaultDataReadReference<float>(METASOUND_GET_PARAM_NAME(InParamFeedbackAmount), InParams.OperatorSettings);
-		FTime MaxDelayTime = InputData.GetOrCreateDefaultValue<FTime>(METASOUND_GET_PARAM_NAME(InParamMaxDelayTime), InParams.OperatorSettings);
-		FTriggerReadRef TriggerReset = InputData.GetOrConstructDataReadReference<FTrigger>(METASOUND_GET_PARAM_NAME(InputResetDelay), InParams.OperatorSettings);
+		// The feedback amount
+		FFloatReadRef Feedback;
 
-		return MakeUnique<FDelayOperator>(InParams, AudioIn, DelayTime, DryLevel, WetLevel, Feedback, MaxDelayTime.GetSeconds(), TriggerReset);
-	}
+		// The audio output
+		FAudioBufferWriteRef AudioOutput;
+
+		// The internal delay buffer
+		Audio::FDelay DelayBuffer;
+
+		// The previous delay time
+		float PrevDelayTimeMsec = 0.f;
+
+		// Feedback sample
+		float FeedbackSample = 0.f;
+
+		// Maximum delay time
+		float MaxDelayTimeSeconds = DelayNodePrivate::DefaultMaxDelaySeconds;
 
-	class FDelayNode : public FNodeFacade
+		// The reset trigger
+		FTriggerReadRef TriggerReset;
+	};
+
+	template<typename ValueType>
+	class METASOUNDSTANDARDNODES_API TDelayNode : public FNodeFacade
 	{
 	public:
-		/**
-		 * Constructor used by the Metasound Frontend.
-		 */
-		FDelayNode(const FNodeInitData& InitData)
-			: FNodeFacade(InitData.InstanceName, InitData.InstanceID, TFacadeOperatorClass<FDelayOperator>())
-		{
-		}
+		TDelayNode(const FNodeInitData& InInitData)
+			: FNodeFacade(InInitData.InstanceName, InInitData.InstanceID, TFacadeOperatorClass<TDelayNodeOperator<ValueType>>())
+		{}
+
+		virtual ~TDelayNode() = default;
 	};
 
 
+	using FDelayNode = TDelayNode<FTime>;
 	METASOUND_REGISTER_NODE(FDelayNode)
+
+	using FDelayNodeAudio = TDelayNode<FAudioBuffer>;
+	METASOUND_REGISTER_NODE(FDelayNodeAudio)
 }
 
 #undef LOCTEXT_NAMESPACE
