metasounds delay node: fix for "Max Delay Time" not being respected, fix for out-of-bounds array access

diff --git a/Engine/Plugins/Runtime/Metasound/Source/MetasoundStandardNodes/Private/MetasoundDelayNode.cpp b/Engine/Plugins/Runtime/Metasound/Source/MetasoundStandardNodes/Private/MetasoundDelayNode.cpp
--- a/Engine/Plugins/Runtime/Metasound/Source/MetasoundStandardNodes/Private/MetasoundDelayNode.cpp
+++ b/Engine/Plugins/Runtime/Metasound/Source/MetasoundStandardNodes/Private/MetasoundDelayNode.cpp
@@ -75,6 +75,9 @@ namespace Metasound
 
 			static float GetDelayLengthSeconds(TDataReadReference<FAudioBuffer>& InDelayParam, int32 InFrameIndex = 0)
 			{
+				if (InFrameIndex < 0 || InFrameIndex >= InDelayParam->Num())
+					return 0.0f;
+
 				const float* Data = InDelayParam->GetData();
 				// the first sample of the audio buffer is the initial delay length
 				return Data[InFrameIndex];
@@ -179,6 +182,8 @@ namespace Metasound
 			, AudioOutput(FAudioBufferWriteRef::CreateNew(InParams.OperatorSettings))
 			, TriggerReset(InTriggerReset)
 		{
+			MaxDelayTimeSeconds = FMath::Clamp(InMaxDelayTimeSeconds, DelayNodePrivate::MinMaxDelaySeconds, DelayNodePrivate::MaxMaxDelaySeconds);
+
 			Reset(InParams);
 		}
 
