[Audio][MS] Add a bunch of new keywords to common MS nodes:
- "gain" and "vca" as keywords to find the MultiplyAudioFloat node
- "if" and "branch" as keywords to find TriggerCompare nodes
- "select" as keyword to find TriggerRoute nodes

diff --git a/Engine/Plugins/Runtime/Metasound/Source/MetasoundStandardNodes/Private/MetasoundMathNodes.cpp b/Engine/Plugins/Runtime/Metasound/Source/MetasoundStandardNodes/Private/MetasoundMathNodes.cpp
--- a/Engine/Plugins/Runtime/Metasound/Source/MetasoundStandardNodes/Private/MetasoundMathNodes.cpp
+++ b/Engine/Plugins/Runtime/Metasound/Source/MetasoundStandardNodes/Private/MetasoundMathNodes.cpp
@@ -66,6 +66,13 @@ namespace Metasound
 		static const TArray<FText> DivideKeywords = { METASOUND_LOCTEXT("DivideMathKeyword", "/") };
 		static const TArray<FText> PowerKeywords = { METASOUND_LOCTEXT("PowerMathKeyword", "^") };
 		static const TArray<FText> ModuloKeywords = { METASOUND_LOCTEXT("ModuloMathKeyword", "%") };
+		
+		static const TArray<FText> MultiplyAudioFloatKeywords =
+			{
+				MultiplyKeywords[0],
+				METASOUND_LOCTEXT("MultiplyGainKeyword", "gain"), // DSP/Modular synth term
+				METASOUND_LOCTEXT("MultiplyVCAKeyword", "vca"),   // Voltage Controlled Amplifier (DAW/Electronics term)
+			};
 
 		template<typename DataType>
 		const FText GetAddDisplayName()
@@ -1340,7 +1347,7 @@ namespace Metasound
 	DEFINE_METASOUND_MATHOP(Multiply, Int32, int32, METASOUND_LOCTEXT("MathMultiplyInt32Node_Description", "Multiplies int32s."), MathOpNames::MultiplyKeywords)
 	DEFINE_METASOUND_MATHOP(Multiply, Audio, FAudioBuffer, METASOUND_LOCTEXT("MathMultiplyBufferNode_Description", "Multiplies buffers together sample-by-sample."), MathOpNames::MultiplyKeywords)
 
-	DEFINE_METASOUND_OPERAND_TYPED_MATHOP(Multiply, Audio, FAudioBuffer, Float, float, METASOUND_LOCTEXT("MathMultiplyAudioByFloatDescription", "Multiplies buffer by float scalars."), MathOpNames::MultiplyKeywords)
+	DEFINE_METASOUND_OPERAND_TYPED_MATHOP(Multiply, Audio, FAudioBuffer, Float, float, METASOUND_LOCTEXT("MathMultiplyAudioByFloatDescription", "Multiplies buffer by float scalars."), MathOpNames::MultiplyAudioFloatKeywords)
 	DEFINE_METASOUND_OPERAND_TYPED_MATHOP(Multiply, Time, FTime, Float, float, METASOUND_LOCTEXT("MathMultiplyTimeNode_Description", "Scales time by floats."), MathOpNames::MultiplyKeywords)
 
 	DEFINE_METASOUND_MATHOP(Divide, Float, float, METASOUND_LOCTEXT("MathDivideFloatNode_Description", "Divide float by another float."), MathOpNames::DivideKeywords)
diff --git a/Engine/Plugins/Runtime/Metasound/Source/MetasoundStandardNodes/Private/MetasoundTriggerCompareNode.cpp b/Engine/Plugins/Runtime/Metasound/Source/MetasoundStandardNodes/Private/MetasoundTriggerCompareNode.cpp
--- a/Engine/Plugins/Runtime/Metasound/Source/MetasoundStandardNodes/Private/MetasoundTriggerCompareNode.cpp
+++ b/Engine/Plugins/Runtime/Metasound/Source/MetasoundStandardNodes/Private/MetasoundTriggerCompareNode.cpp
@@ -21,8 +21,15 @@ namespace Metasound
 				PluginNodeMissingPrompt,
 				InDefaultInterface,
 				{ NodeCategories::Trigger },
-				{ METASOUND_LOCTEXT("EqualsKeyword", "="), METASOUND_LOCTEXT("NotEqualsKeyword", "!="), METASOUND_LOCTEXT("LessThanKeyword", "<"), METASOUND_LOCTEXT("GreaterThanKeyword", ">"), METASOUND_LOCTEXT("LessEqualsKeyword", "<="), METASOUND_LOCTEXT("GreaterEqualsKeyword", ">="),
-				 METASOUND_LOCTEXT("EqualsTextKeyword", "equal"), METASOUND_LOCTEXT("InequalsKeyword", "inequal"), METASOUND_LOCTEXT("NotKeyword", "not"), METASOUND_LOCTEXT("LessThanTextKeyword", "less"), METASOUND_LOCTEXT("GreatThanTextKeyword", "greater")},
+				{
+					METASOUND_LOCTEXT("EqualsKeyword", "="), METASOUND_LOCTEXT("NotEqualsKeyword", "!="),
+					METASOUND_LOCTEXT("LessThanKeyword", "<"), METASOUND_LOCTEXT("GreaterThanKeyword", ">"),
+					METASOUND_LOCTEXT("LessEqualsKeyword", "<="), METASOUND_LOCTEXT("GreaterEqualsKeyword", ">="),
+					METASOUND_LOCTEXT("EqualsTextKeyword", "equal"), METASOUND_LOCTEXT("InequalsKeyword", "inequal"),
+					METASOUND_LOCTEXT("NotKeyword", "not"),
+					METASOUND_LOCTEXT("LessThanTextKeyword", "less"),
+					METASOUND_LOCTEXT("GreatThanTextKeyword", "greater"),
+					METASOUND_LOCTEXT("IfKeyword", "if"), METASOUND_LOCTEXT("BranchKeyword", "branch")},
 				FNodeDisplayStyle()
 			};
 
diff --git a/Engine/Plugins/Runtime/Metasound/Source/MetasoundStandardNodes/Private/MetasoundTriggerRouteNode.cpp b/Engine/Plugins/Runtime/Metasound/Source/MetasoundStandardNodes/Private/MetasoundTriggerRouteNode.cpp
--- a/Engine/Plugins/Runtime/Metasound/Source/MetasoundStandardNodes/Private/MetasoundTriggerRouteNode.cpp
+++ b/Engine/Plugins/Runtime/Metasound/Source/MetasoundStandardNodes/Private/MetasoundTriggerRouteNode.cpp
@@ -80,7 +80,7 @@ namespace Metasound
 				PluginNodeMissingPrompt,
 				InDefaultInterface,
 				{ NodeCategories::Trigger },
-				{ },
+				{ METASOUND_LOCTEXT("TriggerRouteKeyword_Select", "select") },
 				FNodeDisplayStyle()
 			};
 
