Make all audio modules compile with AutoRTFM as they are always hit from Verse.

#rnx

diff --git a/Engine/Plugins/Runtime/AudioModulation/Source/AudioModulation/AudioModulation.Build.cs b/Engine/Plugins/Runtime/AudioModulation/Source/AudioModulation/AudioModulation.Build.cs
--- a/Engine/Plugins/Runtime/AudioModulation/Source/AudioModulation/AudioModulation.Build.cs
+++ b/Engine/Plugins/Runtime/AudioModulation/Source/AudioModulation/AudioModulation.Build.cs
@@ -66,8 +66,6 @@ namespace UnrealBuildTool.Rules
 			{
 				PublicDefinitions.Add("WITH_AUDIOMODULATION_METASOUND_SUPPORT=0");
 			}
-
-			bDisableAutoRTFMInstrumentation = true;
 		}
 	}
 }
diff --git a/Engine/Plugins/Runtime/Metasound/Source/MetasoundEngine/MetasoundEngine.Build.cs b/Engine/Plugins/Runtime/Metasound/Source/MetasoundEngine/MetasoundEngine.Build.cs
--- a/Engine/Plugins/Runtime/Metasound/Source/MetasoundEngine/MetasoundEngine.Build.cs
+++ b/Engine/Plugins/Runtime/Metasound/Source/MetasoundEngine/MetasoundEngine.Build.cs
@@ -45,8 +45,6 @@ namespace UnrealBuildTool.Rules
 					"WaveTable"
 				}
 			);
-
-			bDisableAutoRTFMInstrumentation = true;
 		}
 	}
 }
diff --git a/Engine/Plugins/Runtime/Metasound/Source/MetasoundFrontend/MetasoundFrontend.Build.cs b/Engine/Plugins/Runtime/Metasound/Source/MetasoundFrontend/MetasoundFrontend.Build.cs
--- a/Engine/Plugins/Runtime/Metasound/Source/MetasoundFrontend/MetasoundFrontend.Build.cs
+++ b/Engine/Plugins/Runtime/Metasound/Source/MetasoundFrontend/MetasoundFrontend.Build.cs
@@ -25,8 +25,6 @@ namespace UnrealBuildTool.Rules
 			);
 
 			PublicDefinitions.Add("WITH_METASOUND_FRONTEND=1");
-
-			bDisableAutoRTFMInstrumentation = true;
 		}
 	}
 }
diff --git a/Engine/Plugins/Runtime/Metasound/Source/MetasoundStandardNodes/MetasoundStandardNodes.Build.cs b/Engine/Plugins/Runtime/Metasound/Source/MetasoundStandardNodes/MetasoundStandardNodes.Build.cs
--- a/Engine/Plugins/Runtime/Metasound/Source/MetasoundStandardNodes/MetasoundStandardNodes.Build.cs
+++ b/Engine/Plugins/Runtime/Metasound/Source/MetasoundStandardNodes/MetasoundStandardNodes.Build.cs
@@ -31,8 +31,6 @@ namespace UnrealBuildTool.Rules
 			);
 
 			NumIncludedBytesPerUnityCPPOverride = 120 * 1024;
-
-			bDisableAutoRTFMInstrumentation = true;
 		}
 	}
 }
diff --git a/Engine/Plugins/Runtime/Synthesis/Source/Synthesis/Synthesis.Build.cs b/Engine/Plugins/Runtime/Synthesis/Source/Synthesis/Synthesis.Build.cs
--- a/Engine/Plugins/Runtime/Synthesis/Source/Synthesis/Synthesis.Build.cs
+++ b/Engine/Plugins/Runtime/Synthesis/Source/Synthesis/Synthesis.Build.cs
@@ -31,8 +31,6 @@ namespace UnrealBuildTool.Rules
 			{
 				PrivateDependencyModuleNames.Add("AudioSynesthesiaCore");
 			}
-
-			bDisableAutoRTFMInstrumentation = true;
 		}
 	}
 }
