[MetasoundEditor] Fix undocked analyzer tabs sometimes hanging around after the metasound editor is closed - Using weak pointer for analyzer rack in tabs spawner.

diff --git a/Engine/Plugins/Runtime/Metasound/Source/MetasoundEditor/Private/MetasoundEditor.cpp b/Engine/Plugins/Runtime/Metasound/Source/MetasoundEditor/Private/MetasoundEditor.cpp
--- a/Engine/Plugins/Runtime/Metasound/Source/MetasoundEditor/Private/MetasoundEditor.cpp
+++ b/Engine/Plugins/Runtime/Metasound/Source/MetasoundEditor/Private/MetasoundEditor.cpp
@@ -981,12 +981,12 @@ namespace Metasound
 			.SetGroup(WorkspaceMenuCategoryRef)
 			.SetIcon(FSlateIcon("MetaSoundStyle", "MetasoundEditor.Metasound.Icon"));
 
-			InTabManager->RegisterTabSpawner(TabNamesPrivate::Analyzers, FOnSpawnTab::CreateLambda([InAnalyzerRack = AnalyzerRack](const FSpawnTabArgs&)
+			InTabManager->RegisterTabSpawner(TabNamesPrivate::Analyzers, FOnSpawnTab::CreateLambda([InAnalyzerRack = AnalyzerRack.ToWeakPtr()](const FSpawnTabArgs&)
 			{
 				TSharedRef<SDockTab> DockTab = SNew(SDockTab).Label(LOCTEXT("MetasoundAnalyzersTitle", "Analyzers"));
-				if (InAnalyzerRack.IsValid()) // AnalyzerRack only created for MetaSoundSource, not MetaSoundPatch.
+				if (TSharedPtr<AudioWidgets::FAudioAnalyzerRack> PinnedAnalyzerRack = InAnalyzerRack.Pin()) // AnalyzerRack only created for MetaSoundSource, not MetaSoundPatch.
 				{
-					TSharedRef<SWidget> AnalyzerRackWidget = InAnalyzerRack->CreateWidget(DockTab);
+					TSharedRef<SWidget> AnalyzerRackWidget = PinnedAnalyzerRack->CreateWidget(DockTab);
 					DockTab->SetContent(AnalyzerRackWidget);
 				}
 				return DockTab;
