Core MetaSound Vertex Sub Interface Support
#rnx

diff --git a/Engine/Plugins/Runtime/Metasound/Source/MetasoundEngineTest/Private/EngineTestMetaSoundAutomatedNodeTest.cpp b/Engine/Plugins/Runtime/Metasound/Source/MetasoundEngineTest/Private/EngineTestMetaSoundAutomatedNodeTest.cpp
--- a/Engine/Plugins/Runtime/Metasound/Source/MetasoundEngineTest/Private/EngineTestMetaSoundAutomatedNodeTest.cpp
+++ b/Engine/Plugins/Runtime/Metasound/Source/MetasoundEngineTest/Private/EngineTestMetaSoundAutomatedNodeTest.cpp
@@ -795,7 +795,7 @@ namespace Metasound::EngineTest{
 
 	void CreateVariables(const FOperatorSettings& InOperatorSettings, FInputVertexInterfaceData& OutVertexData)
 	{
-		using namespace MetasoundVertexDataPrivate;
+		using namespace VertexDataPrivate;
 
 		// There is currently no easy way to instantiate variables. The only way
 		// they are instantiated is within the TVariableNode's IOperator. In order
diff --git a/Engine/Plugins/Runtime/Metasound/Source/MetasoundFrontend/Private/MetasoundFrontendDataTypeRegistry.cpp b/Engine/Plugins/Runtime/Metasound/Source/MetasoundFrontend/Private/MetasoundFrontendDataTypeRegistry.cpp
--- a/Engine/Plugins/Runtime/Metasound/Source/MetasoundFrontend/Private/MetasoundFrontendDataTypeRegistry.cpp
+++ b/Engine/Plugins/Runtime/Metasound/Source/MetasoundFrontend/Private/MetasoundFrontendDataTypeRegistry.cpp
@@ -1226,7 +1226,7 @@ namespace Metasound
 
 		void CreateDefaults(const FOperatorSettings& InOperatorSettings, FInputVertexInterfaceData& OutVertexData)
 		{
-			using namespace MetasoundVertexDataPrivate;
+			using namespace VertexDataPrivate;
 
 			auto VertexAccessTypeToDataReferenceAccessType = [](EVertexAccessType InVertexAccessType) -> EDataReferenceAccessType
 			{
diff --git a/Engine/Plugins/Runtime/Metasound/Source/MetasoundGenerator/Private/MetasoundGeneratorBuilder.cpp b/Engine/Plugins/Runtime/Metasound/Source/MetasoundGenerator/Private/MetasoundGeneratorBuilder.cpp
--- a/Engine/Plugins/Runtime/Metasound/Source/MetasoundGenerator/Private/MetasoundGeneratorBuilder.cpp
+++ b/Engine/Plugins/Runtime/Metasound/Source/MetasoundGenerator/Private/MetasoundGeneratorBuilder.cpp
@@ -230,7 +230,7 @@ namespace Metasound
 			}
 
 			// Set any remaining inputs to their default values.
-			for (const MetasoundVertexDataPrivate::FInputBinding& Binding : InputData)
+			for (const VertexDataPrivate::FInputBinding& Binding : InputData)
 			{
 				// Only create data reference if something does not already exist. 
 				if (!Binding.IsBound())
@@ -357,7 +357,7 @@ namespace Metasound
 
 			Frontend::IDataTypeRegistry& DataTypeRegistry = Frontend::IDataTypeRegistry::Get();
 
-			for (MetasoundVertexDataPrivate::FInputBinding& Binding : InOutInterface)
+			for (VertexDataPrivate::FInputBinding& Binding : InOutInterface)
 			{
 				if (const FAnyDataReference* Ref = Binding.GetDataReference())
 				{
@@ -394,7 +394,7 @@ namespace Metasound
 			// can be assigned appropriately
 			const IDataTypeRegistry& DataTypeRegistry = IDataTypeRegistry::Get();
 
-			for (const MetasoundVertexDataPrivate::FInputBinding& Binding : InputData)
+			for (const VertexDataPrivate::FInputBinding& Binding : InputData)
 			{
 				// Only assign inputs that are writable. 
 				if (EDataReferenceAccessType::Write == Binding.GetAccessType())
diff --git a/Engine/Plugins/Runtime/Metasound/Source/MetasoundGraphCore/Private/MetasoundOperatorBuilder.cpp b/Engine/Plugins/Runtime/Metasound/Source/MetasoundGraphCore/Private/MetasoundOperatorBuilder.cpp
--- a/Engine/Plugins/Runtime/Metasound/Source/MetasoundGraphCore/Private/MetasoundOperatorBuilder.cpp
+++ b/Engine/Plugins/Runtime/Metasound/Source/MetasoundGraphCore/Private/MetasoundOperatorBuilder.cpp
@@ -516,7 +516,7 @@ namespace Metasound
 	FOperatorBuilder::FBuildStatus FOperatorBuilder::ValidateOperatorOutputsAreBound(const INode& InNode, const FOutputVertexInterfaceData& InVertexData) const
 	{
 		METASOUND_TRACE_CPUPROFILER_EVENT_SCOPE_COND(Metasound::FOperatorBuilder::ValidateOperatorOutputsAreBound, bVerboseTraceEvents);
-		using FOutputVertexBinding = MetasoundVertexDataPrivate::FOutputBinding;
+		using FOutputVertexBinding = VertexDataPrivate::FOutputBinding;
 
 		FBuildStatus BuildStatus = FBuildStatus::NoError;
 		bool bFoundUnboundVertex = false;
diff --git a/Engine/Plugins/Runtime/Metasound/Source/MetasoundGraphCore/Private/MetasoundVertex.cpp b/Engine/Plugins/Runtime/Metasound/Source/MetasoundGraphCore/Private/MetasoundVertex.cpp
--- a/Engine/Plugins/Runtime/Metasound/Source/MetasoundGraphCore/Private/MetasoundVertex.cpp
+++ b/Engine/Plugins/Runtime/Metasound/Source/MetasoundGraphCore/Private/MetasoundVertex.cpp
@@ -1,11 +1,362 @@
 // Copyright Epic Games, Inc. All Rights Reserved.
 
 #include "MetasoundVertex.h"
+#include "MetasoundVertexPrivate.h"
 
-#include "CoreMinimal.h"
+#include "Algo/Find.h"
+#include "Containers/Array.h"
 
 namespace Metasound
 {
+	namespace VertexPrivate
+	{
+		// Functor used for finding FSubInterfaceLayouts by name. 
+		struct FEqualSubInterfaceName
+		{
+			const FName& NameRef;
+
+			FEqualSubInterfaceName(const FName& InName)
+			: NameRef(InName)
+			{
+			}
+
+			bool operator()(const FSubInterfaceLayout& InOther) const
+			{
+				return NameRef == InOther.SubInterfaceName;
+			}
+		};
+
+		FSubInterfaceDeclarationBuilder::FSubInterfaceDeclarationBuilder(TArray<FSubInterfaceLayout>& OutSubInterfaceLayouts)
+		: SubInterfaces(&OutSubInterfaceLayouts)
+		{
+		}
+
+		FSubInterfaceDeclarationBuilder::~FSubInterfaceDeclarationBuilder()
+		{
+			checkf(CurrentSubInterfaceIndex == INDEX_NONE, TEXT("Failed to beging/end all sub interfaces"));
+		}
+
+		void FSubInterfaceDeclarationBuilder::ReserveSubInterfaces(int32 Num)
+		{
+			SubInterfaces->Reserve(Num);
+		}
+
+		void FSubInterfaceDeclarationBuilder::Add(FBeginSubInterface&& InSubInterface)
+		{
+			PushSubInterfaceDeclaration(InSubInterface.Name);
+		}
+
+		void FSubInterfaceDeclarationBuilder::Add(FEndSubInterface&& InSubInterface)
+		{
+			PopSubInterfaceDeclaration();
+		}
+
+		void FSubInterfaceDeclarationBuilder::OnVertexAdded(const FVertexName& InName)
+		{
+			checkf((CurrentSubInterfaceIndex == INDEX_NONE) || (InName.GetNumber() == NAME_NO_NUMBER_INTERNAL), TEXT("Vertex %s in sub interface %s cannot have a trailing number because it is in a sub interface."), *InName.ToString(), *(*SubInterfaces)[CurrentSubInterfaceIndex].SubInterfaceName.ToString());
+			CurrentNumVertices++;
+		}
+
+		void FSubInterfaceDeclarationBuilder::PushSubInterfaceDeclaration(const FName& InName)
+		{
+			// Sub interface declarations should only happen once because the declaration
+			// needs to be consistent. Relying on developers declaring it identically
+			// is error prone so it is not allowed. The default number of times a
+			// sub interface is replicated can be controlled via the FSubInterfaceDescription.
+			checkf(INDEX_NONE == SubInterfaces->IndexOfByPredicate(FEqualSubInterfaceName{InName}), TEXT("Sub interface %s is already declared."), *InName.ToString());
+
+			SubInterfaces->Emplace(InName);
+			int32 NextSubInterfaceIndex = SubInterfaces->Num() - 1;
+
+			checkf(CurrentSubInterfaceIndex == INDEX_NONE, TEXT("Sub interface %s cannot be embedded inside sub interface %s"), *InName.ToString(), *(*SubInterfaces)[CurrentSubInterfaceIndex].SubInterfaceName.ToString());
+
+			CurrentSubInterfaceIndex = NextSubInterfaceIndex;
+			(*SubInterfaces)[CurrentSubInterfaceIndex].Instances.Emplace(CurrentNumVertices, INDEX_NONE);
+		}
+
+		void FSubInterfaceDeclarationBuilder::PopSubInterfaceDeclaration()
+		{
+			if (ensureMsgf(SubInterfaces->IsValidIndex(CurrentSubInterfaceIndex), TEXT("Sub interface has incorrect Begin/End declaration")))
+			{
+				(*SubInterfaces)[CurrentSubInterfaceIndex].Instances.Last().End = CurrentNumVertices;
+				CurrentSubInterfaceIndex = INDEX_NONE;
+			}
+		}
+
+		FInputVertexInterfaceDeclarationBuilder::FInputVertexInterfaceDeclarationBuilder(TArray<FInputDataVertex>& OutVertices, TArray<FSubInterfaceLayout>& OutInstances)
+		: FSubInterfaceDeclarationBuilder(OutInstances)
+		, Vertices(OutVertices)
+		{
+		}
+
+		void FInputVertexInterfaceDeclarationBuilder::Add(FInputDataVertex&& InVertex)
+		{
+			using FEqualVertexName = VertexPrivate::TEqualVertexName<FInputDataVertex>;
+
+			checkf(!Vertices.ContainsByPredicate(FEqualVertexName(InVertex.VertexName)), TEXT("Duplicate vertex name %s. Vertex names must be unique"), *InVertex.VertexName.ToString());
+
+			Vertices.Emplace(MoveTemp(InVertex));
+			OnVertexAdded(Vertices.Last().VertexName);
+		}
+
+		FOutputVertexInterfaceDeclarationBuilder::FOutputVertexInterfaceDeclarationBuilder(TArray<FOutputDataVertex>& OutVertices, TArray<FSubInterfaceLayout>& OutInstances)
+		: FSubInterfaceDeclarationBuilder(OutInstances)
+		, Vertices(OutVertices)
+		{
+		}
+
+		void FOutputVertexInterfaceDeclarationBuilder::Add(FOutputDataVertex&& InVertex)
+		{
+			using FEqualVertexName = VertexPrivate::TEqualVertexName<FOutputDataVertex>;
+
+			checkf(!Vertices.ContainsByPredicate(FEqualVertexName(InVertex.VertexName)), TEXT("Duplicate vertex name %s. Vertex names must be unique"), *InVertex.VertexName.ToString());
+
+			Vertices.Emplace(MoveTemp(InVertex));
+			OnVertexAdded(Vertices.Last().VertexName);
+		}
+
+		FEnvironmentDeclarationBuilder::FEnvironmentDeclarationBuilder(TArray<FEnvironmentVertex>& OutVertices)
+		: Vertices(OutVertices)
+		{
+		}
+
+		/** Builds an interface from a sub-interface configuration */
+		class FSubInterfaceConfigurationBuilder
+		{
+		public:
+			FSubInterfaceConfigurationBuilder(TArray<FSubInterfaceLayout>& OutSubInterfaceLayouts)
+			: SubInterfaces(OutSubInterfaceLayouts)
+			{
+			}
+
+			virtual ~FSubInterfaceConfigurationBuilder() = default;
+
+
+			virtual void Build(TConstArrayView<FSubInterfaceDescription> InSubInterfaceDescriptions, TConstArrayView<FSubInterfaceConfiguration> InSubInterfaceConfigs)
+			{
+				/** Do nothing if there are no sub interfaces. */
+				if (SubInterfaces.IsEmpty())
+				{
+					return;
+				}
+
+
+				for (int32 LayoutIndex = 0; LayoutIndex < SubInterfaces.Num(); LayoutIndex++)
+				{
+					FSubInterfaceLayout& Layout = SubInterfaces[LayoutIndex];
+
+					const FSubInterfaceDescription* Description = Algo::FindBy(InSubInterfaceDescriptions, Layout.SubInterfaceName, [](const FSubInterfaceDescription& InDesc) { return InDesc.SubInterfaceName; });
+
+					if (ensureMsgf(Description, TEXT("Missing sub interface description %s"), *Layout.SubInterfaceName.ToString()))
+					{
+						// Determine number of sub interface instances. 
+						uint32 Num = Description->NumDefault;
+
+						const FSubInterfaceConfiguration* Config = Algo::FindBy(InSubInterfaceConfigs, Layout.SubInterfaceName, [](const FSubInterfaceConfiguration& InConfig) { return InConfig.SubInterfaceName; });
+						if (Config)
+						{
+							Num = FMath::Clamp(Config->Num, Description->Min, Description->Max);
+						}
+
+						// Construct sub interface instances
+						BuildSubInterfaceInstances(Layout, LayoutIndex, Num);
+					}
+				}
+			}
+
+
+		private:
+
+			void BuildSubInterfaceInstances(FSubInterfaceLayout& InLayout, uint32 InLayoutIndex, uint32 InNum)
+			{
+				// We should always be beginning from a declaration of a sub interface
+				// which enforces that there is only one instance of the sub interface. 
+				check(InLayout.Instances.Num() == 1);
+
+				if (InNum == 1)
+				{
+					// We already have 1 instance from the declaration. Nothing to do. 
+					return;
+				}
+
+				if (InNum == 0)
+				{
+					RemoveSubInterface(InLayout, InLayoutIndex);
+				}
+				else
+				{
+					SetNumSubInterfaces(InLayout, InLayoutIndex, InNum);
+				}
+			}
+
+			void SetNumSubInterfaces(FSubInterfaceLayout& InLayout, uint32 InLayoutIndex, uint32 InNum)
+			{
+				// Assume we are starting with a sub interface declaration. 
+				check(InLayout.Instances.Num() == 1);
+				
+				// Get the position of the prototype for the sub interface. 
+				const int32 ProtoBegin = InLayout.Instances[0].Begin;
+				const int32 ProtoEnd = InLayout.Instances[0].End;
+				const int32 ProtoNum = ProtoEnd - ProtoBegin;
+
+				// Shift the location of sub interfaces that occur after this one 
+				// to account for the about-to-be added vertices. 
+				int32 NumToAdd = (InNum - 1) * ProtoNum;
+				for (FSubInterfaceLayout& Layout : SubInterfaces)
+				{
+					for (FSubInterfaceLayout::FInstance& Instance : Layout.Instances)
+					{
+						if (Instance.Begin > ProtoBegin)
+						{
+							Instance.Begin += NumToAdd;
+							Instance.End += NumToAdd;
+						}
+					}
+				}
+
+				// Add instances
+				TArray<FSubInterfaceLayout::FInstance>& Instances = InLayout.Instances;
+				Instances.Reserve(InNum);
+				Instances.SetNum(InNum);
+
+				for (uint32 Cntr = 1; Cntr < InNum; Cntr++)
+				{
+					// Initialize the position of the sub interface instance
+					const uint32 InsertIndex = ProtoBegin + (Cntr * ProtoNum);
+					Instances[Cntr].Begin = InsertIndex;
+					Instances[Cntr].End = InsertIndex + ProtoNum;
+
+					// Create the vertices of the sub interface instance. 
+					InsertSubInterfaceInstance(ProtoBegin, ProtoNum, InsertIndex, Cntr);
+				}
+			}
+
+			void RemoveSubInterface(FSubInterfaceLayout& InLayout, uint32 InLayoutIndex)
+			{
+				// We should always be beginning from a declaration of a sub interface
+				// which enforces that there is only one instance of the sub interface. 
+				check(InLayout.Instances.Num() == 1);
+
+				// Get the position of the prototype for the sub interface. 
+				const int32 ProtoBegin = InLayout.Instances[0].Begin;
+				const int32 ProtoEnd = InLayout.Instances[0].End;
+				const int32 ProtoNum = ProtoEnd - ProtoBegin;
+
+				// Remove all instances from this subinterface
+				InLayout.Instances.Empty();
+
+				if (ProtoNum > 0)
+				{
+					// Remove actual vertices from interface
+					RemoveVerticesAt(ProtoBegin, ProtoNum);
+
+					// Shift positions of other instances of sub interfaces
+					for (FSubInterfaceLayout& Layout : SubInterfaces)
+					{
+						for (FSubInterfaceLayout::FInstance& Instance : Layout.Instances)
+						{
+							if (Instance.Begin > ProtoBegin)
+							{
+								Instance.Begin -= ProtoNum;
+								Instance.End -= ProtoNum;
+							}
+						}
+					}
+				}
+			}
+
+			virtual void InsertSubInterfaceInstance(uint32 InProtoBegin, uint32 InProtoNum, uint32 InInsertPos, uint32 InSubInterfaceInstanceIndex) = 0;
+			virtual void RemoveVerticesAt(uint32 InVertexIndexBegin, uint32 InNum) = 0;
+
+			TArray<FSubInterfaceLayout>& SubInterfaces;
+		};
+
+		/** TInterfaceConfigurationBuilder builds vertex interfaces from a sub interface
+		 * configuration.*/
+		template<typename VertexType>
+		class TInterfaceConfigurationBuilder : public FSubInterfaceConfigurationBuilder
+		{
+		public:
+			TInterfaceConfigurationBuilder(TArray<VertexType>& OutVertices, TArray<FSubInterfaceLayout>& OutSubInterfaceLayouts)
+			: FSubInterfaceConfigurationBuilder(OutSubInterfaceLayouts)
+			, Vertices(OutVertices)
+			{
+			}
+
+			virtual void Build(TConstArrayView<FSubInterfaceDescription> InSubInterfaceDescriptions, TConstArrayView<FSubInterfaceConfiguration> InSubInterfaceConfigs) override
+			{
+				// Build sub interfaces using parent class.
+				FSubInterfaceConfigurationBuilder::Build(InSubInterfaceDescriptions, InSubInterfaceConfigs);
+
+#if DO_CHECK
+				// Check that there are no duplicate names in the interface
+				for (int32 IndexA = 0; IndexA < Vertices.Num(); IndexA++)
+				{
+					const FVertexName& VertexNameA = Vertices[IndexA].VertexName;
+					for (int32 IndexB = IndexA + 1; IndexB < Vertices.Num(); IndexB++)
+					{
+						checkf(VertexNameA != Vertices[IndexB].VertexName, TEXT("Found duplicate names (%s) in interface"), *VertexNameA.ToString());
+					}
+				}
+#endif
+			}
+		private:
+
+			virtual void InsertSubInterfaceInstance(uint32 InProtoBegin, uint32 InProtoNum, uint32 InInsertPos, uint32 InSubInterfaceInstanceIndex) override
+			{
+				// check that the prototype exists in the vertex array.
+				check((InProtoBegin + InProtoNum) <= static_cast<uint32>(Vertices.Num()));
+				// check that the insert position of the sub interface instance is valid. 
+				check(InInsertPos <= static_cast<uint32>(Vertices.Num()));
+
+				if (InProtoNum > 0)
+				{
+					// Copy vertices from prototype vertices
+					Vertices.Insert(Vertices.GetData() + InProtoBegin, InProtoNum, InInsertPos);
+
+					// Rename vertices
+					const uint32 NewEnd = InInsertPos + InProtoNum;
+					for (uint32 VertexIndex = InInsertPos;  VertexIndex < NewEnd; VertexIndex++)
+					{
+						checkf(
+							Vertices[VertexIndex].VertexName.GetNumber() == NAME_NO_NUMBER_INTERNAL, 
+							TEXT("Prototype vertex %s in sub interface cannot have a trailing number because it is in a sub interface."), 
+							*Vertices[VertexIndex].VertexName.ToString()
+						);
+
+						Vertices[VertexIndex].VertexName.SetNumber(1 + InSubInterfaceInstanceIndex);
+					}
+				}
+			}
+
+			virtual void RemoveVerticesAt(uint32 InVertexIndexBegin, uint32 InNum) override
+			{
+				check((InVertexIndexBegin + InNum) <= static_cast<unsigned>(Vertices.Num()));
+				Vertices.RemoveAt(InVertexIndexBegin, InNum);
+			}
+
+			TArray<VertexType>& Vertices;
+		};
+
+		class FInputInterfaceConfigurationBuilder : public TInterfaceConfigurationBuilder<FInputDataVertex>
+		{
+		public:
+			FInputInterfaceConfigurationBuilder(FInputVertexInterface& InProto)
+			: TInterfaceConfigurationBuilder<FInputDataVertex>(InProto.Vertices, InProto.SubInterfaces)
+			{
+			}
+		};
+
+		class FOutputInterfaceConfigurationBuilder : public TInterfaceConfigurationBuilder<FOutputDataVertex>
+		{
+		public:
+			FOutputInterfaceConfigurationBuilder(FOutputVertexInterface& InProto)
+			: TInterfaceConfigurationBuilder<FOutputDataVertex>(InProto.Vertices, InProto.SubInterfaces)
+			{
+			}
+		};
+	}
+
 	bool operator==(const FInputDataVertex& InLHS, const FInputDataVertex& InRHS)
 	{
 		return (InLHS.VertexName == InRHS.VertexName) && (InLHS.DataTypeName == InRHS.DataTypeName);
@@ -65,6 +416,85 @@ namespace Metasound
 		return InLHS.VertexName.FastLess(InRHS.VertexName);
 	}
 
+	void FInputVertexInterface::ForEachSubInterfaceInstance(const FName& InSubInterfaceName, TFunctionRef<void (TConstArrayView<FInputDataVertex>)> Callable) const
+	{
+		using namespace VertexPrivate;
+
+		if (const FSubInterfaceLayout* Instances = FindSubInterfaceLayout(InSubInterfaceName))
+		{
+			for (const FSubInterfaceLayout::FInstance& Instance : Instances->Instances)
+			{
+				check(Vertices.IsValidIndex(Instance.Begin) && Vertices.IsValidIndex(Instance.End - 1));
+				Callable(MakeConstArrayView<FInputDataVertex>(Vertices.GetData() + Instance.Begin, Instance.End - Instance.Begin));
+			}
+		}
+	}
+
+	void FInputVertexInterface::ForEachSubInterfaceInstance(const FName& InSubInterfaceName, TFunctionRef<void (TArrayView<FInputDataVertex>)> Callable)
+	{
+		using namespace VertexPrivate;
+
+		if (const FSubInterfaceLayout* Instances = FindSubInterfaceLayout(InSubInterfaceName))
+		{
+			for (const FSubInterfaceLayout::FInstance& Instance : Instances->Instances)
+			{
+				check(Vertices.IsValidIndex(Instance.Begin) && Vertices.IsValidIndex(Instance.End - 1));
+				Callable(MakeArrayView<FInputDataVertex>(Vertices.GetData() + Instance.Begin, Instance.End - Instance.Begin));
+			}
+		}
+	}
+
+	TConstArrayView<VertexPrivate::FSubInterfaceLayout> FInputVertexInterface::GetSubInterfaces(const VertexPrivate::FPrivateAccessTag& InTag) const
+	{
+		return SubInterfaces;
+	}
+
+	const VertexPrivate::FSubInterfaceLayout* FInputVertexInterface::FindSubInterfaceLayout(const FName& InName) const
+	{
+		using namespace VertexPrivate;
+
+		return SubInterfaces.FindByPredicate(VertexPrivate::FEqualSubInterfaceName(InName));
+	}
+
+	void FOutputVertexInterface::ForEachSubInterfaceInstance(const FName& InSubInterfaceName, TFunctionRef<void (TConstArrayView<FOutputDataVertex>)> Callable) const
+	{
+		using namespace VertexPrivate;
+
+		if (const FSubInterfaceLayout* Instances = FindSubInterfaceLayout(InSubInterfaceName))
+		{
+			for (const FSubInterfaceLayout::FInstance& Instance : Instances->Instances)
+			{
+				check(Vertices.IsValidIndex(Instance.Begin) && Vertices.IsValidIndex(Instance.End - 1));
+				Callable(MakeConstArrayView<FOutputDataVertex>(Vertices.GetData() + Instance.Begin, Instance.End - Instance.Begin));
+			}
+		}
+	}
+
+	void FOutputVertexInterface::ForEachSubInterfaceInstance(const FName& InSubInterfaceName, TFunctionRef<void (TArrayView<FOutputDataVertex>)> Callable)
+	{
+		using namespace VertexPrivate;
+
+		if (const FSubInterfaceLayout* Instances = FindSubInterfaceLayout(InSubInterfaceName))
+		{
+			for (const FSubInterfaceLayout::FInstance& Instance : Instances->Instances)
+			{
+				check(Vertices.IsValidIndex(Instance.Begin) && Vertices.IsValidIndex(Instance.End - 1));
+				Callable(MakeArrayView<FOutputDataVertex>(Vertices.GetData() + Instance.Begin, Instance.End - Instance.Begin));
+			}
+		}
+	}
+
+	TConstArrayView<VertexPrivate::FSubInterfaceLayout> FOutputVertexInterface::GetSubInterfaces(const VertexPrivate::FPrivateAccessTag& InTag) const
+	{
+		return SubInterfaces;
+	}
+
+	const VertexPrivate::FSubInterfaceLayout* FOutputVertexInterface::FindSubInterfaceLayout(const FName& InName) const
+	{
+		return SubInterfaces.FindByPredicate(VertexPrivate::FEqualSubInterfaceName(InName));
+	}
+
+
 	FVertexInterface::FVertexInterface(const FInputVertexInterface& InInputs, const FOutputVertexInterface& InOutputs)
 	:	InputInterface(InInputs)
 	,	OutputInterface(InOutputs)
@@ -151,6 +581,38 @@ namespace Metasound
 	{
 		return !(InLHS == InRHS);
 	}
+
+	FClassVertexInterface::FClassVertexInterface(FVertexInterface InInterface)
+	: FClassVertexInterface(MoveTemp(InInterface.GetInputInterface()), MoveTemp(InInterface.GetOutputInterface()), MoveTemp(InInterface.GetEnvironmentInterface()))
+	{
+	}
+
+	FClassVertexInterface::FClassVertexInterface(FInputVertexInterface InInputs, FOutputVertexInterface InOutputs, FEnvironmentVertexInterface InEnvironment)
+	: FClassVertexInterface(TArray<FSubInterfaceDescription>{}, MoveTemp(InInputs), MoveTemp(InOutputs), MoveTemp(InEnvironment))
+	{
+	}
+
+	FClassVertexInterface::FClassVertexInterface(TArray<FSubInterfaceDescription> InSubInterfaceDescriptions, FInputVertexInterface InInputs, FOutputVertexInterface InOutputs, FEnvironmentVertexInterface InEnvironment)
+	: SubInterfaces(MoveTemp(InSubInterfaceDescriptions))
+	, Inputs(MoveTemp(InInputs))
+	, Outputs(MoveTemp(InOutputs))
+	, Environment(MoveTemp(InEnvironment))
+	{
+
+	}
+
+	FVertexInterface FClassVertexInterface::CreateVertexInterface(TConstArrayView<FSubInterfaceConfiguration> InSubInterfaceConfigurations) const
+	{
+		using namespace VertexPrivate;
+
+		FInputVertexInterface NewInputs(Inputs);
+		FInputInterfaceConfigurationBuilder(NewInputs).Build(SubInterfaces, InSubInterfaceConfigurations);
+		
+		FOutputVertexInterface NewOutputs(Outputs);
+		FOutputInterfaceConfigurationBuilder(NewOutputs).Build(SubInterfaces, InSubInterfaceConfigurations);
+
+		return FVertexInterface(MoveTemp(NewInputs), MoveTemp(NewOutputs), Environment);
+	}
 }
 
 FString LexToString(Metasound::EVertexAccessType InAccessType)
diff --git a/Engine/Plugins/Runtime/Metasound/Source/MetasoundGraphCore/Private/MetasoundVertexData.cpp b/Engine/Plugins/Runtime/Metasound/Source/MetasoundGraphCore/Private/MetasoundVertexData.cpp
--- a/Engine/Plugins/Runtime/Metasound/Source/MetasoundGraphCore/Private/MetasoundVertexData.cpp
+++ b/Engine/Plugins/Runtime/Metasound/Source/MetasoundGraphCore/Private/MetasoundVertexData.cpp
@@ -2,13 +2,20 @@
 
 #include "MetasoundVertexData.h"
 
+#include "Containers/Array.h"
+#include "Containers/ArrayView.h"
 #include "MetasoundDataReference.h"
 #include "MetasoundLog.h"
 #include "MetasoundThreadLocalDebug.h"
+#include "MetasoundVertex.h"
+#include "MetasoundVertexPrivate.h"
+#include "Templates/Function.h"
+
+#include <type_traits>
 
 namespace Metasound
 {
-	namespace MetasoundVertexDataPrivate
+	namespace VertexDataPrivate
 	{
 		template<typename VertexType>
 		using TBindingType = std::conditional_t<std::is_same_v<VertexType, FInputDataVertex>, FInputBinding, FOutputBinding>;
@@ -224,8 +231,9 @@ namespace Metasound
 				CheckAccessTypeCompatibility(Vertex, *Data);
 			}
 		}
+
 		template<typename VertexType>
-		void EmplaceBindings(TArray<TBindingType<VertexType>>& InArray, const TVertexInterfaceGroup<VertexType>& InVertexInterface)
+		void EmplaceBindings(TArray<TBindingType<VertexType>>& InArray, const TVertexInterfaceImpl<VertexType>& InVertexInterface)
 		{
 			for (const VertexType& DataVertex : InVertexInterface)
 			{
@@ -414,6 +422,53 @@ namespace Metasound
 				}
 			}
 		}
+
+		// Implementation for iterating over all sub interfaces on a FVertexInterfaceData
+		template<typename BindingType>
+		void ForEachSubInterfaceBindingImpl(const VertexPrivate::FSubInterfaceLayout& InSubInterfaceLayout, const FVertexName& InVertexName, uint32 InExpectedNumSubInterfaceInstances, TFunctionRef<void (uint32, BindingType&)> InApplyFunc, TArrayView<BindingType> InOutBindings)
+		{
+			using namespace VertexPrivate;
+
+			// In some scenarios, we need to determine if the number of vertices match the expected num. For instance if binding an array of data references,
+			// we need to check that the array size matches the number of sub interfaces.  
+			if ((InExpectedNumSubInterfaceInstances != INDEX_NONE) && (InSubInterfaceLayout.Instances.Num() != InExpectedNumSubInterfaceInstances))
+			{
+				UE_LOG(LogMetaSound, Verbose, TEXT("Number of instances of sub interface '%s' is %d, expected %d. This may result in unbound vertex data."), *InSubInterfaceLayout.SubInterfaceName.ToString(), InSubInterfaceLayout.Instances.Num(), InExpectedNumSubInterfaceInstances);
+			}
+
+			if (InSubInterfaceLayout.Instances.Num() > 0)
+			{
+				// The first sub interface will contain the original vertex name. 
+				// Find the offset of the original vertex name so it can be located
+				// in subsequent sub interface instances. 
+				const FSubInterfaceLayout::FInstance& FirstInstance = InSubInterfaceLayout.Instances[0];
+				TArrayView<BindingType> FirstSubInterface = InOutBindings.Slice(FirstInstance.Begin, FirstInstance.End - FirstInstance.Begin);
+				const uint32 Offset = FirstSubInterface.IndexOfByPredicate([&InVertexName](const BindingType& Binding) { return Binding.GetVertex().VertexName == InVertexName; });
+
+				if (Offset != INDEX_NONE)
+				{
+					const TArray<FSubInterfaceLayout::FInstance>& Instances = InSubInterfaceLayout.Instances;
+
+					uint32 MaxSubInterfaceIndex = static_cast<uint32>(Instances.Num());
+					if (INDEX_NONE != InExpectedNumSubInterfaceInstances)
+					{
+						MaxSubInterfaceIndex = FMath::Min(MaxSubInterfaceIndex, InExpectedNumSubInterfaceInstances);
+					}
+
+					for (uint32 SubInterfaceIndex = 0; SubInterfaceIndex < MaxSubInterfaceIndex; SubInterfaceIndex++)
+					{
+						const FSubInterfaceLayout::FInstance& Instance = Instances[SubInterfaceIndex];
+						uint32 Pos = Instance.Begin + Offset;
+						check(Pos < static_cast<uint32>(InOutBindings.Num()));
+						InApplyFunc(SubInterfaceIndex, InOutBindings[Pos]);
+					}
+				}
+				else
+				{
+					UE_LOG(LogMetaSound, Warning, TEXT("Failed to find vertex with name '%s' in sub interface with name '%s'."), *InVertexName.ToString(), *InSubInterfaceLayout.SubInterfaceName.ToString());
+				}
+			}
+		}
 	}
 
 	FInputVertexInterfaceData::FInputVertexInterfaceData()
@@ -423,8 +478,9 @@ namespace Metasound
 
 	FInputVertexInterfaceData::FInputVertexInterfaceData(const FInputVertexInterface& InVertexInterface)
 	: bIsVertexInterfaceFrozen(true)
+	, SubInterfaces(InVertexInterface.GetSubInterfaces(VertexPrivate::FPrivateAccessTag{})) 
 	{
-		MetasoundVertexDataPrivate::EmplaceBindings(Bindings, InVertexInterface);
+		VertexDataPrivate::EmplaceBindings(Bindings, InVertexInterface);
 	}
 
 	bool FInputVertexInterfaceData::IsVertexInterfaceFrozen() const
@@ -435,6 +491,7 @@ namespace Metasound
 	void FInputVertexInterfaceData::SetIsVertexInterfaceFrozen(bool bInFreezeVertices)
 	{
 		bIsVertexInterfaceFrozen = bInFreezeVertices;
+		checkf(bIsVertexInterfaceFrozen || SubInterfaces.Num() == 0, TEXT("Interfaces containing sub interfaces cannot be unfrozen"));
 	}
 
 	void FInputVertexInterfaceData::BindVertex(const FVertexName& InVertexName, const FAnyDataReference& InOutDataReference)
@@ -444,7 +501,8 @@ namespace Metasound
 
 	void FInputVertexInterfaceData::BindVertex(const FVertexName& InVertexName, FAnyDataReference& InOutDataReference)
 	{
-		using namespace MetasoundVertexDataPrivate;
+		using namespace VertexDataPrivate;
+
 		auto CreateBinding = [&]()
 		{
 			FInputDataVertex Vertex(InVertexName, InOutDataReference.GetDataTypeName(), FDataVertexMetadata{}, DataReferenceAccessTypeToVertexAccessType(InOutDataReference.GetAccessType()));
@@ -457,6 +515,17 @@ namespace Metasound
 		Apply(InVertexName, CreateBinding, BindData);
 	}
 
+	uint32 FInputVertexInterfaceData::GetNumSubInterfaceInstances(const FName& InSubInterfaceName) const
+	{
+		using namespace VertexPrivate;
+
+		if (const FSubInterfaceLayout* Layout = FindSubInterface(InSubInterfaceName))
+		{
+			return Layout->Instances.Num();
+		}
+		return 0;
+	}
+
 	void FInputVertexInterfaceData::Bind(const FInputVertexInterfaceData& InVertexData)
 	{
 		Bind(const_cast<FInputVertexInterfaceData&>(InVertexData));
@@ -464,32 +533,32 @@ namespace Metasound
 
 	void FInputVertexInterfaceData::Bind(FInputVertexInterfaceData& InVertexData)
 	{
-		MetasoundVertexDataPrivate::Bind(IsVertexInterfaceFrozen(), Bindings, InVertexData.Bindings);
+		VertexDataPrivate::Bind(IsVertexInterfaceFrozen(), Bindings, InVertexData.Bindings);
 	}
 
 	void FInputVertexInterfaceData::SetVertex(const FVertexName& InVertexName, FAnyDataReference&& InDataReference)
 	{
-		MetasoundVertexDataPrivate::SetVertex(IsVertexInterfaceFrozen(), Bindings, InVertexName, MoveTemp(InDataReference));
+		VertexDataPrivate::SetVertex(IsVertexInterfaceFrozen(), Bindings, InVertexName, MoveTemp(InDataReference));
 	}
 
 	void FInputVertexInterfaceData::SetVertex(const FVertexName& InVertexName, const FAnyDataReference& InDataReference)
 	{
-		MetasoundVertexDataPrivate::SetVertex(IsVertexInterfaceFrozen(), Bindings, InVertexName, InDataReference);
+		VertexDataPrivate::SetVertex(IsVertexInterfaceFrozen(), Bindings, InVertexName, InDataReference);
 	}
 
 	void FInputVertexInterfaceData::Set(const FDataReferenceCollection& InCollection)
 	{
-		MetasoundVertexDataPrivate::Set(Bindings, InCollection);
+		VertexDataPrivate::Set(Bindings, InCollection);
 	}
 
 	FDataReferenceCollection FInputVertexInterfaceData::ToDataReferenceCollection() const
 	{
-		return MetasoundVertexDataPrivate::ToDataReferenceCollection(Bindings);
+		return VertexDataPrivate::ToDataReferenceCollection(Bindings);
 	}
 
 	bool FInputVertexInterfaceData::IsVertexBound(const FVertexName& InVertexName) const
 	{
-		return MetasoundVertexDataPrivate::IsVertexBound(Bindings, InVertexName);
+		return VertexDataPrivate::IsVertexBound(Bindings, InVertexName);
 	}
 
 	const FInputDataVertex& FInputVertexInterfaceData::GetVertex(const FVertexName& InVertexName) const
@@ -523,17 +592,30 @@ namespace Metasound
 
 	EDataReferenceAccessType FInputVertexInterfaceData::GetVertexDataAccessType(const FVertexName& InVertexName) const
 	{
-		return MetasoundVertexDataPrivate::GetVertexDataAccessType(Bindings, InVertexName);
+		return VertexDataPrivate::GetVertexDataAccessType(Bindings, InVertexName);
 	}
 
 	bool FInputVertexInterfaceData::AreAllVerticesBound() const
 	{
-		return MetasoundVertexDataPrivate::AreAllVerticesBound(Bindings);
+		return VertexDataPrivate::AreAllVerticesBound(Bindings);
 	}
 
 	const FAnyDataReference* FInputVertexInterfaceData::FindDataReference(const FVertexName& InVertexName) const
 	{
-		return MetasoundVertexDataPrivate::FindDataReference(Bindings, InVertexName);
+		return VertexDataPrivate::FindDataReference(Bindings, InVertexName);
+	}
+
+	TArray<const FAnyDataReference*> FInputVertexInterfaceData::FindSubInterfaceReferences(const FName& InSubInterfaceName, const FVertexName& InVertexName) const
+	{
+		TArray<const FAnyDataReference*> References;
+		auto GetRef = [&References](uint32 InIndex, const FInputBinding& InBinding)
+		{
+			check(References.Num() == InIndex); // This check will be false if the references array 
+												// is not being built in the correct order. 
+			References.Add(InBinding.GetDataReference());
+		};
+		ForEachSubInterfaceBinding(InSubInterfaceName, InVertexName, INDEX_NONE, GetRef);
+		return References;
 	}
 
 	void FInputVertexInterfaceData::Apply(const FVertexName& InVertexName, TFunctionRef<FInputBinding ()> InCreateFunc, TFunctionRef<void (FInputBinding&)> InBindFunc)
@@ -554,14 +636,44 @@ namespace Metasound
 		}
 	}
 
+	void FInputVertexInterfaceData::ForEachSubInterfaceBinding(const FName& InSubInterfaceName, const FVertexName& InVertexName, uint32 InExpectedNumSubInterfaceInstances, TFunctionRef<void (uint32, FInputBinding&)> InApplyFunc)
+	{
+		using namespace VertexPrivate;
+		using namespace VertexDataPrivate;
+
+		if (const FSubInterfaceLayout* Layout = FindSubInterface(InSubInterfaceName))
+		{
+			ForEachSubInterfaceBindingImpl<FInputBinding>(*Layout, InVertexName, InExpectedNumSubInterfaceInstances, InApplyFunc, Bindings);
+		}
+		else
+		{
+			UE_LOG(LogMetaSound, Warning, TEXT("Failed to find sub interface with name '%s'."), *InSubInterfaceName.ToString());
+		}
+	}
+
+	void FInputVertexInterfaceData::ForEachSubInterfaceBinding(const FName& InSubInterfaceName, const FVertexName& InVertexName, uint32 InExpectedNumSubInterfaceInstances, TFunctionRef<void (uint32, const FInputBinding&)> InApplyFunc) const
+	{
+		using namespace VertexPrivate;
+		using namespace VertexDataPrivate;
+
+		if (const FSubInterfaceLayout* Layout = FindSubInterface(InSubInterfaceName))
+		{
+			ForEachSubInterfaceBindingImpl<const FInputBinding>(*Layout, InVertexName, InExpectedNumSubInterfaceInstances, InApplyFunc, Bindings);
+		}
+		else
+		{
+			UE_LOG(LogMetaSound, Warning, TEXT("Failed to find sub interface with name '%s'."), *InSubInterfaceName.ToString());
+		}
+	}
+
 	FInputVertexInterfaceData::FInputBinding* FInputVertexInterfaceData::Find(const FVertexName& InVertexName)
 	{
-		return MetasoundVertexDataPrivate::Find(Bindings, InVertexName);
+		return VertexDataPrivate::Find(Bindings, InVertexName);
 	}
 
 	const FInputVertexInterfaceData::FInputBinding* FInputVertexInterfaceData::Find(const FVertexName& InVertexName) const
 	{
-		return MetasoundVertexDataPrivate::Find(Bindings, InVertexName);
+		return VertexDataPrivate::Find(Bindings, InVertexName);
 	}
 
 	FInputVertexInterfaceData::FInputBinding* FInputVertexInterfaceData::FindChecked(const FVertexName& InVertexName)
@@ -578,6 +690,11 @@ namespace Metasound
 		return Binding;
 	}
 
+	const VertexPrivate::FSubInterfaceLayout* FInputVertexInterfaceData::FindSubInterface(const FName& InSubInterfaceName) const
+	{
+		return SubInterfaces.FindByPredicate([&InSubInterfaceName](const VertexPrivate::FSubInterfaceLayout& InLayout) { return InLayout.SubInterfaceName == InSubInterfaceName; });
+	}
+
 	FOutputVertexInterfaceData::FOutputVertexInterfaceData()
 	: bIsVertexInterfaceFrozen(false)
 	{
@@ -585,8 +702,9 @@ namespace Metasound
 
 	FOutputVertexInterfaceData::FOutputVertexInterfaceData(const FOutputVertexInterface& InVertexInterface)
 	: bIsVertexInterfaceFrozen(true)
+	, SubInterfaces(InVertexInterface.GetSubInterfaces(VertexPrivate::FPrivateAccessTag{})) 
 	{
-		MetasoundVertexDataPrivate::EmplaceBindings(Bindings, InVertexInterface);
+		VertexDataPrivate::EmplaceBindings(Bindings, InVertexInterface);
 	}
 
 	bool FOutputVertexInterfaceData::IsVertexInterfaceFrozen() const
@@ -597,6 +715,7 @@ namespace Metasound
 	void FOutputVertexInterfaceData::SetIsVertexInterfaceFrozen(bool bInFreezeVertices)
 	{
 		bIsVertexInterfaceFrozen = bInFreezeVertices;
+		checkf(bIsVertexInterfaceFrozen || SubInterfaces.Num() == 0, TEXT("Interfaces containing sub interfaces cannot be unfrozen"));
 	}
 
 	void FOutputVertexInterfaceData::BindVertex(const FVertexName& InVertexName, const FAnyDataReference& InOutDataReference)
@@ -606,7 +725,7 @@ namespace Metasound
 
 	void FOutputVertexInterfaceData::BindVertex(const FVertexName& InVertexName, FAnyDataReference& InOutDataReference)
 	{
-		using namespace MetasoundVertexDataPrivate;
+		using namespace VertexDataPrivate;
 		auto CreateBinding = [&]()
 		{
 			FOutputDataVertex Vertex(InVertexName, InOutDataReference.GetDataTypeName(), FDataVertexMetadata{}, DataReferenceAccessTypeToVertexAccessType(InOutDataReference.GetAccessType()));
@@ -619,6 +738,17 @@ namespace Metasound
 		Apply(InVertexName, CreateBinding, BindData);
 	}
 
+	uint32 FOutputVertexInterfaceData::GetNumSubInterfaceInstances(const FName& InSubInterfaceName) const
+	{
+		using namespace VertexPrivate;
+
+		if (const FSubInterfaceLayout* Layout = FindSubInterface(InSubInterfaceName))
+		{
+			return Layout->Instances.Num();
+		}
+		return 0;
+	}
+
 	void FOutputVertexInterfaceData::Bind(const FOutputVertexInterfaceData& InVertexData)
 	{
 		Bind(const_cast<FOutputVertexInterfaceData&>(InVertexData));
@@ -626,32 +756,32 @@ namespace Metasound
 
 	void FOutputVertexInterfaceData::Bind(FOutputVertexInterfaceData& InVertexData)
 	{
-		MetasoundVertexDataPrivate::Bind(IsVertexInterfaceFrozen(), Bindings, InVertexData.Bindings);
+		VertexDataPrivate::Bind(IsVertexInterfaceFrozen(), Bindings, InVertexData.Bindings);
 	}
 
 	void FOutputVertexInterfaceData::SetVertex(const FVertexName& InVertexName, FAnyDataReference&& InDataReference)
 	{
-		MetasoundVertexDataPrivate::SetVertex(IsVertexInterfaceFrozen(), Bindings, InVertexName, MoveTemp(InDataReference));
+		VertexDataPrivate::SetVertex(IsVertexInterfaceFrozen(), Bindings, InVertexName, MoveTemp(InDataReference));
 	}
 
 	void FOutputVertexInterfaceData::SetVertex(const FVertexName& InVertexName, const FAnyDataReference& InDataReference)
 	{
-		MetasoundVertexDataPrivate::SetVertex(IsVertexInterfaceFrozen(), Bindings, InVertexName, InDataReference);
+		VertexDataPrivate::SetVertex(IsVertexInterfaceFrozen(), Bindings, InVertexName, InDataReference);
 	}
 
 	void FOutputVertexInterfaceData::Set(const FDataReferenceCollection& InCollection)
 	{
-		MetasoundVertexDataPrivate::Set(Bindings, InCollection);
+		VertexDataPrivate::Set(Bindings, InCollection);
 	}
 
 	FDataReferenceCollection FOutputVertexInterfaceData::ToDataReferenceCollection() const
 	{
-		return MetasoundVertexDataPrivate::ToDataReferenceCollection(Bindings);
+		return VertexDataPrivate::ToDataReferenceCollection(Bindings);
 	}
 
 	bool FOutputVertexInterfaceData::IsVertexBound(const FVertexName& InVertexName) const
 	{
-		return MetasoundVertexDataPrivate::IsVertexBound(Bindings, InVertexName);
+		return VertexDataPrivate::IsVertexBound(Bindings, InVertexName);
 	}
 
 	const FOutputDataVertex& FOutputVertexInterfaceData::GetVertex(const FVertexName& InVertexName) const
@@ -685,17 +815,30 @@ namespace Metasound
 
 	EDataReferenceAccessType FOutputVertexInterfaceData::GetVertexDataAccessType(const FVertexName& InVertexName) const
 	{
-		return MetasoundVertexDataPrivate::GetVertexDataAccessType(Bindings, InVertexName);
+		return VertexDataPrivate::GetVertexDataAccessType(Bindings, InVertexName);
 	}
 
 	bool FOutputVertexInterfaceData::AreAllVerticesBound() const
 	{
-		return MetasoundVertexDataPrivate::AreAllVerticesBound(Bindings);
+		return VertexDataPrivate::AreAllVerticesBound(Bindings);
 	}
 
 	const FAnyDataReference* FOutputVertexInterfaceData::FindDataReference(const FVertexName& InVertexName) const
 	{
-		return MetasoundVertexDataPrivate::FindDataReference(Bindings, InVertexName);
+		return VertexDataPrivate::FindDataReference(Bindings, InVertexName);
+	}
+
+	TArray<const FAnyDataReference*> FOutputVertexInterfaceData::FindSubInterfaceReferences(const FName& InSubInterfaceName, const FVertexName& InVertexName) const
+	{
+		TArray<const FAnyDataReference*> References;
+		auto GetRef = [&References](uint32 InIndex, const FOutputBinding& InBinding)
+		{
+			check(References.Num() == InIndex); // This check will be false if the references array 
+												// is not being built in the correct order. 
+			References.Add(InBinding.GetDataReference());
+		};
+		ForEachSubInterfaceBinding(InSubInterfaceName, InVertexName, INDEX_NONE, GetRef);
+		return References;
 	}
 
 	void FOutputVertexInterfaceData::Apply(const FVertexName& InVertexName, TFunctionRef<FOutputBinding ()> InCreateFunc, TFunctionRef<void (FOutputBinding&)> InBindFunc)
@@ -716,14 +859,42 @@ namespace Metasound
 		}
 	}
 
+	void FOutputVertexInterfaceData::ForEachSubInterfaceBinding(const FName& InSubInterfaceName, const FVertexName& InVertexName, uint32 InExpectedNumSubInterfaceInstances, TFunctionRef<void (uint32, FOutputBinding&)> InApplyFunc)
+	{
+		using namespace VertexPrivate;
+
+		if (const FSubInterfaceLayout* Layout = FindSubInterface(InSubInterfaceName))
+		{
+			ForEachSubInterfaceBindingImpl<FOutputBinding>(*Layout, InVertexName, InExpectedNumSubInterfaceInstances, InApplyFunc, Bindings);
+		}
+		else
+		{
+			UE_LOG(LogMetaSound, Warning, TEXT("Failed to find sub interface with name '%s'."), *InSubInterfaceName.ToString());
+		}
+	}
+
+	void FOutputVertexInterfaceData::ForEachSubInterfaceBinding(const FName& InSubInterfaceName, const FVertexName& InVertexName, uint32 InExpectedNumSubInterfaceInstances, TFunctionRef<void (uint32, const FOutputBinding&)> InApplyFunc) const
+	{
+		using namespace VertexPrivate;
+
+		if (const FSubInterfaceLayout* Layout = FindSubInterface(InSubInterfaceName))
+		{
+			ForEachSubInterfaceBindingImpl<const FOutputBinding>(*Layout, InVertexName, InExpectedNumSubInterfaceInstances, InApplyFunc, Bindings);
+		}
+		else
+		{
+			UE_LOG(LogMetaSound, Warning, TEXT("Failed to find sub interface with name '%s'."), *InSubInterfaceName.ToString());
+		}
+	}
+
 	FOutputVertexInterfaceData::FOutputBinding* FOutputVertexInterfaceData::Find(const FVertexName& InVertexName)
 	{
-		return MetasoundVertexDataPrivate::Find(Bindings, InVertexName);
+		return VertexDataPrivate::Find(Bindings, InVertexName);
 	}
 
 	const FOutputVertexInterfaceData::FOutputBinding* FOutputVertexInterfaceData::Find(const FVertexName& InVertexName) const
 	{
-		return MetasoundVertexDataPrivate::Find(Bindings, InVertexName);
+		return VertexDataPrivate::Find(Bindings, InVertexName);
 	}
 
 	FOutputVertexInterfaceData::FOutputBinding* FOutputVertexInterfaceData::FindChecked(const FVertexName& InVertexName)
@@ -740,6 +911,11 @@ namespace Metasound
 		return Binding;
 	}
 
+	const VertexPrivate::FSubInterfaceLayout* FOutputVertexInterfaceData::FindSubInterface(const FName& InSubInterfaceName) const
+	{
+		return SubInterfaces.FindByPredicate([&InSubInterfaceName](const VertexPrivate::FSubInterfaceLayout& InLayout) { return InLayout.SubInterfaceName == InSubInterfaceName; });
+	}
+
 	FVertexInterfaceData::FVertexInterfaceData(const FVertexInterface& InVertexInterface)
 	: InputVertexInterfaceData(InVertexInterface.GetInputInterface())
 	, OutputVertexInterfaceData(InVertexInterface.GetOutputInterface())
@@ -780,21 +956,21 @@ namespace Metasound
 
 	void GetVertexInterfaceDataState(const FInputVertexInterfaceData& InVertexInterface, TArray<FVertexDataState>& OutState)
 	{
-		MetasoundVertexDataPrivate::GetVertexInterfaceDataState(InVertexInterface.Bindings, OutState);
+		VertexDataPrivate::GetVertexInterfaceDataState(InVertexInterface.Bindings, OutState);
 	}
 
 	void GetVertexInterfaceDataState(const FOutputVertexInterfaceData& InVertexInterface, TArray<FVertexDataState>& OutState)
 	{
-		MetasoundVertexDataPrivate::GetVertexInterfaceDataState(InVertexInterface.Bindings, OutState);
+		VertexDataPrivate::GetVertexInterfaceDataState(InVertexInterface.Bindings, OutState);
 	}
 
 	void CompareVertexInterfaceDataToPriorState(const FInputVertexInterfaceData& InCurrentInterface, const TArray<FVertexDataState>& InPriorState, TSortedVertexNameMap<FAnyDataReference>& OutUpdates)
 	{
-		MetasoundVertexDataPrivate::CompareVertexInterfaceDataToPriorState(InCurrentInterface.Bindings, InPriorState, OutUpdates);
+		VertexDataPrivate::CompareVertexInterfaceDataToPriorState(InCurrentInterface.Bindings, InPriorState, OutUpdates);
 	}
 
 	void CompareVertexInterfaceDataToPriorState(const FOutputVertexInterfaceData& InCurrentInterface, const TArray<FVertexDataState>& InPriorState, TSortedVertexNameMap<FAnyDataReference>& OutUpdates)
 	{
-		MetasoundVertexDataPrivate::CompareVertexInterfaceDataToPriorState(InCurrentInterface.Bindings, InPriorState, OutUpdates);
+		VertexDataPrivate::CompareVertexInterfaceDataToPriorState(InCurrentInterface.Bindings, InPriorState, OutUpdates);
 	}
 }
diff --git a/Engine/Plugins/Runtime/Metasound/Source/MetasoundGraphCore/Private/MetasoundVertexPrivate.h b/Engine/Plugins/Runtime/Metasound/Source/MetasoundGraphCore/Private/MetasoundVertexPrivate.h
new file mode 100644
--- /dev/null
+++ b/Engine/Plugins/Runtime/Metasound/Source/MetasoundGraphCore/Private/MetasoundVertexPrivate.h
@@ -0,0 +1,12 @@
+// Copyright Epic Games, Inc. All Rights Reserved.
+
+#pragma once
+
+namespace Metasound::VertexPrivate
+{
+	// Used to control API access when a method must be public but we do not want
+	// it to be used outside of this module. 
+	struct FPrivateAccessTag
+	{
+	};
+}
diff --git a/Engine/Plugins/Runtime/Metasound/Source/MetasoundGraphCore/Public/MetasoundDataReference.h b/Engine/Plugins/Runtime/Metasound/Source/MetasoundGraphCore/Public/MetasoundDataReference.h
--- a/Engine/Plugins/Runtime/Metasound/Source/MetasoundGraphCore/Public/MetasoundDataReference.h
+++ b/Engine/Plugins/Runtime/Metasound/Source/MetasoundGraphCore/Public/MetasoundDataReference.h
@@ -541,6 +541,123 @@ namespace Metasound
 	 */
 	class FAnyDataReference : public IDataReference
 	{
+		/** TGetAs allows for template specialization of data retrieval and conversion
+		 * rules for the various flavors of data referencing. The default implementation
+		 * assumes we are attempting to retrieve the value stored in the data reference. 
+		 */
+		template<typename ResultType>
+		struct TGetAs
+		{
+		private:
+			using DataType = typename std::decay_t<std::remove_pointer_t<ResultType>>;
+
+			// Only deref a pointer if the returned type is specifically requested to be
+			// a non-pointer. 
+			static ResultType DerefPointerIfNonPointerResultType(const DataType* InData)
+			{
+				if constexpr (std::is_pointer_v<ResultType>)
+				{
+					return InData;
+				}
+				else
+				{
+					check(InData != nullptr);
+					return *InData;
+				}
+			};
+
+		public:
+		
+			// Return a the underlying data by value or by const pointer depending
+			// upon the template parameter ResultType
+			static ResultType Get(EDataReferenceAccessType InAccessType, const IDataReference& InRef)
+			{
+				check(IsDataReferenceOfType<DataType>(InRef));
+
+				const DataType* DataPtr = nullptr;
+				switch (InAccessType)
+				{
+					case EDataReferenceAccessType::Read:
+					{
+						DataPtr = static_cast<const TDataReadReference<DataType>*>(&InRef)->Get();
+						break;
+					}
+					case EDataReferenceAccessType::Write:
+					{
+						DataPtr = static_cast<const TDataWriteReference<DataType>*>(&InRef)->Get();
+						break;
+					}
+					case EDataReferenceAccessType::Value:
+					{
+						DataPtr = static_cast<const TDataValueReference<DataType>*>(&InRef)->Get();
+						break;
+					}
+					default:
+					{
+						checkNoEntry();
+						DataPtr = static_cast<const DataType*>(InRef.GetRaw());
+					}
+				}
+
+				return DerefPointerIfNonPointerResultType(DataPtr);
+			}
+		};
+
+		/* Template specialization for retrieving data reference as a TDataReadReference<DataType> */
+		template<typename DataType>
+		struct TGetAs<TDataReadReference<DataType>>
+		{
+			static TDataReadReference<DataType> Get(EDataReferenceAccessType InAccessType, const IDataReference& InRef)
+			{
+				check(IsDataReferenceOfType<DataType>(InRef));
+				switch (InAccessType)
+				{
+					case EDataReferenceAccessType::Read:
+					{
+						return *static_cast<const TDataReadReference<DataType>*>(&InRef);
+					}
+					case EDataReferenceAccessType::Write:
+					{
+						return *static_cast<const TDataWriteReference<DataType>*>(&InRef);
+					}
+					case EDataReferenceAccessType::Value:
+					{
+						return *static_cast<const TDataValueReference<DataType>*>(&InRef);
+					}
+					default:
+					{
+						checkNoEntry();
+						return *static_cast<const TDataReadReference<DataType>*>(&InRef);
+					}
+				}
+			}
+		};
+
+		/* Template specialization for retrieving data reference as a TDataValueReference<DataType> */
+		template<typename DataType>
+		struct TGetAs<TDataValueReference<DataType>>
+		{
+			static TDataValueReference<DataType> Get(EDataReferenceAccessType InAccessType, const IDataReference& InRef)
+			{
+				check(IsDataReferenceOfType<DataType>(InRef));
+				checkf(EDataReferenceAccessType::Value== InAccessType, TEXT("Invalid attempt to convert a data ref with \"%s\" access to \"%s\" access"), *LexToString(InAccessType), *LexToString(EDataReferenceAccessType::Value));
+				return *static_cast<const TDataValueReference<DataType>*>(&InRef);
+			}
+		};
+
+		/* Template specialization for retrieving data reference as a TDataWriteReference<DataType> */
+		template<typename DataType>
+		struct TGetAs<TDataWriteReference<DataType>>
+		{
+			static TDataWriteReference<DataType> Get(EDataReferenceAccessType InAccessType, const IDataReference& InRef)
+			{
+				checkf(IsDataReferenceOfType<DataType>(InRef), TEXT("Attempt to get data reference with underlying type \"%s\" when actual underlying type is \"%s\"."), *GetMetasoundDataTypeString<DataType>(), *InRef.GetDataTypeName().ToString());
+				checkf(EDataReferenceAccessType::Write == InAccessType, TEXT("Invalid attempt to convert a data ref with \"%s\" access to \"%s\" access"), *LexToString(InAccessType), *LexToString(EDataReferenceAccessType::Write));
+				return *static_cast<const TDataWriteReference<DataType>*>(&InRef);
+			}
+		};
+
+
 		// Private constructor. 
 		FAnyDataReference(EDataReferenceAccessType InAccessType, const IDataReference& InDataRef)
 		: AccessType(InAccessType)
@@ -624,6 +741,17 @@ namespace Metasound
 			return DataRefPtr->GetRaw();
 		}
 
+		/** Return the data in the specified format. The template argument
+		 * can be any of the support TData*Reference<> types, the underlying DataType
+		 * or a const DataType* pointer. 
+		 */
+		template<typename ReturnType>
+		ReturnType GetAs() const
+		{
+			check(DataRefPtr.IsValid());
+			return TGetAs<ReturnType>::Get(AccessType, *DataRefPtr);
+		}
+
 		/** Returns the current value of a reference. 
 		 *
 		 * This method's behavior is undefined and an assert will be called if
@@ -634,26 +762,7 @@ namespace Metasound
 		{
 			if (DataRefPtr.IsValid())
 			{
-				check(IsDataReferenceOfType<DataType>(*DataRefPtr));
-				switch (AccessType)
-				{
-					case EDataReferenceAccessType::Read:
-					{
-						return static_cast<const TDataReadReference<DataType>*>(DataRefPtr.Get())->Get();
-					}
-					case EDataReferenceAccessType::Write:
-					{
-						return static_cast<const TDataWriteReference<DataType>*>(DataRefPtr.Get())->Get();
-					}
-					case EDataReferenceAccessType::Value:
-					{
-						return static_cast<const TDataValueReference<DataType>*>(DataRefPtr.Get())->Get();
-					}
-					default:
-					{
-						checkNoEntry();
-					}
-				}
+				return GetAs<const DataType*>();
 			}
 			return nullptr;
 		}
@@ -692,30 +801,7 @@ namespace Metasound
 		TDataValueReference<DataType> GetDataValueReference() const
 		{
 			check(DataRefPtr.IsValid());
-			check(IsDataReferenceOfType<DataType>(*DataRefPtr));
-
-			switch (AccessType)
-			{
-				case EDataReferenceAccessType::Value:
-				{
-					return *static_cast<const TDataValueReference<DataType>*>(DataRefPtr.Get());
-				}
-				case EDataReferenceAccessType::Write:
-				{
-					checkf(false, TEXT("Retrieving a value reference from a write reference is not supported."));
-					return ValueCast<DataType>(*static_cast<const TDataWriteReference<DataType>*>(DataRefPtr.Get()));
-				}
-				case EDataReferenceAccessType::Read:
-				{
-					checkf(false, TEXT("Retrieving a value reference from a read reference is not supported."));
-					return ValueCast<DataType>(*static_cast<const TDataReadReference<DataType>*>(DataRefPtr.Get()));
-				}
-				default:
-				{
-					checkf(false, TEXT("Unknown data reference access type."));
-					return *static_cast<const TDataValueReference<DataType>*>(DataRefPtr.Get());
-				}
-			}
+			return GetAs<TDataValueReference<DataType>>();
 		}
 
 		/** Get access to a TDataReadReference. 
@@ -731,29 +817,7 @@ namespace Metasound
 		template<typename DataType>
 		TDataReadReference<DataType> GetDataReadReference() const
 		{
-			check(DataRefPtr.IsValid());
-			check(IsDataReferenceOfType<DataType>(*DataRefPtr));
-
-			switch (AccessType)
-			{
-				case EDataReferenceAccessType::Read:
-				{
-					return *static_cast<const TDataReadReference<DataType>*>(DataRefPtr.Get());
-				}
-				case EDataReferenceAccessType::Write:
-				{
-					return *static_cast<const TDataWriteReference<DataType>*>(DataRefPtr.Get());
-				}
-				case EDataReferenceAccessType::Value:
-				{
-					return *static_cast<const TDataValueReference<DataType>*>(DataRefPtr.Get());
-				}
-				default:
-				{
-					checkNoEntry();
-					return *static_cast<const TDataReadReference<DataType>*>(DataRefPtr.Get());
-				}
-			}
+			return GetAs<TDataReadReference<DataType>>();
 		}
 
 		/** Get access to a TDataWriteReference. 
@@ -769,13 +833,7 @@ namespace Metasound
 		template<typename DataType>
 		TDataWriteReference<DataType> GetDataWriteReference() const
 		{
-			checkf(DataRefPtr.IsValid(), TEXT("FAnyDataReference is in an invalid state"));
-
-			checkf(IsDataReferenceOfType<DataType>(*DataRefPtr), TEXT("Attempt to get data reference with underlying type \"%s\" when actual underlying type is \"%s\"."), *GetMetasoundDataTypeString<DataType>(), *GetDataTypeName().ToString());
-
-			checkf(EDataReferenceAccessType::Write == AccessType, TEXT("Invalid attempt to convert a data ref with \"%s\" access to \"%s\" access"), *LexToString(GetAccessType()), *LexToString(EDataReferenceAccessType::Write));
-
-			return *static_cast<const TDataWriteReference<DataType>*>(DataRefPtr.Get());
+			return GetAs<TDataWriteReference<DataType>>();
 		}
 
 	private:
diff --git a/Engine/Plugins/Runtime/Metasound/Source/MetasoundGraphCore/Public/MetasoundVertex.h b/Engine/Plugins/Runtime/Metasound/Source/MetasoundGraphCore/Public/MetasoundVertex.h
--- a/Engine/Plugins/Runtime/Metasound/Source/MetasoundGraphCore/Public/MetasoundVertex.h
+++ b/Engine/Plugins/Runtime/Metasound/Source/MetasoundGraphCore/Public/MetasoundVertex.h
@@ -15,6 +15,58 @@ namespace Metasound
 	/** Name of a given vertex.  Only unique for a given node interface. */
 	using FVertexName = FName;
 
+	namespace VertexPrivate
+	{
+		// Forward declare for friendship
+		class FInputInterfaceConfigurationBuilder;
+		class FOutputInterfaceConfigurationBuilder;
+
+		struct FPrivateAccessTag;
+
+		/** Contains a list of sub-interface spans where the span indices
+		 * refer to vertex indices of an array containing vertices. */
+		struct FSubInterfaceLayout
+		{
+			struct FInstance
+			{
+				int32 Begin = 0;
+				int32 End = 0; //< Inclusive. 
+			};
+
+			FName SubInterfaceName;
+			TArray<FInstance> Instances;
+		};
+
+		/** Functor for finding vertices with equal names */
+		template<typename VertexType>
+		struct TEqualVertexName
+		{
+			const FVertexName& NameRef;
+
+			TEqualVertexName(const FVertexName& InName)
+			: NameRef(InName)
+			{
+			}
+
+			FORCEINLINE bool operator()(const VertexType& InOther) const
+			{
+				return InOther.VertexName == NameRef;
+			}
+		};
+
+		/** constexpr count of the number of a instances of type T in a parameter pack */
+		template<typename T>
+		struct TNumOfTypeInPack
+		{
+			template<typename ... Ts>
+			static constexpr uint32 Get()
+			{
+				constexpr uint32 Num = ((std::is_base_of_v<T, Ts> ? 1 : 0) + ...);
+				return Num;
+			}
+		};
+	}
+
 	// Vertex metadata
 	struct FDataVertexMetadata
 	{
@@ -195,6 +247,7 @@ namespace Metasound
 		}
 	};
 
+
 	/** FEnvironmentVertex
 	 *
 	 * A vertex for environment variables. 
@@ -230,48 +283,13 @@ namespace Metasound
 		friend bool METASOUNDGRAPHCORE_API operator<(const FEnvironmentVertex& InLHS, const FEnvironmentVertex& InRHS);
 	};
 
-	/** TVertexInterfaceGroups encapsulates multiple related data vertices. It 
+	/** TVertexInterfaceImpls encapsulates multiple related data vertices. It 
 	 * requires that each vertex in the group have a unique FVertexName.
 	 */
 	template<typename VertexType>
-	class TVertexInterfaceGroup
+	class TVertexInterfaceImpl
 	{
-		// Determine if the vertex type is a supported vertex type.
-		template<typename T>
-		struct TIsSupportedVertexType
-		{
-		private:
-			static constexpr bool bIsDerivedFromFInputDataVertex = std::is_base_of<FInputDataVertex, T>::value;
-			static constexpr bool bIsDerivedFromFOutputDataVertex = std::is_base_of<FOutputDataVertex, T>::value;
-			static constexpr bool bIsDerivedFromFEnvironmentVertex = std::is_base_of<FEnvironmentVertex, T>::value;
-			static constexpr bool bIsSupportedVertexType = bIsDerivedFromFInputDataVertex || bIsDerivedFromFOutputDataVertex || bIsDerivedFromFEnvironmentVertex;
-		public:
-			static constexpr bool Value = bIsSupportedVertexType;
-		};
-
-		static_assert(TIsSupportedVertexType<VertexType>::Value, "VertexType must be derived from FInputDataVertex, FOutputDataVertex, or FEnvironmentVertex");
-		
-		using FContainerType = TArray<VertexType>;
-
-		struct FEqualVertexNamePredicate
-		{
-			const FVertexName& NameRef;
-
-			FEqualVertexNamePredicate(const FVertexName& InName)
-			: NameRef(InName)
-			{
-			}
-
-			FEqualVertexNamePredicate(const VertexType& InVertex)
-			: NameRef(InVertex.VertexName)
-			{
-			}
-
-			bool operator()(const VertexType& InOther) const
-			{
-				return InOther.VertexName == NameRef;
-			}
-		};
+		using FEqualVertexName = VertexPrivate::TEqualVertexName<VertexType>;
 
 		void AddOrUpdateVertex(VertexType&& InVertex)
 		{
@@ -284,36 +302,18 @@ namespace Metasound
 				Vertices.Add(MoveTemp(InVertex));
 			}
 		}
-
-		// Required for end of recursion.
-		void CopyInputs()
-		{
-		}
-
-		// Recursive call used to unpack a template parameter pack. This stores
-		// each object in the template parameter pack into the groups internal
-		// storage.
-		//
-		// Assume that InputType is either a sublcass of FDataVertex
-		// and that VertexType is either FInputDataVertex or FOutputDataVertex
-		template<typename CurrentVertexType, typename... RemainingVertexTypes>
-		void CopyInputs(CurrentVertexType&& InInput, RemainingVertexTypes&&... InRemainingInputs)
-		{
-			// Create vertex out of vertex model
-			AddOrUpdateVertex(MoveTemp(InInput));
-			CopyInputs(InRemainingInputs...);
-		}
-
-
 	public:
 
-		using RangedForConstIteratorType = typename FContainerType::RangedForConstIteratorType;
+		using RangedForConstIteratorType = typename TArray<VertexType>::RangedForConstIteratorType;
+
+		TVertexInterfaceImpl() = default;
 		
-		/** TVertexInterfaceGroup constructor with variadic list of vertex
+		/** TVertexInterfaceImpl constructor with variadic list of vertex
 		 * models.
 		 */
 		template<typename... VertexTypes>
-		explicit TVertexInterfaceGroup(VertexTypes&&... InVertices)
+		UE_DEPRECATED(5.6, "Use the constructors of derived classes instead (FInputVertexInterface, FOutputVertexInterface, etc).")
+		explicit TVertexInterfaceImpl(VertexTypes&&... InVertices)
 		{
 			static_assert(
 				(std::is_constructible_v<VertexType, VertexTypes&&> && ...),
@@ -323,7 +323,8 @@ namespace Metasound
 			// over allocation.
 			Vertices.Reserve(sizeof...(VertexTypes));
 
-			CopyInputs(Forward<VertexTypes>(InVertices)...);
+			// Unfold parameter pack and add t
+			(AddOrUpdateVertex(Forward<VertexTypes>(InVertices)), ...);
 		}
 
 		/** Add a vertex to the group. */
@@ -340,26 +341,26 @@ namespace Metasound
 		/** Remove a vertex by key. */
 		bool Remove(const FVertexName& InKey)
 		{
-			int32 NumRemoved = Vertices.RemoveAll(FEqualVertexNamePredicate(InKey));
+			int32 NumRemoved = Vertices.RemoveAll(FEqualVertexName(InKey));
 			return (NumRemoved > 0);
 		}
 
 		/** Returns true if the group contains a vertex with a matching key. */
 		bool Contains(const FVertexName& InKey) const
 		{
-			return Vertices.ContainsByPredicate(FEqualVertexNamePredicate(InKey));
+			return Vertices.ContainsByPredicate(FEqualVertexName(InKey));
 		}
 
 		/** Find a vertex with a given VertexName */
 		VertexType* Find(const FVertexName& InKey)
 		{
-			return Vertices.FindByPredicate(FEqualVertexNamePredicate(InKey));
+			return Vertices.FindByPredicate(FEqualVertexName(InKey));
 		}
 
 		/** Find a vertex with a given VertexName */
 		const VertexType* Find(const FVertexName& InKey) const
 		{
-			return Vertices.FindByPredicate(FEqualVertexNamePredicate(InKey));
+			return Vertices.FindByPredicate(FEqualVertexName(InKey));
 		}
 
 		/** Return the sort order index of a vertex with the given name.
@@ -370,7 +371,7 @@ namespace Metasound
 		 */
 		int32 GetSortOrderIndex(const FVertexName& InKey) const
 		{
-			return Vertices.IndexOfByPredicate(FEqualVertexNamePredicate(InKey));
+			return Vertices.IndexOfByPredicate(FEqualVertexName(InKey));
 		}
 
 		/** Return the vertex for a given vertex key. */
@@ -400,31 +401,296 @@ namespace Metasound
 		}
 
 		/** Compare whether two vertex groups are equal. */
-		friend bool operator==(const TVertexInterfaceGroup<VertexType>& InLHS, const TVertexInterfaceGroup<VertexType>& InRHS)
+		friend bool operator==(const TVertexInterfaceImpl<VertexType>& InLHS, const TVertexInterfaceImpl<VertexType>& InRHS)
 		{
 			return InLHS.Vertices == InRHS.Vertices;
 		}
 
 		/** Compare whether two vertex groups are unequal. */
-		friend bool operator!=(const TVertexInterfaceGroup<VertexType>& InLHS, const TVertexInterfaceGroup<VertexType>& InRHS)
+		friend bool operator!=(const TVertexInterfaceImpl<VertexType>& InLHS, const TVertexInterfaceImpl<VertexType>& InRHS)
 		{
 			return !(InLHS == InRHS);
 		}
 
+	protected:
+
+		TArray<VertexType> Vertices;
+	};
+
+	template<typename VertexType>
+	using TVertexInterfaceGroup UE_DEPRECATED(5.6, "Use TVertexInterfaceImpl instead") = TVertexInterfaceImpl<VertexType>;
+
+	/** Declare the beginning of a sub interface in a vertex interface declaration. */
+	struct FBeginSubInterface
+	{
+		FName Name; //< Name of the sub interface. 
+	};
+
+	/** Declare the ending of a sub interface in a vertex interface declaration. */
+	struct FEndSubInterface
+	{
+	};
+
+	namespace VertexPrivate
+	{
+		/** Interface Declaration Builders create vertex interfaces from a template
+		 * parameter pack. This allows developers to express their node interfaces 
+		 * as declarations as opposed to requiring them to sequentially construct
+		 * their interfaces.  
+		 *
+		 * For example, we can do:
+		 * 
+		 *  FInputVertexInterface Interface
+		 *  {
+		 *  	TInputDataVertex<float>(...),
+		 *  	FBeginSubInterface(...),
+		 *  		TInputDataVertex<int32>(...),
+		 *  	FEndSubInterface(...),
+		 *  	TInputDataVertex<float>(...)
+		 *  };
+		 *
+		 * As opposed to:
+		 *  FInputVertexInterface Interface;
+		 * 
+		 * 	Interface.Add(TInputDataVertex<float>(...));
+		 * 	Interface.Add(FBeginSubInterface(...));
+	     * 	Interface.Add(TInputDataVertex<int32>(...));
+		 * 	Interface.Add(FEndSubInterface(...));
+		 * 	Interface.Add(TInputDataVertex<float>(...);
+		 *
+		 * 
+		 * An essential element of the builders is that they use member function 
+		 * overloading of Add(...) to handle the variety of different objects which
+		 * might be used to describe an interface. 
+		 *
+		 * The builders also size memory allocations exactly to remove any addition
+		 * slack. In general, there can be many interfaces in memory and minimizing
+		 * their memory footprint is important. 
+		 */
+
+		/** Base class for an interface builder which supports sub interfaces. */
+		class METASOUNDGRAPHCORE_API FSubInterfaceDeclarationBuilder
+		{
+		public:
+			FSubInterfaceDeclarationBuilder(TArray<FSubInterfaceLayout>& OutSubInterfaceLayouts);
+			virtual ~FSubInterfaceDeclarationBuilder();
+
+			// Handler for the beginning of a sub interface. 
+			void Add(FBeginSubInterface&& InSubInterface);
+
+			// Handler for the ending of a sub interface. 
+			void Add(FEndSubInterface&& InSubInterface);
+
+			// Allocate memory for a known number of sub interfaces.
+			void ReserveSubInterfaces(int32 Num);
+
+		protected:
+			// Called when a vertex is added to adjust internal counters. 
+			void OnVertexAdded(const FVertexName& InVertexName);
+
+		private:
+
+			void PushSubInterfaceDeclaration(const FName& InName);
+			void PopSubInterfaceDeclaration();
+
+			uint32 CurrentNumVertices = 0;
+			TArray<FSubInterfaceLayout>* SubInterfaces;
+			int32 CurrentSubInterfaceIndex = INDEX_NONE;
+		};
+
+		/** Interface builder for FInputVertexInterface declarations. */
+		class METASOUNDGRAPHCORE_API FInputVertexInterfaceDeclarationBuilder : public FSubInterfaceDeclarationBuilder
+		{
+		public:
+			FInputVertexInterfaceDeclarationBuilder(TArray<FInputDataVertex>& OutVertices, TArray<FSubInterfaceLayout>& OutInstances);
+			template<typename... ArgTypes>
+			void Build(ArgTypes&&... InArgs)
+			{
+				// Size vertices exactly to avoid wasted memory in array allocations
+				constexpr uint32 NumVertices = TNumOfTypeInPack<FInputDataVertex>::Get<ArgTypes...>();
+				if constexpr (NumVertices > 0)
+				{
+					Vertices.Reserve(NumVertices);
+				}
+
+				// Size sub interfaces exactly to avoid wasted memory in array allocations
+				constexpr uint32 NumSubInterfaces = TNumOfTypeInPack<FBeginSubInterface>::Get<ArgTypes...>();
+				if constexpr (NumSubInterfaces > 0)
+				{
+					ReserveSubInterfaces(NumSubInterfaces);
+				}
+
+				// Add all the elements of the vertex interface. This uses a fold expression
+				// to call Add(...) on each input. The various overloads of the Add(...) method
+				// then assemble the appropriate structures.
+				(Add(Forward<ArgTypes>(InArgs)), ...);
+			}
+
+		private:
+
+			using FSubInterfaceDeclarationBuilder::Add; // Unhide Add(...) function overloads in FSubInterfaceDeclarationBuilder
+			void Add(FInputDataVertex&& InVertex);
+
+			TArray<FInputDataVertex>& Vertices;
+		};
+
+		/** Interface builder for FOutputVertexInterface declarations. */
+		class METASOUNDGRAPHCORE_API FOutputVertexInterfaceDeclarationBuilder : public FSubInterfaceDeclarationBuilder
+		{
+		public:
+			FOutputVertexInterfaceDeclarationBuilder(TArray<FOutputDataVertex>& OutVertices, TArray<FSubInterfaceLayout>& OutInstances);
+
+			template<typename... ArgTypes>
+			void Build(ArgTypes&&... InArgs)
+			{
+				// Size vertices exactly to avoid wasted memory in array allocations
+				constexpr uint32 NumVertices = TNumOfTypeInPack<FOutputDataVertex>::Get<ArgTypes...>();
+				if constexpr (NumVertices > 0)
+				{
+					Vertices.Reserve(NumVertices);
+				}
+
+				// Size sub interfaces exactly to avoid wasted memory in array allocations
+				constexpr uint32 NumSubInterfaces = TNumOfTypeInPack<FBeginSubInterface>::Get<ArgTypes...>();
+				if constexpr (NumSubInterfaces > 0)
+				{
+					ReserveSubInterfaces(NumSubInterfaces);
+				}
+
+				// Add all the elements of the vertex interface. This uses a fold expression
+				// to call Add(...) on each input. The various overloads of the Add(...) method
+				// then assemble the appropriate structures.
+				(Add(Forward<ArgTypes>(InArgs)), ...);
+			}
+
+		private:
+			using FSubInterfaceDeclarationBuilder::Add; // Unhide Add(...) function overloads in FSubInterfaceDeclarationBuilder
+			void Add(FOutputDataVertex&& InVertex);
+
+			TArray<FOutputDataVertex>& Vertices;
+		};
+
+		/** Interface builder for FEnvironmentInterface declarations. */
+		class METASOUNDGRAPHCORE_API FEnvironmentDeclarationBuilder 
+		{
+		public:
+			FEnvironmentDeclarationBuilder(TArray<FEnvironmentVertex>& OutVertices);
+
+			template<typename... ArgTypes>
+			void Build(ArgTypes&&... InArgs)
+			{
+				// Size vertices exactly to avoid wasted memory in array allocations
+				constexpr uint32 NumVertices = TNumOfTypeInPack<FEnvironmentVertex>::Get<ArgTypes...>();
+				if constexpr (NumVertices > 0)
+				{
+					Vertices.Reserve(NumVertices);
+				}
+
+				// Add all the elements of the vertex interface. 
+				(Vertices.Emplace(MoveTemp(InArgs)), ...);
+			}
+		private:
+
+			TArray<FEnvironmentVertex>& Vertices;
+		};
+	}
+
+	/** Interface representing the inputs of a node. */
+	class METASOUNDGRAPHCORE_API FInputVertexInterface : public TVertexInterfaceImpl<FInputDataVertex>
+	{
+	public:
+		FInputVertexInterface() = default;
+
+		/** Construct an FInputVertexInterface from a parameter pack. This allows 
+		 * node interfaces to be declared in a single function.
+		 *
+		 * For example:
+		 * 
+		 *  FInputVertexInterface Interface
+		 *  {
+		 *  	TInputDataVertex<float>(...),
+		 *  	FBeginSubInterface(...),
+		 *  		TInputDataVertex<int32>(...),
+		 *  	FEndSubInterface(...),
+		 *  	TInputDataVertex<float>(...)
+		 *  };
+		 */
+		template<typename... ArgTypes>
+		explicit FInputVertexInterface(ArgTypes && ... InArgs)
+		{
+			VertexPrivate::FInputVertexInterfaceDeclarationBuilder InterfaceBuilder(Vertices, SubInterfaces);
+			InterfaceBuilder.Build(Forward<ArgTypes>(InArgs)...);
+		}
+
+		/** Iterate through all repetitions of a sub interface. */
+		void ForEachSubInterfaceInstance(const FName& InSubInterfaceName, TFunctionRef<void (TConstArrayView<FInputDataVertex>)> Callable) const;
+		/** Iterate through all repetitions of a sub interface. */
+		void ForEachSubInterfaceInstance(const FName& InSubInterfaceName, TFunctionRef<void (TArrayView<FInputDataVertex>)> Callable);
+		TConstArrayView<VertexPrivate::FSubInterfaceLayout> GetSubInterfaces(const VertexPrivate::FPrivateAccessTag& InTag) const;
+
 	private:
+		friend class VertexPrivate::FInputInterfaceConfigurationBuilder;
+		const VertexPrivate::FSubInterfaceLayout* FindSubInterfaceLayout(const FName& InName) const;
 
-		FContainerType Vertices;
+		TArray<VertexPrivate::FSubInterfaceLayout> SubInterfaces;
 	};
 
-	
-	/** FInputVertexInterface is a TVertexInterfaceGroup which holds FInputDataVertexes. */
-	typedef TVertexInterfaceGroup<FInputDataVertex> FInputVertexInterface;
+	/** Interface representing the outputs of a node. */
+	class METASOUNDGRAPHCORE_API FOutputVertexInterface : public TVertexInterfaceImpl<FOutputDataVertex>
+	{
+	public:
+		FOutputVertexInterface() = default;
+
+		/** Construct an FOutputVertexInterface from a parameter pack. This allows 
+		 * node interfaces to be declared in a single function.
+		 *
+		 * For example:
+		 * 
+		 *  FOutputVertexInterface Interface
+		 *  {
+		 *  	TOutputDataVertex<float>(...),
+		 *  	FBeginSubInterface(...),
+		 *  		TOutputDataVertex<int32>(...),
+		 *  	FEndSubInterface(...),
+		 *  	TOutputDataVertex<float>(...)
+		 *  };
+		 */
+		template<typename... ArgTypes>
+		explicit FOutputVertexInterface(ArgTypes && ... InArgs)
+		{
+			VertexPrivate::FOutputVertexInterfaceDeclarationBuilder InterfaceBuilder(Vertices, SubInterfaces);
+			InterfaceBuilder.Build(Forward<ArgTypes>(InArgs)...);
+		}
+
+		/** Iterate through all repetitions of a sub interface. */
+		void ForEachSubInterfaceInstance(const FName& InSubInterfaceName, TFunctionRef<void (TConstArrayView<FOutputDataVertex>)> Callable) const;
+		/** Iterate through all repetitions of a sub interface. */
+		void ForEachSubInterfaceInstance(const FName& InSubInterfaceName, TFunctionRef<void (TArrayView<FOutputDataVertex>)> Callable);
+
+		TConstArrayView<VertexPrivate::FSubInterfaceLayout> GetSubInterfaces(const VertexPrivate::FPrivateAccessTag& InTag) const;
+
+	private:
+
+		friend class VertexPrivate::FOutputInterfaceConfigurationBuilder;
+
+		const VertexPrivate::FSubInterfaceLayout* FindSubInterfaceLayout(const FName& InName) const;
+
+		TArray<VertexPrivate::FSubInterfaceLayout> SubInterfaces;
+	};
 
-	/** FOutputVertexInterface is a TVertexInterfaceGroup which holds FOutputDataVertexes. */
-	typedef TVertexInterfaceGroup<FOutputDataVertex> FOutputVertexInterface;
+	/** Interface representing the environment variables used by a node. */
+	class METASOUNDGRAPHCORE_API FEnvironmentVertexInterface : public TVertexInterfaceImpl<FEnvironmentVertex>
+	{
+	public:
+		FEnvironmentVertexInterface() = default;
 
-	/** FEnvironmentVertexInterface is a TVertexInterfaceGroup which holds FEnvironmentVertexes. */
-	typedef TVertexInterfaceGroup<FEnvironmentVertex> FEnvironmentVertexInterface;
+		template<typename... ArgTypes>
+		explicit FEnvironmentVertexInterface(ArgTypes && ... InArgs)
+		{
+			VertexPrivate::FEnvironmentDeclarationBuilder InterfaceBuilder(Vertices);
+			InterfaceBuilder.Build(Forward<ArgTypes>(InArgs)...);
+		}
+	};
 
 	/** FVertexInterface provides access to a collection of input and output vertex
 	 * interfaces. 
@@ -491,6 +757,44 @@ namespace Metasound
 			FEnvironmentVertexInterface EnvironmentInterface;
 	};
 
+	/** A description of a sub interface which is used when declaring a FClassVertexInterface */
+	struct FSubInterfaceDescription
+	{
+		FName SubInterfaceName;
+		uint32 Min; 		//< Minimum number of instances of the sub interface. 
+		uint32 Max; 		//< Maximum number of instances of the sub interface. 
+		uint32 NumDefault; 	//< Default number of instances of the sub interface if desired num is unspecified. 
+	};
+
+	/** A sub interface configuration is used to create a FVertexInterface from 
+	 * a FClassVertexInterface. */
+	struct FSubInterfaceConfiguration
+	{
+		FName SubInterfaceName;
+		uint32 Num;
+	};
+
+	/** FClassVertexInterface describes the interface of a node class. It is an 
+	 * immutable factory for FVertexInterfaces. */
+	class METASOUNDGRAPHCORE_API FClassVertexInterface
+	{
+	public:
+		FClassVertexInterface() = default;
+		FClassVertexInterface(FVertexInterface InInterface);
+		FClassVertexInterface(FInputVertexInterface InInputs, FOutputVertexInterface InOutputs, FEnvironmentVertexInterface InEnvironment={});
+		FClassVertexInterface(TArray<FSubInterfaceDescription> InSubInterfactDescriptions, FInputVertexInterface InInputs, FOutputVertexInterface InOutputs, FEnvironmentVertexInterface InEnvironment={});
+
+		/** Create a FVertexInterface with the given configuration. */
+		FVertexInterface CreateVertexInterface(TConstArrayView<FSubInterfaceConfiguration> InSubInterfaceConfigurations) const;
+
+	private:
+
+		TArray<FSubInterfaceDescription> SubInterfaces;
+		FInputVertexInterface Inputs;
+		FOutputVertexInterface Outputs;
+		FEnvironmentVertexInterface Environment;
+	};
+
 	/**
 	 * This struct is used to pass in any arguments required for constructing a single node instance.
 	 * because of this, all FNode implementations have to implement a constructor that takes an FNodeInitData instance.
diff --git a/Engine/Plugins/Runtime/Metasound/Source/MetasoundGraphCore/Public/MetasoundVertexData.h b/Engine/Plugins/Runtime/Metasound/Source/MetasoundGraphCore/Public/MetasoundVertexData.h
--- a/Engine/Plugins/Runtime/Metasound/Source/MetasoundGraphCore/Public/MetasoundVertexData.h
+++ b/Engine/Plugins/Runtime/Metasound/Source/MetasoundGraphCore/Public/MetasoundVertexData.h
@@ -11,6 +11,7 @@
 #include "MetasoundLog.h"
 #include "MetasoundOperatorSettings.h"
 #include "MetasoundVertex.h"
+#include "Traits/ElementType.h"
 
 #include <type_traits>
 
@@ -20,7 +21,7 @@
 
 namespace Metasound
 {
-	namespace MetasoundVertexDataPrivate
+	namespace VertexDataPrivate
 	{
 		// Tests to see that the access type of the data reference is compatible 
 		// with the access type of the vertex.
@@ -55,8 +56,86 @@ namespace Metasound
 		{
 			static constexpr EVertexAccessType VertexAccess = EVertexAccessType::Value;
 		};
+
+		/** TGetReferencDataTypeT is a template trait which determines the underlying
+		 * type stored in a data reference
+		 */
+		template<typename ...>
+		struct TGetReferenceDataTypeT
+		{
+			static_assert("Invalid template parameter. This must be instantiated with a TDataReadReference<>, TDataWriteReference or TDataValueReference");
+		};
+
+		template<typename DataType>
+		struct TGetReferenceDataTypeT<TDataReadReference<DataType>>
+		{
+			using Type = DataType;
+		};
+
+		template<typename DataType>
+		struct TGetReferenceDataTypeT<TDataWriteReference<DataType>>
+		{
+			using Type = DataType;
+		};
+
+		template<typename DataType>
+		struct TGetReferenceDataTypeT<TDataValueReference<DataType>>
+		{
+			using Type = DataType;
+		};
 		
+		/** TGetReferencDataType is the underlying type stored in a data reference */
+		template<typename T>
+		using TGetReferenceDataType = typename TGetReferenceDataTypeT<T>::Type;
+
+		/** TFactoryFacace is a factory for data types which uses the class template
+		 * type to determine which flavor of DataReference to create. 
+		 */
+		template<typename Type>
+		struct TFactoryFacade
+		{
+			static Type CreateWithLiteral(const FOperatorSettings& InSettings, const FLiteral& InLiteral)
+			{
+				return TDataTypeLiteralFactory<Type>::CreateExplicitArgs(InSettings, InLiteral);
+			}
+
+			template<typename ... ArgTypes>
+			static Type CreateWithArgs(ArgTypes&&... Args)
+			{
+				return Type(Forward<ArgTypes>(Args)...);
+			}
+		};
+
+		template<typename DataType>
+		struct TFactoryFacade<TDataReadReference<DataType>>
+		{
+			static TDataReadReference<DataType> CreateWithLiteral(const FOperatorSettings& InSettings, const FLiteral& InLiteral)
+			{
+				return TDataReadReferenceLiteralFactory<DataType>::CreateExplicitArgs(InSettings, InLiteral);
+			}
+
+			template<typename ... ArgTypes>
+			static TDataReadReference<DataType> CreateWithArgs(ArgTypes&&... Args)
+			{
+				return TDataReadReference<DataType>::CreateNew(Forward<ArgTypes>(Args)...);
+			}
+		};
 
+		template<typename DataType>
+		struct TFactoryFacade<TDataWriteReference<DataType>>
+		{
+			static TDataWriteReference<DataType> CreateWithLiteral(const FOperatorSettings& InSettings, const FLiteral& InLiteral)
+			{
+				return TDataWriteReferenceLiteralFactory<DataType>::CreateExplicitArgs(InSettings, InLiteral);
+			}
+
+			template<typename ... ArgTypes>
+			static TDataWriteReference<DataType> CreateWithArgs(ArgTypes&&... Args)
+			{
+				return TDataWriteReference<DataType>::CreateNew(Forward<ArgTypes>(Args)...);
+			}
+		};
+		
 		// An input binding which connects an FInputVertex to a IDataReference
 		class METASOUNDGRAPHCORE_API FInputBinding
 		{
@@ -66,29 +145,16 @@ namespace Metasound
 			FInputBinding(const FInputDataVertex& InVertex);
 			FInputBinding(const FVertexName& InVertexName, FAnyDataReference&& InReference);
 
-			template<typename DataType>
-			void BindRead(TDataReadReference<DataType>& InOutDataReference)
-			{
-				if (Data.IsSet())
-				{
-					InOutDataReference = Data->GetDataReadReference<DataType>();
-				}
-				else
-				{
-					Set(InOutDataReference);
-				}
-			}
-
-			template<typename DataType>
-			void BindRead(TDataWriteReference<DataType>& InOutDataReference)
+			template<typename DataReferenceType>
+			void BindRead(DataReferenceType& InOutDataReference)
 			{
 				if (Data.IsSet())
 				{
-					InOutDataReference = Data->GetDataWriteReference<DataType>();
+					InOutDataReference = Data->GetAs<DataReferenceType>();
 				}
 				else
 				{
-					Set(TDataReadReference<DataType>(InOutDataReference));
+					Set(TDataReadReference<TGetReferenceDataType<DataReferenceType>>(InOutDataReference));
 				}
 			}
 
@@ -130,107 +196,48 @@ namespace Metasound
 
 			FDataReferenceID GetDataReferenceID() const;
 
-			// Get data read reference assuming data is bound and read or write accessible.
-			template<typename DataType>
-			TDataReadReference<DataType> GetDataReadReference() const
+			// Get the bound data reference. This assumes the data reference exists
+			// and is convertible to the AsType.
+			template<typename AsType>
+			AsType GetDataReferenceAs() const
 			{
 				check(Data.IsSet());
-				return Data->GetDataReadReference<DataType>();
+				return Data->GetAs<AsType>();
 			}
 
-			// Get the bound data read reference if it exists. Otherwise create and 
-			// return a data read reference by constructing one using the Vertex's 
+			// Get the bound data reference if it exists. Otherwise create and 
+			// return a data reference by constructing one using the Vertex's 
 			// default literal.
-			template< typename DataType > 
-			TDataReadReference<DataType> GetOrCreateDefaultDataReadReference(const FOperatorSettings& InSettings) const
+			template<typename AsType>
+			AsType GetOrCreateDefaultAs(const FOperatorSettings& InSettings) const
 			{
 				if (Data.IsSet())
 				{
-					return Data->GetDataReadReference<DataType>();
+					return GetDataReferenceAs<AsType>();
 				}
 				else
 				{
-					return TDataReadReferenceLiteralFactory<DataType>::CreateExplicitArgs(InSettings, Vertex.GetDefaultLiteral());
+					return TFactoryFacade<AsType>::CreateWithLiteral(InSettings, Vertex.GetDefaultLiteral());
 				}
 			}
 
-			// Get the bound data read reference if it exists. Otherwise create and 
-			// return a data read reference by constructing one the supplied constructor 
+			// Get the bound data reference if it exists. Otherwise create and 
+			// return a data reference by constructing one the supplied constructor 
 			// arguments.
-			template<typename DataType, typename ... ConstructorArgTypes>
-			TDataReadReference<DataType> GetOrConstructDataReadReference(ConstructorArgTypes&&... ConstructorArgs) const
+			template<typename AsType, typename ... ArgTypes>
+			AsType GetOrConstructAs(ArgTypes&&... Args) const
 			{
 				if (Data.IsSet())
 				{
-					return Data->GetDataReadReference<DataType>();
+					return GetDataReferenceAs<AsType>();
 				}
-
-				return TDataReadReference<DataType>::CreateNew(Forward<ConstructorArgTypes>(ConstructorArgs)...);
-			}
-
-			// Get data write reference assuming data is bound and write accessible.
-			template<typename DataType> 
-			TDataWriteReference<DataType> GetDataWriteReference() const
-			{
-				check(Data.IsSet());
-				return Data->GetDataWriteReference<DataType>();
-			}
-
-			// Get the bound data write reference if it exists. Otherwise create and 
-			// return a data write reference by constructing one using the Vertex's 
-			// default literal.
-			template< typename DataType > 
-			TDataWriteReference<DataType> GetOrCreateDefaultDataWriteReference(const FOperatorSettings& InSettings) const
-			{
-				if (Data.IsSet())
-				{
-					return Data->GetDataWriteReference<DataType>();
-				}
-				return TDataWriteReferenceLiteralFactory<DataType>::CreateExplicitArgs(InSettings, Vertex.GetDefaultLiteral());
-			}
-
-			// Get the bound data write reference if it exists. Otherwise create and 
-			// return a data write reference by constructing one the supplied constructor 
-			// arguments.
-			template<typename DataType, typename ... ConstructorArgTypes>
-			TDataWriteReference<DataType> GetOrConstructDataWriteReference(ConstructorArgTypes&&... ConstructorArgs) const
-			{
-				if (Data.IsSet())
-				{
-					return Data->GetDataWriteReference<DataType>();
-				}
-				return TDataWriteReference<DataType>::CreateNew(Forward<ConstructorArgTypes>(ConstructorArgs)...);
-			}
-
-			// Gets the value of the bound data reference if it exists. Returns
-			// nullptr if not bound.
-			template<typename DataType>
-			const DataType* GetValue() const
-			{
-				if (Data.IsSet())
+				else
 				{
-					return Data->GetValue<DataType>();
+					return TFactoryFacade<AsType>::CreateWithArgs(Forward<ArgTypes>(Args)...);
 				}
-				return nullptr;
 			}
 
 
-			// Gets the value of the bound data reference if it exists. Otherwise
-			// create and return a value by constructing one using the Vertex's 
-			// default literal.
-			template<typename DataType>
-			DataType GetOrCreateDefaultValue(const FOperatorSettings& InSettings) const
-			{
-				if (Data.IsSet())
-				{
-					if (const DataType* Value = Data->GetValue<DataType>())
-					{
-						return *Value;
-					}
-				}
-				return TDataTypeLiteralFactory<DataType>::CreateExplicitArgs(InSettings, Vertex.GetDefaultLiteral());
-			}
-
 			// Set the value with a constant value reference.
 			template<typename DataType>
 			void SetValue(const DataType& InValue)
@@ -247,18 +254,6 @@ namespace Metasound
 		};
 
 
-		// Factory for creating a new FInputBinding from a IDataReference derived class and a vertex name.
-		struct FInputBindingFactory
-		{
-			template<typename DataReferenceType>
-			static FInputBinding CreateBinding(const FVertexName& InVertexName, const DataReferenceType& InRef)
-			{
-				FInputDataVertex Vertex(InVertexName, InRef.GetDataTypeName(), FDataVertexMetadata{}, TGetVertexAccess<DataReferenceType>::VertexAccess);
-				return FInputBinding(Vertex);
-			}
-		};
-
-
 		// Binds an IDataReference to a FOutputDataVertex
 		class METASOUNDGRAPHCORE_API FOutputBinding
 		{
@@ -273,22 +268,10 @@ namespace Metasound
 				Set(InOutDataReference);
 			}
 
-			template<typename DataType>
-			void BindRead(TDataReadReference<DataType>& InOutDataReference)
-			{
-				Set(InOutDataReference);
-			}
-
-			template<typename DataType>
-			void BindRead(TDataValueReference<DataType>& InOutDataReference)
-			{
-				Set(TDataReadReference<DataType>(InOutDataReference));
-			}
-
-			template<typename DataType>
-			void BindRead(TDataWriteReference<DataType>& InOutDataReference)
+			template<typename DataReferenceType>
+			void BindRead(DataReferenceType& InOutDataReference)
 			{
-				Set(TDataReadReference<DataType>(InOutDataReference));
+				Set(TDataReadReference<TGetReferenceDataType<DataReferenceType>>(InOutDataReference));
 			}
 
 			template<typename DataType>
@@ -323,59 +306,29 @@ namespace Metasound
 
 			FDataReferenceID GetDataReferenceID() const;
 
-			// Get data read reference assuming data is bound and read or write accessible.
-			template<typename DataType>
-			TDataReadReference<DataType> GetDataReadReference() const
-			{
-				check(Data.IsSet());
-				return Data->GetDataReadReference<DataType>();
-			}
-
-			// Get the bound data read reference if it exists. Otherwise create and 
-			// return a data read reference by constructing one the supplied constructor 
-			// arguments.
-			template<typename DataType, typename ... ConstructorArgTypes>
-			TDataReadReference<DataType> GetOrConstructDataReadReference(ConstructorArgTypes&&... ConstructorArgs) const
-			{
-				if (Data.IsSet())
-				{
-					return Data->GetDataReadReference<DataType>();
-				}
-
-				return TDataReadReference<DataType>::CreateNew(Forward<ConstructorArgTypes>(ConstructorArgs)...);
-			}
-
-			// Get data write reference assuming data is bound and write accessible.
-			template<typename DataType> 
-			TDataWriteReference<DataType> GetDataWriteReference() const
+			// Get the bound data reference. This assumes the data reference exists
+			// and is convertible to the AsType.
+			template<typename AsType>
+			AsType GetDataReferenceAs() const
 			{
 				check(Data.IsSet());
-				return Data->GetDataWriteReference<DataType>();
+				return Data->GetAs<AsType>();
 			}
 
-			// Get the bound data write reference if it exists. Otherwise create and 
-			// return a data write reference by constructing one the supplied constructor 
+			// Get the bound data reference if it exists. Otherwise create and 
+			// return a data reference by constructing one the supplied constructor 
 			// arguments.
-			template<typename DataType, typename ... ConstructorArgTypes>
-			TDataWriteReference<DataType> GetOrConstructDataWriteReference(ConstructorArgTypes&&... ConstructorArgs) const
+			template<typename AsType, typename ... ArgTypes>
+			AsType GetOrConstructAs(ArgTypes&&... Args) const
 			{
 				if (Data.IsSet())
 				{
-					return Data->GetDataWriteReference<DataType>();
+					return GetDataReferenceAs<AsType>();
 				}
-				return TDataWriteReference<DataType>::CreateNew(Forward<ConstructorArgTypes>(ConstructorArgs)...);
-			}
-
-			// Gets the value of the bound data reference if it exists. Returns
-			// nullptr if not bound.
-			template<typename DataType>
-			const DataType* GetValue() const
-			{
-				if (Data.IsSet())
+				else
 				{
-					return Data->GetValue<DataType>();
+					return TFactoryFacade<AsType>::CreateWithArgs(Forward<ArgTypes>(Args)...);
 				}
-				return nullptr;
 			}
 
 			// Set the value with a constant value reference.
@@ -392,13 +345,41 @@ namespace Metasound
 			TOptional<FAnyDataReference> Data;
 		};
 
-		// Factory for creating FOutputBinding from an IDataReference derived object.
-		struct FOutputBindingFactory
+
+		// Functor for creating a new FInputBinding from a IDataReference derived class and a vertex name.
+		template<typename DataReferenceType>
+		struct TCreateInputBinding
 		{
-			template<typename DataReferenceType>
-			static FOutputBinding CreateBinding(const FVertexName& InVertexName, const DataReferenceType& InRef)
+			const FVertexName& VertexName;
+			DataReferenceType& Ref;
+
+			TCreateInputBinding(const FVertexName& InVertexName, DataReferenceType& InRef)
+			: VertexName(InVertexName)
+			, Ref(InRef)
+			{}
+
+			FInputBinding operator()() const
 			{
-				FOutputDataVertex Vertex(InVertexName, InRef.GetDataTypeName(), FDataVertexMetadata{}, TGetVertexAccess<DataReferenceType>::VertexAccess);
+				FInputDataVertex Vertex(VertexName, Ref.GetDataTypeName(), FDataVertexMetadata{}, TGetVertexAccess<std::decay_t<DataReferenceType>>::VertexAccess);
+				return FInputBinding(Vertex);
+			}
+		};
+
+		// Functor for creating a new FOutputBinding from a IDataReference derived class and a vertex name.
+		template<typename DataReferenceType>
+		struct TCreateOutputBinding
+		{
+			const FVertexName& VertexName;
+			DataReferenceType& Ref;
+
+			TCreateOutputBinding(const FVertexName& InVertexName, DataReferenceType& InRef)
+			: VertexName(InVertexName)
+			, Ref(InRef)
+			{}
+
+			FOutputBinding operator()() const
+			{
+				FOutputDataVertex Vertex(VertexName, Ref.GetDataTypeName(), FDataVertexMetadata{}, TGetVertexAccess<std::decay_t<DataReferenceType>>::VertexAccess);
 				return FOutputBinding(Vertex);
 			}
 		};
@@ -477,7 +458,7 @@ namespace Metasound
 	/** An input vertex interface with optionally bound data references. */
 	class METASOUNDGRAPHCORE_API FInputVertexInterfaceData
 	{
-		using FInputBinding = MetasoundVertexDataPrivate::FInputBinding;
+		using FInputBinding = VertexDataPrivate::FInputBinding;
 
 	public:
 
@@ -528,28 +509,20 @@ namespace Metasound
 		template<typename DataType>
 		void BindReadVertex(const FVertexName& InVertexName, TDataReadReference<DataType>& InOutDataReference)
 		{
-			auto CreateBinding = [&]()
-			{
-				return MetasoundVertexDataPrivate::FInputBindingFactory::CreateBinding(InVertexName, InOutDataReference);
-			};
-
+			using namespace VertexDataPrivate;
 			auto BindData = [&](FInputBinding& Binding) { Binding.BindRead(InOutDataReference); };
 			
-			Apply(InVertexName, CreateBinding, BindData);
+			Apply(InVertexName, TCreateInputBinding(InVertexName, InOutDataReference), BindData);
 		}
 		
 		/** Bind a read vertex from a write reference. */
 		template<typename DataType>
 		void BindReadVertex(const FVertexName& InVertexName, TDataWriteReference<DataType>& InOutDataReference)
 		{
-			auto CreateBinding = [&]()
-			{
-				return MetasoundVertexDataPrivate::FInputBindingFactory::CreateBinding(InVertexName, InOutDataReference);
-			};
-
-			auto BindData = [&](FInputBinding& Binding) { Binding.BindRead(InOutDataReference); };
+			using namespace VertexDataPrivate;
 			
-			Apply(InVertexName, CreateBinding, BindData);
+			auto BindData = [&](FInputBinding& Binding) { Binding.BindRead(InOutDataReference); };
+			Apply(InVertexName, TCreateInputBinding(InVertexName, InOutDataReference), BindData);
 		}
 
 		/** Bind a write vertex from a write reference. Slated for deprecation. */
@@ -563,14 +536,9 @@ namespace Metasound
 		template<typename DataType>
 		void BindWriteVertex(const FVertexName& InVertexName, TDataWriteReference<DataType>& InOutDataReference)
 		{
-			auto CreateBinding = [&]()
-			{
-				return MetasoundVertexDataPrivate::FInputBindingFactory::CreateBinding(InVertexName, InOutDataReference);
-			};
-
+			using namespace VertexDataPrivate;
 			auto BindData = [&](FInputBinding& Binding) { Binding.BindWrite(InOutDataReference); };
-			
-			Apply(InVertexName, CreateBinding, BindData);
+			Apply(InVertexName, TCreateInputBinding(InVertexName, InOutDataReference), BindData);
 		}
 
 		/** Bind a vertex with a any data reference. Slated for deprecation. */
@@ -578,6 +546,69 @@ namespace Metasound
 		/** Bind a vertex with a any data reference. */
 		void BindVertex(const FVertexName& InVertexName, FAnyDataReference& InOutDataReference);
 
+		/** Return the number of instances of a particular sub interface. */
+		uint32 GetNumSubInterfaceInstances(const FName& InSubInterfaceName) const;
+
+		/** Bind the vertices of a sub interface vertex. 
+		 * @param InSubInterfaceName  - Name of sub interface.
+		 * @param InVertexName        - Name of prototype vertex on sub interface.
+		 * @param InOutDataReferences - An indexable sequence container of data references. 
+		 */
+		template<typename RangeType>
+		void BindSubInterfaceVertices(const FName& InSubInterfaceName, const FName& InVertexName, RangeType& InOutDataReferences)
+		{
+			auto BindData = [&InOutDataReferences](uint32 Index, FInputBinding& Binding)
+			{
+				Binding.Bind(InOutDataReferences[Index]);
+			};
+			ForEachSubInterfaceBinding(InSubInterfaceName, InVertexName, InOutDataReferences.Num(), BindData);
+		}
+
+		/** Bind the vertices of a sub interface vertex as write references
+		 * @param InSubInterfaceName  - Name of sub interface.
+		 * @param InVertexName        - Name of prototype vertex on sub interface.
+		 * @param InOutDataReferences - An indexable sequence container of data references. 
+		 */
+		template<typename RangeType>
+		void BindSubInterfaceWriteVertices(const FName& InSubInterfaceName, const FName& InVertexName, RangeType& InOutDataReferences)
+		{
+			auto BindData = [&InOutDataReferences](uint32 Index, FInputBinding& Binding)
+			{
+				Binding.BindWrite(InOutDataReferences[Index]);
+			};
+			ForEachSubInterfaceBinding(InSubInterfaceName, InVertexName, InOutDataReferences.Num(), BindData);
+		}
+
+		/** Bind the vertices of a sub interface vertex as read references
+		 * @param InSubInterfaceName  - Name of sub interface.
+		 * @param InVertexName        - Name of prototype vertex on sub interface.
+		 * @param InOutDataReferences - An indexable sequence container of data references. 
+		 */
+		template<typename RangeType>
+		void BindSubInterfaceReadVertices(const FName& InSubInterfaceName, const FName& InVertexName, RangeType& InOutDataReferences)
+		{
+			auto BindData = [&InOutDataReferences](uint32 Index, FInputBinding& Binding)
+			{
+				Binding.BindRead(InOutDataReferences[Index]);
+			};
+			ForEachSubInterfaceBinding(InSubInterfaceName, InVertexName, InOutDataReferences.Num(), BindData);
+		}
+
+		/** Bind the vertices of a sub interface vertex as value references
+		 * @param InSubInterfaceName  - Name of sub interface.
+		 * @param InVertexName        - Name of prototype vertex on sub interface.
+		 * @param InOutDataReferences - An indexable sequence container of data references. 
+		 */
+		template<typename RangeType>
+		void SetSubInterfaceValues(const FName& InSubInterfaceName, const FName& InVertexName, const RangeType& InValues)
+		{
+			auto BindData = [&InValues](uint32 Index, FInputBinding& Binding)
+			{
+				Binding.SetValue<TElementType_T<RangeType>>(InValues[Index]);
+			};
+			ForEachSubInterfaceBinding(InSubInterfaceName, InVertexName, InValues.Num(), BindData);
+		}
+
 		/** Bind vertex data using other vertex data. Slated for deprecation. */
 		void Bind(const FInputVertexInterfaceData& InOutInputVertexData);
 		/** Bind vertex data using other vertex data. */
@@ -644,7 +675,11 @@ namespace Metasound
 		const DataType* GetValue(const FVertexName& InVertexName) const
 		{
 			const FInputBinding* Binding = FindChecked(InVertexName);
-			return Binding->GetValue<DataType>();
+			if (Binding->IsBound())
+			{
+				return Binding->GetDataReferenceAs<const DataType*>();
+			}
+			return nullptr;
 		}
 
 		/**  Gets the value of the bound data reference if it exists. Otherwise
@@ -654,7 +689,7 @@ namespace Metasound
 		DataType GetOrCreateDefaultValue(const FVertexName& InVertexName, const FOperatorSettings& InSettings) const
 		{
 			const FInputBinding* Binding = FindChecked(InVertexName);
-			return Binding->GetOrCreateDefaultValue<DataType>(InSettings);
+			return Binding->GetOrCreateDefaultAs<DataType>(InSettings);
 		}
 
 		/** Get data read reference assuming data is bound and read or write accessible. */
@@ -662,7 +697,7 @@ namespace Metasound
 		TDataReadReference<DataType> GetDataReadReference(const FVertexName& InVertexName) const
 		{
 			const FInputBinding* Binding = FindChecked(InVertexName);
-			return Binding->GetDataReadReference<DataType>();
+			return Binding->GetDataReferenceAs<TDataReadReference<DataType>>();
 		}
 
 		/** Get the bound data read reference if it exists. Otherwise create and 
@@ -674,7 +709,7 @@ namespace Metasound
 		{
 			if (const FInputBinding* Binding = Find(InVertexName))
 			{
-				return Binding->GetOrCreateDefaultDataReadReference<DataType>(InSettings);
+				return Binding->GetOrCreateDefaultAs<TDataReadReference<DataType>>(InSettings);
 			}
 			else
 			{
@@ -693,7 +728,7 @@ namespace Metasound
 		{
 			if (const FInputBinding* Binding = Find(InVertexName))
 			{
-				return Binding->GetOrConstructDataReadReference<DataType>(Forward<ConstructorArgTypes>(ConstructorArgs)...);
+				return Binding->GetOrConstructAs<TDataReadReference<DataType>>(Forward<ConstructorArgTypes>(ConstructorArgs)...);
 			}
 			else
 			{
@@ -709,7 +744,7 @@ namespace Metasound
 		TDataWriteReference<DataType> GetDataWriteReference(const FVertexName& InVertexName) const
 		{
 			const FInputBinding* Binding = FindChecked(InVertexName);
-			return Binding->GetDataWriteReference<DataType>();
+			return Binding->GetDataReferenceAs<TDataWriteReference<DataType>>();
 		}
 
 		/** Get the bound data write reference if it exists. Otherwise create and 
@@ -721,7 +756,7 @@ namespace Metasound
 		{
 			if (const FInputBinding* Binding = Find(InVertexName))
 			{
-				return Binding->GetOrCreateDefaultDataWriteReference<DataType>(InSettings);
+				return Binding->GetOrCreateDefaultAs<TDataWriteReference<DataType>>(InSettings);
 			}
 			else
 			{
@@ -740,7 +775,7 @@ namespace Metasound
 		{
 			if (const FInputBinding* Binding = Find(InVertexName))
 			{
-				return Binding->GetOrConstructDataWriteReference<DataType>(Forward<ConstructorArgTypes>(ConstructorArgs)...);
+				return Binding->GetOrConstructAs<TDataWriteReference<DataType>>(Forward<ConstructorArgTypes>(ConstructorArgs)...);
 			}
 			else
 			{
@@ -749,12 +784,75 @@ namespace Metasound
 
 			return TDataWriteReference<DataType>::CreateNew(Forward<ConstructorArgTypes>(ConstructorArgs)...);
 		}
+
+		/** Get or create an array of data references for all instances of a sub interface vertex.
+		 * @param InSubInterfaceName - Name of the sub interface
+		 * @param InVertexName       - Name of the prototype vertex in the sub interface.
+		 * @param InSettings         - Operator settings to use if creating a data reference.
+		 */
+		template<typename DataReferenceType>
+		TArray<DataReferenceType> GetOrCreateDefaultSubInterfaceAs(const FName& InSubInterfaceName, const FVertexName& InVertexName, const FOperatorSettings& InSettings) const
+		{
+			TArray<DataReferenceType> References;
+			auto GetOrCreateDefault = [&References, &InSettings](uint32 InIndex, const FInputBinding& InBinding)
+			{
+				check(References.Num() == InIndex); // This check will be false if the references array 
+													// is not being built in the correct order. 
+				References.Add(InBinding.GetOrCreateDefaultAs<DataReferenceType>(InSettings));
+			};
+			ForEachSubInterfaceBinding(InSubInterfaceName, InVertexName, INDEX_NONE, GetOrCreateDefault);
+			return References;
+		}
+
+		/** Get or create an array of data read references for all instances of a sub interface vertex.
+		 * @param InSubInterfaceName - Name of the sub interface
+		 * @param InVertexName       - Name of the prototype vertex in the sub interface.
+		 * @param InSettings         - Operator settings to use if creating a data reference.
+		 */
+		template<typename DataType>
+		TArray<TDataReadReference<DataType>> GetOrCreateDefaultSubInterfaceDataReadReferences(const FName& InSubInterfaceName, const FVertexName& InVertexName, const FOperatorSettings& InSettings) const
+		{
+			return GetOrCreateDefaultSubInterfaceAs<TDataReadReference<DataType>>(InSubInterfaceName, InVertexName, InSettings);
+		}
+
+		/** Get or create an array of data write references for all instances of a sub interface vertex.
+		 * @param InSubInterfaceName - Name of the sub interface
+		 * @param InVertexName       - Name of the prototype vertex in the sub interface.
+		 * @param InSettings         - Operator settings to use if creating a data reference.
+		 */
+		template<typename DataType>
+		TArray<TDataWriteReference<DataType>> GetOrCreateDefaultSubInterfaceDataWriteReferences(const FName& InSubInterfaceName, const FVertexName& InVertexName, const FOperatorSettings& InSettings) const
+		{
+			return GetOrCreateDefaultSubInterfaceAs<TDataWriteReference<DataType>>(InSubInterfaceName, InVertexName, InSettings);
+		}
+
+		/** Get or create an array of values for all instances of a sub interface vertex.
+		 * @param InSubInterfaceName - Name of the sub interface
+		 * @param InVertexName       - Name of the prototype vertex in the sub interface.
+		 * @param InSettings         - Operator settings to use if creating a data reference.
+		 */
+		template<typename DataType>
+		TArray<DataType> GetOrCreateDefaultSubInterfaceValues(const FName& InSubInterfaceName, const FVertexName& InVertexName, const FOperatorSettings& InSettings) const
+		{
+			return GetOrCreateDefaultSubInterfaceAs<DataType>(InSubInterfaceName, InVertexName, InSettings);
+		}
+
+		/** Find data references bound to the sub interface vertices. If there is
+		 * no bound data for a sub interface instance, the array element will be
+		 * a nullptr. 
+		 *
+		 * @param InSubInterfaceName - Name of the sub interface
+		 * @param InVertexName       - Name of the prototype vertex in the sub interface.
+		 */
+		TArray<const FAnyDataReference*> FindSubInterfaceReferences(const FName& InSubInterfaceName, const FVertexName& InVertexName) const;
 	
 	private:
 		friend METASOUNDGRAPHCORE_API void GetVertexInterfaceDataState(const FInputVertexInterfaceData& InVertexInterface, TArray<FVertexDataState>& OutState);
 		friend METASOUNDGRAPHCORE_API void CompareVertexInterfaceDataToPriorState(const FInputVertexInterfaceData& InCurrentInterface, const TArray<FVertexDataState>& InPriorState, TSortedVertexNameMap<FAnyDataReference>& OutUpdates);
 
 		void Apply(const FVertexName& InVertexName, TFunctionRef<FInputBinding ()> InCreateFunc, TFunctionRef<void (FInputBinding&)> InApplyFunc);
+		void ForEachSubInterfaceBinding(const FName& InSubInterfaceName, const FVertexName& InVertexName, uint32 InExpectedNumSubInterfaceInstances, TFunctionRef<void (uint32, FInputBinding&)> InApplyFunc);
+		void ForEachSubInterfaceBinding(const FName& InSubInterfaceName, const FVertexName& InVertexName, uint32 InExpectedNumSubInterfaceInstances, TFunctionRef<void (uint32, const FInputBinding&)> InApplyFunc) const;
 
 		FInputBinding* Find(const FVertexName& InVertexName);
 		const FInputBinding* Find(const FVertexName& InVertexName) const;
@@ -762,15 +860,18 @@ namespace Metasound
 		FInputBinding* FindChecked(const FVertexName& InVertexName);
 		const FInputBinding* FindChecked(const FVertexName& InVertexName) const;
 
+		const VertexPrivate::FSubInterfaceLayout* FindSubInterface(const FName& InSubInterfaceName) const;
+		
 		bool bIsVertexInterfaceFrozen = false;
 		TArray<FInputBinding> Bindings;
+		TArray<VertexPrivate::FSubInterfaceLayout> SubInterfaces;
 	};
 
 
 	/** An output vertex interface with optionally bound data references. */
 	class METASOUNDGRAPHCORE_API FOutputVertexInterfaceData
 	{
-		using FOutputBinding = MetasoundVertexDataPrivate::FOutputBinding;
+		using FOutputBinding = VertexDataPrivate::FOutputBinding;
 
 	public:
 
@@ -806,14 +907,9 @@ namespace Metasound
 		template<typename DataType>
 		void BindValueVertex(const FVertexName& InVertexName, const TDataValueReference<DataType>& InOutDataReference)
 		{
-			auto CreateBinding = [&]()
-			{
-				return MetasoundVertexDataPrivate::FOutputBindingFactory::CreateBinding(InVertexName, InOutDataReference);
-			};
-
+			using namespace VertexDataPrivate;
 			auto BindData = [&](FOutputBinding& Binding) { Binding.BindValue(InOutDataReference); };
-
-			Apply(InVertexName, CreateBinding, BindData);
+			Apply(InVertexName, TCreateOutputBinding(InVertexName, InOutDataReference), BindData);
 		}
 
 		/** Bind a read vertex from a value reference. */
@@ -834,14 +930,9 @@ namespace Metasound
 		template<typename DataType>
 		void BindReadVertex(const FVertexName& InVertexName, TDataReadReference<DataType>& InOutDataReference)
 		{
-			auto CreateBinding = [&]()
-			{
-				return MetasoundVertexDataPrivate::FOutputBindingFactory::CreateBinding(InVertexName, InOutDataReference);
-			};
-
+			using namespace VertexDataPrivate;
 			auto BindData = [&](FOutputBinding& Binding) { Binding.BindRead(InOutDataReference); };
-
-			Apply(InVertexName, CreateBinding, BindData);
+			Apply(InVertexName, TCreateOutputBinding(InVertexName, InOutDataReference), BindData);
 		}
 
 		/** Bind a read vertex from a write reference. Slated for deprecation. */
@@ -855,14 +946,9 @@ namespace Metasound
 		template<typename DataType>
 		void BindReadVertex(const FVertexName& InVertexName, TDataWriteReference<DataType>& InOutDataReference)
 		{
-			auto CreateBinding = [&]()
-			{
-				return MetasoundVertexDataPrivate::FOutputBindingFactory::CreateBinding(InVertexName, InOutDataReference);
-			};
-
+			using namespace VertexDataPrivate;
 			auto BindData = [&](FOutputBinding& Binding) { Binding.BindRead(InOutDataReference); };
-
-			Apply(InVertexName, CreateBinding, BindData);
+			Apply(InVertexName, TCreateOutputBinding(InVertexName, InOutDataReference), BindData);
 		}
 
 		/** Bind a write vertex from a write reference. Slated for deprecation. */
@@ -876,14 +962,9 @@ namespace Metasound
 		template<typename DataType>
 		void BindWriteVertex(const FVertexName& InVertexName, TDataWriteReference<DataType>& InOutDataReference)
 		{
-			auto CreateBinding = [&]()
-			{
-				return MetasoundVertexDataPrivate::FOutputBindingFactory::CreateBinding(InVertexName, InOutDataReference);
-			};
-
+			using namespace VertexDataPrivate;
 			auto BindData = [&](FOutputBinding& Binding) { Binding.BindWrite(InOutDataReference); };
-
-			Apply(InVertexName, CreateBinding, BindData);
+			Apply(InVertexName, TCreateOutputBinding(InVertexName, InOutDataReference), BindData);
 		}
 
 		/** Bind a vertex with a any data reference. Slated for deprecation. */
@@ -891,6 +972,77 @@ namespace Metasound
 		/** Bind a vertex with a any data reference. */
 		void BindVertex(const FVertexName& InVertexName, FAnyDataReference& InOutDataReference);
 
+		/** Return the number of instances of a particular sub interface. */
+		uint32 GetNumSubInterfaceInstances(const FName& InSubInterfaceName) const;
+
+		/** Bind the vertices of a sub interface vertex. 
+		 * @param InSubInterfaceName  - Name of sub interface.
+		 * @param InVertexName        - Name of prototype vertex on sub interface.
+		 * @param InOutDataReferences - An indexable sequence container of data references. 
+		 */
+		template<typename RangeType>
+		void BindSubInterfaceVertices(const FName& InSubInterfaceName, const FName& InVertexName, RangeType& InOutDataReferences)
+		{
+			using namespace VertexDataPrivate;
+
+			auto BindData = [&InOutDataReferences](uint32 Index, FOutputBinding& Binding)
+			{
+				Binding.Bind(InOutDataReferences[Index]);
+			};
+			ForEachSubInterfaceBinding(InSubInterfaceName, InVertexName, InOutDataReferences.Num(), BindData);
+		}
+
+		/** Bind the vertices of a sub interface vertex as write references
+		 * @param InSubInterfaceName  - Name of sub interface.
+		 * @param InVertexName        - Name of prototype vertex on sub interface.
+		 * @param InOutDataReferences - An indexable sequence container of data references. 
+		 */
+		template<typename RangeType>
+		void BindSubInterfaceWriteVertices(const FName& InSubInterfaceName, const FName& InVertexName, RangeType& InOutDataReferences)
+		{
+			using namespace VertexDataPrivate;
+
+			auto BindData = [&InOutDataReferences](uint32 Index, FOutputBinding& Binding)
+			{
+				Binding.BindWrite(InOutDataReferences[Index]);
+			};
+			ForEachSubInterfaceBinding(InSubInterfaceName, InVertexName, InOutDataReferences.Num(), BindData);
+		}
+
+		/** Bind the vertices of a sub interface vertex as read references
+		 * @param InSubInterfaceName  - Name of sub interface.
+		 * @param InVertexName        - Name of prototype vertex on sub interface.
+		 * @param InOutDataReferences - An indexable sequence container of data references. 
+		 */
+		template<typename RangeType>
+		void BindSubInterfaceReadVertices(const FName& InSubInterfaceName, const FName& InVertexName, RangeType& InOutDataReferences)
+		{
+			using namespace VertexDataPrivate;
+
+			auto BindData = [&InOutDataReferences](uint32 Index, FOutputBinding& Binding)
+			{
+				Binding.BindRead(InOutDataReferences[Index]);
+			};
+			ForEachSubInterfaceBinding(InSubInterfaceName, InVertexName, InOutDataReferences.Num(), BindData);
+		}
+
+		/** Bind the vertices of a sub interface vertex as value references
+		 * @param InSubInterfaceName  - Name of sub interface.
+		 * @param InVertexName        - Name of prototype vertex on sub interface.
+		 * @param InOutDataReferences - An indexable sequence container of data references. 
+		 */
+		template<typename RangeType>
+		void SetSubInterfaceValues(const FName& InSubInterfaceName, const FName& InVertexName, const RangeType& InValues)
+		{
+			using namespace VertexDataPrivate;
+
+			auto BindData = [&InValues](uint32 Index, FOutputBinding& Binding)
+			{
+				Binding.SetValue<TElementType_T<RangeType>>(InValues[Index]);
+			};
+			ForEachSubInterfaceBinding(InSubInterfaceName, InVertexName, InValues.Num(), BindData);
+		}
+
 		/** Bind vertex data using other vertex data. Slated for deprecation. */
 		void Bind(const FOutputVertexInterfaceData& InOutOutputVertexData);
 		/** Bind vertex data using other vertex data. */
@@ -957,7 +1109,14 @@ namespace Metasound
 		const DataType* GetValue(const FVertexName& InVertexName) const
 		{
 			const FOutputBinding* Binding = FindChecked(InVertexName);
-			return Binding->GetValue<DataType>();
+			if (Binding->IsBound())
+			{
+				return Binding->GetDataReferenceAs<const DataType*>();
+			}
+			else
+			{
+				return nullptr;
+			}
 		}
 
 
@@ -966,7 +1125,7 @@ namespace Metasound
 		TDataReadReference<DataType> GetDataReadReference(const FVertexName& InVertexName) const
 		{
 			const FOutputBinding* Binding = FindChecked(InVertexName);
-			return Binding->GetDataReadReference<DataType>();
+			return Binding->GetDataReferenceAs<TDataReadReference<DataType>>();
 		}
 
 		/** Get the bound data read reference if it exists. Otherwise create and 
@@ -978,7 +1137,7 @@ namespace Metasound
 		{
 			if (const FOutputBinding* Binding = Find(InVertexName))
 			{
-				return Binding->GetOrConstructDataReadReference<DataType>(Forward<ConstructorArgTypes>(ConstructorArgs)...);
+				return Binding->GetOrConstructAs<TDataReadReference<DataType>>(Forward<ConstructorArgTypes>(ConstructorArgs)...);
 			}
 			else
 			{
@@ -994,7 +1153,7 @@ namespace Metasound
 		TDataWriteReference<DataType> GetDataWriteReference(const FVertexName& InVertexName) const
 		{
 			const FOutputBinding* Binding = FindChecked(InVertexName);
-			return Binding->GetDataWriteReference<DataType>();
+			return Binding->GetDataReferenceAs<TDataWriteReference<DataType>>();
 		}
 
 		/** Get the bound data write reference if it exists. Otherwise create and 
@@ -1006,7 +1165,7 @@ namespace Metasound
 		{
 			if (const FOutputBinding* Binding = Find(InVertexName))
 			{
-				return Binding->GetOrConstructDataWriteReference<DataType>(Forward<ConstructorArgTypes>(ConstructorArgs)...);
+				return Binding->GetOrConstructAs<TDataWriteReference<DataType>>(Forward<ConstructorArgTypes>(ConstructorArgs)...);
 			}
 			else
 			{
@@ -1015,21 +1174,35 @@ namespace Metasound
 
 			return TDataWriteReference<DataType>::CreateNew(Forward<ConstructorArgTypes>(ConstructorArgs)...);
 		}
+
+
+		/** Find data references bound to the sub interface vertices. If there is
+		 * no bound data for a sub interface instance, the array element will be
+		 * a nullptr. 
+		 *
+		 * @param InSubInterfaceName - Name of the sub interface
+		 * @param InVertexName       - Name of the prototype vertex in the sub interface.
+		 */
+		TArray<const FAnyDataReference*> FindSubInterfaceReferences(const FName& InSubInterfaceName, const FVertexName& InVertexName) const;
 	
 	private:
 		friend METASOUNDGRAPHCORE_API void GetVertexInterfaceDataState(const FOutputVertexInterfaceData& InVertexInterface, TArray<FVertexDataState>& OutState);
 		friend METASOUNDGRAPHCORE_API void CompareVertexInterfaceDataToPriorState(const FOutputVertexInterfaceData& InCurrentInterface, const TArray<FVertexDataState>& InPriorState, TSortedVertexNameMap<FAnyDataReference>& OutUpdates);
 
 		void Apply(const FVertexName& InVertexName, TFunctionRef<FOutputBinding ()> InCreateFunc, TFunctionRef<void (FOutputBinding&)> InBindFunc);
+		void ForEachSubInterfaceBinding(const FName& InSubInterfaceName, const FVertexName& InVertexName, uint32 InExpectedNumSubInterfaceInstances, TFunctionRef<void (uint32, FOutputBinding&)> InApplyFunc);
+		void ForEachSubInterfaceBinding(const FName& InSubInterfaceName, const FVertexName& InVertexName, uint32 InExpectedNumSubInterfaceInstances, TFunctionRef<void (uint32, const FOutputBinding&)> InApplyFunc) const;
 
 		FOutputBinding* Find(const FVertexName& InVertexName);
 		const FOutputBinding* Find(const FVertexName& InVertexName) const;
+		const VertexPrivate::FSubInterfaceLayout* FindSubInterface(const FName& InSubInterfaceName) const;
 
 		FOutputBinding* FindChecked(const FVertexName& InVertexName);
 		const FOutputBinding* FindChecked(const FVertexName& InVertexName) const;
 
 		bool bIsVertexInterfaceFrozen = false;
 		TArray<FOutputBinding> Bindings;
+		TArray<VertexPrivate::FSubInterfaceLayout> SubInterfaces;
 	};
 
 	/** A vertex interface with optionally bound data. */
