Fix for issue with deleting MetaSound nodes not deleting their connection on their respective page
#rnx

diff --git a/Engine/Plugins/Runtime/Metasound/Source/MetasoundEditor/Private/MetasoundEditorGraphNode.cpp b/Engine/Plugins/Runtime/Metasound/Source/MetasoundEditor/Private/MetasoundEditorGraphNode.cpp
--- a/Engine/Plugins/Runtime/Metasound/Source/MetasoundEditor/Private/MetasoundEditorGraphNode.cpp
+++ b/Engine/Plugins/Runtime/Metasound/Source/MetasoundEditor/Private/MetasoundEditorGraphNode.cpp
@@ -916,11 +916,15 @@ bool UMetasoundEditorGraphOutputNode::RemoveFromDocument() const
 {
 	if (Output)
 	{
-		// When removing ed graph output nodes, just remove the location as all frontend page graphs
-		// require the graph vertex node to exist and contain matching NodeIDs across all pages.
+		// When removing ed graph output nodes, disconnect, but only remove the location as all
+		// frontend page graphs require the graph vertex node to exist and contain matching NodeIDs
+		// across all pages.
 		const FGuid& NodeID = GetNodeID();
 		UMetaSoundBuilderBase& Builder = GetBuilderChecked();
-		const int32 NumLocationsRemoved = Builder.GetBuilder().RemoveNodeLocation(NodeID);
+		FMetaSoundFrontendDocumentBuilder& DocBuilder = Builder.GetBuilder();
+
+		DocBuilder.RemoveEdges(NodeID);
+		const int32 NumLocationsRemoved = DocBuilder.RemoveNodeLocation(NodeID);
 
 		if (FMetasoundAssetBase* MetaSoundAsset = Metasound::IMetasoundUObjectRegistry::Get().GetObjectAsAssetBase(&GetMetasoundChecked()))
 		{
