Cleanup MetaSound node registry file name and class name.

#rnx

diff --git a/Engine/Plugins/Runtime/Metasound/Source/MetasoundFrontend/Private/MetasoundAssetBase.cpp b/Engine/Plugins/Runtime/Metasound/Source/MetasoundFrontend/Private/MetasoundAssetBase.cpp
--- a/Engine/Plugins/Runtime/Metasound/Source/MetasoundFrontend/Private/MetasoundAssetBase.cpp
+++ b/Engine/Plugins/Runtime/Metasound/Source/MetasoundFrontend/Private/MetasoundAssetBase.cpp
@@ -22,10 +22,11 @@
 #include "MetasoundFrontendDocumentBuilder.h"
 #include "MetasoundFrontendDocumentIdGenerator.h"
 #include "MetasoundFrontendDocumentVersioning.h"
+#include "MetasoundFrontendNodeClassRegistry.h"
+#include "MetasoundFrontendNodeClassRegistryPrivate.h"
 #include "MetasoundFrontendNodeTemplateRegistry.h"
 #include "MetasoundFrontendProxyDataCache.h"
 #include "MetasoundFrontendRegistries.h"
-#include "MetasoundFrontendRegistryContainerImpl.h"
 #include "MetasoundFrontendSearchEngine.h"
 #include "MetasoundFrontendTransform.h"
 #include "MetasoundGlobals.h"
@@ -106,11 +107,6 @@ namespace Metasound
 					return FrontendClass;
 				}
 
-				virtual TUniquePtr<INodeRegistryEntry> Clone() const override
-				{
-					return MakeUnique<FDocumentNodeRegistryEntryForSerialization>(*this);
-				}
-
 				virtual const TSet<FMetasoundFrontendVersion>* GetImplementedInterfaces() const override
 				{
 					return &Interfaces;
@@ -360,7 +356,7 @@ void FMetasoundAssetBase::UpdateAndRegisterForExecution(Metasound::Frontend::FMe
 	CacheRegistryMetadata();
 #endif // WITH_EDITOR
 
-	GraphRegistryKey = FRegistryContainerImpl::Get().RegisterGraph(Owner);
+	GraphRegistryKey = FNodeClassRegistry::Get().RegisterGraph(Owner);
 	if (GraphRegistryKey.IsValid())
 	{
 #if WITH_EDITORONLY_DATA
@@ -437,7 +433,7 @@ void FMetasoundAssetBase::UpdateAndRegisterForSerialization(FName CookPlatformNa
 
 		if (GraphRegistryKey.IsValid())
 		{
-			FRegistryContainerImpl::Get().UnregisterNode(GraphRegistryKey.NodeKey);
+			FNodeClassRegistry::Get().UnregisterNode(GraphRegistryKey.NodeKey);
 			GraphRegistryKey = { };
 		}
 
@@ -449,7 +445,7 @@ void FMetasoundAssetBase::UpdateAndRegisterForSerialization(FName CookPlatformNa
 		const FTopLevelAssetPath AssetPath = DocInterface->GetAssetPathChecked();
 		TUniquePtr<INodeRegistryEntry> RegistryEntry = MakeUnique<AssetBasePrivate::FDocumentNodeRegistryEntryForSerialization>(Document, AssetPath);
 
-		const FNodeRegistryKey NodeKey = FRegistryContainerImpl::Get().RegisterNode(MoveTemp(RegistryEntry));
+		const FNodeClassRegistryKey NodeKey = FNodeClassRegistry::Get().RegisterNode(MoveTemp(RegistryEntry));
 		GraphRegistryKey = FGraphRegistryKey { NodeKey, AssetPath };
 	}
 
@@ -485,7 +481,7 @@ void FMetasoundAssetBase::OnNotifyBeginDestroy()
 	{
 		if (GraphRegistryKey.IsValid())
 		{
-			FRegistryContainerImpl::Get().UnregisterNode(GraphRegistryKey.NodeKey);
+			FNodeClassRegistry::Get().UnregisterNode(GraphRegistryKey.NodeKey);
 			GraphRegistryKey = { };
 		}
 	}
@@ -510,7 +506,7 @@ void FMetasoundAssetBase::UnregisterGraphWithFrontend()
 		UObject* OwningAsset = GetOwningAsset();
 		if (ensureAlways(OwningAsset))
 		{
-			const bool bSuccess = FRegistryContainerImpl::Get().UnregisterGraph(GraphRegistryKey, OwningAsset);
+			const bool bSuccess = FNodeClassRegistry::Get().UnregisterGraph(GraphRegistryKey, OwningAsset);
 			if (!bSuccess)
 			{
 				UE_LOG(LogMetaSound, Verbose, TEXT("Failed to unregister node with key %s for asset %s. No registry entry exists with that key."), *GraphRegistryKey.ToString(), *GetOwningAssetName());
@@ -905,7 +901,7 @@ const Metasound::Frontend::FGraphRegistryKey& FMetasoundAssetBase::GetGraphRegis
 	return GraphRegistryKey;
 }
 
-const Metasound::Frontend::FNodeRegistryKey& FMetasoundAssetBase::GetRegistryKey() const
+const Metasound::Frontend::FNodeClassRegistryKey& FMetasoundAssetBase::GetRegistryKey() const
 {
 	return GraphRegistryKey.NodeKey;
 }
diff --git a/Engine/Plugins/Runtime/Metasound/Source/MetasoundFrontend/Private/MetasoundFrontendDataTypeRegistry.cpp b/Engine/Plugins/Runtime/Metasound/Source/MetasoundFrontend/Private/MetasoundFrontendDataTypeRegistry.cpp
--- a/Engine/Plugins/Runtime/Metasound/Source/MetasoundFrontend/Private/MetasoundFrontendDataTypeRegistry.cpp
+++ b/Engine/Plugins/Runtime/Metasound/Source/MetasoundFrontend/Private/MetasoundFrontendDataTypeRegistry.cpp
@@ -194,11 +194,6 @@ namespace Metasound
 					return DataTypeEntry->CreateInputNode(MoveTemp(InputParams));
 				}
 
-				virtual TUniquePtr<INodeRegistryEntry> Clone() const override
-				{
-					return MakeUnique<FInputNodeRegistryEntry>(DataTypeEntry);
-				}
-
 			private:
 				
 				TSharedRef<IDataTypeRegistryEntry, ESPMode::ThreadSafe> DataTypeEntry;
@@ -234,11 +229,6 @@ namespace Metasound
 					return DataTypeEntry->CreateConstructorInputNode(MoveTemp(InputParams));
 				}
 
-				virtual TUniquePtr<INodeRegistryEntry> Clone() const override
-				{
-					return MakeUnique<FConstructorInputNodeRegistryEntry>(DataTypeEntry);
-				}
-
 			private:
 				
 				TSharedRef<IDataTypeRegistryEntry, ESPMode::ThreadSafe> DataTypeEntry;
@@ -268,11 +258,6 @@ namespace Metasound
 					return DataTypeEntry->CreateOutputNode(MoveTemp(InParams));
 				}
 
-				virtual TUniquePtr<INodeRegistryEntry> Clone() const override
-				{
-					return MakeUnique<FOutputNodeRegistryEntry>(DataTypeEntry);
-				}
-
 			private:
 				
 				TSharedRef<IDataTypeRegistryEntry, ESPMode::ThreadSafe> DataTypeEntry;
@@ -302,11 +287,6 @@ namespace Metasound
 					return DataTypeEntry->CreateConstructorOutputNode(MoveTemp(InParams));
 				}
 
-				virtual TUniquePtr<INodeRegistryEntry> Clone() const override
-				{
-					return MakeUnique<FConstructorOutputNodeRegistryEntry>(DataTypeEntry);
-				}
-
 			private:
 				
 				TSharedRef<IDataTypeRegistryEntry, ESPMode::ThreadSafe> DataTypeEntry;
@@ -337,11 +317,6 @@ namespace Metasound
 					return DataTypeEntry->CreateLiteralNode(MoveTemp(InParams));
 				}
 
-				virtual TUniquePtr<INodeRegistryEntry> Clone() const override
-				{
-					return MakeUnique<FLiteralNodeRegistryEntry>(DataTypeEntry);
-				}
-
 			private:
 				
 				TSharedRef<IDataTypeRegistryEntry, ESPMode::ThreadSafe> DataTypeEntry;
@@ -372,11 +347,6 @@ namespace Metasound
 					return DataTypeEntry->CreateVariableNode(MoveTemp(InParams));
 				}
 
-				virtual TUniquePtr<INodeRegistryEntry> Clone() const override
-				{
-					return MakeUnique<FVariableNodeRegistryEntry>(DataTypeEntry);
-				}
-
 			private:
 				
 				TSharedRef<IDataTypeRegistryEntry, ESPMode::ThreadSafe> DataTypeEntry;
@@ -406,11 +376,6 @@ namespace Metasound
 					return DataTypeEntry->CreateVariableMutatorNode(InParams);
 				}
 
-				virtual TUniquePtr<INodeRegistryEntry> Clone() const override
-				{
-					return MakeUnique<FVariableMutatorNodeRegistryEntry>(DataTypeEntry);
-				}
-
 			private:
 				
 				TSharedRef<IDataTypeRegistryEntry, ESPMode::ThreadSafe> DataTypeEntry;
@@ -440,11 +405,6 @@ namespace Metasound
 					return DataTypeEntry->CreateVariableAccessorNode(InParams);
 				}
 
-				virtual TUniquePtr<INodeRegistryEntry> Clone() const override
-				{
-					return MakeUnique<FVariableAccessorNodeRegistryEntry>(DataTypeEntry);
-				}
-
 			private:
 				
 				TSharedRef<IDataTypeRegistryEntry, ESPMode::ThreadSafe> DataTypeEntry;
@@ -474,11 +434,6 @@ namespace Metasound
 					return DataTypeEntry->CreateVariableDeferredAccessorNode(InParams);
 				}
 
-				virtual TUniquePtr<INodeRegistryEntry> Clone() const override
-				{
-					return MakeUnique<FVariableDeferredAccessorNodeRegistryEntry>(DataTypeEntry);
-				}
-
 			private:
 				
 				TSharedRef<IDataTypeRegistryEntry, ESPMode::ThreadSafe> DataTypeEntry;
diff --git a/Engine/Plugins/Runtime/Metasound/Source/MetasoundFrontend/Private/MetasoundFrontendNodeClassRegistry.cpp b/Engine/Plugins/Runtime/Metasound/Source/MetasoundFrontend/Private/MetasoundFrontendNodeClassRegistry.cpp
new file mode 100644
--- /dev/null
+++ b/Engine/Plugins/Runtime/Metasound/Source/MetasoundFrontend/Private/MetasoundFrontendNodeClassRegistry.cpp
@@ -0,0 +1,1364 @@
+// Copyright Epic Games, Inc. All Rights Reserved.
+
+#include "MetasoundFrontendNodeClassRegistry.h"
+#include "MetasoundFrontendNodeClassRegistryPrivate.h"
+
+#include "Algo/ForEach.h"
+#include "Algo/Transform.h"
+#include "Containers/Array.h"
+#include "Containers/UnrealString.h"
+#include "Containers/Set.h"
+#include "HAL/PlatformTime.h"
+#include "Misc/Guid.h"
+#include "Misc/ScopeLock.h"
+#include "Tasks/Pipe.h"
+#include "Tasks/Task.h"
+#include "Templates/SharedPointer.h"
+#include "Templates/UniquePtr.h"
+#include "UObject/NameTypes.h"
+#include "UObject/ScriptInterface.h"
+
+#include "MetasoundDocumentInterface.h"
+#include "MetasoundFrontendDataTypeRegistry.h"
+#include "MetasoundFrontendDocumentBuilder.h"
+#include "MetasoundFrontendGraph.h"
+#include "MetasoundFrontendProxyDataCache.h"
+#include "MetasoundFrontendRegistryTransaction.h"
+#include "MetasoundFrontendSearchEngine.h"
+#include "MetasoundLog.h"
+#include "MetasoundRouter.h"
+#include "MetasoundTrace.h"
+
+namespace Metasound::Frontend
+{
+	namespace ConsoleVariables
+	{
+		static bool bDisableAsyncGraphRegistration = false;
+		FAutoConsoleVariableRef CVarMetaSoundDisableAsyncGraphRegistration(
+			TEXT("au.MetaSound.DisableAsyncGraphRegistration"),
+			Metasound::Frontend::ConsoleVariables::bDisableAsyncGraphRegistration,
+			TEXT("Disables async registration of MetaSound graphs\n")
+			TEXT("Default: false"),
+			ECVF_Default);
+	} // namespace ConsoleVariables
+
+	namespace RegistryPrivate
+	{
+		TScriptInterface<IMetaSoundDocumentInterface> BuildRegistryDocument(TScriptInterface<IMetaSoundDocumentInterface> DocumentInterface, bool bAsync)
+		{
+			METASOUND_TRACE_CPUPROFILER_EVENT_SCOPE(Metasound::Frontend::BuildRegistryDocument);
+
+			UObject* DocObject = DocumentInterface.GetObject();
+			check(DocObject);
+			const FMetasoundFrontendDocument& Document = DocumentInterface->GetConstDocument();
+
+#if WITH_EDITOR
+			// Only assets require template node processing
+			if (DocObject->IsAsset())
+			{
+				DocumentInterface = &UMetaSoundBuilderDocument::Create(*DocumentInterface.GetInterface());
+				FMetaSoundFrontendDocumentBuilder Builder(DocumentInterface);
+				Builder.TransformTemplateNodes();
+			}
+
+			return DocumentInterface;
+#else
+			const bool bIsBuilding = DocumentInterface->IsActivelyBuilding();
+			const bool bForceCopy = bIsBuilding && bAsync;
+
+	#if !NO_LOGGING
+		// Force a copy if async registration is enabled and we need to protect against race conditions from external modifications.
+
+		#if WITH_EDITORONLY_DATA
+			// Only assets require template node processing and support document attachment
+			if (DocObject->IsAsset())
+			{
+				const FMetaSoundFrontendDocumentBuilder& OriginalDocBuilder = IDocumentBuilderRegistry::GetChecked().FindOrBeginBuilding(DocumentInterface);
+				const bool bContainsTemplateDependency = OriginalDocBuilder.ContainsDependencyOfType(EMetasoundFrontendClassType::Template);
+				if (bContainsTemplateDependency)
+				{
+					UE_LOG(LogMetaSound, Error,
+						TEXT("Template node processing disabled but provided asset class at '%s' to register contains template nodes. Runtime graph will fail to build."),
+						*OriginalDocBuilder.GetDebugName());
+				}
+
+				// Destroy builder if one didn't exist before running template check to ensure
+				// that builder existence doesn't inadvertently cause potential future re-registration
+				// calls to perform unnecessary document copy below.
+				if (!bIsBuilding)
+				{
+					const FMetasoundFrontendClassName& ClassName = Document.RootGraph.Metadata.GetClassName();
+					IDocumentBuilderRegistry::GetChecked().FinishBuilding(ClassName);
+				}
+			}
+		#endif // WITH_EDITORONLY_DATA
+	#endif // !NO_LOGGING
+
+			if (bForceCopy)
+			{
+				return &UMetaSoundBuilderDocument::Create(*DocumentInterface.GetInterface());
+			}
+			else
+			{
+				return DocumentInterface;
+			}
+#endif // WITH_EDITOR
+		}
+
+		// FGraphNode is used to create unique INodes based off of a IGraph.
+		//
+		// Individual nodes need to reflect their InstanceName and InstanceID, but otherwise
+		// they simply encapsulate a shared set of behavior. To minimize memory usage, a single
+		// shared IGraph is used for all nodes referring to the same IGraph.
+		class FGraphNode : public INode
+		{
+			// This adapter class forwards the correct FBuilderOperatorParams
+			// to the graph's operator creation method. Many operator creation
+			// methods downcast the supplied INode in `FBuilderOperatorParams`
+			// and so it is required that it point to the correct runtime instance
+			// when calling CreateOperator(...)
+			class FGraphOperatorFactoryAdapter : public IOperatorFactory 
+			{
+			public:
+				FGraphOperatorFactoryAdapter(const IGraph& InGraph)
+				: Graph(&InGraph)
+				, GraphFactory(InGraph.GetDefaultOperatorFactory())
+				{
+				}
+
+				virtual ~FGraphOperatorFactoryAdapter() = default;
+
+				virtual TUniquePtr<IOperator> CreateOperator(const FBuildOperatorParams& InParams, FBuildResults& OutResults) override
+				{
+					FBuildOperatorParams ForwardParams 
+					{
+						*Graph,  // Point to correct INode instance
+						InParams.OperatorSettings,
+						InParams.InputData,
+						InParams.Environment,
+						InParams.Builder,
+						InParams.GraphRenderCost
+					};
+					return GraphFactory->CreateOperator(ForwardParams, OutResults);
+				}
+			private:
+				const IGraph* Graph; // Only store pointer because owning node keeps wrapped IGraph alive.  
+				FOperatorFactorySharedRef GraphFactory;
+			};
+
+		public:
+			FGraphNode(const FNodeInitData& InNodeInitData, TSharedRef<const IGraph> InGraphToWrap)
+			: InstanceID(InNodeInitData.InstanceID)
+			, Factory(MakeShared<FGraphOperatorFactoryAdapter>(*InGraphToWrap))
+			, Graph(MoveTemp(InGraphToWrap))
+			{
+			}
+
+			virtual const FName& GetInstanceName() const override
+			{
+				// Use the instance name of underlying graph because it refers
+				// to the actual asset name.
+				return Graph->GetInstanceName();
+			}
+
+			virtual const FGuid& GetInstanceID() const override
+			{
+				return InstanceID;
+			}
+
+			virtual const FNodeClassMetadata& GetMetadata() const override
+			{
+				return Graph->GetMetadata();
+			}
+
+			virtual const FVertexInterface& GetVertexInterface() const override
+			{
+				return Graph->GetVertexInterface();
+			}
+
+			virtual bool SetVertexInterface(const FVertexInterface& InInterface) override
+			{
+				// Cannot set vertex interface because of const reference to IGraph
+				// Return true if the supplied interface is the same as the existing interface.
+				return Graph->GetVertexInterface() == InInterface;
+			}
+
+			virtual bool IsVertexInterfaceSupported(const FVertexInterface& InInterface) const override
+			{
+				return Graph->IsVertexInterfaceSupported(InInterface);
+			}
+
+			virtual FOperatorFactorySharedRef GetDefaultOperatorFactory() const override
+			{
+				return Factory;
+			}
+
+		private:
+
+			FGuid InstanceID;
+			TSharedRef<FGraphOperatorFactoryAdapter> Factory;
+			TSharedRef<const IGraph> Graph;
+		};
+
+
+		// FDocumentNodeClassRegistryEntry encapsulates a node registry entry for a FGraph
+		class FDocumentNodeClassRegistryEntry : public INodeClassRegistryEntry
+		{
+		public:
+			FDocumentNodeClassRegistryEntry(const FMetasoundFrontendGraphClass& InGraphClass, const TSet<FMetasoundFrontendVersion>& InInterfaces, FNodeClassInfo&& InNodeClassInfo, TSharedPtr<const IGraph> InGraph)
+			: FrontendClass(InGraphClass)
+			, Interfaces(InInterfaces)
+			, ClassInfo(MoveTemp(InNodeClassInfo))
+			, Graph(InGraph)
+			{
+				FrontendClass.Metadata.SetType(EMetasoundFrontendClassType::External);
+			}
+
+			FDocumentNodeClassRegistryEntry(const FDocumentNodeClassRegistryEntry&) = default;
+
+			virtual ~FDocumentNodeClassRegistryEntry() = default;
+
+			virtual const FNodeClassInfo& GetClassInfo() const override
+			{
+				return ClassInfo;
+			}
+
+			virtual TUniquePtr<INode> CreateNode(const FNodeInitData& InNodeInitData) const override
+			{
+				if (Graph.IsValid())
+				{
+					return MakeUnique<FGraphNode>(InNodeInitData, Graph.ToSharedRef());
+				}
+				else
+				{
+					UE_LOG(LogMetaSound, Error, TEXT("Cannot create MetaSound node from asset %s due to prior failure to build graph"), *ClassInfo.AssetPath.ToString());
+					return TUniquePtr<INode>();
+				}
+			}
+
+			virtual TUniquePtr<INode> CreateNode(FDefaultLiteralNodeConstructorParams&&) const override { return nullptr; }
+			virtual TUniquePtr<INode> CreateNode(FDefaultNamedVertexNodeConstructorParams&&) const override { return nullptr; }
+			virtual TUniquePtr<INode> CreateNode(FDefaultNamedVertexWithLiteralNodeConstructorParams&&) const override { return nullptr; }
+
+			virtual const FMetasoundFrontendClass& GetFrontendClass() const override
+			{
+				return FrontendClass;
+			}
+
+			virtual const TSet<FMetasoundFrontendVersion>* GetImplementedInterfaces() const override
+			{
+				return &Interfaces;
+			}
+
+			virtual bool IsNative() const override
+			{
+				return false;
+			}
+
+		private:
+
+			FMetasoundFrontendClass FrontendClass;
+			TSet<FMetasoundFrontendVersion> Interfaces;
+			FNodeClassInfo ClassInfo;
+			TSharedPtr<const IGraph> Graph;
+		};
+	} // namespace RegistryPrivate
+
+
+
+
+
+	FNodeClassRegistryTransaction::FNodeClassRegistryTransaction(ETransactionType InType, const FNodeClassInfo& InNodeClassInfo, FNodeClassRegistryTransaction::FTimeType InTimestamp)
+	: Type(InType)
+	, NodeClassInfo(InNodeClassInfo)
+	, Timestamp(InTimestamp)
+	{
+	}
+
+	FNodeClassRegistryTransaction::ETransactionType FNodeClassRegistryTransaction::GetTransactionType() const
+	{
+		return Type;
+	}
+
+	const FNodeClassInfo& FNodeClassRegistryTransaction::GetNodeClassInfo() const
+	{
+		return NodeClassInfo;
+	}
+
+	FNodeClassRegistryKey FNodeClassRegistryTransaction::GetNodeRegistryKey() const
+	{
+		return FNodeClassRegistryKey(NodeClassInfo);
+	}
+
+	FNodeClassRegistryTransaction::FTimeType FNodeClassRegistryTransaction::GetTimestamp() const
+	{
+		return Timestamp;
+	}
+
+	namespace NodeClassRegistryKey
+	{
+		FNodeClassRegistryKey CreateKey(EMetasoundFrontendClassType InType, const FString& InFullClassName, int32 InMajorVersion, int32 InMinorVersion)
+		{
+			if (InType == EMetasoundFrontendClassType::Graph)
+			{
+				// No graphs are registered. Any registered graph should be registered as an external node.
+				InType = EMetasoundFrontendClassType::External;
+			}
+
+			FMetasoundFrontendClassName ClassName;
+			FMetasoundFrontendClassName::Parse(InFullClassName, ClassName);
+			return FNodeClassRegistryKey(InType, ClassName, InMajorVersion, InMinorVersion);
+		}
+
+		const FNodeClassRegistryKey& GetInvalid()
+		{
+			return FNodeClassRegistryKey::GetInvalid();
+		}
+
+		bool IsValid(const FNodeClassRegistryKey& InKey)
+		{
+			return InKey.IsValid();
+		}
+
+		bool IsEqual(const FNodeClassRegistryKey& InLHS, const FNodeClassRegistryKey& InRHS)
+		{
+			return InLHS == InRHS;
+		}
+
+		bool IsEqual(const FMetasoundFrontendClassMetadata& InLHS, const FMetasoundFrontendClassMetadata& InRHS)
+		{
+			if (InLHS.GetClassName() == InRHS.GetClassName())
+			{
+				if (InLHS.GetType() == InRHS.GetType())
+				{
+					if (InLHS.GetVersion() == InRHS.GetVersion())
+					{
+						return true;
+					}
+				}
+			}
+			return false;
+		}
+
+		bool IsEqual(const FNodeClassInfo& InLHS, const FMetasoundFrontendClassMetadata& InRHS)
+		{
+			if (InLHS.ClassName == InRHS.GetClassName())
+			{
+				if (InLHS.Type == InRHS.GetType())
+				{
+					if (InLHS.Version == InRHS.GetVersion())
+					{
+						return true;
+					}
+				}
+			}
+			return false;
+		}
+
+		FNodeClassRegistryKey CreateKey(const FNodeClassMetadata& InNodeMetadata)
+		{
+			return FNodeClassRegistryKey(InNodeMetadata);
+		}
+
+		FNodeClassRegistryKey CreateKey(const FMetasoundFrontendClassMetadata& InNodeMetadata)
+		{
+			checkf(InNodeMetadata.GetType() != EMetasoundFrontendClassType::Graph, TEXT("Cannot create key from 'graph' type. Likely meant to use CreateKey overload that is provided FMetasoundFrontendGraphClass"));
+			return FNodeClassRegistryKey(InNodeMetadata);
+		}
+
+		FNodeClassRegistryKey CreateKey(const FMetasoundFrontendGraphClass& InGraphClass)
+		{
+			return FNodeClassRegistryKey(InGraphClass);
+		}
+
+		FNodeClassRegistryKey CreateKey(const FNodeClassInfo& InClassInfo)
+		{
+			return FNodeClassRegistryKey(InClassInfo);
+		}
+	} // namespace NodeClassRegistryKey
+	  
+	void FNodeClassRegistry::BuildAndRegisterGraphFromDocument(const FMetasoundFrontendDocument& InDocument, const FProxyDataCache& InProxyDataCache, FNodeClassInfo&& InNodeClassInfo)
+	{
+		using namespace RegistryPrivate;
+
+		METASOUND_TRACE_CPUPROFILER_EVENT_SCOPE(Metasound::FNodeClassRegistry::BuildAndRegisterGraphFromDocument);
+		METASOUND_TRACE_CPUPROFILER_EVENT_SCOPE_TEXT(*FString::Printf(TEXT("Metasound::FNodeClassRegistry::BuildAndRegisterGraphFromDocument asset %s"), *InNodeClassInfo.AssetPath.ToString()));
+
+		// Use the asset class id for the graph id because it should be locally unique. Unlike other ids, it is regenerated on asset duplicate. 
+		TUniquePtr<FFrontendGraph> FrontendGraph = FFrontendGraphBuilder::CreateGraph(InDocument, InProxyDataCache, InNodeClassInfo.AssetPath.ToString(), /*GraphId=*/InNodeClassInfo.AssetClassID);
+		if (!FrontendGraph.IsValid())
+		{
+			UE_LOG(LogMetaSound, Error, TEXT("Failed to build MetaSound graph in asset '%s'"), *InNodeClassInfo.AssetPath.ToString());
+		}
+
+		TSharedPtr<const FGraph> GraphToRegister = MakeShareable<const FGraph>(FrontendGraph.Release());
+
+		TUniquePtr<INodeClassRegistryEntry> RegistryEntry = MakeUnique<FDocumentNodeClassRegistryEntry>(InDocument.RootGraph, InDocument.Interfaces, MoveTemp(InNodeClassInfo), GraphToRegister);
+
+		const FNodeClassRegistryKey RegistryKey = RegisterNodeInternal(MoveTemp(RegistryEntry));
+
+		// Key must use the asset path provided to class info and *NOT* that of the
+		// document's owning DocumentInterface object, as that may be a built transient object
+		// with a different transient asset path.
+		const FGraphRegistryKey GraphKey { RegistryKey, InNodeClassInfo.AssetPath };
+		RegisterGraphInternal(GraphKey, GraphToRegister);
+	}
+
+	FNodeClassRegistry* FNodeClassRegistry::LazySingleton = nullptr;
+
+	FNodeClassRegistry& FNodeClassRegistry::Get()
+	{
+		if (!LazySingleton)
+		{
+			LazySingleton = new Metasound::Frontend::FNodeClassRegistry();
+		}
+
+		return *LazySingleton;
+	}
+
+	void FNodeClassRegistry::Shutdown()
+	{
+		if (nullptr != LazySingleton)
+		{
+			delete LazySingleton;
+			LazySingleton = nullptr;
+		}
+	}
+
+	FNodeClassRegistry::FNodeClassRegistry()
+	: TransactionBuffer(MakeShared<FNodeClassRegistryTransactionBuffer>())
+	, AsyncRegistrationPipe( UE_SOURCE_LOCATION )
+	{
+	}
+
+	void FNodeClassRegistry::RegisterPendingNodes()
+	{
+		METASOUND_LLM_SCOPE;
+		METASOUND_TRACE_CPUPROFILER_EVENT_SCOPE(metasound::FNodeClassRegistry::RegisterPendingNodes);
+		{
+			FScopeLock ScopeLock(&LazyInitCommandCritSection);
+
+			for (TUniqueFunction<void()>& Command : LazyInitCommands)
+			{
+				Command();
+			}
+
+			LazyInitCommands.Empty();
+		}
+
+		if (!IsRunningCommandlet())
+		{
+			// Prime search engine after bulk registration.
+			ISearchEngine::Get().Prime();
+		}
+	}
+
+	bool FNodeClassRegistry::EnqueueInitCommand(TUniqueFunction<void()>&& InFunc)
+	{
+
+		FScopeLock ScopeLock(&LazyInitCommandCritSection);
+		if (LazyInitCommands.Num() >= MaxNumNodesAndDatatypesToInitialize)
+		{
+			UE_LOG(LogMetaSound, Warning, TEXT("Registering more that %d nodes and datatypes for metasounds! Consider increasing MetasoundFrontendRegistryContainer::MaxNumNodesAndDatatypesToInitialize."), MaxNumNodesAndDatatypesToInitialize);
+		}
+
+		LazyInitCommands.Add(MoveTemp(InFunc));
+		return true;
+	}
+
+	void FNodeClassRegistry::SetObjectReferencer(TUniquePtr<IObjectReferencer> InReferencer)
+	{
+		FScopeLock LockActiveReg(&ActiveRegistrationTasksCriticalSection);
+		checkf(ActiveRegistrationTasks.IsEmpty(), TEXT("Object Referencer cannot be set while registry is actively being manipulated"));
+		ObjectReferencer = MoveTemp(InReferencer);
+	}
+
+	TUniquePtr<Metasound::INode> FNodeClassRegistry::CreateNode(const FNodeClassRegistryKey& InKey, const Metasound::FNodeInitData& InInitData) const
+	{
+		TUniquePtr<INode> Node;
+		auto CreateNodeLambda = [&Node, &InInitData](const INodeClassRegistryEntry& Entry)
+		{ 
+			Node = Entry.CreateNode(InInitData); 
+		};
+
+		if (!AccessNodeEntryThreadSafe(InKey, CreateNodeLambda))
+		{
+			// Creation of external nodes can rely on assets being unavailable due to errors in loading order, asset(s)
+			// missing, etc. 
+			UE_LOG(LogMetaSound, Error, TEXT("Could not find node [RegistryKey:%s]"), *InKey.ToString());
+		}
+
+		return MoveTemp(Node);
+	}
+
+	TUniquePtr<Metasound::INode> FNodeClassRegistry::CreateNode(const FNodeClassRegistryKey& InKey, FDefaultLiteralNodeConstructorParams&& InParams) const
+	{
+		TUniquePtr<INode> Node;
+		auto CreateNodeLambda = [&Node, &InParams](const INodeClassRegistryEntry& Entry)
+		{ 
+			Node = Entry.CreateNode(MoveTemp(InParams)); 
+		};
+
+		if (!AccessNodeEntryThreadSafe(InKey, CreateNodeLambda))
+		{
+			UE_LOG(LogMetaSound, Error, TEXT("Could not find node [RegistryKey:%s]"), *InKey.ToString());
+		}
+
+		return MoveTemp(Node);
+	}
+
+	TUniquePtr<Metasound::INode> FNodeClassRegistry::CreateNode(const FNodeClassRegistryKey& InKey, FDefaultNamedVertexNodeConstructorParams&& InParams) const
+	{
+		TUniquePtr<INode> Node;
+		auto CreateNodeLambda = [&Node, &InParams](const INodeClassRegistryEntry& Entry) 
+		{ 
+			Node = Entry.CreateNode(MoveTemp(InParams)); 
+		};
+
+		if (!AccessNodeEntryThreadSafe(InKey, CreateNodeLambda))
+		{
+			UE_LOG(LogMetaSound, Error, TEXT("Could not find node [RegistryKey:%s]"), *InKey.ToString());
+		}
+
+		return MoveTemp(Node);
+	}
+
+	TUniquePtr<Metasound::INode> FNodeClassRegistry::CreateNode(const FNodeClassRegistryKey& InKey, FDefaultNamedVertexWithLiteralNodeConstructorParams&& InParams) const
+	{
+		TUniquePtr<INode> Node;
+		auto CreateNodeLambda = [&Node, &InParams](const INodeClassRegistryEntry& Entry)
+		{ 
+			Node = Entry.CreateNode(MoveTemp(InParams));
+		};
+
+		if (!AccessNodeEntryThreadSafe(InKey, CreateNodeLambda))
+		{
+			UE_LOG(LogMetaSound, Error, TEXT("Could not find node [RegistryKey:%s]"), *InKey.ToString());
+		}
+
+		return MoveTemp(Node);
+	}
+
+	TArray<::Metasound::Frontend::FConverterNodeInfo> FNodeClassRegistry::GetPossibleConverterNodes(const FName& FromDataType, const FName& ToDataType)
+	{
+		FConverterNodeClassRegistryKey InKey = { FromDataType, ToDataType };
+		if (!ConverterNodeClassRegistry.Contains(InKey))
+		{
+			return TArray<FConverterNodeInfo>();
+		}
+		else
+		{
+			return ConverterNodeClassRegistry[InKey].PotentialConverterNodes;
+		}
+	}
+
+	TUniquePtr<FNodeClassRegistryTransactionStream> FNodeClassRegistry::CreateTransactionStream()
+	{
+		return MakeUnique<FNodeClassRegistryTransactionStream>(TransactionBuffer);
+	}
+
+	FGraphRegistryKey FNodeClassRegistry::RegisterGraph(const TScriptInterface<IMetaSoundDocumentInterface>& InDocumentInterface)
+	{
+		using namespace UE;
+
+		METASOUND_TRACE_CPUPROFILER_EVENT_SCOPE(Metasound::FNodeClassRegistry::RegisterGraph);
+
+		check(InDocumentInterface);
+		check(IsInGameThread());
+
+		const FMetasoundFrontendDocument& Document = InDocumentInterface->GetConstDocument();
+		const FTopLevelAssetPath AssetPath = InDocumentInterface->GetAssetPathChecked();
+		const FGraphRegistryKey RegistryKey { FNodeClassRegistryKey(Document.RootGraph), AssetPath };
+
+		if (!RegistryKey.IsValid())
+		{
+			// Do not attempt to build and register a MetaSound with an invalid registry key
+			UE_LOG(LogMetaSound, Warning, TEXT("Registry key is invalid when attemping to register graph for asset %s"), *AssetPath.ToString());
+			return RegistryKey;
+		}
+
+		METASOUND_TRACE_CPUPROFILER_EVENT_SCOPE_TEXT(*FString::Printf(TEXT("FNodeClassRegistry::RegisterGraph key:%s, asset %s"), *RegistryKey.ToString(), *AssetPath.ToString()));
+
+		// Wait for any async tasks that are in flight which correspond to the same graph prior to building, even if this is a synchronous call.
+		WaitForAsyncGraphRegistration(RegistryKey);
+
+		const bool bAsync = !ConsoleVariables::bDisableAsyncGraphRegistration;
+
+		// Use the asset path of the provided document interface object for identification, *NOT* the
+		// built version as the build process may in fact create a new object with a transient path.
+		const TScriptInterface<IMetaSoundDocumentInterface> RegistryDocInterface = RegistryPrivate::BuildRegistryDocument(InDocumentInterface, bAsync);
+
+		UObject* OwningObject = RegistryDocInterface.GetObject();
+		check(OwningObject);
+
+		// Proxies are created synchronously to avoid creating proxies in async tasks. Proxies
+		// are created from UObjects which need to be protected from GC and non-GT access.
+		FProxyDataCache ProxyDataCache;
+		ProxyDataCache.CreateAndCacheProxies(Document);
+
+#if !NO_LOGGING
+		if (UE_LOG_ACTIVE(LogMetaSound, Verbose))
+		{
+			const FGuid PageID = IDocumentBuilderRegistry::GetChecked().ResolveTargetPageID(Document.RootGraph);
+			const bool bContainsMultipleGraphs = Document.RootGraph.GetConstGraphPages().Num() > 1;
+			if (bContainsMultipleGraphs || PageID != Metasound::Frontend::DefaultPageID)
+			{
+				UE_LOG(LogMetaSound, Verbose, TEXT("Registered MetaSound '%s' Graph Page with PageID '%s'."),
+					*AssetPath.GetAssetName().ToString(),
+					*PageID.ToString());
+				if (bContainsMultipleGraphs)
+				{
+					UE_LOG(LogMetaSound, Verbose, TEXT("Graphs found with following PageIDs Implemented:"));
+					Document.RootGraph.IterateGraphPages([](const FMetasoundFrontendGraph& Graph)
+					{
+						UE_LOG(LogMetaSound, Verbose, TEXT("    - %s'"), *Graph.PageID.ToString());
+					});
+				}
+			}
+		}
+#endif // !NO_LOGGING
+
+		// Store update to newly registered node in history so nodes
+		// can be queried by transaction ID
+		FNodeClassInfo NodeClassInfo(Document.RootGraph, AssetPath);
+		{
+			FNodeClassRegistryTransaction::FTimeType Timestamp = FPlatformTime::Cycles64();
+			TransactionBuffer->AddTransaction(FNodeClassRegistryTransaction(FNodeClassRegistryTransaction::ETransactionType::NodeRegistration, NodeClassInfo, Timestamp));
+		}
+
+		if (bAsync)
+		{
+			Tasks::FTask BuildAndRegisterTask = AsyncRegistrationPipe.Launch(
+				UE_SOURCE_LOCATION,
+				[RegistryKey, ClassInfo = MoveTemp(NodeClassInfo), RegistryDocInterface, ProxyDataCache = MoveTemp(ProxyDataCache)]() mutable
+				{
+					FNodeClassRegistry& Registry = FNodeClassRegistry::Get();
+					// Unregister the graph before re-registering
+					Registry.UnregisterGraphInternal(RegistryKey);
+					Registry.BuildAndRegisterGraphFromDocument(RegistryDocInterface->GetConstDocument(), ProxyDataCache, MoveTemp(ClassInfo));
+					Registry.RemoveRegistrationTask(RegistryKey, FNodeClassRegistryTransaction::ETransactionType::NodeRegistration);
+					Registry.RemoveDocumentReference(RegistryDocInterface);
+				}
+			);
+
+			AddDocumentReference(RegistryDocInterface);
+			AddRegistrationTask(RegistryKey, FActiveRegistrationTaskInfo
+			{
+				FNodeClassRegistryTransaction::ETransactionType::NodeRegistration,
+				BuildAndRegisterTask,
+				AssetPath
+			});
+		}
+		else
+		{
+			UnregisterGraphInternal(RegistryKey);
+
+			// Build and register graph synchronously
+			BuildAndRegisterGraphFromDocument(RegistryDocInterface->GetConstDocument(), ProxyDataCache, MoveTemp(NodeClassInfo));
+		}
+
+		return RegistryKey;
+	}
+
+	void FNodeClassRegistry::AddRegistrationTask(const FGraphRegistryKey& InKey, FActiveRegistrationTaskInfo&& TaskInfo)
+	{
+		FScopeLock LockActiveReg(&ActiveRegistrationTasksCriticalSection);
+		ActiveRegistrationTasks.FindOrAdd(InKey.NodeKey).Add(MoveTemp(TaskInfo));
+	}
+
+	void FNodeClassRegistry::RemoveRegistrationTask(const FGraphRegistryKey& InKey, FNodeClassRegistryTransaction::ETransactionType TransactionType)
+	{
+		FScopeLock LockActiveReg(&ActiveRegistrationTasksCriticalSection);
+
+		int32 NumRemoved = 0;
+		if (InKey.AssetPath.IsNull())  // Null provided path instructs to remove all tasks related to the underlying Node Registry Key
+		{
+			NumRemoved = ActiveRegistrationTasks.Remove(InKey.NodeKey);
+		}
+		else if (TArray<FActiveRegistrationTaskInfo>* TaskInfos = ActiveRegistrationTasks.Find(InKey.NodeKey))
+		{
+			auto MatchesEntryInTask = [&InKey, &TransactionType](const FActiveRegistrationTaskInfo& Info)
+			{
+				const bool bIsAssetPath = Info.AssetPath == InKey.AssetPath;
+				const bool bIsTransactionType = Info.TransactionType == TransactionType;
+				return bIsAssetPath && bIsTransactionType;
+			};
+
+			NumRemoved = TaskInfos->RemoveAllSwap(MatchesEntryInTask, EAllowShrinking::No);
+			if (TaskInfos->IsEmpty())
+			{
+				ActiveRegistrationTasks.Remove(InKey.NodeKey);
+			}
+		}
+
+
+		if (NumRemoved == 0)
+		{
+			const bool bIsCooking = IsRunningCookCommandlet();
+			if (ensureMsgf(!bIsCooking,
+				TEXT("Failed to find active %s tasks for the graph '%s': Async registration is not supported while cooking"),
+				*FNodeClassRegistryTransaction::LexToString(TransactionType),
+				*InKey.ToString()))
+			{
+				UE_LOG(LogMetaSound, Warning,
+					TEXT("Failed to find active %s tasks for the graph '%s'."),
+					*FNodeClassRegistryTransaction::LexToString(TransactionType),
+					*InKey.ToString());
+			}
+		}
+	}
+
+	void FNodeClassRegistry::AddDocumentReference(TScriptInterface<IMetaSoundDocumentInterface> DocumentInterface)
+	{
+		FScopeLock LockActiveReg(&ObjectReferencerCriticalSection);
+		if (UObject* Object = DocumentInterface.GetObject())
+		{
+			if (ObjectReferencer)
+			{
+				ObjectReferencer->AddObject(Object);
+			}
+		}
+	}
+
+	void FNodeClassRegistry::RemoveDocumentReference(TScriptInterface<IMetaSoundDocumentInterface> DocumentInterface)
+	{
+		FScopeLock LockActiveReg(&ObjectReferencerCriticalSection);
+		if (UObject* Object = DocumentInterface.GetObject())
+		{
+			if (ObjectReferencer)
+			{
+				ObjectReferencer->RemoveObject(Object);
+			}
+		}
+	}
+
+	void FNodeClassRegistry::RegisterGraphInternal(const FGraphRegistryKey& InKey, TSharedPtr<const FGraph> InGraph)
+	{
+		using namespace RegistryPrivate;
+
+		FScopeLock Lock(&RegistryMapsCriticalSection);
+
+#if !NO_LOGGING
+		if (RegisteredGraphs.Contains(InKey))
+		{
+			UE_LOG(LogMetaSound, Warning, TEXT("Graph is already registered with the same registry key '%s'. The existing registered graph will be replaced with the new graph."), *InKey.ToString());
+		}
+#endif // !NO_LOGGING
+
+		RegisteredGraphs.Add(InKey, InGraph);
+	}
+
+	bool FNodeClassRegistry::UnregisterGraphInternal(const FGraphRegistryKey& InKey)
+	{
+		using namespace RegistryPrivate;
+
+		METASOUND_TRACE_CPUPROFILER_EVENT_SCOPE_TEXT(*InKey.ToString(TEXT("FNodeClassRegistry::UnregisterGraphInternal")));
+
+		FScopeLock Lock(&RegistryMapsCriticalSection);
+		{
+			if (!RegisteredGraphs.Contains(InKey))
+			{
+				return false;
+			}
+
+			const int32 GraphUnregistered = RegisteredGraphs.Remove(InKey) > 0;
+			const bool bNodeUnregistered = UnregisterNodeInternal(InKey.NodeKey);
+
+#if !NO_LOGGING
+			if (GraphUnregistered)
+			{
+				UE_LOG(LogMetaSound, VeryVerbose, TEXT("Unregistered graph with key '%s'"), *InKey.ToString());
+			}
+			else
+			{
+				// Avoid warning if in cook as we always expect a graph to not get registered/
+				// unregistered while cooking (as its unnecessary for serialization).
+				if (bNodeUnregistered && !IsRunningCookCommandlet())
+				{
+					UE_LOG(LogMetaSound, Warning, TEXT("Graph '%s' was not found, but analogous registered node class was when unregistering."), *InKey.ToString());
+				}
+			}
+#endif // !NO_LOGGING
+
+			return bNodeUnregistered;
+		}
+
+		return false;
+	}
+
+	bool FNodeClassRegistry::UnregisterGraph(const FGraphRegistryKey& InRegistryKey, const TScriptInterface<IMetaSoundDocumentInterface>& InDocumentInterface)
+	{
+		using namespace UE;
+		using namespace RegistryPrivate;
+
+		METASOUND_TRACE_CPUPROFILER_EVENT_SCOPE(Metasound::FNodeClassRegistry::UnregisterGraph);
+		METASOUND_TRACE_CPUPROFILER_EVENT_SCOPE_TEXT(*InRegistryKey.ToString(TEXT("FNodeClassRegistry::UnregisterGraph")));
+
+		check(IsInGameThread());
+		check(InDocumentInterface);
+
+		// Do not attempt to unregister a MetaSound with an invalid registry key
+		if (!InRegistryKey.IsValid())
+		{
+			UE_LOG(LogMetaSound, Warning, TEXT("Registry key is invalid when attempting to unregister graph (%s)"), *InRegistryKey.ToString());
+			return false;
+		}
+
+		const FMetasoundFrontendDocument& Document = InDocumentInterface->GetConstDocument();
+		FNodeClassInfo NodeClassInfo(Document.RootGraph.Metadata);
+
+		// This is a hack to avoid requiring the asset path to be passed while unregistering.
+		// The asset path may be invalid by this point if the object being unregistered is being GC'ed.
+		// FNodeClassInfo needs to be deprecated in favor of more precise types as a key, editor data, etc.
+		// Its currently kind of a dumping ground as it stands.
+		NodeClassInfo.Type = EMetasoundFrontendClassType::External;
+
+		// Store update to unregistered node in history so nodes can be queried by transaction ID
+		{
+			FNodeClassRegistryTransaction::FTimeType Timestamp = FPlatformTime::Cycles64();
+			TransactionBuffer->AddTransaction(FNodeClassRegistryTransaction(FNodeClassRegistryTransaction::ETransactionType::NodeUnregistration, NodeClassInfo, Timestamp));
+		}
+
+		// Async registration is only available if:
+		// 1. The IMetaSoundDocumentInterface is not actively modified by a builder
+		//    (built graph must be released synchronously to avoid a race condition on
+		//    reading/writing the IMetaSoundDocumentInterface on the Game Thread)
+		// 2. Async registration is globally disabled via console variable.
+		const bool bAsync = !(InDocumentInterface->IsActivelyBuilding() || ConsoleVariables::bDisableAsyncGraphRegistration);
+		if (bAsync)
+		{
+			// Wait for any async tasks that are in flight which correspond to the same graph
+			WaitForAsyncGraphRegistration(InRegistryKey);
+
+			Tasks::FTask UnregisterTask = AsyncRegistrationPipe.Launch(UE_SOURCE_LOCATION, [RegistryKey = InRegistryKey]()
+			{
+				FNodeClassRegistry& Registry = FNodeClassRegistry::Get();
+				Registry.UnregisterGraphInternal(RegistryKey);
+				Registry.RemoveRegistrationTask(RegistryKey, FNodeClassRegistryTransaction::ETransactionType::NodeUnregistration);
+			});
+
+			AddRegistrationTask(InRegistryKey, FActiveRegistrationTaskInfo
+			{
+				FNodeClassRegistryTransaction::ETransactionType::NodeUnregistration,
+				UnregisterTask,
+				InRegistryKey.AssetPath
+			});
+		}
+		else
+		{
+			UnregisterGraphInternal(InRegistryKey);
+		}
+
+		return true;
+	}
+
+	TSharedPtr<const Metasound::FGraph> FNodeClassRegistry::GetGraph(const FGraphRegistryKey& InKey) const
+	{
+		WaitForAsyncGraphRegistration(InKey);
+
+		TSharedPtr<const FGraph> Graph;
+		{
+			FScopeLock Lock(&RegistryMapsCriticalSection);
+			if (const TSharedPtr<const FGraph>* RegisteredGraph = RegisteredGraphs.Find(InKey))
+			{
+				Graph = *RegisteredGraph;
+			}
+		}
+
+		if (!Graph)
+		{
+			UE_LOG(LogMetaSound, Error, TEXT("Could not find graph with registry graph key '%s'."), *InKey.ToString());
+		}
+
+		return Graph;
+	}
+
+	FNodeClassRegistryKey FNodeClassRegistry::RegisterNodeInternal(TUniquePtr<INodeClassRegistryEntry>&& InEntry)
+	{
+		using namespace RegistryPrivate;
+
+		METASOUND_LLM_SCOPE;
+
+		if (!InEntry.IsValid())
+		{
+			return { };
+		}
+
+		const FNodeClassRegistryKey Key = FNodeClassRegistryKey(InEntry->GetClassInfo());
+		METASOUND_TRACE_CPUPROFILER_EVENT_SCOPE_TEXT(*Key.ToString(TEXT("FNodeClassRegistry::RegisterNodeInternal")))
+
+		TSharedRef<INodeClassRegistryEntry, ESPMode::ThreadSafe> Entry(InEntry.Release());
+		{
+			FScopeLock Lock(&RegistryMapsCriticalSection);
+
+			// check to see if an identical node was already registered, and log
+			if (const TSharedRef<INodeClassRegistryEntry>* ExistingEntry = RegisteredNodes.Find(Key))
+			{
+				const FNodeClassInfo& ClassInfo = (*ExistingEntry)->GetClassInfo();
+				UE_LOG(LogMetaSound, Error,
+					TEXT("Node with registry key '%s' already registered by asset '%s' encountered while registering node with asset %s. MetaSounds which depend on these assets may utilize incorrect asset."),
+					*Key.ToString(),
+					*ClassInfo.AssetPath.ToString(),
+					*Entry->GetClassInfo().AssetPath.ToString());
+			}
+
+			// Store registry elements in map so nodes can be queried using registry key.
+			RegisteredNodes.Add(Key, Entry);
+		}
+
+		return Key;
+	}
+
+	FNodeClassRegistryKey FNodeClassRegistry::RegisterNode(TUniquePtr<INodeClassRegistryEntry>&& InEntry)
+	{
+		const FNodeClassInfo ClassInfo = InEntry->GetClassInfo();
+		const FNodeClassRegistryKey Key = RegisterNodeInternal(MoveTemp(InEntry));
+
+		if (Key.IsValid())
+		{
+			// Store update to newly registered node in history so nodes
+			// can be queried by transaction ID
+			const FNodeClassRegistryTransaction::FTimeType Timestamp = FPlatformTime::Cycles64();
+			TransactionBuffer->AddTransaction(FNodeClassRegistryTransaction(FNodeClassRegistryTransaction::ETransactionType::NodeRegistration, ClassInfo, Timestamp));
+		}
+
+		return Key;
+	}
+
+	FNodeClassRegistryKey FNodeClassRegistry::RegisterNodeTemplate(TUniquePtr<INodeTemplateRegistryEntry>&& InEntry)
+	{
+		METASOUND_LLM_SCOPE;
+
+		FNodeClassRegistryKey Key;
+
+		if (InEntry.IsValid())
+		{
+			TSharedRef<INodeTemplateRegistryEntry, ESPMode::ThreadSafe> Entry(InEntry.Release());
+
+			FNodeClassRegistryTransaction::FTimeType Timestamp = FPlatformTime::Cycles64();
+
+			Key = FNodeClassRegistryKey(Entry->GetClassInfo());
+
+			{
+				FScopeLock Lock(&RegistryMapsCriticalSection);
+				// check to see if an identical node was already registered, and log
+				ensureAlwaysMsgf(
+					!RegisteredNodeTemplates.Contains(Key),
+					TEXT("Node template with registry key '%s' already registered. The previously registered node will be overwritten."),
+					*Key.ToString());
+
+				// Store registry elements in map so nodes can be queried using registry key.
+				RegisteredNodeTemplates.Add(Key, Entry);
+			}
+
+			// Store update to newly registered node in history so nodes
+			// can be queried by transaction ID
+
+			TransactionBuffer->AddTransaction(FNodeClassRegistryTransaction(FNodeClassRegistryTransaction::ETransactionType::NodeRegistration, Entry->GetClassInfo(), Timestamp));
+		}
+
+		return Key;
+	}
+
+	bool FNodeClassRegistry::UnregisterNodeInternal(const FNodeClassRegistryKey& InKey, FNodeClassInfo* OutClassInfo)
+	{
+		METASOUND_LLM_SCOPE;
+
+		if (InKey.IsValid())
+		{
+			METASOUND_TRACE_CPUPROFILER_EVENT_SCOPE_TEXT(*FString::Printf(TEXT("FNodeClassRegistry::UnregisterNodeInternal key %s"), *InKey.ToString()))
+
+			FScopeLock Lock(&RegistryMapsCriticalSection);
+			if (const TSharedRef<INodeClassRegistryEntry, ESPMode::ThreadSafe>* EntryPtr = RegisteredNodes.Find(InKey))
+			{
+				const TSharedRef<INodeClassRegistryEntry>& Entry = *EntryPtr;
+				if (OutClassInfo)
+				{
+					*OutClassInfo = Entry->GetClassInfo();
+				}
+				const uint32 NumRemoved = RegisteredNodes.RemoveSingle(InKey, Entry);
+				if (ensure(NumRemoved == 1))
+				{
+					return true;
+				}
+			}
+		}
+
+		if (OutClassInfo)
+		{
+			*OutClassInfo = { };
+		}
+		return false;
+	}
+
+	bool FNodeClassRegistry::UnregisterNode(const FNodeClassRegistryKey& InKey)
+	{
+		FNodeClassInfo ClassInfo;
+		if (UnregisterNodeInternal(InKey, &ClassInfo))
+		{
+			const FNodeClassRegistryTransaction::FTimeType Timestamp = FPlatformTime::Cycles64();
+			TransactionBuffer->AddTransaction(FNodeClassRegistryTransaction(FNodeClassRegistryTransaction::ETransactionType::NodeUnregistration, ClassInfo, Timestamp));
+
+			return true;
+		}
+
+		return false;
+	}
+
+	bool FNodeClassRegistry::UnregisterNodeTemplate(const FNodeClassRegistryKey& InKey)
+	{
+		METASOUND_LLM_SCOPE;
+
+		if (InKey.IsValid())
+		{
+			if (const INodeTemplateRegistryEntry* Entry = FindNodeTemplateEntry(InKey))
+			{
+				FNodeClassRegistryTransaction::FTimeType Timestamp = FPlatformTime::Cycles64();
+
+				TransactionBuffer->AddTransaction(FNodeClassRegistryTransaction(FNodeClassRegistryTransaction::ETransactionType::NodeUnregistration, Entry->GetClassInfo(), Timestamp));
+
+				{
+					FScopeLock Lock(&RegistryMapsCriticalSection);
+					RegisteredNodeTemplates.Remove(InKey);
+				}
+				return true;
+			}
+		}
+
+		return false;
+	}
+
+	bool FNodeClassRegistry::RegisterConversionNode(const FConverterNodeClassRegistryKey& InNodeKey, const FConverterNodeInfo& InNodeInfo)
+	{
+		if (!ConverterNodeClassRegistry.Contains(InNodeKey))
+		{
+			ConverterNodeClassRegistry.Add(InNodeKey);
+		}
+
+		FConverterNodeClassRegistryValue& ConverterNodeList = ConverterNodeClassRegistry[InNodeKey];
+
+		if (ensureAlways(!ConverterNodeList.PotentialConverterNodes.Contains(InNodeInfo)))
+		{
+			ConverterNodeList.PotentialConverterNodes.Add(InNodeInfo);
+			return true;
+		}
+		else
+		{
+			// If we hit this, someone attempted to add the same converter node to our list multiple times.
+			return false;
+		}
+	}
+
+	bool FNodeClassRegistry::IsNodeRegistered(const FNodeClassRegistryKey& InKey) const
+	{
+		auto IsNodeRegisteredInternal = [this, &InKey]() -> bool
+		{
+			FScopeLock Lock(&RegistryMapsCriticalSection);
+			return RegisteredNodes.Contains(InKey) || RegisteredNodeTemplates.Contains(InKey);
+		};
+
+		if (IsNodeRegisteredInternal())
+		{
+			return true;
+		}
+		else
+		{
+			WaitForAsyncRegistrationInternal(InKey, nullptr /* InAssetPath */);
+			return IsNodeRegisteredInternal();
+		}
+	}
+
+	bool FNodeClassRegistry::IsGraphRegistered(const FGraphRegistryKey& InKey) const
+	{
+		WaitForAsyncGraphRegistration(InKey);
+
+		{
+			FScopeLock Lock(&RegistryMapsCriticalSection);
+			return RegisteredGraphs.Contains(InKey);
+		}
+	}
+
+	bool FNodeClassRegistry::IsNodeNative(const FNodeClassRegistryKey& InKey) const
+	{
+		bool bIsNative = false;
+		auto SetIsNative = [&bIsNative](const INodeClassRegistryEntry& Entry) 
+		{ 
+			bIsNative = Entry.IsNative();
+		};
+
+		if (AccessNodeEntryThreadSafe(InKey, SetIsNative))
+		{
+			return bIsNative;
+		}
+
+		if (const INodeTemplateRegistryEntry* TemplateEntry = FindNodeTemplateEntry(InKey))
+		{
+			return true;
+		}
+		return false;
+	}
+
+	bool FNodeClassRegistry::FindFrontendClassFromRegistered(const FNodeClassRegistryKey& InKey, FMetasoundFrontendClass& OutClass)
+	{
+		auto SetFrontendClass = [&OutClass](const INodeClassRegistryEntry& Entry)
+		{
+			OutClass = Entry.GetFrontendClass();
+		};
+
+		if (AccessNodeEntryThreadSafe(InKey, SetFrontendClass))
+		{
+			return true;
+		}
+
+		if (const INodeTemplateRegistryEntry* Entry = FindNodeTemplateEntry(InKey))
+		{
+			OutClass = Entry->GetFrontendClass();
+			return true;
+		}
+
+		return false;
+	}
+
+	bool FNodeClassRegistry::FindImplementedInterfacesFromRegistered(const Metasound::Frontend::FNodeClassRegistryKey& InKey, TSet<FMetasoundFrontendVersion>& OutInterfaceVersions) const 
+	{
+		bool bDidCopy = false;
+
+		auto CopyImplementedInterfaces = [&OutInterfaceVersions, &bDidCopy](const INodeClassRegistryEntry& Entry)
+		{
+			if (const TSet<FMetasoundFrontendVersion>* Interfaces = Entry.GetImplementedInterfaces())
+			{
+				OutInterfaceVersions = *Interfaces;
+				bDidCopy = true;
+			}
+		};
+
+		AccessNodeEntryThreadSafe(InKey, CopyImplementedInterfaces);
+
+		return bDidCopy;
+	}
+
+	bool FNodeClassRegistry::FindInputNodeRegistryKeyForDataType(const FName& InDataTypeName, const EMetasoundFrontendVertexAccessType InAccessType, FNodeClassRegistryKey& OutKey)
+	{
+		FMetasoundFrontendClass Class;
+		switch (InAccessType)
+		{
+			case EMetasoundFrontendVertexAccessType::Reference:
+			{
+				if (IDataTypeRegistry::Get().GetFrontendInputClass(InDataTypeName, Class))
+				{
+					OutKey = FNodeClassRegistryKey(Class.Metadata);
+					return true;
+				}
+			}
+			break;
+
+			case EMetasoundFrontendVertexAccessType::Value:
+			{
+				if (IDataTypeRegistry::Get().GetFrontendConstructorInputClass(InDataTypeName, Class))
+				{
+					OutKey = FNodeClassRegistryKey(Class.Metadata);
+					return true;
+				}
+			}
+			break;
+
+			default:
+			case EMetasoundFrontendVertexAccessType::Unset:
+			{
+				return false;
+			}
+			break;
+		}
+
+		return false;
+	}
+
+	bool FNodeClassRegistry::FindVariableNodeRegistryKeyForDataType(const FName& InDataTypeName, FNodeClassRegistryKey& OutKey)
+	{
+		FMetasoundFrontendClass Class;
+		if (IDataTypeRegistry::Get().GetFrontendLiteralClass(InDataTypeName, Class))
+		{
+			OutKey = FNodeClassRegistryKey(Class.Metadata);
+			return true;
+		}
+		return false;
+	}
+
+	bool FNodeClassRegistry::FindOutputNodeRegistryKeyForDataType(const FName& InDataTypeName, const EMetasoundFrontendVertexAccessType InAccessType, FNodeClassRegistryKey& OutKey)
+	{
+		FMetasoundFrontendClass Class;
+		switch (InAccessType)
+		{
+			case EMetasoundFrontendVertexAccessType::Reference:
+			{
+				if (IDataTypeRegistry::Get().GetFrontendOutputClass(InDataTypeName, Class))
+				{
+					OutKey = FNodeClassRegistryKey(Class.Metadata);
+					return true;
+				}
+			}
+			break;
+
+			case EMetasoundFrontendVertexAccessType::Value:
+			{
+				if (IDataTypeRegistry::Get().GetFrontendConstructorOutputClass(InDataTypeName, Class))
+				{
+					OutKey = FNodeClassRegistryKey(Class.Metadata);
+					return true;
+				}
+			}
+			break;
+		}
+
+		return false;
+	}
+
+	void FNodeClassRegistry::IterateRegistry(FIterateMetasoundFrontendClassFunction InIterFunc, EMetasoundFrontendClassType InClassType) const
+	{
+		UE_LOG(LogMetaSound, Warning, TEXT("Calling FMetasoundRegistryContainer::IterateRegistry(...) is not threadsafe. Please use Metasound::Frontend::ISearchEngine instead"));
+		auto WrappedFunc = [&](const TPair<FNodeClassRegistryKey, TSharedPtr<INodeClassRegistryEntry, ESPMode::ThreadSafe>>& Pair)
+		{
+			InIterFunc(Pair.Value->GetFrontendClass());
+		};
+
+		if (EMetasoundFrontendClassType::Invalid == InClassType)
+		{
+			// Iterate through all classes. 
+			Algo::ForEach(RegisteredNodes, WrappedFunc);
+		}
+		else
+		{
+			// Only call function on classes of certain type.
+			auto IsMatchingClassType = [&](const TPair<FNodeClassRegistryKey, TSharedPtr<INodeClassRegistryEntry, ESPMode::ThreadSafe>>& Pair)
+			{
+				return Pair.Value->GetClassInfo().Type == InClassType;
+			};
+			Algo::ForEachIf(RegisteredNodes, IsMatchingClassType, WrappedFunc);
+		}
+	}
+
+	bool FNodeClassRegistry::AccessNodeEntryThreadSafe(const FNodeClassRegistryKey& InKey, TFunctionRef<void(const INodeClassRegistryEntry&)> InFunc) const
+	{
+		auto TryAccessNodeEntry = [this, &InKey, &InFunc]() -> bool
+		{
+			FScopeLock Lock(&RegistryMapsCriticalSection);
+			if (const TSharedRef<INodeClassRegistryEntry, ESPMode::ThreadSafe>* Entry = RegisteredNodes.Find(InKey))
+			{
+				InFunc(*(*Entry));
+				return true;
+			}
+			return false;
+		};
+
+		if (TryAccessNodeEntry())
+		{
+			return true;
+		}
+		else
+		{
+			// Wait for any async registration tasks related to the registry key. 
+			WaitForAsyncRegistrationInternal(InKey, nullptr /* InAssetPath */);
+			return TryAccessNodeEntry();
+		}
+	}
+
+	const INodeTemplateRegistryEntry* FNodeClassRegistry::FindNodeTemplateEntry(const FNodeClassRegistryKey& InKey) const
+	{
+		FScopeLock Lock(&RegistryMapsCriticalSection);
+		if (const TSharedRef<INodeTemplateRegistryEntry, ESPMode::ThreadSafe>* Entry = RegisteredNodeTemplates.Find(InKey))
+		{
+			return &Entry->Get();
+		}
+
+		return nullptr;
+	}
+
+	void FNodeClassRegistry::WaitForAsyncGraphRegistration(const FGraphRegistryKey& InKey) const
+	{
+		WaitForAsyncRegistrationInternal(InKey.NodeKey, &InKey.AssetPath);
+	}
+
+	void FNodeClassRegistry::WaitForAsyncRegistrationInternal(const FNodeClassRegistryKey& InRegistryKey, const FTopLevelAssetPath* InAssetPath) const
+	{
+		using namespace UE::Tasks;
+
+		if (AsyncRegistrationPipe.IsInContext())
+		{
+			// It is not safe to wait for an async registration task from within the async registration pipe because it will result in a deadlock. 
+			UE_LOG(LogMetaSound, Verbose, TEXT("Async registration pipe is already in context for registering key %s. Task will not be waited for."), *InRegistryKey.ToString());
+			return;
+		}
+
+		TArray<FTask> TasksToWaitFor;
+		{
+			FScopeLock Lock(&ActiveRegistrationTasksCriticalSection);
+			if (const TArray<FActiveRegistrationTaskInfo>* FoundTasks = ActiveRegistrationTasks.Find(InRegistryKey))
+			{
+				// Filter by asset path or ignore if not provided
+				Algo::TransformIf(*FoundTasks, TasksToWaitFor,
+					[&InAssetPath](const FActiveRegistrationTaskInfo& TaskInfo) { return !InAssetPath || InAssetPath->IsNull() || TaskInfo.AssetPath == *InAssetPath; },
+					[](const FActiveRegistrationTaskInfo& TaskInfo) { return TaskInfo.Task; });
+			}
+		}
+
+		for (const FTask& Task : TasksToWaitFor)
+		{
+			METASOUND_TRACE_CPUPROFILER_EVENT_SCOPE(Metasound::FNodeClassRegistry::WaitForRegistrationTaskToComplete);
+			if (Task.IsValid())
+			{
+				Task.Wait();
+			}
+		}
+	}
+
+	INodeClassRegistry* INodeClassRegistry::Get()
+	{
+		return &Metasound::Frontend::FNodeClassRegistry::Get();
+	}
+
+	void INodeClassRegistry::ShutdownMetasoundFrontend()
+	{
+		Metasound::Frontend::FNodeClassRegistry::Shutdown();
+	}
+
+	bool INodeClassRegistry::GetFrontendClassFromRegistered(const FNodeClassRegistryKey& InKey, FMetasoundFrontendClass& OutClass)
+	{
+		INodeClassRegistry* Registry = INodeClassRegistry::Get();
+
+		if (ensure(nullptr != Registry))
+		{
+			return Registry->FindFrontendClassFromRegistered(InKey, OutClass);
+		}
+
+		return false;
+	}
+
+
+	bool INodeClassRegistry::GetInputNodeRegistryKeyForDataType(const FName& InDataTypeName, const EMetasoundFrontendVertexAccessType InAccessType, FNodeClassRegistryKey& OutKey)
+	{
+		if (INodeClassRegistry* Registry = INodeClassRegistry::Get())
+		{
+			return Registry->FindInputNodeRegistryKeyForDataType(InDataTypeName, InAccessType, OutKey);
+		}
+		return false;
+	}
+
+	bool INodeClassRegistry::GetVariableNodeRegistryKeyForDataType(const FName& InDataTypeName, FNodeClassRegistryKey& OutKey)
+	{
+		if (INodeClassRegistry* Registry = INodeClassRegistry::Get())
+		{
+			return Registry->FindVariableNodeRegistryKeyForDataType(InDataTypeName, OutKey);
+		}
+		return false;
+	}
+
+	bool INodeClassRegistry::GetOutputNodeRegistryKeyForDataType(const FName& InDataTypeName, const EMetasoundFrontendVertexAccessType InVertexAccessType, FNodeClassRegistryKey& OutKey)
+	{
+		if (INodeClassRegistry* Registry = INodeClassRegistry::Get())
+		{
+			return Registry->FindOutputNodeRegistryKeyForDataType(InDataTypeName, InVertexAccessType, OutKey);
+		}
+		return false;
+	}
+
+} // namespace Metasound::Frontend
diff --git a/Engine/Plugins/Runtime/Metasound/Source/MetasoundFrontend/Private/MetasoundFrontendNodeClassRegistryPrivate.h b/Engine/Plugins/Runtime/Metasound/Source/MetasoundFrontend/Private/MetasoundFrontendNodeClassRegistryPrivate.h
new file mode 100644
--- /dev/null
+++ b/Engine/Plugins/Runtime/Metasound/Source/MetasoundFrontend/Private/MetasoundFrontendNodeClassRegistryPrivate.h
@@ -0,0 +1,202 @@
+// Copyright Epic Games, Inc. All Rights Reserved.
+
+#pragma once
+
+#include "Containers/Array.h"
+#include "Containers/Map.h"
+#include "HAL/CriticalSection.h"
+#include "Tasks/Pipe.h"
+
+#include "MetasoundFrontendDataTypeRegistry.h"
+#include "MetasoundFrontendNodeClassRegistry.h"
+#include "MetasoundFrontendRegistryTransaction.h"
+
+struct FMetasoundFrontendDocument; 
+
+namespace Metasound
+{
+	class FProxyDataCache;
+	class FGraph;
+}
+
+namespace Metasound::Frontend
+{
+	struct FNodeClassInfo;
+
+	using FNodeClassRegistryTransactionBuffer = TTransactionBuffer<FNodeClassRegistryTransaction>;
+	using FNodeClassRegistryTransactionStream = TTransactionStream<FNodeClassRegistryTransaction>; 
+
+	/** INodeClassRegistryTemplateEntry declares the interface for a node registry entry.
+	 * Each node class in the registry must satisfy this interface.
+	 */
+	class INodeTemplateRegistryEntry
+	{
+	public:
+		virtual ~INodeTemplateRegistryEntry() = default;
+
+		/** Return FNodeClassInfo for the node class.
+		 *
+		 * Implementations of method should avoid any expensive operations
+		 * (e.g. loading from disk, allocating memory) as this method is called
+		 * frequently when querying nodes.
+		 */
+		virtual const FNodeClassInfo& GetClassInfo() const = 0;
+
+		/** Return a FMetasoundFrontendClass which describes the node. */
+		virtual const FMetasoundFrontendClass& GetFrontendClass() const = 0;
+	};
+
+	// Registry container private implementation.
+	class FNodeClassRegistry : public INodeClassRegistry
+	{
+
+		FNodeClassRegistry();
+	public:
+
+		FNodeClassRegistry(const FNodeClassRegistry&) = delete;
+		FNodeClassRegistry& operator=(const FNodeClassRegistry&) = delete;
+
+		static FNodeClassRegistry& Get();
+		static void Shutdown();
+
+		virtual ~FNodeClassRegistry() = default;
+
+		// Add a function to the init command array.
+		virtual bool EnqueueInitCommand(TUniqueFunction<void()>&& InFunc) override;
+
+		virtual void SetObjectReferencer(TUniquePtr<IObjectReferencer> InReferencer) override;
+
+		// This is called on module startup. This invokes any registration commands enqueued by our registration macros.
+		virtual void RegisterPendingNodes() override;
+
+		// Wait for async graph registration to complete for a specific graph
+		virtual void WaitForAsyncGraphRegistration(const FGraphRegistryKey& InRegistryKey) const override;
+
+		// Retrieve a registered graph. 
+		//
+		// If the graph is registered asynchronously, this will wait until the registration task has completed.
+		virtual TSharedPtr<const FGraph> GetGraph(const FGraphRegistryKey& InRegistryKey) const override;
+
+		/** Register external node with the frontend.
+		 *
+		 * @param InEntry - Entry to register.
+		 *
+		 * @return True on success.
+		 */
+		virtual FNodeClassRegistryKey RegisterNode(TUniquePtr<Metasound::Frontend::INodeClassRegistryEntry>&& InEntry) override;
+
+		virtual bool UnregisterNode(const FNodeClassRegistryKey& InKey) override;
+		virtual bool IsNodeRegistered(const FNodeClassRegistryKey& InKey) const override;
+		virtual bool IsGraphRegistered(const FGraphRegistryKey& InKey) const override;
+		virtual bool IsNodeNative(const FNodeClassRegistryKey& InKey) const override;
+
+		virtual bool RegisterConversionNode(const FConverterNodeClassRegistryKey& InNodeKey, const FConverterNodeInfo& InNodeInfo) override;
+
+		virtual void IterateRegistry(FIterateMetasoundFrontendClassFunction InIterFunc, EMetasoundFrontendClassType InClassType = EMetasoundFrontendClassType::Invalid) const override;
+
+		// Find Frontend class data.
+		virtual bool FindFrontendClassFromRegistered(const FNodeClassRegistryKey& InKey, FMetasoundFrontendClass& OutClass) override;
+
+
+		virtual bool FindImplementedInterfacesFromRegistered(const Metasound::Frontend::FNodeClassRegistryKey& InKey, TSet<FMetasoundFrontendVersion>& OutInterfaceVersions) const override;
+		
+		virtual bool FindInputNodeRegistryKeyForDataType(const FName& InDataTypeName, const EMetasoundFrontendVertexAccessType InAccessType, FNodeClassRegistryKey& OutKey) override;
+		virtual bool FindVariableNodeRegistryKeyForDataType(const FName& InDataTypeName, FNodeClassRegistryKey& OutKey) override;
+		virtual bool FindOutputNodeRegistryKeyForDataType(const FName& InDataTypeName, const EMetasoundFrontendVertexAccessType InAccessType, FNodeClassRegistryKey& OutKey) override;
+
+		// Create a new instance of a C++ implemented node from the registry.
+		virtual TUniquePtr<Metasound::INode> CreateNode(const FNodeClassRegistryKey& InKey, const Metasound::FNodeInitData& InInitData) const override;
+		virtual TUniquePtr<INode> CreateNode(const FNodeClassRegistryKey& InKey, FDefaultLiteralNodeConstructorParams&&) const override;
+		virtual TUniquePtr<INode> CreateNode(const FNodeClassRegistryKey& InKey, FDefaultNamedVertexNodeConstructorParams&&) const override;
+		virtual TUniquePtr<INode> CreateNode(const FNodeClassRegistryKey& InKey, FDefaultNamedVertexWithLiteralNodeConstructorParams&&) const override;
+
+		// Returns a list of possible nodes to use to convert from FromDataType to ToDataType.
+		// Returns an empty array if none are available.
+		virtual TArray<FConverterNodeInfo> GetPossibleConverterNodes(const FName& FromDataType, const FName& ToDataType) override;
+
+		// Register a graph from an IMetaSoundDocumentInterface
+		FGraphRegistryKey RegisterGraph(const TScriptInterface<IMetaSoundDocumentInterface>& InDocumentInterface);
+
+		// Unregister a graph
+		bool UnregisterGraph(const FGraphRegistryKey& InRegistryKey, const TScriptInterface<IMetaSoundDocumentInterface>& InDocumentInterface);
+
+		// Private implementation until hardened and used for template nodes other than reroutes.
+		FNodeClassRegistryKey RegisterNodeTemplate(TUniquePtr<Metasound::Frontend::INodeTemplateRegistryEntry>&& InEntry);
+		bool UnregisterNodeTemplate(const FNodeClassRegistryKey& InKey);
+
+		// Create a transaction stream for any newly transactions
+		TUniquePtr<FNodeClassRegistryTransactionStream> CreateTransactionStream();
+
+	private:
+		struct FActiveRegistrationTaskInfo
+		{
+			FNodeClassRegistryTransaction::ETransactionType TransactionType = FNodeClassRegistryTransaction::ETransactionType::NodeRegistration;
+			UE::Tasks::FTask Task;
+			FTopLevelAssetPath AssetPath;
+		};
+
+		void BuildAndRegisterGraphFromDocument(const FMetasoundFrontendDocument& InDocument, const FProxyDataCache& InProxyDataCache, FNodeClassInfo&& InNodeClassInfo);
+
+		void AddRegistrationTask(const FGraphRegistryKey& InKey, FActiveRegistrationTaskInfo&& TaskInfo);
+		void RemoveRegistrationTask(const FGraphRegistryKey& InKey, FNodeClassRegistryTransaction::ETransactionType TransactionType);
+
+		// Adds reference to document's owning UObject to internal ObjectReferencer,
+		// indicating async registration/unregistration task(s) are depending on it.
+		void AddDocumentReference(TScriptInterface<IMetaSoundDocumentInterface> DocumentInterface);
+
+		// Removes reference to document's owning UObject from internal ObjectReferencer,
+		// indicating async registration/unregistration task(s) are no longer depending on it.
+		void RemoveDocumentReference(TScriptInterface<IMetaSoundDocumentInterface> DocumentInterface);
+
+		bool UnregisterNodeInternal(const FNodeClassRegistryKey& InKey, FNodeClassInfo* OutClassInfo = nullptr);
+		FNodeClassRegistryKey RegisterNodeInternal(TUniquePtr<INodeClassRegistryEntry>&& InEntry);
+
+		static FNodeClassRegistry* LazySingleton;
+
+		void WaitForAsyncRegistrationInternal(const FNodeClassRegistryKey& InRegistryKey, const FTopLevelAssetPath* InAssetPath) const;
+
+		void RegisterGraphInternal(const FGraphRegistryKey& InKey, TSharedPtr<const FGraph> InGraph);
+		bool UnregisterGraphInternal(const FGraphRegistryKey& InKey);
+
+		// Access a node entry safely. Node entries can be added/removed asynchronously. Functions passed to this method will be
+		// executed in a manner where access to the node registry entry is safe from threading issues. 
+		//
+		// @returns true if a node registry entry was found and the function executed. False if the entry was not 
+		//          found and the function not executed. 
+		bool AccessNodeEntryThreadSafe(const FNodeClassRegistryKey& InKey, TFunctionRef<void(const INodeClassRegistryEntry&)> InFunc) const;
+
+		const INodeTemplateRegistryEntry* FindNodeTemplateEntry(const FNodeClassRegistryKey& InKey) const;
+
+
+		// This buffer is used to enqueue nodes and datatypes to register when the module has been initialized,
+		// in order to avoid bad behavior with ensures, logs, etc. on static initialization.
+		// The bad news is that TInlineAllocator is the safest allocator to use on static init.
+		// The good news is that none of these lambdas typically have captures, so this should have low memory overhead.
+		static constexpr int32 MaxNumNodesAndDatatypesToInitialize = 2048;
+		TArray<TUniqueFunction<void()>, TInlineAllocator<MaxNumNodesAndDatatypesToInitialize>> LazyInitCommands;
+		
+		FCriticalSection LazyInitCommandCritSection;
+
+		// Registry in which we keep all information about nodes implemented in C++.
+		TMultiMap<FNodeClassRegistryKey, TSharedRef<INodeClassRegistryEntry, ESPMode::ThreadSafe>> RegisteredNodes;
+
+		// Registry in which we keep all information about dynamically-generated templated nodes via in C++.
+		TMap<FNodeClassRegistryKey, TSharedRef<INodeTemplateRegistryEntry, ESPMode::ThreadSafe>> RegisteredNodeTemplates;
+
+		// Map of all registered graphs
+		TMap<FGraphRegistryKey, TSharedPtr<const FGraph>> RegisteredGraphs;
+
+		// Registry in which we keep lists of possible nodes to use to convert between two datatypes
+		TMap<FConverterNodeClassRegistryKey, FConverterNodeClassRegistryValue> ConverterNodeClassRegistry;
+
+		TSharedRef<FNodeClassRegistryTransactionBuffer> TransactionBuffer;
+
+		mutable FCriticalSection RegistryMapsCriticalSection;
+		mutable FCriticalSection ActiveRegistrationTasksCriticalSection;
+		mutable FCriticalSection ObjectReferencerCriticalSection;
+
+		UE::Tasks::FPipe AsyncRegistrationPipe;
+		TMap<FNodeClassRegistryKey, TArray<FActiveRegistrationTaskInfo>> ActiveRegistrationTasks;
+		TUniquePtr<INodeClassRegistry::IObjectReferencer> ObjectReferencer;
+	};
+}
diff --git a/Engine/Plugins/Runtime/Metasound/Source/MetasoundFrontend/Private/MetasoundFrontendNodeTemplateRegistry.cpp b/Engine/Plugins/Runtime/Metasound/Source/MetasoundFrontend/Private/MetasoundFrontendNodeTemplateRegistry.cpp
--- a/Engine/Plugins/Runtime/Metasound/Source/MetasoundFrontend/Private/MetasoundFrontendNodeTemplateRegistry.cpp
+++ b/Engine/Plugins/Runtime/Metasound/Source/MetasoundFrontend/Private/MetasoundFrontendNodeTemplateRegistry.cpp
@@ -4,7 +4,8 @@
 
 #include "MetasoundFrontendDocument.h"
 #include "MetasoundFrontendDocumentBuilder.h"
-#include "MetasoundFrontendRegistryContainerImpl.h"
+#include "MetasoundFrontendNodeClassRegistry.h"
+#include "MetasoundFrontendNodeClassRegistryPrivate.h"
 #include "MetasoundFrontendRegistryTransaction.h"
 #include "MetasoundFrontendTransform.h"
 #include "MetasoundLog.h"
@@ -39,14 +40,14 @@ namespace Metasound::Frontend
 		FNodeTemplateRegistry() = default;
 		virtual ~FNodeTemplateRegistry() = default;
 
-		virtual const INodeTemplate* FindTemplate(const FNodeRegistryKey& InKey) const override;
+		virtual const INodeTemplate* FindTemplate(const FNodeClassRegistryKey& InKey) const override;
 		virtual const INodeTemplate* FindTemplate(const FMetasoundFrontendClassName& InClassName) const override;
 
 		void Register(TUniquePtr<INodeTemplate>&& InEntry);
-		void Unregister(const FNodeRegistryKey& InKey);
+		void Unregister(const FNodeClassRegistryKey& InKey);
 
 	private:
-		TMap<FNodeRegistryKey, TUniquePtr<const INodeTemplate>> Templates;
+		TMap<FNodeClassRegistryKey, TUniquePtr<const INodeTemplate>> Templates;
 		TMultiMap<FMetasoundFrontendClassName, const INodeTemplate*> TemplateByClassName;
 	};
 
@@ -54,7 +55,7 @@ namespace Metasound::Frontend
 	{
 		if (ensure(InTemplate.IsValid()))
 		{
-			const FNodeRegistryKey Key = FNodeRegistryKey(InTemplate->GetFrontendClass().Metadata);
+			const FNodeClassRegistryKey Key = FNodeClassRegistryKey(InTemplate->GetFrontendClass().Metadata);
 			if (ensure(Key.IsValid()))
 			{
 				TUniquePtr<const INodeTemplate>& Entry = Templates.Add(Key, MoveTemp(InTemplate));
@@ -63,7 +64,7 @@ namespace Metasound::Frontend
 		}
 	}
 
-	void FNodeTemplateRegistry::Unregister(const FNodeRegistryKey& InKey)
+	void FNodeTemplateRegistry::Unregister(const FNodeClassRegistryKey& InKey)
 	{
 		TUniquePtr<const INodeTemplate> Removed;
 		if (ensure(Templates.RemoveAndCopyValue(InKey, Removed)))
@@ -72,7 +73,7 @@ namespace Metasound::Frontend
 		}
 	}
 
-	const INodeTemplate* FNodeTemplateRegistry::FindTemplate(const FNodeRegistryKey& InKey) const
+	const INodeTemplate* FNodeTemplateRegistry::FindTemplate(const FNodeClassRegistryKey& InKey) const
 	{
 		if (const TUniquePtr<const INodeTemplate>* Template = Templates.Find(InKey))
 		{
@@ -152,7 +153,7 @@ namespace Metasound::Frontend
 
 	void RegisterNodeTemplate(TUniquePtr<INodeTemplate>&& InTemplate)
 	{
-		class FTemplateRegistryEntry : public INodeRegistryTemplateEntry
+		class FTemplateRegistryEntry : public INodeTemplateRegistryEntry
 		{
 			const FNodeClassInfo ClassInfo;
 			const FMetasoundFrontendClass FrontendClass;
@@ -178,8 +179,8 @@ namespace Metasound::Frontend
 			}
 		};
 
-		TUniquePtr<INodeRegistryTemplateEntry> RegEntry = TUniquePtr<INodeRegistryTemplateEntry>(new FTemplateRegistryEntry(*InTemplate.Get()));
-		FRegistryContainerImpl::Get().RegisterNodeTemplate(MoveTemp(RegEntry));
+		TUniquePtr<INodeTemplateRegistryEntry> RegEntry = TUniquePtr<INodeTemplateRegistryEntry>(new FTemplateRegistryEntry(*InTemplate.Get()));
+		FNodeClassRegistry::Get().RegisterNodeTemplate(MoveTemp(RegEntry));
 
 		static_cast<FNodeTemplateRegistry&>(INodeTemplateRegistry::Get()).Register(MoveTemp(InTemplate));
 	}
@@ -188,20 +189,20 @@ namespace Metasound::Frontend
 	{
 		FMetasoundFrontendClassName ClassName;
 		FMetasoundFrontendClassName::Parse(InVersion.Name.ToString(), ClassName);
-		const FNodeRegistryKey Key = FNodeRegistryKey(EMetasoundFrontendClassType::Template, ClassName, InVersion.Number);
+		const FNodeClassRegistryKey Key = FNodeClassRegistryKey(EMetasoundFrontendClassType::Template, ClassName, InVersion.Number);
 		if (ensure(Key.IsValid()))
 		{
-			FRegistryContainerImpl::Get().UnregisterNodeTemplate(Key);
+			FNodeClassRegistry::Get().UnregisterNodeTemplate(Key);
 			static_cast<FNodeTemplateRegistry&>(INodeTemplateRegistry::Get()).Unregister(Key);
 		}
 	}
 
 	void UnregisterNodeTemplate(const FMetasoundFrontendClassName& InClassName, const FMetasoundFrontendVersionNumber& InTemplateVersion)
 	{
-		const FNodeRegistryKey Key = FNodeRegistryKey(EMetasoundFrontendClassType::Template, InClassName, InTemplateVersion);
+		const FNodeClassRegistryKey Key = FNodeClassRegistryKey(EMetasoundFrontendClassType::Template, InClassName, InTemplateVersion);
 		if (ensure(Key.IsValid()))
 		{
-			FRegistryContainerImpl::Get().UnregisterNodeTemplate(Key);
+			FNodeClassRegistry::Get().UnregisterNodeTemplate(Key);
 			static_cast<FNodeTemplateRegistry&>(INodeTemplateRegistry::Get()).Unregister(Key);
 		}
 	}
diff --git a/Engine/Plugins/Runtime/Metasound/Source/MetasoundFrontend/Private/MetasoundFrontendQuerySteps.cpp b/Engine/Plugins/Runtime/Metasound/Source/MetasoundFrontend/Private/MetasoundFrontendQuerySteps.cpp
--- a/Engine/Plugins/Runtime/Metasound/Source/MetasoundFrontend/Private/MetasoundFrontendQuerySteps.cpp
+++ b/Engine/Plugins/Runtime/Metasound/Source/MetasoundFrontend/Private/MetasoundFrontendQuerySteps.cpp
@@ -5,9 +5,9 @@
 #include "Algo/MaxElement.h"
 #include "MetasoundFrontend.h"
 #include "MetasoundFrontendDocument.h"
+#include "MetasoundFrontendNodeClassRegistry.h"
+#include "MetasoundFrontendNodeClassRegistryPrivate.h"
 #include "MetasoundFrontendQuery.h"
-#include "MetasoundFrontendRegistryContainer.h"
-#include "MetasoundFrontendRegistryContainerImpl.h"
 #include "MetasoundFrontendRegistryTransaction.h"
 #include "MetasoundLog.h"
 
@@ -18,7 +18,7 @@ namespace Metasound
 	public:
 		FNodeClassRegistrationEventsPimpl()
 		{
-			TransactionStream = Frontend::FRegistryContainerImpl::Get().CreateTransactionStream();
+			TransactionStream = Frontend::FNodeClassRegistry::Get().CreateTransactionStream();
 		}
 
 		virtual void Stream(TArray<FFrontendQueryValue>& OutValues) override
@@ -27,7 +27,7 @@ namespace Metasound
 
 			auto AddEntry = [&OutValues](const FNodeRegistryTransaction& InTransaction)
 			{
-				OutValues.Emplace(TInPlaceType<FNodeRegistryTransaction>(), InTransaction);
+				OutValues.Emplace(TInPlaceType<FNodeClassRegistryTransaction>(), InTransaction);
 			};
 
 			if (TransactionStream.IsValid())
@@ -37,7 +37,7 @@ namespace Metasound
 		}
 
 	private:
-		TUniquePtr<Frontend::FNodeRegistryTransactionStream> TransactionStream;
+		TUniquePtr<Frontend::FNodeClassRegistryTransactionStream> TransactionStream;
 	};
 
 	FNodeClassRegistrationEvents::FNodeClassRegistrationEvents()
diff --git a/Engine/Plugins/Runtime/Metasound/Source/MetasoundFrontend/Private/MetasoundFrontendRegistryContainer.cpp b/Engine/Plugins/Runtime/Metasound/Source/MetasoundFrontend/Private/MetasoundFrontendRegistryContainer.cpp
deleted file mode 100644
diff --git a/Engine/Plugins/Runtime/Metasound/Source/MetasoundFrontend/Private/MetasoundFrontendRegistryContainerImpl.cpp b/Engine/Plugins/Runtime/Metasound/Source/MetasoundFrontend/Private/MetasoundFrontendRegistryContainerImpl.cpp
deleted file mode 100644
diff --git a/Engine/Plugins/Runtime/Metasound/Source/MetasoundFrontend/Private/MetasoundFrontendRegistryContainerImpl.h b/Engine/Plugins/Runtime/Metasound/Source/MetasoundFrontend/Private/MetasoundFrontendRegistryContainerImpl.h
deleted file mode 100644
diff --git a/Engine/Plugins/Runtime/Metasound/Source/MetasoundFrontend/Private/MetasoundFrontendRegistryKey.cpp b/Engine/Plugins/Runtime/Metasound/Source/MetasoundFrontend/Private/MetasoundFrontendRegistryKey.cpp
--- a/Engine/Plugins/Runtime/Metasound/Source/MetasoundFrontend/Private/MetasoundFrontendRegistryKey.cpp
+++ b/Engine/Plugins/Runtime/Metasound/Source/MetasoundFrontend/Private/MetasoundFrontendRegistryKey.cpp
@@ -3,7 +3,6 @@
 #include "MetasoundFrontendRegistryKey.h"
 #include "MetasoundAssetManager.h"
 
-
 namespace Metasound::Frontend
 {
 	namespace NodeClassInfoPrivate
@@ -18,24 +17,6 @@ namespace Metasound::Frontend
 	{
 	}
 
-	FNodeClassInfo::FNodeClassInfo(const FMetasoundFrontendGraphClass& InClass, const FSoftObjectPath& InAssetPath)
-		: ClassName(InClass.Metadata.GetClassName())
-		, Type(EMetasoundFrontendClassType::External) // Overridden as it is considered the same as an external class in registries
-		, AssetClassID(FGuid(ClassName.Name.ToString()))
-		, Version(InClass.Metadata.GetVersion())
-	{
-		using namespace NodeClassInfoPrivate;
-
-		ensure(AssetPath.TrySetPath(InAssetPath.ToString()));
-		ensure(!AssetPath.IsNull());
-
-#if WITH_EDITORONLY_DATA
-		Algo::Transform(InClass.Interface.Inputs, InputTypes, GetVertexTypeName);
-		Algo::Transform(InClass.Interface.Outputs, OutputTypes, GetVertexTypeName);
-		bIsPreset = InClass.PresetOptions.bIsPreset;
-#endif // WITH_EDITORONLY_DATA
-	}
-
 	FNodeClassInfo::FNodeClassInfo(const FMetasoundFrontendGraphClass& InClass, const FTopLevelAssetPath& InAssetPath)
 		: ClassName(InClass.Metadata.GetClassName())
 		, Type(EMetasoundFrontendClassType::External) // Overridden as it is considered the same as an external class in registries
@@ -54,86 +35,76 @@ namespace Metasound::Frontend
 #endif // WITH_EDITORONLY_DATA
 	}
 
-	UObject* FNodeClassInfo::LoadAsset() const
-	{
-		return nullptr;
-	}
-
-	FNodeRegistryKey::FNodeRegistryKey(EMetasoundFrontendClassType InType, const FMetasoundFrontendClassName& InClassName, int32 InMajorVersion, int32 InMinorVersion)
+	FNodeClassRegistryKey::FNodeClassRegistryKey(EMetasoundFrontendClassType InType, const FMetasoundFrontendClassName& InClassName, int32 InMajorVersion, int32 InMinorVersion)
 		: Type(InType)
 		, ClassName(InClassName)
 		, Version({ InMajorVersion, InMinorVersion })
 	{
 	}
 
-	FNodeRegistryKey::FNodeRegistryKey(EMetasoundFrontendClassType InType, const FMetasoundFrontendClassName& InClassName, const FMetasoundFrontendVersionNumber& InVersion)
+	FNodeClassRegistryKey::FNodeClassRegistryKey(EMetasoundFrontendClassType InType, const FMetasoundFrontendClassName& InClassName, const FMetasoundFrontendVersionNumber& InVersion)
 		: Type(InType)
 		, ClassName(InClassName)
 		, Version(InVersion)
 	{
 	}
 
-	FNodeRegistryKey::FNodeRegistryKey(const FNodeClassMetadata& InNodeMetadata)
+	FNodeClassRegistryKey::FNodeClassRegistryKey(const FNodeClassMetadata& InNodeMetadata)
 		: Type(EMetasoundFrontendClassType::External) // Overridden as it is considered the same as an external class in registries
 		, ClassName(InNodeMetadata.ClassName)
 		, Version({ InNodeMetadata.MajorVersion, InNodeMetadata.MinorVersion })
 	{
 	}
 
-	FNodeRegistryKey::FNodeRegistryKey(const FMetasoundFrontendClassMetadata& InNodeMetadata)
+	FNodeClassRegistryKey::FNodeClassRegistryKey(const FMetasoundFrontendClassMetadata& InNodeMetadata)
 		: Type(InNodeMetadata.GetType())
 		, ClassName(InNodeMetadata.GetClassName())
 		, Version(InNodeMetadata.GetVersion())
 	{
-		checkf(InNodeMetadata.GetType() != EMetasoundFrontendClassType::Graph, TEXT("Cannot create key from 'graph' type. Likely meant to use FNodeRegistryKey ctor that is provided FMetasoundFrontendGraphClass"));
+		checkf(InNodeMetadata.GetType() != EMetasoundFrontendClassType::Graph, TEXT("Cannot create key from 'graph' type. Likely meant to use FNodeClassRegistryKey ctor that is provided FMetasoundFrontendGraphClass"));
 	}
 
-	FNodeRegistryKey::FNodeRegistryKey(const FMetasoundFrontendGraphClass& InGraphClass)
+	FNodeClassRegistryKey::FNodeClassRegistryKey(const FMetasoundFrontendGraphClass& InGraphClass)
 		: Type(EMetasoundFrontendClassType::External) // Type overridden as all graphs are considered the same as an external class in the registry
 		, ClassName(InGraphClass.Metadata.GetClassName())
 		, Version(InGraphClass.Metadata.GetVersion())
 	{
 	}
 
-	FNodeRegistryKey::FNodeRegistryKey(const FNodeClassInfo& InClassInfo)
+	FNodeClassRegistryKey::FNodeClassRegistryKey(const FNodeClassInfo& InClassInfo)
 		: Type(InClassInfo.Type)
 		, ClassName(InClassInfo.ClassName)
 		, Version(InClassInfo.Version)
 	{
-		checkf(InClassInfo.Type != EMetasoundFrontendClassType::Graph, TEXT("Cannot create key from 'graph' type. Likely meant to use FNodeRegistryKey ctor that is provided FMetasoundFrontendGraphClass"));
-	}
-
-	FNodeRegistryKey::FNodeRegistryKey(const FString& InKeyString)
-	{
-		Parse(InKeyString, *this);
+		checkf(InClassInfo.Type != EMetasoundFrontendClassType::Graph, TEXT("Cannot create key from 'graph' type. Likely meant to use FNodeClassRegistryKey ctor that is provided FMetasoundFrontendGraphClass"));
 	}
 
-	FNodeRegistryKey::FNodeRegistryKey(const FAssetKey& AssetKey)
+	FNodeClassRegistryKey::FNodeClassRegistryKey(const FAssetKey& AssetKey)
 		: Type(EMetasoundFrontendClassType::External)
 		, ClassName(AssetKey.ClassName)
 		, Version(AssetKey.Version)
 	{
 	}
 
-	const FNodeRegistryKey& FNodeRegistryKey::GetInvalid()
+	const FNodeClassRegistryKey& FNodeClassRegistryKey::GetInvalid()
 	{
-		static const FNodeRegistryKey InvalidKey;
+		static const FNodeClassRegistryKey InvalidKey;
 		return InvalidKey;
 	}
 
-	bool FNodeRegistryKey::IsValid() const
+	bool FNodeClassRegistryKey::IsValid() const
 	{
 		return Type != EMetasoundFrontendClassType::Invalid && ClassName.IsValid() && Version.IsValid();
 	}
 
-	void FNodeRegistryKey::Reset()
+	void FNodeClassRegistryKey::Reset()
 	{
 		Type = EMetasoundFrontendClassType::Invalid;
 		ClassName = { };
 		Version = { };
 	}
 
-	FString FNodeRegistryKey::ToString() const
+	FString FNodeClassRegistryKey::ToString() const
 	{
 		TStringBuilder<128> KeyStringBuilder;
 		KeyStringBuilder.Append(LexToString(Type));
@@ -146,7 +117,7 @@ namespace Metasound::Frontend
 		return KeyStringBuilder.ToString();
 	}
 
-	FString FNodeRegistryKey::ToString(const FString& InScopeHeader) const
+	FString FNodeClassRegistryKey::ToString(const FString& InScopeHeader) const
 	{
 		checkf(InScopeHeader.Len() < 128, TEXT("Scope text is limited to 128 characters"));
 
@@ -159,7 +130,7 @@ namespace Metasound::Frontend
 		return Builder.ToString();
 	}
 
-	bool FNodeRegistryKey::Parse(const FString& InKeyString, FNodeRegistryKey& OutKey)
+	bool FNodeClassRegistryKey::Parse(const FString& InKeyString, FNodeClassRegistryKey& OutKey)
 	{
 		TArray<FString> Tokens;
 		InKeyString.ParseIntoArray(Tokens, TEXT("_"));
@@ -179,7 +150,7 @@ namespace Metasound::Frontend
 						Version.Major = FCString::Atoi(*MajorVersionString);
 						Version.Minor = FCString::Atoi(*MinorVersionString);
 
-						OutKey = FNodeRegistryKey(Type, ClassName, Version.Major, Version.Minor);
+						OutKey = FNodeClassRegistryKey(Type, ClassName, Version.Major, Version.Minor);
 						return true;
 					}
 				}
@@ -189,7 +160,7 @@ namespace Metasound::Frontend
 		return false;
 	}
 
-	FString FGraphRegistryKey::ToString() const
+	FString FGraphClassRegistryKey::ToString() const
 	{
 		TStringBuilder<256> Builder;
 		Builder.Append(NodeKey.ToString());
@@ -200,7 +171,7 @@ namespace Metasound::Frontend
 		return Builder.ToString();
 	}
 
-	FString FGraphRegistryKey::ToString(const FString& InScopeHeader) const
+	FString FGraphClassRegistryKey::ToString(const FString& InScopeHeader) const
 	{
 		TStringBuilder<512> Builder;
 		Builder.Append(InScopeHeader);
@@ -210,7 +181,7 @@ namespace Metasound::Frontend
 		return Builder.ToString();
 	}
 
-	bool FGraphRegistryKey::IsValid() const
+	bool FGraphClassRegistryKey::IsValid() const
 	{
 		return NodeKey.IsValid() && AssetPath.IsValid();
 	}
diff --git a/Engine/Plugins/Runtime/Metasound/Source/MetasoundFrontend/Public/MetasoundFrontendDataTypeRegistry.h b/Engine/Plugins/Runtime/Metasound/Source/MetasoundFrontend/Public/MetasoundFrontendDataTypeRegistry.h
--- a/Engine/Plugins/Runtime/Metasound/Source/MetasoundFrontend/Public/MetasoundFrontendDataTypeRegistry.h
+++ b/Engine/Plugins/Runtime/Metasound/Source/MetasoundFrontend/Public/MetasoundFrontendDataTypeRegistry.h
@@ -16,6 +16,8 @@
 
 namespace Metasound
 {
+	class IDataChannel;
+
 	namespace Frontend
 	{
 		using IParameterAssignmentFunction = TUniqueFunction<void(const void*,void*)>;
diff --git a/Engine/Plugins/Runtime/Metasound/Source/MetasoundFrontend/Public/MetasoundFrontendNodeClassRegistry.h b/Engine/Plugins/Runtime/Metasound/Source/MetasoundFrontend/Public/MetasoundFrontendNodeClassRegistry.h
new file mode 100644
--- /dev/null
+++ b/Engine/Plugins/Runtime/Metasound/Source/MetasoundFrontend/Public/MetasoundFrontendNodeClassRegistry.h
@@ -0,0 +1,317 @@
+// Copyright Epic Games, Inc. All Rights Reserved.
+#pragma once
+
+#include "MetasoundFrontendDocument.h"
+#include "MetasoundFrontendRegistryKey.h"
+#include "MetasoundNodeInterface.h"
+#include "MetasoundNodeConstructorParams.h"
+#include "MetasoundVertex.h"
+#include "Templates/Function.h"
+
+// Forward declare
+class UObject;
+
+namespace Metasound
+{
+	// Forward Declarations
+	namespace Frontend
+	{
+		using FIterateMetasoundFrontendClassFunction = TFunctionRef<void(const FMetasoundFrontendClass&)>;
+	}
+
+	using FIterateMetasoundFrontendClassFunction UE_DEPRECATED(5.6, "Use TIterateMetasoundFrontendClassFunction defined in Metasound::Frontend namespace") = Frontend::FIterateMetasoundFrontendClassFunction;
+
+	class FGraph;
+} // namespace Metasound
+
+namespace Metasound::Frontend
+{
+	/** INodeClassRegistryEntry declares the interface for a node registry entry.
+		* Each node class in the registry must satisfy this interface. 
+		*/
+	class METASOUNDFRONTEND_API INodeClassRegistryEntry
+	{
+	public:
+		virtual ~INodeClassRegistryEntry() = default;
+
+		/** Return FNodeClassInfo for the node class.
+			*
+			* Implementations of method should avoid any expensive operations 
+			* (e.g. loading from disk, allocating memory) as this method is called
+			* frequently when querying nodes.
+			*/
+		virtual const FNodeClassInfo& GetClassInfo() const = 0;
+
+		/** Create a node given FDefaultNamedVertexNodeConstructorParams.
+			*
+			* If a node can be created with FDefaultNamedVertexNodeConstructorParams, this function
+			* should return a valid node pointer.
+			*/
+		virtual TUniquePtr<INode> CreateNode(FDefaultNamedVertexNodeConstructorParams&&) const = 0;
+
+		/** Create a node given FDefaultNamedVertexWithLiteralNodeConstructorParams.
+			*
+			* If a node can be created with FDefaultNamedVertexWithLiteralNodeConstructorParams, this function
+			* should return a valid node pointer.
+			*/
+		virtual TUniquePtr<INode> CreateNode(FDefaultNamedVertexWithLiteralNodeConstructorParams&&) const = 0;
+
+		/** Create a node given FDefaultLiteralNodeConstructorParams.
+			*
+			* If a node can be created with FDefaultLiteralNodeConstructorParams, this function
+			* should return a valid node pointer.
+			*/
+		virtual TUniquePtr<INode> CreateNode(FDefaultLiteralNodeConstructorParams&&) const = 0;
+
+		/** Create a node given FNodeInitData.
+			*
+			* If a node can be created with FNodeInitData, this function
+			* should return a valid node pointer.
+			*/
+		virtual TUniquePtr<INode> CreateNode(const FNodeInitData&) const = 0;
+
+		/** Return a FMetasoundFrontendClass which describes the node. */
+		virtual const FMetasoundFrontendClass& GetFrontendClass() const = 0;
+
+		/** Clone this registry entry. */
+		UE_DEPRECATED(5.6, "Node Class Registration Entries do not need to be cloned.")
+		virtual TUniquePtr<INodeClassRegistryEntry> Clone() const { return nullptr; }
+
+		/** Returns set of implemented interface versions.
+			*
+			* Returns nullptr if node class implementation does not support interface implementation.
+			*/
+		virtual const TSet<FMetasoundFrontendVersion>* GetImplementedInterfaces() const = 0;
+
+		/** Whether or not the node is natively defined */
+		virtual bool IsNative() const = 0;
+	};
+	using INodeRegistryEntry = INodeClassRegistryEntry;
+
+	struct METASOUNDFRONTEND_API FConverterNodeClassRegistryKey
+	{
+		// The datatype one would like to convert from.
+		FName FromDataType;
+
+		// The datatype one would like to convert to.
+		FName ToDataType;
+
+		FORCEINLINE bool operator==(const FConverterNodeClassRegistryKey& Other) const
+		{
+			return FromDataType == Other.FromDataType && ToDataType == Other.ToDataType;
+		}
+
+		friend uint32 GetTypeHash(const FConverterNodeClassRegistryKey& InKey)
+		{
+			return HashCombine(GetTypeHash(InKey.FromDataType), GetTypeHash(InKey.ToDataType));
+		}
+	};
+	using FConverterNodeRegistryKey = FConverterNodeClassRegistryKey;
+
+	struct METASOUNDFRONTEND_API FConverterNodeClassInfo
+	{
+		// If this node has multiple input pins, we use this to designate which pin should be used.
+		FVertexName PreferredConverterInputPin;
+
+		// If this node has multiple output pins, we use this to designate which pin should be used.
+		FVertexName PreferredConverterOutputPin;
+
+		// The key for this node in the node registry.
+		FNodeClassRegistryKey NodeKey;
+
+		FORCEINLINE bool operator==(const FConverterNodeClassInfo& Other) const
+		{
+			return NodeKey == Other.NodeKey;
+		}
+	};
+	using FConverterNodeInfo = FConverterNodeClassInfo;
+
+	struct METASOUNDFRONTEND_API FConverterNodeClassRegistryValue
+	{
+		// A list of nodes that can perform a conversion between the two datatypes described in the FConverterNodeClassRegistryKey for this map element.
+		TArray<FConverterNodeClassInfo> PotentialConverterNodes;
+	};
+	using FConverterNodeClassRegistryValue = FConverterNodeClassRegistryValue;
+
+	using FRegistryTransactionID = int32;
+
+	class METASOUNDFRONTEND_API FNodeClassRegistryTransaction 
+	{
+	public:
+		using FTimeType = uint64;
+
+		/** Describes the type of transaction. */
+		enum class ETransactionType : uint8
+		{
+			NodeRegistration,     //< Something was added to the registry.
+			NodeUnregistration,  //< Something was removed from the registry.
+			Invalid
+		};
+
+		static FString LexToString(ETransactionType InTransactionType)
+		{
+			using namespace Metasound;
+
+			switch (InTransactionType)
+			{
+				case ETransactionType::NodeRegistration:
+				{
+					return TEXT("Node Registration");
+				}
+
+				case ETransactionType::NodeUnregistration:
+				{
+					return TEXT("Node Unregistration");
+				}
+
+				case ETransactionType::Invalid:
+				{
+					return TEXT("Invalid");
+				}
+
+				default:
+				{
+					checkNoEntry();
+					return TEXT("");
+				}
+			}
+		}
+
+		FNodeClassRegistryTransaction(ETransactionType InType, const FNodeClassInfo& InNodeClassInfo, FTimeType InTimestamp);
+
+		ETransactionType GetTransactionType() const;
+		const FNodeClassInfo& GetNodeClassInfo() const;
+		FNodeClassRegistryKey GetNodeRegistryKey() const;
+		FTimeType GetTimestamp() const;
+
+	private:
+		ETransactionType Type;
+		FNodeClassInfo NodeClassInfo;
+		FTimeType Timestamp;
+	};
+	using FNodeRegistryTransaction = FNodeClassRegistryTransaction;
+
+	namespace NodeClassRegistryKey
+	{
+		// Returns true if the class metadata represent the same entry in the node registry.
+		METASOUNDFRONTEND_API bool IsEqual(const FMetasoundFrontendClassMetadata& InLHS, const FMetasoundFrontendClassMetadata& InRHS);
+
+		// Returns true if the class info and class metadata represent the same entry in the node registry.
+		METASOUNDFRONTEND_API bool IsEqual(const FNodeClassInfo& InLHS, const FMetasoundFrontendClassMetadata& InRHS);
+	}
+	namespace NodeRegistryKey = NodeClassRegistryKey;
+
+	/**
+	 * Singleton registry for all types and nodes.
+	 */
+	class METASOUNDFRONTEND_API INodeClassRegistry
+	{
+
+	public:
+		// The MetaSound frontend does not rely on the Engine module and therefore
+		// does not have the ability to provide GC protection. This interface allows
+		// external modules to provide GC protection so that async tasks can safely
+		// use UObjects
+		class IObjectReferencer
+		{
+		public:
+			virtual ~IObjectReferencer() = default;
+
+			// Called when an object should be referenced.
+			virtual void AddObject(UObject* InObject) = 0;
+
+			// Called when an object no longer needs to be referenced. 
+			virtual void RemoveObject(UObject* InObject) = 0;
+		};
+
+
+		static INodeClassRegistry* Get();
+
+		UE_DEPRECATED(5.6, "This function is no longer supported")
+		static void ShutdownMetasoundFrontend();
+
+
+		static bool GetFrontendClassFromRegistered(const FNodeClassRegistryKey& InKey, FMetasoundFrontendClass& OutClass);
+
+
+		static bool GetInputNodeRegistryKeyForDataType(const FName& InDataTypeName, const EMetasoundFrontendVertexAccessType InAccessType, FNodeClassRegistryKey& OutKey);
+		static bool GetVariableNodeRegistryKeyForDataType(const FName& InDataTypeName, FNodeClassRegistryKey& OutKey);
+		static bool GetOutputNodeRegistryKeyForDataType(const FName& InDataTypeName, const EMetasoundFrontendVertexAccessType InAccessType, FNodeClassRegistryKey& OutKey);
+
+
+	protected:
+		INodeClassRegistry() = default;
+	public:
+		virtual ~INodeClassRegistry() = default;
+
+		INodeClassRegistry(const INodeClassRegistry&) = delete;
+		INodeClassRegistry& operator=(const INodeClassRegistry&) = delete;
+
+		// Enqueue and init command for registering a node or data type.
+		// The command queue will be processed on module init or when calling `RegisterPendingNodes()`
+		virtual bool EnqueueInitCommand(TUniqueFunction<void()>&& InFunc) = 0;
+
+		virtual void SetObjectReferencer(TUniquePtr<IObjectReferencer> InReferencer) = 0;
+
+		// This is called on module startup. This invokes any registration commands enqueued by our registration macros.
+		virtual void RegisterPendingNodes() = 0;
+
+		// Wait for async graph registration to complete for a specific graph
+		virtual void WaitForAsyncGraphRegistration(const FGraphRegistryKey& InRegistryKey) const = 0;
+
+		// Retrieve a registered graph. 
+		//
+		// If the graph is registered asynchronously, this will wait until the registration task has completed.
+		virtual TSharedPtr<const Metasound::FGraph> GetGraph(const FGraphRegistryKey& InRegistryKey) const = 0;
+
+		/** Register an external node with the frontend.
+		 *
+		 * @param InCreateNode - Function for creating node from FNodeInitData.
+		 * @param InCreateDescription - Function for creating a FMetasoundFrontendClass.
+		 *
+		 * @return A node registration key. If the registration failed, then the registry 
+		 *         key will be invalid.
+		 */
+		virtual FNodeClassRegistryKey RegisterNode(TUniquePtr<INodeClassRegistryEntry>&& InEntry) = 0;
+
+		/** Unregister an external node from the frontend.
+		 *
+		 * @param InKey - The registration key for the node.
+		 *
+		 * @return True on success, false on failure.
+		 */
+		virtual bool UnregisterNode(const FNodeClassRegistryKey& InKey) = 0;
+
+		/** Returns true if the provided registry key corresponds to a valid registered node. */
+		virtual bool IsNodeRegistered(const FNodeClassRegistryKey& InKey) const = 0;
+
+		/** Returns true if the provided registry key (node key and asset path) corresponds to a valid registered graph. */
+		virtual bool IsGraphRegistered(const FGraphRegistryKey& InKey) const = 0;
+
+		/** Returns true if the provided registry key corresponds to a valid registered node that is natively defined. */
+		virtual bool IsNodeNative(const FNodeClassRegistryKey& InKey) const = 0;
+
+		// Iterates class types in registry.  If InClassType is set to a valid class type (optional), only iterates classes of the given type
+		virtual void IterateRegistry(FIterateMetasoundFrontendClassFunction InIterFunc, EMetasoundFrontendClassType InClassType = EMetasoundFrontendClassType::Invalid) const = 0;
+
+		// Query for MetaSound Frontend document objects.
+		virtual bool FindFrontendClassFromRegistered(const FNodeClassRegistryKey& InKey, FMetasoundFrontendClass& OutClass) = 0;
+
+		virtual bool FindImplementedInterfacesFromRegistered(const FNodeClassRegistryKey& InKey, TSet<FMetasoundFrontendVersion>& OutInterfaceVersions) const = 0;
+		
+		virtual bool FindInputNodeRegistryKeyForDataType(const FName& InDataTypeName, const EMetasoundFrontendVertexAccessType InAccessType, FNodeClassRegistryKey& OutKey) = 0;
+		virtual bool FindVariableNodeRegistryKeyForDataType(const FName& InDataTypeName, FNodeClassRegistryKey& OutKey) = 0;
+		virtual bool FindOutputNodeRegistryKeyForDataType(const FName& InDataTypeName, const EMetasoundFrontendVertexAccessType InAccessType, FNodeClassRegistryKey& OutKey) = 0;
+
+		virtual TUniquePtr<Metasound::INode> CreateNode(const FNodeClassRegistryKey& InKey, const Metasound::FNodeInitData&) const = 0;
+		virtual TUniquePtr<Metasound::INode> CreateNode(const FNodeClassRegistryKey& InKey, Metasound::FDefaultLiteralNodeConstructorParams&&) const = 0;
+		virtual TUniquePtr<Metasound::INode> CreateNode(const FNodeClassRegistryKey& InKey, Metasound::FDefaultNamedVertexNodeConstructorParams&&) const = 0;
+		virtual TUniquePtr<Metasound::INode> CreateNode(const FNodeClassRegistryKey& InKey, Metasound::FDefaultNamedVertexWithLiteralNodeConstructorParams&&) const = 0;
+
+		virtual bool RegisterConversionNode(const FConverterNodeClassRegistryKey& InNodeKey, const FConverterNodeClassInfo& InNodeInfo) = 0;
+
+		// Returns a list of possible nodes to use to convert from FromDataType to ToDataType.
+		// Returns an empty array if none are available.
+		virtual TArray<FConverterNodeClassInfo> GetPossibleConverterNodes(const FName& FromDataType, const FName& ToDataType) = 0;
+	};
+} // namespace Metasound::Frontend
diff --git a/Engine/Plugins/Runtime/Metasound/Source/MetasoundFrontend/Public/MetasoundFrontendRegistryContainer.h b/Engine/Plugins/Runtime/Metasound/Source/MetasoundFrontend/Public/MetasoundFrontendRegistryContainer.h
--- a/Engine/Plugins/Runtime/Metasound/Source/MetasoundFrontend/Public/MetasoundFrontendRegistryContainer.h
+++ b/Engine/Plugins/Runtime/Metasound/Source/MetasoundFrontend/Public/MetasoundFrontendRegistryContainer.h
@@ -1,6 +1,9 @@
 // Copyright Epic Games, Inc. All Rights Reserved.
 #pragma once
 
+#include "MetasoundFrontendNodeClassRegistry.h"
+
+#if UE_ENABLE_INCLUDE_ORDER_DEPRECATED_IN_5_6
 #include "IAudioProxyInitializer.h"
 #include "MetasoundDataReference.h"
 #include "MetasoundEnum.h"
@@ -13,347 +16,12 @@
 #include "MetasoundRouter.h"
 #include "MetasoundVertex.h"
 #include "Templates/Function.h"
+#endif
 
 // Forward Declarations
 class IMetaSoundDocumentInterface;
 template <typename InInterfaceType>
 class TScriptInterface;
-class UObject;
-
-
-namespace Metasound
-{
-	// Forward Declarations
-	using FIterateMetasoundFrontendClassFunction = TFunctionRef<void(const FMetasoundFrontendClass&)>;
-	class FGraph;
-} // namespace Metasound
-
-namespace Metasound::Frontend
-{
-	/** INodeRegistryEntry declares the interface for a node registry entry.
-		* Each node class in the registry must satisfy this interface. 
-		*/
-	class METASOUNDFRONTEND_API INodeRegistryEntry
-	{
-	public:
-		virtual ~INodeRegistryEntry() = default;
-
-		/** Return FNodeClassInfo for the node class.
-			*
-			* Implementations of method should avoid any expensive operations 
-			* (e.g. loading from disk, allocating memory) as this method is called
-			* frequently when querying nodes.
-			*/
-		virtual const FNodeClassInfo& GetClassInfo() const = 0;
-
-		/** Create a node given FDefaultNamedVertexNodeConstructorParams.
-			*
-			* If a node can be created with FDefaultNamedVertexNodeConstructorParams, this function
-			* should return a valid node pointer.
-			*/
-		virtual TUniquePtr<INode> CreateNode(FDefaultNamedVertexNodeConstructorParams&&) const = 0;
-
-		/** Create a node given FDefaultNamedVertexWithLiteralNodeConstructorParams.
-			*
-			* If a node can be created with FDefaultNamedVertexWithLiteralNodeConstructorParams, this function
-			* should return a valid node pointer.
-			*/
-		virtual TUniquePtr<INode> CreateNode(FDefaultNamedVertexWithLiteralNodeConstructorParams&&) const = 0;
-
-		/** Create a node given FDefaultLiteralNodeConstructorParams.
-			*
-			* If a node can be created with FDefaultLiteralNodeConstructorParams, this function
-			* should return a valid node pointer.
-			*/
-		virtual TUniquePtr<INode> CreateNode(FDefaultLiteralNodeConstructorParams&&) const = 0;
-
-		/** Create a node given FNodeInitData.
-			*
-			* If a node can be created with FNodeInitData, this function
-			* should return a valid node pointer.
-			*/
-		virtual TUniquePtr<INode> CreateNode(const FNodeInitData&) const = 0;
-
-		/** Return a FMetasoundFrontendClass which describes the node. */
-		virtual const FMetasoundFrontendClass& GetFrontendClass() const = 0;
-
-		/** Clone this registry entry. */
-		virtual TUniquePtr<INodeRegistryEntry> Clone() const = 0;
-
-		/** Returns set of implemented interface versions.
-			*
-			* Returns nullptr if node class implementation does not support interface implementation.
-			*/
-		virtual const TSet<FMetasoundFrontendVersion>* GetImplementedInterfaces() const = 0;
-
-		/** Whether or not the node is natively defined */
-		virtual bool IsNative() const = 0;
-	};
-
-	struct METASOUNDFRONTEND_API FConverterNodeRegistryKey
-	{
-		// The datatype one would like to convert from.
-		FName FromDataType;
-
-		// The datatype one would like to convert to.
-		FName ToDataType;
-
-		FORCEINLINE bool operator==(const FConverterNodeRegistryKey& Other) const
-		{
-			return FromDataType == Other.FromDataType && ToDataType == Other.ToDataType;
-		}
-
-		friend uint32 GetTypeHash(const ::Metasound::Frontend::FConverterNodeRegistryKey& InKey)
-		{
-			return HashCombine(GetTypeHash(InKey.FromDataType), GetTypeHash(InKey.ToDataType));
-		}
-	};
-
-	struct METASOUNDFRONTEND_API FConverterNodeInfo
-	{
-		// If this node has multiple input pins, we use this to designate which pin should be used.
-		FVertexName PreferredConverterInputPin;
-
-		// If this node has multiple output pins, we use this to designate which pin should be used.
-		FVertexName PreferredConverterOutputPin;
-
-		// The key for this node in the node registry.
-		FNodeRegistryKey NodeKey;
-
-		FORCEINLINE bool operator==(const FConverterNodeInfo& Other) const
-		{
-			return NodeKey == Other.NodeKey;
-		}
-	};
-
-	struct METASOUNDFRONTEND_API FConverterNodeRegistryValue
-	{
-		// A list of nodes that can perform a conversion between the two datatypes described in the FConverterNodeRegistryKey for this map element.
-		TArray<FConverterNodeInfo> PotentialConverterNodes;
-	};
-
-	using FRegistryTransactionID = int32;
-
-	class METASOUNDFRONTEND_API FNodeRegistryTransaction 
-	{
-	public:
-		using FTimeType = uint64;
-
-		/** Describes the type of transaction. */
-		enum class ETransactionType : uint8
-		{
-			NodeRegistration,     //< Something was added to the registry.
-			NodeUnregistration,  //< Something was removed from the registry.
-			Invalid
-		};
-
-		static FString LexToString(ETransactionType InTransactionType)
-		{
-			using namespace Metasound;
-
-			switch (InTransactionType)
-			{
-				case ETransactionType::NodeRegistration:
-				{
-					return TEXT("Node Registration");
-				}
-
-				case ETransactionType::NodeUnregistration:
-				{
-					return TEXT("Node Unregistration");
-				}
-
-				case ETransactionType::Invalid:
-				{
-					return TEXT("Invalid");
-				}
-
-				default:
-				{
-					checkNoEntry();
-					return TEXT("");
-				}
-			}
-		}
-
-		FNodeRegistryTransaction(ETransactionType InType, const FNodeClassInfo& InNodeClassInfo, FTimeType InTimestamp);
-
-		ETransactionType GetTransactionType() const;
-		const FNodeClassInfo& GetNodeClassInfo() const;
-		FNodeRegistryKey GetNodeRegistryKey() const;
-		FTimeType GetTimestamp() const;
-
-	private:
-		ETransactionType Type;
-		FNodeClassInfo NodeClassInfo;
-		FTimeType Timestamp;
-	};
-
-	namespace NodeRegistryKey
-	{
-		UE_DEPRECATED(5.4, "Use FNodeRegistryKey::GetInvalid() instead")
-		METASOUNDFRONTEND_API const FNodeRegistryKey& GetInvalid();
-
-		UE_DEPRECATED(5.4, "Use FNodeRegistryKey::IsValid() instead")
-		METASOUNDFRONTEND_API bool IsValid(const FNodeRegistryKey& InKey);
-
-		UE_DEPRECATED(5.4, "Use FNodeRegistryKey equality operator instead")
-		METASOUNDFRONTEND_API bool IsEqual(const FNodeRegistryKey& InLHS, const FNodeRegistryKey& InRHS);
-
-		// Returns true if the class metadata represent the same entry in the node registry.
-		METASOUNDFRONTEND_API bool IsEqual(const FMetasoundFrontendClassMetadata& InLHS, const FMetasoundFrontendClassMetadata& InRHS);
-
-		// Returns true if the class info and class metadata represent the same entry in the node registry.
-		METASOUNDFRONTEND_API bool IsEqual(const FNodeClassInfo& InLHS, const FMetasoundFrontendClassMetadata& InRHS);
-
-		UE_DEPRECATED(5.4, "Use FNodeRegistryKey constructors instead")
-		METASOUNDFRONTEND_API FNodeRegistryKey CreateKey(EMetasoundFrontendClassType InType, const FString& InFullClassName, int32 InMajorVersion, int32 InMinorVersion);
-
-		UE_DEPRECATED(5.4, "Use FNodeRegistryKey constructors instead")
-		METASOUNDFRONTEND_API FNodeRegistryKey CreateKey(const FNodeClassMetadata& InNodeMetadata);
-
-		UE_DEPRECATED(5.4, "Use FNodeRegistryKey constructor instead")
-		METASOUNDFRONTEND_API FNodeRegistryKey CreateKey(const FMetasoundFrontendClassMetadata& InNodeMetadata);
-
-		UE_DEPRECATED(5.4, "Use FNodeRegistryKey constructor instead")
-		METASOUNDFRONTEND_API FNodeRegistryKey CreateKey(const FMetasoundFrontendGraphClass& InNodeMetadata);
-
-		UE_DEPRECATED(5.4, "Use FNodeRegistryKey constructor instead")
-		METASOUNDFRONTEND_API FNodeRegistryKey CreateKey(const FNodeClassInfo& ClassInfo);
-	}
-} // namespace Metasound::Frontend
-
-/**
- * Singleton registry for all types and nodes.
- */
-class METASOUNDFRONTEND_API FMetasoundFrontendRegistryContainer
-{
-
-public:
-	// The MetaSound frontend does not rely on the Engine module and therefore
-	// does not have the ability to provide GC protection. This interface allows
-	// external modules to provide GC protection so that async tasks can safely
-	// use UObjects
-	class IObjectReferencer
-	{
-	public:
-		virtual ~IObjectReferencer() = default;
-
-		// Called when an object should be referenced.
-		virtual void AddObject(UObject* InObject) = 0;
-
-		// Called when an object no longer needs to be referenced. 
-		virtual void RemoveObject(UObject* InObject) = 0;
-	};
-
-	using FNodeClassInfo = Metasound::Frontend::FNodeClassInfo;
-	using FConverterNodeRegistryKey = ::Metasound::Frontend::FConverterNodeRegistryKey;
-	using FConverterNodeRegistryValue = ::Metasound::Frontend::FConverterNodeRegistryValue;
-	using FConverterNodeInfo = ::Metasound::Frontend::FConverterNodeInfo;
-
-	using FGraphRegistryKey = Metasound::Frontend::FGraphRegistryKey;
-	using FNodeRegistryKey = Metasound::Frontend::FNodeRegistryKey;
-	using FNodeClassMetadata = Metasound::FNodeClassMetadata;
-
-	static FMetasoundFrontendRegistryContainer* Get();
-	static void ShutdownMetasoundFrontend();
-
-	UE_DEPRECATED(5.4, "Use FNodeRegistryKey constructor that takes FNodeClassMetadata instead")
-	static FNodeRegistryKey GetRegistryKey(const FNodeClassMetadata& InNodeMetadata);
-
-	UE_DEPRECATED(5.4, "Use FNodeRegistryKey constructor that takes FMetasoundFrontendClassMetadata instead")
-	static FNodeRegistryKey GetRegistryKey(const FMetasoundFrontendClassMetadata& InNodeMetadata);
-
-	UE_DEPRECATED(5.4, "Use FNodeRegistryKey constructor that takes FNodeClassInfo instead")
-	static FNodeRegistryKey GetRegistryKey(const FNodeClassInfo& ClassInfo);
-
-	static bool GetFrontendClassFromRegistered(const FNodeRegistryKey& InKey, FMetasoundFrontendClass& OutClass);
-
-	UE_DEPRECATED(5.4, "Use GetFrontendClassFromRegistered instead")
-	static bool GetNodeClassInfoFromRegistered(const FNodeRegistryKey& InKey, FNodeClassInfo& OutInfo);
-
-	static bool GetInputNodeRegistryKeyForDataType(const FName& InDataTypeName, const EMetasoundFrontendVertexAccessType InAccessType, FNodeRegistryKey& OutKey);
-	static bool GetVariableNodeRegistryKeyForDataType(const FName& InDataTypeName, FNodeRegistryKey& OutKey);
-	static bool GetOutputNodeRegistryKeyForDataType(const FName& InDataTypeName, const EMetasoundFrontendVertexAccessType InAccessType, FNodeRegistryKey& OutKey);
-
-
-	FMetasoundFrontendRegistryContainer() = default;
-	virtual ~FMetasoundFrontendRegistryContainer() = default;
-
-	FMetasoundFrontendRegistryContainer(const FMetasoundFrontendRegistryContainer&) = delete;
-	FMetasoundFrontendRegistryContainer& operator=(const FMetasoundFrontendRegistryContainer&) = delete;
-
-	// Enqueue and init command for registering a node or data type.
-	// The command queue will be processed on module init or when calling `RegisterPendingNodes()`
-	virtual bool EnqueueInitCommand(TUniqueFunction<void()>&& InFunc) = 0;
-
-	virtual void SetObjectReferencer(TUniquePtr<IObjectReferencer> InReferencer) = 0;
-
-	// This is called on module startup. This invokes any registration commands enqueued by our registration macros.
-	virtual void RegisterPendingNodes() = 0;
-
-	// Wait for async graph registration to complete for a specific graph
-	virtual void WaitForAsyncGraphRegistration(const FGraphRegistryKey& InRegistryKey) const = 0;
-
-	// Retrieve a registered graph. 
-	//
-	// If the graph is registered asynchronously, this will wait until the registration task has completed.
-	virtual TSharedPtr<const Metasound::FGraph> GetGraph(const FGraphRegistryKey& InRegistryKey) const = 0;
-
-	/** Register an external node with the frontend.
-	 *
-	 * @param InCreateNode - Function for creating node from FNodeInitData.
-	 * @param InCreateDescription - Function for creating a FMetasoundFrontendClass.
-	 *
-	 * @return A node registration key. If the registration failed, then the registry 
-	 *         key will be invalid.
-	 */
-	virtual FNodeRegistryKey RegisterNode(TUniquePtr<Metasound::Frontend::INodeRegistryEntry>&& InEntry) = 0;
-
-	/** Unregister an external node from the frontend.
-	 *
-	 * @param InKey - The registration key for the node.
-	 *
-	 * @return True on success, false on failure.
-	 */
-	virtual bool UnregisterNode(const FNodeRegistryKey& InKey) = 0;
-
-	/** Returns true if the provided registry key corresponds to a valid registered node. */
-	virtual bool IsNodeRegistered(const FNodeRegistryKey& InKey) const = 0;
-
-	/** Returns true if the provided registry key (node key and asset path) corresponds to a valid registered graph. */
-	virtual bool IsGraphRegistered(const FGraphRegistryKey& InKey) const = 0;
-
-	/** Returns true if the provided registry key corresponds to a valid registered node that is natively defined. */
-	virtual bool IsNodeNative(const FNodeRegistryKey& InKey) const = 0;
-
-	// Iterates class types in registry.  If InClassType is set to a valid class type (optional), only iterates classes of the given type
-	virtual void IterateRegistry(Metasound::FIterateMetasoundFrontendClassFunction InIterFunc, EMetasoundFrontendClassType InClassType = EMetasoundFrontendClassType::Invalid) const = 0;
-
-	// Query for MetaSound Frontend document objects.
-	virtual bool FindFrontendClassFromRegistered(const Metasound::Frontend::FNodeRegistryKey& InKey, FMetasoundFrontendClass& OutClass) = 0;
-
-	UE_DEPRECATED(5.4, "This implementation of FindImplementedInterfacesFromRegistered(...). Please use bool FindImplementedInterfacesFromRegistered(const Metasound::Frontend::FNodeRegistryKey& InKey, TSet<FMetasoundFrontendVersion>& OutInterfacVersion)")
-	virtual const TSet<FMetasoundFrontendVersion>* FindImplementedInterfacesFromRegistered(const Metasound::Frontend::FNodeRegistryKey& InKey) const = 0;
-
-	virtual bool FindImplementedInterfacesFromRegistered(const Metasound::Frontend::FNodeRegistryKey& InKey, TSet<FMetasoundFrontendVersion>& OutInterfaceVersions) const = 0;
-	
-	UE_DEPRECATED(5.4, "Use FindFrontendClassFromRegistered instead")
-	virtual bool FindNodeClassInfoFromRegistered(const Metasound::Frontend::FNodeRegistryKey& InKey, FNodeClassInfo& OutInfo) = 0;
-	virtual bool FindInputNodeRegistryKeyForDataType(const FName& InDataTypeName, const EMetasoundFrontendVertexAccessType InAccessType, FNodeRegistryKey& OutKey) = 0;
-	virtual bool FindVariableNodeRegistryKeyForDataType(const FName& InDataTypeName, FNodeRegistryKey& OutKey) = 0;
-	virtual bool FindOutputNodeRegistryKeyForDataType(const FName& InDataTypeName, const EMetasoundFrontendVertexAccessType InAccessType, FNodeRegistryKey& OutKey) = 0;
-
-	virtual TUniquePtr<Metasound::INode> CreateNode(const FNodeRegistryKey& InKey, const Metasound::FNodeInitData&) const = 0;
-	virtual TUniquePtr<Metasound::INode> CreateNode(const FNodeRegistryKey& InKey, Metasound::FDefaultLiteralNodeConstructorParams&&) const = 0;
-	virtual TUniquePtr<Metasound::INode> CreateNode(const FNodeRegistryKey& InKey, Metasound::FDefaultNamedVertexNodeConstructorParams&&) const = 0;
-	virtual TUniquePtr<Metasound::INode> CreateNode(const FNodeRegistryKey& InKey, Metasound::FDefaultNamedVertexWithLiteralNodeConstructorParams&&) const = 0;
-
-	virtual bool RegisterConversionNode(const FConverterNodeRegistryKey& InNodeKey, const FConverterNodeInfo& InNodeInfo) = 0;
-
-	// Returns a list of possible nodes to use to convert from FromDataType to ToDataType.
-	// Returns an empty array if none are available.
-	virtual TArray<FConverterNodeInfo> GetPossibleConverterNodes(const FName& FromDataType, const FName& ToDataType) = 0;
-};
 
+using FMetasoundFrontendRegistryContainer = Metasound::Frontend::INodeClassRegistry;
 
diff --git a/Engine/Plugins/Runtime/Metasound/Source/MetasoundFrontend/Public/MetasoundFrontendRegistryKey.h b/Engine/Plugins/Runtime/Metasound/Source/MetasoundFrontend/Public/MetasoundFrontendRegistryKey.h
--- a/Engine/Plugins/Runtime/Metasound/Source/MetasoundFrontend/Public/MetasoundFrontendRegistryKey.h
+++ b/Engine/Plugins/Runtime/Metasound/Source/MetasoundFrontend/Public/MetasoundFrontendRegistryKey.h
@@ -49,40 +49,31 @@ namespace Metasound::Frontend
 		// should the class originate from one).
 		FNodeClassInfo(const FMetasoundFrontendClassMetadata& InMetadata);
 
-		UE_DEPRECATED(5.4, "Use constructor that takes an FTopLevelAssetPath instead.")
-		FNodeClassInfo(const FMetasoundFrontendGraphClass& InClass, const FSoftObjectPath& InAssetPath);
-
 		// Constructor used to generate NodeClassInfo from an asset
 		FNodeClassInfo(const FMetasoundFrontendGraphClass& InClass, const FTopLevelAssetPath& InAssetPath);
-
-		UE_DEPRECATED(5.4, "NodeClassInfo no longer supports directly loading an asset from its stored path.")
-		UObject* LoadAsset() const;
 	};
 
-	struct METASOUNDFRONTEND_API FNodeRegistryKey
+	struct METASOUNDFRONTEND_API FNodeClassRegistryKey
 	{
 		EMetasoundFrontendClassType Type = EMetasoundFrontendClassType::Invalid;
 		FMetasoundFrontendClassName ClassName;
 		FMetasoundFrontendVersionNumber Version;
 
-		FNodeRegistryKey() = default;
-		FNodeRegistryKey(EMetasoundFrontendClassType InType, const FMetasoundFrontendClassName& InClassName, int32 InMajorVersion, int32 InMinorVersion);
-		FNodeRegistryKey(EMetasoundFrontendClassType InType, const FMetasoundFrontendClassName& InClassName, const FMetasoundFrontendVersionNumber& InVersion);
-		FNodeRegistryKey(const FNodeClassMetadata& InNodeMetadata);
-		FNodeRegistryKey(const FMetasoundFrontendClassMetadata& InNodeMetadata);
-		FNodeRegistryKey(const FMetasoundFrontendGraphClass& InGraphClass);
-		FNodeRegistryKey(const FNodeClassInfo& InClassInfo);
-		FNodeRegistryKey(const FAssetKey& AssetKey);
-
-		UE_DEPRECATED(5.4, "Implicit String ctor is no longer supported.")
-		FNodeRegistryKey(const FString& InKeyString);
+		FNodeClassRegistryKey() = default;
+		FNodeClassRegistryKey(EMetasoundFrontendClassType InType, const FMetasoundFrontendClassName& InClassName, int32 InMajorVersion, int32 InMinorVersion);
+		FNodeClassRegistryKey(EMetasoundFrontendClassType InType, const FMetasoundFrontendClassName& InClassName, const FMetasoundFrontendVersionNumber& InVersion);
+		FNodeClassRegistryKey(const FNodeClassMetadata& InNodeMetadata);
+		FNodeClassRegistryKey(const FMetasoundFrontendClassMetadata& InNodeMetadata);
+		FNodeClassRegistryKey(const FMetasoundFrontendGraphClass& InGraphClass);
+		FNodeClassRegistryKey(const FNodeClassInfo& InClassInfo);
+		FNodeClassRegistryKey(const FAssetKey& AssetKey);
 
-		FORCEINLINE friend bool operator==(const FNodeRegistryKey& InLHS, const FNodeRegistryKey& InRHS)
+		FORCEINLINE friend bool operator==(const FNodeClassRegistryKey& InLHS, const FNodeClassRegistryKey& InRHS)
 		{
 			return (InLHS.Type == InRHS.Type) && (InLHS.ClassName == InRHS.ClassName) && (InLHS.Version == InRHS.Version);
 		}
 
-		FORCEINLINE friend bool operator<(const FNodeRegistryKey& InLHS, const FNodeRegistryKey& InRHS)
+		FORCEINLINE friend bool operator<(const FNodeClassRegistryKey& InLHS, const FNodeClassRegistryKey& InRHS)
 		{
 			if (static_cast<uint8>(InLHS.Type) == static_cast<uint8>(InRHS.Type))
 			{
@@ -97,37 +88,18 @@ namespace Metasound::Frontend
 			return static_cast<uint8>(InLHS.Type) < static_cast<uint8>(InRHS.Type);
 		}
 
-		friend FORCEINLINE uint32 GetTypeHash(const FNodeRegistryKey& InKey)
+		friend FORCEINLINE uint32 GetTypeHash(const FNodeClassRegistryKey& InKey)
 		{
 			const int32 Hash = HashCombineFast(static_cast<uint32>(InKey.Type), GetTypeHash(InKey.ClassName));
 			return HashCombineFast(Hash, GetTypeHash(InKey.Version));
 		}
 
-		UE_DEPRECATED(5.4, "Implicit String IsEmpty() check is no longer valid. Use IsValid() instead.")
-		bool IsEmpty() const
-		{
-			return !IsValid();
-		}
-
 		// Returns invalid (default constructed) key
-		static const FNodeRegistryKey& GetInvalid();
+		static const FNodeClassRegistryKey& GetInvalid();
 
 		// Returns whether or not instance is valid or not
 		bool IsValid() const;
 
-		UE_DEPRECATED(5.4, "Implicit cast to FString from FNodeRegistryKey is deprecated, use ToString().")
-		FORCEINLINE operator FString() const
-		{
-			return ToString();
-		}
-
-		UE_DEPRECATED(5.4, "Implicit cast to TCHAR* is deprecated, use ToString().")
-		FORCEINLINE const TCHAR* operator*() const
-		{
-			static const FString DeprecatedString;
-			return *DeprecatedString;
-		}
-
 		// Resets the key back to an invalid default state
 		void Reset();
 
@@ -139,12 +111,13 @@ namespace Metasound::Frontend
 
 		// Parses string representation of key into registry key.  For debug and deserialization use only.
 		// Returns true if parsed successfully.
-		static bool Parse(const FString& InKeyString, FNodeRegistryKey& OutKey);
+		static bool Parse(const FString& InKeyString, FNodeClassRegistryKey& OutKey);
 	};
+	using FNodeRegistryKey = FNodeClassRegistryKey;
 
-	struct METASOUNDFRONTEND_API FGraphRegistryKey
+	struct METASOUNDFRONTEND_API FGraphClassRegistryKey
 	{
-		FNodeRegistryKey NodeKey;
+		FNodeClassRegistryKey NodeKey;
 		FTopLevelAssetPath AssetPath;
 
 		FString ToString() const;
@@ -154,15 +127,16 @@ namespace Metasound::Frontend
 
 		bool IsValid() const;
 
-		FORCEINLINE friend bool operator==(const FGraphRegistryKey& InLHS, const FGraphRegistryKey& InRHS)
+		FORCEINLINE friend bool operator==(const FGraphClassRegistryKey& InLHS, const FGraphClassRegistryKey& InRHS)
 		{
 			return (InLHS.NodeKey == InRHS.NodeKey) && (InLHS.AssetPath == InRHS.AssetPath);
 		}
 
-		friend FORCEINLINE uint32 GetTypeHash(const FGraphRegistryKey& InKey)
+		friend FORCEINLINE uint32 GetTypeHash(const FGraphClassRegistryKey& InKey)
 		{
 			const int32 Hash = HashCombineFast(GetTypeHash(InKey.NodeKey), GetTypeHash(InKey.AssetPath));
 			return Hash;
 		}
 	};
+	using FGraphRegistryKey = FGraphClassRegistryKey;
 } // namespace Metasound::Frontend
diff --git a/Engine/Plugins/Runtime/Metasound/Source/MetasoundFrontend/Public/MetasoundNodeRegistrationMacro.h b/Engine/Plugins/Runtime/Metasound/Source/MetasoundFrontend/Public/MetasoundNodeRegistrationMacro.h
--- a/Engine/Plugins/Runtime/Metasound/Source/MetasoundFrontend/Public/MetasoundNodeRegistrationMacro.h
+++ b/Engine/Plugins/Runtime/Metasound/Source/MetasoundFrontend/Public/MetasoundNodeRegistrationMacro.h
@@ -5,7 +5,7 @@
 #include "CoreMinimal.h"
 #include "MetasoundFrontend.h"
 #include "MetasoundFrontendDocument.h"
-#include "MetasoundFrontendRegistries.h"
+#include "MetasoundFrontendNodeClassRegistry.h"
 #include "MetasoundLog.h"
 #include "MetasoundNodeInterface.h"
 
@@ -80,8 +80,6 @@ namespace Metasound
 				return FrontendClass;
 			}
 
-			virtual TUniquePtr<INodeRegistryEntry> Clone() const override = 0;
-
 			virtual const TSet<FMetasoundFrontendVersion>* GetImplementedInterfaces() const override
 			{
 				return nullptr;
@@ -110,11 +108,6 @@ namespace Metasound
 			{
 				return MakeUnique<FNodeType>(InParams);
 			}
-
-			virtual TUniquePtr<Frontend::INodeRegistryEntry> Clone() const override
-			{
-				return MakeUnique<FNodeRegistryEntry>(*this);
-			}
 		};
 
 
diff --git a/Engine/Plugins/Runtime/Metasound/Source/MetasoundStandardNodes/Private/MetasoundAbsoluteValueNodes.cpp b/Engine/Plugins/Runtime/Metasound/Source/MetasoundStandardNodes/Private/MetasoundAbsoluteValueNodes.cpp
--- a/Engine/Plugins/Runtime/Metasound/Source/MetasoundStandardNodes/Private/MetasoundAbsoluteValueNodes.cpp
+++ b/Engine/Plugins/Runtime/Metasound/Source/MetasoundStandardNodes/Private/MetasoundAbsoluteValueNodes.cpp
@@ -8,6 +8,7 @@
 #include "MetasoundPrimitives.h"
 #include "MetasoundStandardNodesCategories.h"
 #include "MetasoundStandardNodesNames.h"
+#include "MetasoundTime.h"
 #include "MetasoundAudioBuffer.h"
 #include "Internationalization/Text.h"
 #include "MetasoundParamHelper.h"
diff --git a/Engine/Plugins/Runtime/Metasound/Source/MetasoundStandardNodes/Private/MetasoundAdditiveSynthNode.cpp b/Engine/Plugins/Runtime/Metasound/Source/MetasoundStandardNodes/Private/MetasoundAdditiveSynthNode.cpp
--- a/Engine/Plugins/Runtime/Metasound/Source/MetasoundStandardNodes/Private/MetasoundAdditiveSynthNode.cpp
+++ b/Engine/Plugins/Runtime/Metasound/Source/MetasoundStandardNodes/Private/MetasoundAdditiveSynthNode.cpp
@@ -4,6 +4,7 @@
 #include "DSP/SineWaveTableOsc.h"
 #include "MetasoundFacade.h"
 #include "Internationalization/Text.h"
+#include "MetasoundAudioBuffer.h"
 #include "MetasoundExecutableOperator.h"
 #include "MetasoundNodeRegistrationMacro.h"
 #include "MetasoundParamHelper.h"
