Non-unity fixes metasound plugin

#rnx

diff --git a/Engine/Plugins/Runtime/Metasound/Source/MetasoundFrontend/Public/MetasoundFrontendArchetypeRegistry.h b/Engine/Plugins/Runtime/Metasound/Source/MetasoundFrontend/Public/MetasoundFrontendArchetypeRegistry.h
--- a/Engine/Plugins/Runtime/Metasound/Source/MetasoundFrontend/Public/MetasoundFrontendArchetypeRegistry.h
+++ b/Engine/Plugins/Runtime/Metasound/Source/MetasoundFrontend/Public/MetasoundFrontendArchetypeRegistry.h
@@ -2,6 +2,8 @@
 // DEPRECATED MetasoundFrontendArchetypeRegistry.h. Use Interfaces/MetasoundFrontendInterfaceRegistry.h instead
 #pragma once
 
+#include "Misc/CoreMiscDefines.h"
+
 // HEADER_UNIT_SKIP - Deprecated
 
 UE_DEPRECATED_HEADER(5.3, "Use Interfaces/MetasoundFrontendInterfaceRegistry.h instead instead of MetasoundFrontendArchetypeRegistry.h.")
diff --git a/Engine/Plugins/Runtime/Metasound/Source/MetasoundGraphCore/Public/MetasoundNodeRegister.h b/Engine/Plugins/Runtime/Metasound/Source/MetasoundGraphCore/Public/MetasoundNodeRegister.h
deleted file mode 100644
diff --git a/Engine/Plugins/Runtime/Metasound/Source/MetasoundStandardNodes/Private/MetasoundFlangerNode.cpp b/Engine/Plugins/Runtime/Metasound/Source/MetasoundStandardNodes/Private/MetasoundFlangerNode.cpp
--- a/Engine/Plugins/Runtime/Metasound/Source/MetasoundStandardNodes/Private/MetasoundFlangerNode.cpp
+++ b/Engine/Plugins/Runtime/Metasound/Source/MetasoundStandardNodes/Private/MetasoundFlangerNode.cpp
@@ -3,6 +3,7 @@
 #include "DSP/BufferVectorOperations.h"
 #include "DSP/Flanger.h"
 #include "Internationalization/Text.h"
+#include "MetasoundAudioBuffer.h"
 #include "MetasoundExecutableOperator.h"
 #include "MetasoundFacade.h"
 #include "MetasoundNodeRegistrationMacro.h"
