Deprecate const bind APIs

#rnx

diff --git a/Engine/Plugins/Runtime/Metasound/Source/MetasoundEngine/Private/MetasoundWaveTableOscillatorNode.cpp b/Engine/Plugins/Runtime/Metasound/Source/MetasoundEngine/Private/MetasoundWaveTableOscillatorNode.cpp
--- a/Engine/Plugins/Runtime/Metasound/Source/MetasoundEngine/Private/MetasoundWaveTableOscillatorNode.cpp
+++ b/Engine/Plugins/Runtime/Metasound/Source/MetasoundEngine/Private/MetasoundWaveTableOscillatorNode.cpp
@@ -142,23 +142,7 @@ namespace Metasound
 
 		virtual void BindOutputs(FOutputVertexInterfaceData& InOutVertexData) override
 		{
-			InOutVertexData.BindReadVertex("Out", TDataReadReference<FAudioBuffer>(OutBufferWriteRef));
-		}
-
-		virtual FDataReferenceCollection GetInputs() const override
-		{
-			// This should never be called. Bind(...) is called instead. This method
-			// exists as a stop-gap until the API can be deprecated and removed.
-			checkNoEntry();
-			return {};
-		}
-
-		virtual FDataReferenceCollection GetOutputs() const override
-		{
-			// This should never be called. Bind(...) is called instead. This method
-			// exists as a stop-gap until the API can be deprecated and removed.
-			checkNoEntry();
-			return {};
+			InOutVertexData.BindReadVertex("Out", OutBufferWriteRef);
 		}
 
 		void Execute()
diff --git a/Engine/Plugins/Runtime/Metasound/Source/MetasoundEngine/Private/MetasoundWaveTablePlayerNode.cpp b/Engine/Plugins/Runtime/Metasound/Source/MetasoundEngine/Private/MetasoundWaveTablePlayerNode.cpp
--- a/Engine/Plugins/Runtime/Metasound/Source/MetasoundEngine/Private/MetasoundWaveTablePlayerNode.cpp
+++ b/Engine/Plugins/Runtime/Metasound/Source/MetasoundEngine/Private/MetasoundWaveTablePlayerNode.cpp
@@ -178,8 +178,8 @@ namespace Metasound
 
 		virtual void BindOutputs(FOutputVertexInterfaceData& InOutVertexData) override
 		{
-			InOutVertexData.BindReadVertex("MonoOut", TDataReadReference<FAudioBuffer>(OutBufferWriteRef));
-			InOutVertexData.BindReadVertex("OnFinished", TDataReadReference<FTrigger>(OutOnFinishedRef));
+			InOutVertexData.BindReadVertex("MonoOut", OutBufferWriteRef);
+			InOutVertexData.BindReadVertex("OnFinished", OutOnFinishedRef);
 		}
 
 		virtual FDataReferenceCollection GetInputs() const override
diff --git a/Engine/Plugins/Runtime/Metasound/Source/MetasoundFrontend/Public/MetasoundInputNode.h b/Engine/Plugins/Runtime/Metasound/Source/MetasoundFrontend/Public/MetasoundInputNode.h
--- a/Engine/Plugins/Runtime/Metasound/Source/MetasoundFrontend/Public/MetasoundInputNode.h
+++ b/Engine/Plugins/Runtime/Metasound/Source/MetasoundFrontend/Public/MetasoundInputNode.h
@@ -188,7 +188,7 @@ namespace Metasound
 
 			virtual void BindOutputs(FOutputVertexInterfaceData& InOutVertexData) override
 			{
-				InOutVertexData.BindVertex(DataReferenceName, DataRef.GetDataReadReference<DataType>());
+				InOutVertexData.BindVertex(DataReferenceName, DataRef);
 			}
 
 			virtual FExecuteFunction GetExecuteFunction() override
diff --git a/Engine/Plugins/Runtime/Metasound/Source/MetasoundFrontend/Public/MetasoundReceiveNode.h b/Engine/Plugins/Runtime/Metasound/Source/MetasoundFrontend/Public/MetasoundReceiveNode.h
--- a/Engine/Plugins/Runtime/Metasound/Source/MetasoundFrontend/Public/MetasoundReceiveNode.h
+++ b/Engine/Plugins/Runtime/Metasound/Source/MetasoundFrontend/Public/MetasoundReceiveNode.h
@@ -126,23 +126,7 @@ namespace Metasound
 				virtual void BindOutputs(FOutputVertexInterfaceData& InOutVertexData) override
 				{
 					using namespace ReceiveNodeInfo;
-					InOutVertexData.BindReadVertex<TDataType>(METASOUND_GET_PARAM_NAME(Output), TDataReadReference<TDataType>(OutputData));
-				}
-
-				virtual FDataReferenceCollection GetInputs() const override
-				{
-					// This should never be called. Bind(...) is called instead. This method
-					// exists as a stop-gap until the API can be deprecated and removed.
-					checkNoEntry();
-					return {};
-				}
-
-				virtual FDataReferenceCollection GetOutputs() const override
-				{
-					// This should never be called. Bind(...) is called instead. This method
-					// exists as a stop-gap until the API can be deprecated and removed.
-					checkNoEntry();
-					return {};
+					InOutVertexData.BindReadVertex<TDataType>(METASOUND_GET_PARAM_NAME(Output), OutputData);
 				}
 
 				void Execute()
diff --git a/Engine/Plugins/Runtime/Metasound/Source/MetasoundGraphCore/Private/MetasoundFacade.cpp b/Engine/Plugins/Runtime/Metasound/Source/MetasoundGraphCore/Private/MetasoundFacade.cpp
--- a/Engine/Plugins/Runtime/Metasound/Source/MetasoundGraphCore/Private/MetasoundFacade.cpp
+++ b/Engine/Plugins/Runtime/Metasound/Source/MetasoundGraphCore/Private/MetasoundFacade.cpp
@@ -9,29 +9,11 @@ namespace Metasound
 	{
 	}
 
-	FNodeFacade::FFactory::FFactory(FOriginalCreateOperatorFunction InCreateFunc)
-	:	CreateFunc(WrapOriginalCreateOperatorFunction(InCreateFunc))
-	{
-	}
-
-
 	TUniquePtr<IOperator> FNodeFacade::FFactory::CreateOperator(const FBuildOperatorParams& InParams, FBuildResults& OutResults)
 	{
 		return CreateFunc(InParams, OutResults);
 	}
 
-	FNodeFacade::FFactory::FCreateOperatorFunction FNodeFacade::FFactory::WrapOriginalCreateOperatorFunction(FOriginalCreateOperatorFunction InCreateFunc)
-	{
-		return [InCreateFunc](const FBuildOperatorParams& InParams, FBuildResults& OutResults)
-		{
-			FDataReferenceCollection InputDataReferences = InParams.InputData.ToDataReferenceCollection();
-PRAGMA_DISABLE_DEPRECATION_WARNINGS // Temporary for the deprecation of FCreateOperatorParams
-			FCreateOperatorParams DeprecatedParams{ InParams.Node, InParams.OperatorSettings, InputDataReferences, InParams.Environment, InParams.Builder };
-			return InCreateFunc(DeprecatedParams, OutResults.Errors);
-PRAGMA_ENABLE_DEPRECATION_WARNINGS
-		};
-	}
-
 	const FVertexInterface& FNodeFacade::GetVertexInterface() const
 	{
 		return VertexInterface;
diff --git a/Engine/Plugins/Runtime/Metasound/Source/MetasoundGraphCore/Private/MetasoundOperatorBuilder.cpp b/Engine/Plugins/Runtime/Metasound/Source/MetasoundGraphCore/Private/MetasoundOperatorBuilder.cpp
--- a/Engine/Plugins/Runtime/Metasound/Source/MetasoundGraphCore/Private/MetasoundOperatorBuilder.cpp
+++ b/Engine/Plugins/Runtime/Metasound/Source/MetasoundGraphCore/Private/MetasoundOperatorBuilder.cpp
@@ -508,7 +508,9 @@ namespace Metasound
 			check(NodePtr);
 			if (const FGraphOperatorData::FOperatorInfo* OpInfo = InOutContext.GraphOperatorData.OperatorMap.Find(GetOperatorID(NodePtr)))
 			{
+PRAGMA_DISABLE_DEPRECATION_WARNINGS
 				OutNodeVertexData.Emplace(NodePtr->GetInstanceID(), OpInfo->VertexData.GetOutputs().ToDataReferenceCollection());
+PRAGMA_ENABLE_DEPRECATION_WARNINGS
 			}
 		}
 	}
diff --git a/Engine/Plugins/Runtime/Metasound/Source/MetasoundGraphCore/Private/MetasoundVertexData.cpp b/Engine/Plugins/Runtime/Metasound/Source/MetasoundGraphCore/Private/MetasoundVertexData.cpp
--- a/Engine/Plugins/Runtime/Metasound/Source/MetasoundGraphCore/Private/MetasoundVertexData.cpp
+++ b/Engine/Plugins/Runtime/Metasound/Source/MetasoundGraphCore/Private/MetasoundVertexData.cpp
@@ -121,20 +121,6 @@ namespace Metasound
 		}
 
 
-		void FInputBinding::Bind(FAnyDataReference& InOutDataReference)
-		{
-			check(Vertex.DataTypeName == InOutDataReference.GetDataTypeName());
-			if (Data.IsSet())
-			{
-				InOutDataReference = *Data;
-			}
-			else
-			{
-				Data.Emplace(InOutDataReference);
-				CheckAccessTypeCompatibility(Vertex, *Data);
-			}
-		}
-
 		void FInputBinding::Bind(FInputBinding& InBinding)
 		{
 			check(Vertex.DataTypeName == InBinding.GetVertex().DataTypeName);
@@ -214,12 +200,6 @@ namespace Metasound
 			return nullptr;
 		}
 
-		void FOutputBinding::Bind(FAnyDataReference& InOutDataReference)
-		{
-			check(Vertex.DataTypeName == InOutDataReference.GetDataTypeName());
-			Data.Emplace(InOutDataReference);
-			CheckAccessTypeCompatibility(Vertex, *Data);
-		}
 
 		void FOutputBinding::Bind(FOutputBinding& InBinding)
 		{
diff --git a/Engine/Plugins/Runtime/Metasound/Source/MetasoundGraphCore/Public/MetasoundBuilderInterface.h b/Engine/Plugins/Runtime/Metasound/Source/MetasoundGraphCore/Public/MetasoundBuilderInterface.h
--- a/Engine/Plugins/Runtime/Metasound/Source/MetasoundGraphCore/Public/MetasoundBuilderInterface.h
+++ b/Engine/Plugins/Runtime/Metasound/Source/MetasoundGraphCore/Public/MetasoundBuilderInterface.h
@@ -59,35 +59,6 @@ namespace Metasound
 		TMap<FGuid, FDataReferenceCollection> InternalDataReferences;
 	};
 
-	/** FCreateOperatorParams holds the parameters provided to operator factories
-	 * during the creation of an IOperator
-	 */
-	struct UE_DEPRECATED(5.4, "Update APIs to use FBuildOperatorParams. Operator creation routines should have the following signature \"TUniquePtr<IOperator> CreateOperator(const FBuildOperatorParams&, FBuildResults&)\"") FCreateOperatorParams;
-
-	struct FCreateOperatorParams
-	{
-		/** The node associated with this factory and the desired IOperator. */
-		const INode& Node;
-
-		/** General operator settings for the graph. */
-		const FOperatorSettings& OperatorSettings;
-
-		/** Collection of input parameters available for to an IOperator. */
-		const FDataReferenceCollection& InputDataReferences;
-
-		/** Environment settings available. */
-		const FMetasoundEnvironment& Environment;
-
-		/** Pointer to builder actively building graph. */
-		const IOperatorBuilder* Builder = nullptr;
-
-		/** Implicit conversion to FResetParams for convenience. */
-		operator IOperator::FResetParams() const
-		{
-			return IOperator::FResetParams{OperatorSettings, Environment};
-		}
-	};
-
 	/** FBuildOperatorParams holds the parameters provided to operator factories
 	 * during the creation of an IOperator
 	 */
@@ -118,10 +89,8 @@ namespace Metasound
 		}
 	};
 
-	/** Parameters for building an operator from a graph. 
-	  *
-	  * This object is slated to be deprecated in UE 5.2
-	  */
+	struct UE_DEPRECATED(5.6, "Use FBuildGrpahOperatorParams instead") FBuildGraphParams;
+	/** Parameters for building an operator from a graph.  */
 	struct FBuildGraphParams
 	{
 		/** Reference to graph being built. */
diff --git a/Engine/Plugins/Runtime/Metasound/Source/MetasoundGraphCore/Public/MetasoundDataReference.h b/Engine/Plugins/Runtime/Metasound/Source/MetasoundGraphCore/Public/MetasoundDataReference.h
--- a/Engine/Plugins/Runtime/Metasound/Source/MetasoundGraphCore/Public/MetasoundDataReference.h
+++ b/Engine/Plugins/Runtime/Metasound/Source/MetasoundGraphCore/Public/MetasoundDataReference.h
@@ -748,8 +748,15 @@ namespace Metasound
 		template<typename ReturnType>
 		ReturnType GetAs() const
 		{
-			check(DataRefPtr.IsValid());
-			return TGetAs<ReturnType>::Get(AccessType, *DataRefPtr);
+			if constexpr (std::is_same_v<FAnyDataReference, ReturnType>)
+			{
+				return *this;
+			}
+			else
+			{
+				check(DataRefPtr.IsValid());
+				return TGetAs<ReturnType>::Get(AccessType, *DataRefPtr);
+			}
 		}
 
 		/** Returns the current value of a reference. 
diff --git a/Engine/Plugins/Runtime/Metasound/Source/MetasoundGraphCore/Public/MetasoundFacade.h b/Engine/Plugins/Runtime/Metasound/Source/MetasoundGraphCore/Public/MetasoundFacade.h
--- a/Engine/Plugins/Runtime/Metasound/Source/MetasoundGraphCore/Public/MetasoundFacade.h
+++ b/Engine/Plugins/Runtime/Metasound/Source/MetasoundGraphCore/Public/MetasoundFacade.h
@@ -21,12 +21,6 @@ namespace Metasound
 				template<typename T, T> 
 				struct Helper;
 
-				PRAGMA_DISABLE_DEPRECATION_WARNINGS // Temporary for deprecation of FCreateOperatorParams
-				// Check for "static TUniquePtr<IOperator> U::CreateOperator(const FCreateOperatorParams& Inparams, FBuildErrorArray& OutErrors)"
-				template<typename T>
-				static uint8 Check(Helper<TUniquePtr<IOperator>(*)(const FCreateOperatorParams& InParams, FBuildErrorArray& OutErrors), &T::CreateOperator>*);
-				PRAGMA_ENABLE_DEPRECATION_WARNINGS
-
 				// Check for "static TUniquePtr<IOperator> U::CreateOperator(const FBuildOperatorParams& Inparams, FBuildErrorArray& OutErrors)"
 				template<typename T>
 				static uint8 Check(Helper<TUniquePtr<IOperator>(*)(const FBuildOperatorParams& InParams, FBuildResults& OutResults), &T::CreateOperator>*);
@@ -96,19 +90,12 @@ namespace Metasound
 			{
 				using FCreateOperatorFunction = TFunction<TUniquePtr<IOperator>(const FBuildOperatorParams& InParams, FBuildResults& OutResults)>;
 
-				PRAGMA_DISABLE_DEPRECATION_WARNINGS // Temporary for deprecation of FCreateOperatorParams
-				using FOriginalCreateOperatorFunction = TFunction<TUniquePtr<IOperator>(const FCreateOperatorParams& InParams, FBuildErrorArray& OutErrors)>;
-				PRAGMA_ENABLE_DEPRECATION_WARNINGS
-
 				public:
 					FFactory(FCreateOperatorFunction InCreateFunc);
-					FFactory(FOriginalCreateOperatorFunction InCreateFunc);
 
 					virtual TUniquePtr<IOperator> CreateOperator(const FBuildOperatorParams& InParams, FBuildResults& OutResults) override;
 
 				private:
-					// Wrap the prior version of CreateOperator with the new version of CreateOperator
-					static FCreateOperatorFunction WrapOriginalCreateOperatorFunction(FOriginalCreateOperatorFunction InCreateFunc);
 
 					FCreateOperatorFunction CreateFunc;
 			};
diff --git a/Engine/Plugins/Runtime/Metasound/Source/MetasoundGraphCore/Public/MetasoundVariable.h b/Engine/Plugins/Runtime/Metasound/Source/MetasoundGraphCore/Public/MetasoundVariable.h
--- a/Engine/Plugins/Runtime/Metasound/Source/MetasoundGraphCore/Public/MetasoundVariable.h
+++ b/Engine/Plugins/Runtime/Metasound/Source/MetasoundGraphCore/Public/MetasoundVariable.h
@@ -6,6 +6,7 @@
 #include "MetasoundDataReference.h"
 #include "MetasoundLiteral.h"
 #include "MetasoundOperatorSettings.h"
+#include "MetasoundVertexData.h"
 
 namespace Metasound
 {
@@ -73,6 +74,18 @@ namespace Metasound
 			return *DataReference;
 		}
 
+		void BindDataReference(FOutputVertexInterfaceData& InVertexData, const FVertexName& InVertexName)
+		{
+			checkf(DataReference.IsSet(), TEXT("InitDataReference must be called before accessing"));
+			InVertexData.BindReadVertex(InVertexName, *DataReference);
+		}
+
+		void BindDelayedDataReference(FOutputVertexInterfaceData& InVertexData, const FVertexName& InVertexName)
+		{
+			checkf(DelayedDataReference.IsSet(), TEXT("InitDelayedDataReference must be called before accessing"));
+			InVertexData.BindReadVertex(InVertexName, *DelayedDataReference);
+		}
+
 		void InitDelayedDataReference(const FOperatorSettings& InOperatorSettings)
 		{
 			if (!DelayedDataReference.IsSet())
diff --git a/Engine/Plugins/Runtime/Metasound/Source/MetasoundGraphCore/Public/MetasoundVariableNodes.h b/Engine/Plugins/Runtime/Metasound/Source/MetasoundGraphCore/Public/MetasoundVariableNodes.h
--- a/Engine/Plugins/Runtime/Metasound/Source/MetasoundGraphCore/Public/MetasoundVariableNodes.h
+++ b/Engine/Plugins/Runtime/Metasound/Source/MetasoundGraphCore/Public/MetasoundVariableNodes.h
@@ -111,24 +111,6 @@ namespace Metasound
 
 			virtual ~TVariableOperator() = default;
 
-			virtual FDataReferenceCollection GetInputs() const override
-			{
-				checkNoEntry();
-				return FDataReferenceCollection{};
-			}
-
-			virtual FDataReferenceCollection GetOutputs() const override
-			{
-				checkNoEntry();
-				return FDataReferenceCollection{};
-			}
-
-			virtual void Bind(FVertexInterfaceData& InVertexData) const override
-			{
-				const_cast<TVariableOperator<DataType>*>(this)->BindInputs(InVertexData.GetInputs());
-				const_cast<TVariableOperator<DataType>*>(this)->BindOutputs(InVertexData.GetOutputs());
-			}
-
 			virtual void BindInputs(FInputVertexInterfaceData& InVertexData) override
 			{
 			}
@@ -222,24 +204,6 @@ namespace Metasound
 
 			virtual ~TVariableMutatorOperator() = default;
 
-			virtual FDataReferenceCollection GetInputs() const override
-			{
-				checkNoEntry();
-				return FDataReferenceCollection{};
-			}
-
-			virtual FDataReferenceCollection GetOutputs() const override
-			{
-				checkNoEntry();
-				return FDataReferenceCollection{};
-			}
-
-			virtual void Bind(FVertexInterfaceData& InVertexData) const override
-			{
-				const_cast<TVariableMutatorOperator<DataType>*>(this)->BindInputs(InVertexData.GetInputs());
-				const_cast<TVariableMutatorOperator<DataType>*>(this)->BindOutputs(InVertexData.GetOutputs());
-			}
-			
 			virtual void BindInputs(FInputVertexInterfaceData& InVertexData) override
 			{
 				using namespace VariableNames; 
@@ -395,24 +359,6 @@ namespace Metasound
 
 			virtual ~TVariableDeferredAccessorOperator() = default;
 
-			virtual FDataReferenceCollection GetInputs() const override
-			{
-				checkNoEntry();
-				return FDataReferenceCollection{};
-			}
-
-			virtual FDataReferenceCollection GetOutputs() const override
-			{
-				checkNoEntry();
-				return FDataReferenceCollection{};
-			}
-
-			virtual void Bind(FVertexInterfaceData& InVertexData) const override
-			{
-				const_cast<TVariableDeferredAccessorOperator<DataType>*>(this)->BindInputs(InVertexData.GetInputs());
-				const_cast<TVariableDeferredAccessorOperator<DataType>*>(this)->BindOutputs(InVertexData.GetOutputs());
-			}
-
 			virtual void BindInputs(FInputVertexInterfaceData& InVertexData) override
 			{
 				using namespace VariableNames; 
@@ -427,7 +373,7 @@ namespace Metasound
 				using namespace VariableNames; 
 
 				InVertexData.BindWriteVertex<FVariable>(METASOUND_GET_PARAM_NAME(OutputVariable), Variable);
-				InVertexData.BindReadVertex<DataType>(METASOUND_GET_PARAM_NAME(OutputData), Variable->GetDelayedDataReference());
+				Variable->BindDelayedDataReference(InVertexData, METASOUND_GET_PARAM_NAME(OutputData));
 			}
 
 		private:
@@ -511,24 +457,6 @@ namespace Metasound
 
 			virtual ~TVariableAccessorOperator() = default;
 
-			virtual FDataReferenceCollection GetInputs() const override
-			{
-				checkNoEntry();
-				return FDataReferenceCollection{};
-			}
-
-			virtual FDataReferenceCollection GetOutputs() const override
-			{
-				checkNoEntry();
-				return FDataReferenceCollection{};
-			}
-
-			virtual void Bind(FVertexInterfaceData& InVertexData) const override
-			{
-				const_cast<TVariableAccessorOperator<DataType>*>(this)->BindInputs(InVertexData.GetInputs());
-				const_cast<TVariableAccessorOperator<DataType>*>(this)->BindOutputs(InVertexData.GetOutputs());
-			}
-
 			virtual void BindInputs(FInputVertexInterfaceData& InVertexData) override
 			{
 				using namespace VariableNames; 
@@ -542,8 +470,8 @@ namespace Metasound
 			{
 				using namespace VariableNames; 
 
-				InVertexData.BindReadVertex<DataType>(METASOUND_GET_PARAM_NAME(OutputData), Variable->GetDataReference());
 				InVertexData.BindWriteVertex<FVariable>(METASOUND_GET_PARAM_NAME(OutputVariable), Variable);
+				Variable->BindDataReference(InVertexData, METASOUND_GET_PARAM_NAME(OutputData));
 			}
 
 		private:
diff --git a/Engine/Plugins/Runtime/Metasound/Source/MetasoundGraphCore/Public/MetasoundVertexData.h b/Engine/Plugins/Runtime/Metasound/Source/MetasoundGraphCore/Public/MetasoundVertexData.h
--- a/Engine/Plugins/Runtime/Metasound/Source/MetasoundGraphCore/Public/MetasoundVertexData.h
+++ b/Engine/Plugins/Runtime/Metasound/Source/MetasoundGraphCore/Public/MetasoundVertexData.h
@@ -135,6 +135,20 @@ namespace Metasound
 				return TDataWriteReference<DataType>::CreateNew(Forward<ArgTypes>(Args)...);
 			}
 		};
+
+		template<typename T>
+		struct TTriggerDeprecatedBindAPIWarning
+		{
+		};
+
+		template<typename T>
+		struct TTriggerDeprecatedBindAPIWarning<const T>
+		{
+			UE_DEPRECATED(5.6, "Please use a non-const lvalue reference for data references passed to this function. Const data references can sometimes be created by passing in temporary objects and/or triggering implicit conversion.")
+			FORCEINLINE void Trigger() 
+			{
+			}
+		};
 		
 		// An input binding which connects an FInputVertex to a IDataReference
 		class METASOUNDGRAPHCORE_API FInputBinding
@@ -171,7 +185,36 @@ namespace Metasound
 				}
 			}
 
-			void Bind(FAnyDataReference& InOutDataReference);
+			template<typename DataReferenceType>
+			void Bind(DataReferenceType& InOutDataReference)
+			{
+				if constexpr (std::is_const_v<DataReferenceType>)
+				{
+					// Bind operations need to be able to modify the provided
+					// data reference in place. A prior API had performed binding
+					// on const references without modifications. That API has 
+					// since been deprecated. 
+					//
+					// If you trigger this deprecation warning, it is likely 
+					// that the code is using the deprecated API, or that it is
+					// incidentally calling this API because it passes a temporary 
+					// object to this function which then is bound to a const reference.
+					TTriggerDeprecatedBindAPIWarning<DataReferenceType>::Trigger();
+					Bind(const_cast<std::remove_const_t<DataReferenceType>&>(InOutDataReference));
+				}
+				else
+				{
+					if (Data.IsSet())
+					{
+						InOutDataReference = Data->GetAs<DataReferenceType>();
+					}
+					else
+					{
+						Set(InOutDataReference);
+					}
+				}
+			}
+
 			void Bind(FInputBinding& InBinding);
 
 			// Set the data reference, overwriting any existing bound data references
@@ -280,7 +323,29 @@ namespace Metasound
 				Set(InOutDataReference);
 			}
 
-			void Bind(FAnyDataReference& InOutDataReference);
+			template<typename DataReferenceType>
+			void Bind(DataReferenceType& InOutDataReference)
+			{
+				if constexpr (std::is_const_v<DataReferenceType>)
+				{
+					// Bind operations need to be able to modify the provided
+					// data reference in place. A prior API had performed binding
+					// on const references without modifications. That API has 
+					// since been deprecated. 
+					//
+					// If you trigger this deprecation warning, it is likely 
+					// that the code is using the deprecated API, or that it is
+					// incidentally calling this API because it passes a temporary 
+					// object to this function which then is bound to a const reference.
+					TTriggerDeprecatedBindAPIWarning<DataReferenceType>::Trigger();
+					Bind(const_cast<std::remove_const_t<DataReferenceType>&>(InOutDataReference));
+				}
+				else
+				{
+					Set(InOutDataReference);
+				}
+			}
+
 			void Bind(FOutputBinding& InBinding);
 
 			// Set the data reference, overwriting any existing bound data references
@@ -498,8 +563,9 @@ namespace Metasound
 			Apply(InVertexName, CreateBinding, BindData);
 		}
 
-		/** Bind a read vertex from a read reference. Slated for deprecation. */
+		/** Bind a read vertex from a read reference.*/
 		template<typename DataType>
+		UE_DEPRECATED(5.6, "Please use a non-const lvalue reference for data references passed to this function. Const data references can sometimes be created by passing in temporary objects and/or triggering implicit conversion.")
 		void BindReadVertex(const FVertexName& InVertexName, const TDataReadReference<DataType>& InOutDataReference)
 		{
 			BindReadVertex(InVertexName, const_cast<TDataReadReference<DataType>&>(InOutDataReference));
@@ -525,8 +591,9 @@ namespace Metasound
 			Apply(InVertexName, TCreateInputBinding(InVertexName, InOutDataReference), BindData);
 		}
 
-		/** Bind a write vertex from a write reference. Slated for deprecation. */
+		/** Bind a write vertex from a write reference.*/
 		template<typename DataType>
+		UE_DEPRECATED(5.6, "Please use a non-const lvalue reference for data references passed to this function. Const data references can sometimes be created by passing in temporary objects and/or triggering implicit conversion.")
 		void BindWriteVertex(const FVertexName& InVertexName, const TDataWriteReference<DataType>& InOutDataReference)
 		{
 			BindWriteVertex(InVertexName, const_cast<TDataWriteReference<DataType>&>(InOutDataReference));
@@ -541,8 +608,19 @@ namespace Metasound
 			Apply(InVertexName, TCreateInputBinding(InVertexName, InOutDataReference), BindData);
 		}
 
-		/** Bind a vertex with a any data reference. Slated for deprecation. */
+		/** Bind a vertex with a any data reference. */
+		UE_DEPRECATED(5.6, "Please use a non-const lvalue reference for data references passed to this function. Const data references can sometimes be created by passing in temporary objects and/or triggering implicit conversion.")
 		void BindVertex(const FVertexName& InVertexName, const FAnyDataReference& InOutDataReference);
+		
+		/** Bind a vertex with a a data reference. */
+		template<typename DataReferenceType>
+		void BindVertex(const FVertexName& InVertexName, DataReferenceType& InOutDataReference)
+		{
+			using namespace VertexDataPrivate;
+			auto BindData = [&](FInputBinding& Binding) { Binding.Bind(InOutDataReference); };
+			Apply(InVertexName, TCreateInputBinding(InVertexName, InOutDataReference), BindData);
+		}
+		
 		/** Bind a vertex with a any data reference. */
 		void BindVertex(const FVertexName& InVertexName, FAnyDataReference& InOutDataReference);
 
@@ -609,8 +687,10 @@ namespace Metasound
 			ForEachSubInterfaceBinding(InSubInterfaceName, InVertexName, InValues.Num(), BindData);
 		}
 
-		/** Bind vertex data using other vertex data. Slated for deprecation. */
+		/** Bind vertex data using other vertex data.*/
+		UE_DEPRECATED(5.6, "Please use a non-const lvalue reference for data references passed to this function. Const data references can sometimes be created by passing in temporary objects and/or triggering implicit conversion.")
 		void Bind(const FInputVertexInterfaceData& InOutInputVertexData);
+
 		/** Bind vertex data using other vertex data. */
 		void Bind(FInputVertexInterfaceData& InOutInputVertexData);
 
@@ -620,9 +700,11 @@ namespace Metasound
 		void SetVertex(const FVertexName& InVertexName, FAnyDataReference&& InDataReference);
 
 		/** Set a data references to vertices with matching vertex names. Ignores any existing data bound to the vertex. */
+		UE_DEPRECATED(5.6, "Do not use FDataReferenceCollection.")
 		void Set(const FDataReferenceCollection& InCollection);
 
 		/** Convert vertex data to a data reference collection. */
+		UE_DEPRECATED(5.6, "Do not use FDataReferenceCollection.")
 		FDataReferenceCollection ToDataReferenceCollection() const;
 
 		/** Return the vertex associated with the vertex name. */
@@ -900,11 +982,19 @@ namespace Metasound
 		template<typename DataType>
 		void SetValue(const FVertexName& InVertexName, const DataType& InValue)
 		{
-			BindValueVertex<DataType>(InVertexName, TDataValueReference<DataType>::CreateNew(InValue));
+			auto CreateBinding = [&]()
+			{
+				FOutputDataVertex Vertex(InVertexName, GetMetasoundDataTypeName<DataType>(), FDataVertexMetadata{}, EVertexAccessType::Value);
+				return FOutputBinding(Vertex);
+			};
+
+			auto BindData = [&](FOutputBinding& Binding) { Binding.SetValue<DataType>(InValue); };
+			Apply(InVertexName, CreateBinding, BindData);
 		}
 
 		/** Bind a value vertex from a value reference. */
 		template<typename DataType>
+		UE_DEPRECATED(5.6, "Please use a non-const lvalue reference for data references passed to this function. Const data references can sometimes be created by passing in temporary objects and/or triggering implicit conversion.")
 		void BindValueVertex(const FVertexName& InVertexName, const TDataValueReference<DataType>& InOutDataReference)
 		{
 			using namespace VertexDataPrivate;
@@ -914,13 +1004,15 @@ namespace Metasound
 
 		/** Bind a read vertex from a value reference. */
 		template<typename DataType>
+		UE_DEPRECATED(5.6, "Please use a non-const lvalue reference for data references passed to this function. Const data references can sometimes be created by passing in temporary objects and/or triggering implicit conversion.")
 		void BindReadVertex(const FVertexName& InVertexName, const TDataValueReference<DataType>& InOutDataReference)
 		{
 			BindReadVertex(InVertexName, static_cast<TDataReadReference<DataType>>(InOutDataReference));
 		}
 
-		/** Bind a read vertex from a read reference. Slated for deprecation. */
+		/** Bind a read vertex from a read reference.*/
 		template<typename DataType>
+		UE_DEPRECATED(5.6, "Please use a non-const lvalue reference for data references passed to this function. Const data references can sometimes be created by passing in temporary objects and/or triggering implicit conversion.")
 		void BindReadVertex(const FVertexName& InVertexName, const TDataReadReference<DataType>& InOutDataReference)
 		{
 			BindReadVertex(InVertexName, const_cast<TDataReadReference<DataType>&>(InOutDataReference));
@@ -935,8 +1027,9 @@ namespace Metasound
 			Apply(InVertexName, TCreateOutputBinding(InVertexName, InOutDataReference), BindData);
 		}
 
-		/** Bind a read vertex from a write reference. Slated for deprecation. */
+		/** Bind a read vertex from a write reference.*/
 		template<typename DataType>
+		UE_DEPRECATED(5.6, "Please use a non-const lvalue reference for data references passed to this function. Const data references can sometimes be created by passing in temporary objects and/or triggering implicit conversion.")
 		void BindReadVertex(const FVertexName& InVertexName, const TDataWriteReference<DataType>& InOutDataReference)
 		{
 			BindReadVertex(InVertexName, const_cast<TDataWriteReference<DataType>&>(InOutDataReference));
@@ -951,8 +1044,9 @@ namespace Metasound
 			Apply(InVertexName, TCreateOutputBinding(InVertexName, InOutDataReference), BindData);
 		}
 
-		/** Bind a write vertex from a write reference. Slated for deprecation. */
+		/** Bind a write vertex from a write reference.*/
 		template<typename DataType>
+		UE_DEPRECATED(5.6, "Please use a non-const lvalue reference for data references passed to this function. Const data references can sometimes be created by passing in temporary objects and/or triggering implicit conversion.")
 		void BindWriteVertex(const FVertexName& InVertexName, const TDataWriteReference<DataType>& InOutDataReference)
 		{
 			BindWriteVertex(InVertexName, const_cast<TDataWriteReference<DataType>&>(InOutDataReference));
@@ -967,9 +1061,20 @@ namespace Metasound
 			Apply(InVertexName, TCreateOutputBinding(InVertexName, InOutDataReference), BindData);
 		}
 
-		/** Bind a vertex with a any data reference. Slated for deprecation. */
-		void BindVertex(const FVertexName& InVertexName, const FAnyDataReference& InOutDataReference);
 		/** Bind a vertex with a any data reference. */
+		UE_DEPRECATED(5.6, "Please use a non-const lvalue reference for data references passed to this function. Const data references can sometimes be created by passing in temporary objects and/or triggering implicit conversion.")
+		void BindVertex(const FVertexName& InVertexName, const FAnyDataReference& InOutDataReference);
+		
+		/** Bind a vertex with a data reference. */
+		template<typename DataReferenceType>
+		void BindVertex(const FVertexName& InVertexName, DataReferenceType& InOutDataReference)
+		{
+			using namespace VertexDataPrivate;
+			auto BindData = [&](FOutputBinding& Binding) { Binding.Bind(InOutDataReference); };
+			Apply(InVertexName, TCreateOutputBinding(InVertexName, InOutDataReference), BindData);
+		}
+
+		/** Bind a vertex with any data reference. */		
 		void BindVertex(const FVertexName& InVertexName, FAnyDataReference& InOutDataReference);
 
 		/** Return the number of instances of a particular sub interface. */
@@ -1043,7 +1148,8 @@ namespace Metasound
 			ForEachSubInterfaceBinding(InSubInterfaceName, InVertexName, InValues.Num(), BindData);
 		}
 
-		/** Bind vertex data using other vertex data. Slated for deprecation. */
+		/** Bind vertex data using other vertex data.*/
+		UE_DEPRECATED(5.6, "Please use a non-const lvalue reference for data references passed to this function. Const data references can sometimes be created by passing in temporary objects and/or triggering implicit conversion.")
 		void Bind(const FOutputVertexInterfaceData& InOutOutputVertexData);
 		/** Bind vertex data using other vertex data. */
 		void Bind(FOutputVertexInterfaceData& InOutOutputVertexData);
@@ -1053,10 +1159,12 @@ namespace Metasound
 		/** Set a data reference to a vertex, ignoring any existing data bound to the vertex*/
 		void SetVertex(const FVertexName& InVertexName, FAnyDataReference&& InDataReference);
 
-		/** Bind a data references to vertices with matching vertex names. */
+		/** Set a data references to vertices with matching vertex names. */
+		UE_DEPRECATED(5.6, "Do not use FDataReferenceCollection.")
 		void Set(const FDataReferenceCollection& InCollection);
 
 		/** Converts the vertex data to a data reference collection. */
+		UE_DEPRECATED(5.6, "Do not use FDataReferenceCollection.")
 		FDataReferenceCollection ToDataReferenceCollection() const;
 
 		/** Returns true if a vertex with the given vertex name exists and is bound
@@ -1216,7 +1324,10 @@ namespace Metasound
 		FVertexInterfaceData(const FVertexInterface& InVertexInterface);
 
 		/** Set vertex data using other vertex data. */
+		UE_DEPRECATED(5.6, "Please use a non-const lvalue reference for data  passed to this function. Const data  can sometimes be created by passing in temporary objects and/or triggering implicit conversion.")
 		void Bind(const FVertexInterfaceData& InVertexData);
+		
+		/** Set vertex data using other vertex data. */
 		void Bind(FVertexInterfaceData& InVertexData);
 
 		/** Get input vertex interface data. */
diff --git a/Engine/Plugins/Runtime/Metasound/Source/MetasoundStandardNodes/Private/MetasoundMathNodes.cpp b/Engine/Plugins/Runtime/Metasound/Source/MetasoundStandardNodes/Private/MetasoundMathNodes.cpp
--- a/Engine/Plugins/Runtime/Metasound/Source/MetasoundStandardNodes/Private/MetasoundMathNodes.cpp
+++ b/Engine/Plugins/Runtime/Metasound/Source/MetasoundStandardNodes/Private/MetasoundMathNodes.cpp
@@ -290,23 +290,7 @@ namespace Metasound
 			{
 				using namespace MathOpNames;
 
-				InOutVertexData.BindReadVertex(METASOUND_GET_PARAM_NAME(OutMath), TDataClassReadRef(ValueRef));
-			}
-
-			virtual FDataReferenceCollection GetInputs() const override
-			{
-				// This should never be called. Bind(...) is called instead. This method
-				// exists as a stop-gap until the API can be deprecated and removed.
-				checkNoEntry();
-				return {};
-			}
-
-			virtual FDataReferenceCollection GetOutputs() const override
-			{
-				// This should never be called. Bind(...) is called instead. This method
-				// exists as a stop-gap until the API can be deprecated and removed.
-				checkNoEntry();
-				return {};
+				InOutVertexData.BindReadVertex(METASOUND_GET_PARAM_NAME(OutMath), ValueRef);
 			}
 		};
 
diff --git a/Engine/Plugins/Runtime/Metasound/Source/MetasoundStandardNodes/Private/MetasoundTriggerOnThresholdNode.cpp b/Engine/Plugins/Runtime/Metasound/Source/MetasoundStandardNodes/Private/MetasoundTriggerOnThresholdNode.cpp
--- a/Engine/Plugins/Runtime/Metasound/Source/MetasoundStandardNodes/Private/MetasoundTriggerOnThresholdNode.cpp
+++ b/Engine/Plugins/Runtime/Metasound/Source/MetasoundStandardNodes/Private/MetasoundTriggerOnThresholdNode.cpp
@@ -251,25 +251,8 @@ namespace Metasound
 		virtual void BindOutputs(FOutputVertexInterfaceData& InOutVertexData) override
 		{
 			using namespace TriggerOnThresholdVertexNames;
-			InOutVertexData.BindReadVertex(METASOUND_GET_PARAM_NAME(OutPin), FTriggerWriteRef(Out));
+			InOutVertexData.BindReadVertex(METASOUND_GET_PARAM_NAME(OutPin), Out);
 		}
-
-		virtual FDataReferenceCollection GetInputs() const override
-		{
-			// This should never be called. Bind(...) is called instead. This method
-			// exists as a stop-gap until the API can be deprecated and removed.
-			checkNoEntry();
-			return {};
-		}
-
-		virtual FDataReferenceCollection GetOutputs() const override
-		{
-			// This should never be called. Bind(...) is called instead. This method
-			// exists as a stop-gap until the API can be deprecated and removed.
-			checkNoEntry();
-			return {};
-		}
-
 		static const FVertexInterface DeclareVertexInterface()
 		{
 			return TTriggerOnThresholdHelper<ValueType, ThresholdType>::DeclareVertexInterface(DefaultThreshold);
