MetaSound Frontend Graph Builder cleanup in prep of further work in this part of the API

diff --git a/Engine/Plugins/Runtime/Metasound/Source/MetasoundEngineTest/Private/NodeTestGraphBuilder.cpp b/Engine/Plugins/Runtime/Metasound/Source/MetasoundEngineTest/Private/NodeTestGraphBuilder.cpp
--- a/Engine/Plugins/Runtime/Metasound/Source/MetasoundEngineTest/Private/NodeTestGraphBuilder.cpp
+++ b/Engine/Plugins/Runtime/Metasound/Source/MetasoundEngineTest/Private/NodeTestGraphBuilder.cpp
@@ -90,7 +90,7 @@ namespace Metasound::Test
 		
 		const TSet<FName> TransmittableInputNames;
 		const FString UnknownAsset = TEXT("UnknownAsset");
-		if (const TUniquePtr<FFrontendGraph> Graph = FFrontendGraphBuilder::CreateGraph(Document, UnknownAsset))
+		if (const TUniquePtr<FFrontendGraph> Graph = Frontend::FGraphBuilder::CreateGraph(Document, UnknownAsset))
 		{
 			FOperatorBuilderSettings BuilderSettings;
 			BuilderSettings.bFailOnAnyError = true;
diff --git a/Engine/Plugins/Runtime/Metasound/Source/MetasoundFrontend/Private/MetasoundFrontendGraph.cpp b/Engine/Plugins/Runtime/Metasound/Source/MetasoundFrontend/Private/MetasoundFrontendGraph.cpp
--- a/Engine/Plugins/Runtime/Metasound/Source/MetasoundFrontend/Private/MetasoundFrontendGraph.cpp
+++ b/Engine/Plugins/Runtime/Metasound/Source/MetasoundFrontend/Private/MetasoundFrontendGraph.cpp
@@ -2,56 +2,13 @@
 
 #include "MetasoundFrontendGraph.h"
 
-#include "Algo/AllOf.h"
-#include "Algo/AnyOf.h"
-#include "Algo/TopologicalSort.h"
-#include "Algo/Transform.h"
-#include "CoreMinimal.h"
-#include "CoreGlobals.h"
 #include "MetasoundDocumentInterface.h"
 #include "MetasoundFrontend.h"
-#include "MetasoundFrontendDataTypeRegistry.h"
-#include "MetasoundFrontendNodeTemplateRegistry.h"
-#include "MetasoundFrontendProxyDataCache.h"
-#include "MetasoundFrontendRegistries.h"
 #include "MetasoundGraph.h"
-#include "MetasoundLiteralNode.h"
-#include "MetasoundLog.h"
 #include "MetasoundNodeInterface.h"
 
-
 namespace Metasound
 {
-	namespace Frontend
-	{
-		namespace GraphPrivate
-		{
-			FNodeInitData CreateNodeInitData(const FMetasoundFrontendNode& InNode)
-			{
-				FNodeInitData InitData;
-
-				InitData.InstanceName = InNode.Name;
-				InitData.InstanceID = InNode.GetID();
-
-				return InitData;
-			}
-
-			template <typename ResolveType>
-			FGuid ResolveTargetPageID(const ResolveType& InToResolve)
-			{
-				// Registry is not available in tests, so for now resolution is considered successful at this level
-				// if registry is not initialized and providing a resolved page ID. TODO: Add a test implementation
-				// that returns the default page (or whatever page behavior is desired for testing).
-				if (IDocumentBuilderRegistry* BuilderRegistry = IDocumentBuilderRegistry::Get())
-				{
-					return IDocumentBuilderRegistry::GetChecked().ResolveTargetPageID(InToResolve);
-				}
-
-				return Frontend::DefaultPageID;
-			}
-		} // namespace GraphPrivate
-	} // namespace Frontend
-
 	FFrontendGraph::FFrontendGraph(const FString& InInstanceName, const FGuid& InInstanceID)
 	:	FGraph(InInstanceName, InInstanceID)
 	{
@@ -66,8 +23,7 @@ namespace Metasound
 
 			// Input nodes need an extra Index value to keep track of their position in the graph's inputs.
 			InputNodes.Add(InIndex, InNode.Get());
-			AddInputDataDestination(*InNode, InVertexName);
-			AddNode(InDependencyId, InNode);
+			AddInputNode(InDependencyId, InVertexName, MoveTemp(InNode));
 		}
 	}
 
@@ -80,8 +36,7 @@ namespace Metasound
 
 			// Output nodes need an extra Index value to keep track of their position in the graph's inputs.
 			OutputNodes.Add(InIndex, InNode.Get());
-			AddOutputDataSource(*InNode, InVertexName);
-			AddNode(InNodeID, InNode);
+			AddOutputNode(InNodeID, InVertexName, MoveTemp(InNode));
 		}
 	}
 
@@ -108,792 +63,4 @@ namespace Metasound
 
 		return nullptr;
 	}
-
-	bool FFrontendGraph::OwnsAllReferencedNodes() const
-	{
-		// This function is deprecated in 5.3 and is no longer functional.
-		return true;
-	}
-
-	TUniquePtr<INode> FFrontendGraphBuilder::CreateVariableNode(const FBuildContext& InContext, const FMetasoundFrontendNode& InNode, const FMetasoundFrontendGraph& InGraph)
-	{
-		using namespace Metasound::Frontend;
-
-		const FMetasoundFrontendVariable* FrontendVariable = FindVariableForVariableNode(InNode, InGraph);
-
-		if (nullptr != FrontendVariable)
-		{
-			IDataTypeRegistry& DataTypeRegistry = IDataTypeRegistry::Get();
-			const bool IsLiteralParsableByDataType = DataTypeRegistry.IsLiteralTypeSupported(FrontendVariable->TypeName, FrontendVariable->Literal.GetType());
-
-			if (IsLiteralParsableByDataType)
-			{
-				FLiteral Literal = FrontendVariable->Literal.ToLiteral(FrontendVariable->TypeName, &InContext.DataTypeRegistry, &InContext.ProxyDataCache);
-
-				FVariableNodeConstructorParams InitParams =
-				{
-					InNode.Name,
-					InNode.GetID(),
-					MoveTemp(Literal)
-				};
-
-				return DataTypeRegistry.CreateVariableNode(FrontendVariable->TypeName, MoveTemp(InitParams));
-			}
-			else
-			{
-				UE_LOG(LogMetaSound, Error, TEXT("Cannot create variable node [NodeID:%s]. [Variable:%s] cannot be constructed with the provided literal type."), *InNode.GetID().ToString(), *FrontendVariable->Name.ToString());
-			}
-		}
-		else
-		{
-			UE_LOG(LogMetaSound, Error, TEXT("Cannot create variable node [NodeID:%s]. No variable found for variable node."), *InNode.GetID().ToString());
-		}
-
-		return TUniquePtr<INode>(nullptr);
-	}
-
-	TUniquePtr<INode> FFrontendGraphBuilder::CreateInputNode(const FBuildContext& InContext, const FMetasoundFrontendNode& InNode, const FMetasoundFrontendClass& InClass, const FMetasoundFrontendClassInput& InOwningGraphClassInput)
-	{
-		using namespace Metasound::Frontend;
-
-		const FMetasoundFrontendLiteral* FrontendLiteral = FindInputLiteralForInputNode(InNode, InClass, InOwningGraphClassInput);
-
-		if (nullptr != FrontendLiteral)
-		{
-			if (ensure(InNode.Interface.Inputs.Num() == 1))
-			{
-				IDataTypeRegistry& DataTypeRegistry = IDataTypeRegistry::Get();
-				const FMetasoundFrontendVertex& InputVertex = InNode.Interface.Inputs[0];
-
-				const bool IsLiteralParsableByDataType = DataTypeRegistry.IsLiteralTypeSupported(InputVertex.TypeName, FrontendLiteral->GetType());
-
-				if (IsLiteralParsableByDataType)
-				{
-					FLiteral Literal = FrontendLiteral->ToLiteral(InputVertex.TypeName, &InContext.DataTypeRegistry, &InContext.ProxyDataCache);
-
-					FInputNodeConstructorParams InitParams =
-					{
-						InNode.Name,
-						InNode.GetID(),
-						InputVertex.Name,
-						MoveTemp(Literal)
-					};
-
-					ensureAlwaysMsgf(InOwningGraphClassInput.AccessType != EMetasoundFrontendVertexAccessType::Unset, TEXT("Graph Class Input cannot be set to access type of 'Unset'"));
-					if (InOwningGraphClassInput.AccessType == EMetasoundFrontendVertexAccessType::Reference)
-					{
-						return DataTypeRegistry.CreateInputNode(InputVertex.TypeName, MoveTemp(InitParams));
-					}
-					else // InOwningGraphClassInput.AccessType == EMetasoundFrontendVertexAccessType::Value
-					{
-						return DataTypeRegistry.CreateConstructorInputNode(InputVertex.TypeName, MoveTemp(InitParams));
-					}
-				}
-				else
-				{
-					UE_LOG(LogMetaSound, Error, TEXT("Cannot create input node [NodeID:%s]. [Vertex:%s] cannot be constructed with the provided literal type."), *InNode.GetID().ToString(), *InputVertex.Name.ToString());
-				}
-			}
-		}
-		else
-		{
-			UE_LOG(LogMetaSound, Error, TEXT("Cannot create input node [NodeID:%s]. No default literal set for input node."), *InNode.GetID().ToString());
-		}
-
-		return TUniquePtr<INode>(nullptr);
-	}
-
-	TUniquePtr<INode> FFrontendGraphBuilder::CreateOutputNode(const FBuildContext& InContext, const FMetasoundFrontendNode& InNode, const FMetasoundFrontendClass& InClass, FBuildGraphContext& InGraphContext, const TSet<FNodeIDVertexID>& InEdgeDestinations)
-	{
-		using namespace Metasound::Frontend;
-
-		check(InClass.Metadata.GetType() == EMetasoundFrontendClassType::Output);
-		check(InNode.ClassID == InClass.ID);
-
-		if (ensure(InNode.Interface.Outputs.Num() == 1))
-		{
-			const FMetasoundFrontendVertex& OutputVertex = InNode.Interface.Outputs[0];
-
-			FOutputNodeConstructorParams InitParams =
-			{
-				InNode.Name,
-				InNode.GetID(),
-				OutputVertex.Name
-			};
-
-			{
-				const FNodeInitData InitData = GraphPrivate::CreateNodeInitData(InNode);
-				TArray<FDefaultLiteralData> DefaultLiteralData = GetInputDefaultLiteralData(InContext, InNode, InitData, InEdgeDestinations);
-				for (FDefaultLiteralData& Data : DefaultLiteralData)
-				{
-					InGraphContext.DefaultInputs.Emplace(FNodeIDVertexID { InNode.GetID(), Data.DestinationVertexID }, MoveTemp(Data));
-				}
-			}
-
-			ensure(InClass.Interface.Outputs.Num() == 1);
-			if (InClass.Interface.Outputs[0].AccessType == EMetasoundFrontendVertexAccessType::Reference)
-			{
-				return IDataTypeRegistry::Get().CreateOutputNode(OutputVertex.TypeName, MoveTemp(InitParams));
-			}
-			else // InClass.Interface.Outputs[0].AccessType == EMetasoundFrontendVertexAccessType::Value
-			{
-				return IDataTypeRegistry::Get().CreateConstructorOutputNode(OutputVertex.TypeName, MoveTemp(InitParams));
-			}
-		}
-		return TUniquePtr<INode>(nullptr);
-	}
-
-	TUniquePtr<INode> FFrontendGraphBuilder::CreateExternalNode(const FBuildContext& InContext, const FMetasoundFrontendNode& InNode, const FMetasoundFrontendClass& InClass, FBuildGraphContext& InGraphContext, const TSet<FNodeIDVertexID>& InEdgeDestinations)
-	{
-		using namespace Frontend;
-
-		check(InNode.ClassID == InClass.ID);
-
-		const FNodeInitData InitData = GraphPrivate::CreateNodeInitData(InNode);
-		{
-			TArray<FDefaultLiteralData> DefaultLiteralData = GetInputDefaultLiteralData(InContext, InNode, InitData, InEdgeDestinations);
-			for (FDefaultLiteralData& Data : DefaultLiteralData)
-			{
-				InGraphContext.DefaultInputs.Emplace(FNodeIDVertexID{ InNode.GetID(), Data.DestinationVertexID }, MoveTemp(Data));
-			}
-		}
-
-		// TODO: handle check to see if node interface conforms to class interface here. 
-		// TODO: check to see if external object supports class interface.
-
-		const FNodeRegistryKey Key = FNodeRegistryKey(InClass.Metadata);
-		return FMetasoundFrontendRegistryContainer::Get()->CreateNode(Key, InitData);
-	}
-
-	const FMetasoundFrontendClassInput* FFrontendGraphBuilder::FindClassInputForInputNode(const FMetasoundFrontendGraphClass& InOwningGraph, const FMetasoundFrontendNode& InInputNode, int32& OutClassInputIndex)
-	{
-		OutClassInputIndex = INDEX_NONE;
-
-		// Input nodes should have exactly one input.
-		if (ensure(InInputNode.Interface.Inputs.Num() == 1))
-		{
-			const FName& TypeName = InInputNode.Interface.Inputs[0].TypeName;
-
-			auto IsMatchingInput = [&](const FMetasoundFrontendClassInput& GraphInput)
-			{
-				return (InInputNode.GetID() == GraphInput.NodeID);
-			};
-
-			OutClassInputIndex = InOwningGraph.Interface.Inputs.IndexOfByPredicate(IsMatchingInput);
-			if (INDEX_NONE != OutClassInputIndex)
-			{
-				return &InOwningGraph.Interface.Inputs[OutClassInputIndex];
-			}
-		}
-		return nullptr;
-	}
-
-	const FMetasoundFrontendClassOutput* FFrontendGraphBuilder::FindClassOutputForOutputNode(const FMetasoundFrontendGraphClass& InOwningGraph, const FMetasoundFrontendNode& InOutputNode, int32& OutClassOutputIndex)
-	{
-		OutClassOutputIndex = INDEX_NONE;
-
-		// Output nodes should have exactly one output
-		if (ensure(InOutputNode.Interface.Outputs.Num() == 1))
-		{
-			const FName& TypeName = InOutputNode.Interface.Outputs[0].TypeName;
-
-			auto IsMatchingOutput = [&](const FMetasoundFrontendClassOutput& GraphOutput)
-			{
-				return (InOutputNode.GetID() == GraphOutput.NodeID);
-			};
-
-			OutClassOutputIndex = InOwningGraph.Interface.Outputs.IndexOfByPredicate(IsMatchingOutput);
-			if (INDEX_NONE != OutClassOutputIndex)
-			{
-				return &InOwningGraph.Interface.Outputs[OutClassOutputIndex];
-			}
-		}
-		return nullptr;
-	}
-
-	const FMetasoundFrontendVariable* FFrontendGraphBuilder::FindVariableForVariableNode(const FMetasoundFrontendNode& InVariableNode, const FMetasoundFrontendGraph& InGraph)
-	{
-		const FGuid& DesiredID = InVariableNode.GetID();
-		return InGraph.Variables.FindByPredicate([&](const FMetasoundFrontendVariable& InVar) { return InVar.VariableNodeID == DesiredID; });
-	}
-
-	const FMetasoundFrontendLiteral* FFrontendGraphBuilder::FindInputLiteralForInputNode(const FMetasoundFrontendNode& InInputNode, const FMetasoundFrontendClass& InInputNodeClass, const FMetasoundFrontendClassInput& InOwningGraphClassInput)
-	{
-		using namespace Frontend;
-
-		// Default value priority is:
-		// 1. A value set directly on the node
-		// 2. A default value of the owning graph
-		// 3. A default value on the input node class.
-
-		const FMetasoundFrontendLiteral* Literal = nullptr;
-
-		// Check for default value directly on node.
-		if (ensure(InInputNode.Interface.Inputs.Num() == 1))
-		{
-			const FMetasoundFrontendVertex& InputVertex = InInputNode.Interface.Inputs[0];
-
-			// Find input literal matching VertexID
-			const FMetasoundFrontendVertexLiteral* VertexLiteral = InInputNode.InputLiterals.FindByPredicate(
-				[&](const FMetasoundFrontendVertexLiteral& InVertexLiteral)
-				{
-					return InVertexLiteral.VertexID == InputVertex.VertexID;
-				}
-			);
-
-			if (nullptr != VertexLiteral)
-			{
-				Literal = &VertexLiteral->Value;
-			}
-		}
-
-		// Check for default value on owning graph.
-		if (nullptr == Literal)
-		{
-			// Find Class Default that is not invalid
-			const FGuid PageID = GraphPrivate::ResolveTargetPageID(InOwningGraphClassInput);
-			const FMetasoundFrontendLiteral& DefaultLiteral = InOwningGraphClassInput.FindConstDefaultChecked(PageID);
-			if (DefaultLiteral.IsValid())
-			{
-				Literal = &DefaultLiteral;
-			}
-		}
-
-		// Check for default value on input node class
-		if (nullptr == Literal && ensure(InInputNodeClass.Interface.Inputs.Num() == 1))
-		{
-			const FMetasoundFrontendClassInput& InputNodeClassInput = InInputNodeClass.Interface.Inputs.Last();
-			const FGuid PageID = GraphPrivate::ResolveTargetPageID(InputNodeClassInput);
-			const FMetasoundFrontendLiteral& DefaultLiteral = InputNodeClassInput.FindConstDefaultChecked(PageID);
-			if (DefaultLiteral.IsValid())
-			{
-				Literal = &DefaultLiteral;
-			}
-		}
-
-		return Literal;
-	}
-
-	bool FFrontendGraphBuilder::AddNodesToGraph(FBuildGraphContext& InGraphContext)
-	{
-		TSet<FNodeIDVertexID> GraphEdgeDestinations;
-		const TArray<FMetasoundFrontendEdge>& GraphEdges = InGraphContext.PagedGraph.Edges;
-		Algo::Transform(GraphEdges, GraphEdgeDestinations, [](const FMetasoundFrontendEdge& Edge) 
-		{
-			return FNodeIDVertexID{Edge.ToNodeID, Edge.ToVertexID};
-		});
-
-		for (const FMetasoundFrontendNode& Node : InGraphContext.PagedGraph.Nodes)
-		{
-			const FMetasoundFrontendClass* NodeClass = InGraphContext.BuildContext.FrontendClasses.FindRef(Node.ClassID);
-
-			if (ensure(nullptr != NodeClass))
-			{
-				switch (NodeClass->Metadata.GetType())
-				{
-					case EMetasoundFrontendClassType::Input:
-					{
-						int32 InputIndex = INDEX_NONE;
-						const FMetasoundFrontendClassInput* ClassInput = FindClassInputForInputNode(InGraphContext.GraphClass, Node, InputIndex);
-
-						if ((nullptr != ClassInput) && (INDEX_NONE != InputIndex))
-						{
-							TSharedPtr<const INode> InputNode(CreateInputNode(InGraphContext.BuildContext, Node, *NodeClass, *ClassInput).Release());
-							InGraphContext.Graph->AddInputNode(Node.GetID(), InputIndex, ClassInput->Name, InputNode);
-						}
-						else
-						{
-							const FString GraphClassIDString = InGraphContext.GraphClass.ID.ToString();
-							UE_LOG(LogMetaSound, Error, TEXT("MetaSound '%s': Failed to match input node [NodeID:%s, NodeName:%s] to owning graph [ClassID:%s] input."), *InGraphContext.BuildContext.DebugAssetName, *Node.GetID().ToString(), *Node.Name.ToString(), *GraphClassIDString);
-							return false;
-						}
-					}
-					break;
-
-					case EMetasoundFrontendClassType::Output:
-					{
-						int32 OutputIndex = INDEX_NONE;
-						const FMetasoundFrontendClassOutput* ClassOutput = FindClassOutputForOutputNode(InGraphContext.GraphClass, Node, OutputIndex);
-						if ((nullptr != ClassOutput) && (INDEX_NONE != OutputIndex))
-						{
-							TSharedPtr<const INode> OutputNode(CreateOutputNode(InGraphContext.BuildContext, Node, *NodeClass, InGraphContext, GraphEdgeDestinations).Release());
-							InGraphContext.Graph->AddOutputNode(Node.GetID(), OutputIndex, ClassOutput->Name, OutputNode);
-						}
-						else
-						{
-							const FString GraphClassIDString = InGraphContext.GraphClass.ID.ToString();
-							UE_LOG(LogMetaSound, Error, TEXT("MetaSound '%s': Failed to match output node [NodeID:%s, NodeName:%s] to owning graph [ClassID:%s] output."), *InGraphContext.BuildContext.DebugAssetName, *Node.GetID().ToString(), *Node.Name.ToString(), *GraphClassIDString);
-							return false;
-						}
-					}
-					break;
-
-					case EMetasoundFrontendClassType::Graph:
-					{
-						const TSharedPtr<const INode> SubgraphPtr = InGraphContext.BuildContext.Graphs.FindRef(Node.ClassID);
-
-						if (SubgraphPtr.IsValid())
-						{
-							InGraphContext.Graph->AddNode(Node.GetID(), SubgraphPtr);
-						}
-						else
-						{
-							UE_LOG(LogMetaSound, Error, TEXT("MetaSound '%s': Failed to find subgraph for node [NodeID:%s, NodeName:%s, ClassID:%s]"), *InGraphContext.BuildContext.DebugAssetName, *Node.GetID().ToString(), *Node.Name.ToString(), *Node.ClassID.ToString());
-							return false;
-						}
-					}
-					break;
-
-					case EMetasoundFrontendClassType::Literal:
-					{
-						checkNoEntry(); // Unsupported.
-						return false;
-					}
-
-					case EMetasoundFrontendClassType::Variable:
-					{
-						TSharedPtr<const INode> VariableNode(CreateVariableNode(InGraphContext.BuildContext, Node, InGraphContext.PagedGraph).Release());
-						InGraphContext.Graph->AddNode(Node.GetID(), VariableNode);
-					}
-					break;
-
-					// Templates, variable accessors and mutators are
-					// constructed with the same parameters as external nodes.
-					case EMetasoundFrontendClassType::Template:
-					case EMetasoundFrontendClassType::VariableAccessor:
-					case EMetasoundFrontendClassType::VariableDeferredAccessor:
-					case EMetasoundFrontendClassType::VariableMutator:
-					case EMetasoundFrontendClassType::External:
-					default:
-					{
-						TSharedPtr<const INode> ExternalNode(CreateExternalNode(InGraphContext.BuildContext, Node, *NodeClass, InGraphContext, GraphEdgeDestinations).Release());
-						InGraphContext.Graph->AddNode(Node.GetID(), ExternalNode);
-					}
-					break;
-				}
-			}
-		}
-
-		return true;
-	}
-
-	bool FFrontendGraphBuilder::AddEdgesToGraph(FBuildGraphContext& InGraphContext)
-	{
-		// Pair of frontend node and core node. The frontend node can be one of
-		// several types.
-		struct FCoreNodeAndFrontendVertex
-		{
-			const INode* Node = nullptr;
-			const FMetasoundFrontendVertex* Vertex = nullptr;
-		};
-
-		TMap<FNodeIDVertexID, FCoreNodeAndFrontendVertex> NodeSourcesByID;
-		TMap<FNodeIDVertexID, FCoreNodeAndFrontendVertex> NodeDestinationsByID;
-
-		// Add nodes to NodeID/VertexID map
-		for (const FMetasoundFrontendNode& Node : InGraphContext.PagedGraph.Nodes)
-		{
-			const INode* CoreNode = InGraphContext.Graph->FindNode(Node.GetID());
-			if (nullptr == CoreNode)
-			{
-				UE_LOG(LogMetaSound, Warning, TEXT("Metasound '%s': Could not find referenced node [Name:%s, NodeID:%s]"), *InGraphContext.BuildContext.DebugAssetName, *Node.Name.ToString(), *Node.GetID().ToString());
-				return false;
-			}
-
-			for (const FMetasoundFrontendVertex& Vertex : Node.Interface.Inputs)
-			{
-				NodeDestinationsByID.Add(FNodeIDVertexID(Node.GetID(), Vertex.VertexID), FCoreNodeAndFrontendVertex({CoreNode, &Vertex}));
-			}
-
-			for (const FMetasoundFrontendVertex& Vertex : Node.Interface.Outputs)
-			{
-				NodeSourcesByID.Add(FNodeIDVertexID(Node.GetID(), Vertex.VertexID), FCoreNodeAndFrontendVertex({CoreNode, &Vertex}));
-			}
-		};
-
-		for (const FMetasoundFrontendEdge& Edge : InGraphContext.PagedGraph.Edges)
-		{
-			const FNodeIDVertexID DestinationKey(Edge.ToNodeID, Edge.ToVertexID);
-			const FCoreNodeAndFrontendVertex* DestinationNodeAndVertex = NodeDestinationsByID.Find(DestinationKey);
-
-			if (nullptr == DestinationNodeAndVertex)
-			{
-				UE_LOG(LogMetaSound, Error, TEXT("MetaSound '%s': Failed to add edge. Could not find destination [NodeID:%s, VertexID:%s]"), *InGraphContext.BuildContext.DebugAssetName, *Edge.ToNodeID.ToString(), *Edge.ToVertexID.ToString());
-				return false;
-			}
-
-			if (nullptr == DestinationNodeAndVertex->Node)
-			{
-				UE_LOG(LogMetaSound, Warning, TEXT("MetaSound '%s': 'Failed to add edge. Null destination node [NodeID:%s]"), *InGraphContext.BuildContext.DebugAssetName, *Edge.ToNodeID.ToString());
-				return false;
-			}
-
-			const FNodeIDVertexID SourceKey(Edge.FromNodeID, Edge.FromVertexID);
-			const FCoreNodeAndFrontendVertex* SourceNodeAndVertex = NodeSourcesByID.Find(SourceKey);
-
-			if (nullptr == SourceNodeAndVertex)
-			{
-				UE_LOG(LogMetaSound, Error, TEXT("MetaSound '%s': Failed to add edge. Could not find source [NodeID:%s, VertexID:%s]"), *InGraphContext.BuildContext.DebugAssetName, *Edge.FromNodeID.ToString(), *Edge.FromVertexID.ToString());
-				return false;
-			}
-
-			if (nullptr == SourceNodeAndVertex->Node)
-			{
-				UE_LOG(LogMetaSound, Warning, TEXT("MetaSound '%s': Skipping edge. Null source node [NodeID:%s]"), *InGraphContext.BuildContext.DebugAssetName, *Edge.FromNodeID.ToString());
-				return false;
-			}
-
-			const INode* FromNode = SourceNodeAndVertex->Node;
-			const FVertexName FromVertexKey = SourceNodeAndVertex->Vertex->Name;
-
-			const INode* ToNode = DestinationNodeAndVertex->Node;
-			const FVertexName ToVertexKey = DestinationNodeAndVertex->Vertex->Name;
-
-			bool bSuccess = InGraphContext.Graph->AddDataEdge(*FromNode, FromVertexKey,  *ToNode, ToVertexKey);
-
-			// If succeeded, remove input as viable vertex to construct default variable, as it has been superceded by a connection.
-			if (bSuccess)
-			{
-				FNodeIDVertexID DestinationPair { ToNode->GetInstanceID(), DestinationNodeAndVertex->Vertex->VertexID };
-				InGraphContext.DefaultInputs.Remove(DestinationPair);
-			}
-			else
-			{
-				UE_LOG(LogMetaSound, Error, TEXT("MetaSound '%s': Failed to connect edge from [NodeID:%s, VertexID:%s] to [NodeID:%s, VertexID:%s]"), *InGraphContext.BuildContext.DebugAssetName, *Edge.FromNodeID.ToString(), *Edge.FromVertexID.ToString(), *Edge.ToNodeID.ToString(), *Edge.ToVertexID.ToString());
-				return false;
-			}
-		}
-
-		return true;
-	}
-
-	bool FFrontendGraphBuilder::AddDefaultInputLiterals(FBuildGraphContext& InGraphContext)
-	{
-		using namespace Metasound::Frontend;
-		using namespace LiteralNodeNames; 
-
-		using FIterator = FDefaultInputByIDMap::TIterator;
-
-		for (FDefaultInputByIDMap::ElementType& Pair : InGraphContext.DefaultInputs)
-		{
-			FDefaultLiteralData& LiteralData = Pair.Value;
-			const FGuid LiteralNodeID = Frontend::CreateLocallyUniqueId();
-
-			// 1. Construct and add the default variable to the graph
-			{
-				TUniquePtr<const INode> DefaultLiteral = IDataTypeRegistry::Get().CreateLiteralNode(LiteralData.TypeName, MoveTemp(LiteralData.InitParams));
-				InGraphContext.Graph->AddNode(LiteralNodeID, TSharedPtr<const INode>(DefaultLiteral.Release()));
-			}
-
-			// 2. Connect the default variable to the expected input
-			const INode* FromNode = InGraphContext.Graph->FindNode(LiteralNodeID);
-			if (nullptr == FromNode)
-			{
-				UE_LOG(LogMetaSound, Error, TEXT("MetaSound '%s': Failed to find node in graph [NodeID:%s]"), *InGraphContext.BuildContext.DebugAssetName, *LiteralNodeID.ToString());
-				return false;
-			}
-			const FVertexName& FromVertexName = METASOUND_GET_PARAM_NAME(OutputValue);
-
-			const INode* ToNode = InGraphContext.Graph->FindNode(LiteralData.DestinationNodeID);
-			if (nullptr == ToNode)
-			{
-				UE_LOG(LogMetaSound, Error, TEXT("MetaSound '%s': Failed to node in graph [NodeID:%s]"), *InGraphContext.BuildContext.DebugAssetName, *LiteralData.DestinationNodeID.ToString());
-				return false;
-			}
-			const FVertexName& ToVertexName = LiteralData.DestinationVertexKey;
-
-			bool bSuccess = InGraphContext.Graph->AddDataEdge(*FromNode, FromVertexName, *ToNode, ToVertexName);
-			if (!bSuccess)
-			{
-				UE_LOG(LogMetaSound, Error, TEXT("MetaSound '%s': Failed to connect default variable edge: from '%s' to '%s'"), *InGraphContext.BuildContext.DebugAssetName, *FromVertexName.ToString(), *ToVertexName.ToString());
-				return false;
-			}
-		}
-
-		// Clear default inputs because literals have been moved out of default input map.
-		InGraphContext.DefaultInputs.Reset();
-		return true;
-	}
-
-	TArray<FFrontendGraphBuilder::FDefaultLiteralData> FFrontendGraphBuilder::GetInputDefaultLiteralData(const FBuildContext& InContext, const FMetasoundFrontendNode& InNode, const FNodeInitData& InInitData, const TSet<FNodeIDVertexID>& InEdgeDestinations)
-	{
-		TArray<FDefaultLiteralData> DefaultLiteralData;
-
-		TArray<FMetasoundFrontendVertex> InputVertices = InNode.Interface.Inputs;
-		for (const FMetasoundFrontendVertexLiteral& Literal : InNode.InputLiterals)
-		{
-			FLiteralNodeConstructorParams InitParams;
-			FName TypeName;
-
-			FGuid VertexID = Literal.VertexID;
-			FVertexName DestinationVertexName;
-			bool bRequiresDefault = false;
-			auto RemoveAndBuildParams = [&](const FMetasoundFrontendVertex& Vertex)
-			{
-				if (Vertex.VertexID == VertexID)
-				{
-					// Only build params and forward along to connect to dynamically generated
-					// literal node if edge is not explicitly connected to vertex as destination.
-					bRequiresDefault = !InEdgeDestinations.Contains({InNode.GetID(), Vertex.VertexID});
-					if (bRequiresDefault)
-					{
-						InitParams.Literal = Literal.Value.ToLiteral(Vertex.TypeName, &InContext.DataTypeRegistry, &InContext.ProxyDataCache);
-						InitParams.InstanceID = Frontend::CreateLocallyUniqueId();
-						InitParams.NodeName = "Literal";
-						TypeName = Vertex.TypeName;
-
-						DestinationVertexName = Vertex.Name;
-					}
-
-					return true;
-				}
-
-				return false;
-			};
-
-			const bool bRemoved = InputVertices.RemoveAllSwap(RemoveAndBuildParams, EAllowShrinking::No) > 0;
-			if (ensure(bRemoved))
-			{
-				if (bRequiresDefault)
-				{
-					DefaultLiteralData.Emplace(FDefaultLiteralData
-						{
-							InNode.GetID(),
-							VertexID,
-							DestinationVertexName,
-							TypeName,
-							MoveTemp(InitParams)
-						});
-				}
-			}
-		}
-
-		return DefaultLiteralData;
-	}
-
-	/** Check that all dependencies are C++ class dependencies. */
-	bool FFrontendGraphBuilder::IsFlat(const FMetasoundFrontendDocument& InDocument)
-	{
-		if (InDocument.Subgraphs.Num() > 0)
-		{
-			return false;
-		}
-
-PRAGMA_DISABLE_DEPRECATION_WARNINGS
-		return IsFlat(InDocument.RootGraph, InDocument.Dependencies);
-PRAGMA_ENABLE_DEPRECATION_WARNINGS
-	}
-
-	bool FFrontendGraphBuilder::IsFlat(const FMetasoundFrontendGraphClass& InRoot, const TArray<FMetasoundFrontendClass>& InDependencies)
-	{
-		// All dependencies are external dependencies in a flat graph
-		auto IsClassExternal = [&](const FMetasoundFrontendClass& InDesc) 
-		{
-			const bool bIsExternal = (InDesc.Metadata.GetType() == EMetasoundFrontendClassType::External) ||
-				(InDesc.Metadata.GetType() == EMetasoundFrontendClassType::Template) ||
-				(InDesc.Metadata.GetType() == EMetasoundFrontendClassType::Input) ||
-				(InDesc.Metadata.GetType() == EMetasoundFrontendClassType::Output);
-			return bIsExternal;
-		};
-		const bool bIsEveryDependencyExternal = Algo::AllOf(InDependencies, IsClassExternal);
-
-		if (!bIsEveryDependencyExternal)
-		{
-			return false;
-		}
-
-		// All the dependencies are met 
-		TSet<FGuid> AvailableDependencies;
-		Algo::Transform(InDependencies, AvailableDependencies, [](const FMetasoundFrontendClass& InDesc) { return InDesc.ID; });
-
-		auto IsDependencyMet = [&](const FMetasoundFrontendNode& InNode) 
-		{ 
-			return AvailableDependencies.Contains(InNode.ClassID);
-		};
-
-		const bool bIsEveryDependencyMet = Algo::AllOf(InRoot.GetConstDefaultGraph().Nodes, IsDependencyMet);
-
-		return bIsEveryDependencyMet;
-	}
-
-	bool FFrontendGraphBuilder::SortSubgraphDependencies(TArray<const FMetasoundFrontendGraphClass*>& Subgraphs)
-	{
-		// Helper for caching and querying subgraph dependencies
-		struct FSubgraphDependencyLookup
-		{
-			FSubgraphDependencyLookup(TArrayView<const FMetasoundFrontendGraphClass*> InGraphs)
-			{
-				// Map ClassID to graph pointer. 
-				TMap<FGuid, const FMetasoundFrontendGraphClass*> ClassIDAndGraph;
-				for (const FMetasoundFrontendGraphClass* Graph: InGraphs)
-				{
-					ClassIDAndGraph.Add(Graph->ID, Graph);
-				}
-
-				// Cache subgraph dependencies.
-				for (const FMetasoundFrontendGraphClass* GraphClass : InGraphs)
-				{
-					for (const FMetasoundFrontendNode& Node : GraphClass->GetConstDefaultGraph().Nodes)
-					{
-						if (ClassIDAndGraph.Contains(Node.ClassID))
-						{
-							DependencyMap.Add(GraphClass, ClassIDAndGraph[Node.ClassID]);
-						}
-					}
-				}
-			}
-
-			TArray<const FMetasoundFrontendGraphClass*> operator()(const FMetasoundFrontendGraphClass* InParent) const
-			{
-				TArray<const FMetasoundFrontendGraphClass*> Dependencies;
-				DependencyMap.MultiFind(InParent, Dependencies);
-				return Dependencies;
-			}
-
-		private:
-
-			TMultiMap<const FMetasoundFrontendGraphClass*, const FMetasoundFrontendGraphClass*> DependencyMap;
-		};
-
-		bool bSuccess = Algo::TopologicalSort(Subgraphs, FSubgraphDependencyLookup(Subgraphs));
-		if (!bSuccess)
-		{
-			UE_LOG(LogMetaSound, Error, TEXT("Failed to topologically sort subgraphs. Possible recursive subgraph dependency"));
-		}
-
-		return bSuccess;
-	}
-
-	TUniquePtr<FFrontendGraph> FFrontendGraphBuilder::CreateGraph(FBuildContext& InContext, const FMetasoundFrontendGraphClass& InGraphClass)
-	{
-		using namespace Frontend;
-
-		const FString GraphName = InContext.DebugAssetName;
-
-		const FGuid PageID = GraphPrivate::ResolveTargetPageID(InGraphClass);
-		const FMetasoundFrontendGraph& PageGraph = InGraphClass.FindConstGraphChecked(PageID);
-		FBuildGraphContext BuildGraphContext
-		{
-			MakeUnique<FFrontendGraph>(GraphName, InContext.GraphId),
-			InGraphClass,
-			PageGraph,
-			InContext
-		};
-
-		bool bSuccess = AddNodesToGraph(BuildGraphContext);
-
-		if (bSuccess)
-		{
-			bSuccess = AddEdgesToGraph(BuildGraphContext);
-		}
-
-		if (bSuccess)
-		{
-			bSuccess = AddDefaultInputLiterals(BuildGraphContext);
-		}
-
-		if (bSuccess)
-		{
-			return MoveTemp(BuildGraphContext.Graph);
-		}
-		else
-		{
-			return TUniquePtr<FFrontendGraph>(nullptr);
-		}
-	}
-
-	TUniquePtr<FFrontendGraph> FFrontendGraphBuilder::CreateGraph(const FMetasoundFrontendGraphClass& InGraph, const TArray<FMetasoundFrontendGraphClass>& InSubgraphs, const TArray<FMetasoundFrontendClass>& InDependencies, const TSet<FName>& InTransmittableInputNames, const FString& InDebugAssetName)
-	{
-		return CreateGraph(InGraph, InSubgraphs, InDependencies, InDebugAssetName);
-	}
-
-	TUniquePtr<FFrontendGraph> FFrontendGraphBuilder::CreateGraph(const FMetasoundFrontendGraphClass& InGraph, const TArray<FMetasoundFrontendGraphClass>& InSubgraphs, const TArray<FMetasoundFrontendClass>& InDependencies, const Frontend::FProxyDataCache& InProxyDataCache, const FString& InDebugAssetName, const FGuid InGraphId)
-	{
-		FBuildContext Context
-		{
-			InDebugAssetName, 					// DebugAssetName
-			{}, 								// FrontendClasses
-			{}, 								// Graphs
-			Frontend::IDataTypeRegistry::Get(), // DataTypeRegistry
-			InProxyDataCache, 					// ProxyDataCache
-			InGraphId							// GraphId
-		};
-
-		// Gather all references to node classes from external dependencies and subgraphs.
-		for (const FMetasoundFrontendClass& ExtClass : InDependencies)
-		{
-			Context.FrontendClasses.Add(ExtClass.ID, &ExtClass);
-		}
-		for (const FMetasoundFrontendClass& ExtClass : InSubgraphs)
-		{
-			Context.FrontendClasses.Add(ExtClass.ID, &ExtClass);
-		}
-
-		// Sort subgraphs so that dependent subgraphs are created in correct order.
-		TArray<const FMetasoundFrontendGraphClass*> FrontendSubgraphPtrs;
-		Algo::Transform(InSubgraphs, FrontendSubgraphPtrs, [](const FMetasoundFrontendGraphClass& InClass) { return &InClass; });
-
-		bool bSuccess = SortSubgraphDependencies(FrontendSubgraphPtrs);
-		if (!bSuccess)
-		{
-			UE_LOG(LogMetaSound, Error, TEXT("Failed to create graph due to failed subgraph ordering in asset '%s'."), *InDebugAssetName);
-			return TUniquePtr<FFrontendGraph>(nullptr);
-		}
-
-		// Create each subgraph.
-		for (const FMetasoundFrontendGraphClass* FrontendSubgraphPtr : FrontendSubgraphPtrs)
-		{
-			TSharedPtr<const INode> Subgraph(CreateGraph(Context, *FrontendSubgraphPtr).Release());
-			if (!Subgraph.IsValid())
-			{
-				UE_LOG(LogMetaSound, Warning, TEXT("Failed to create subgraph [SubgraphName: %s] in asset '%s'"), *FrontendSubgraphPtr->Metadata.GetClassName().ToString(), *InDebugAssetName);
-			}
-			else
-			{
-				// Add subgraphs to context so they are accessible for subsequent graphs.
-				Context.Graphs.Add(FrontendSubgraphPtr->ID, Subgraph);
-			}
-		}
-
-		// Create parent graph.
-		return CreateGraph(Context, InGraph);
-	}
-
-	TUniquePtr<FFrontendGraph> FFrontendGraphBuilder::CreateGraph(const FMetasoundFrontendDocument& InDocument, const FString& InDebugAssetName)
-	{
-		// Create proxies before creating graph.
-		Frontend::FProxyDataCache ProxyDataCache;
-		ProxyDataCache.CreateAndCacheProxies(InDocument);
-		
-		return CreateGraph(InDocument, ProxyDataCache, InDebugAssetName, Frontend::CreateLocallyUniqueId());
-	}
-
-	TUniquePtr<FFrontendGraph> FFrontendGraphBuilder::CreateGraph(const FMetasoundFrontendGraphClass& InGraph, const TArray<FMetasoundFrontendGraphClass>& InSubgraphs, const TArray<FMetasoundFrontendClass>& InDependencies, const FString& InDebugAssetName)
-	{
-		// Create proxies before building graph
-		Frontend::FProxyDataCache ProxyDataCache;
-		ProxyDataCache.CreateAndCacheProxies(InGraph);
-
-		for (const FMetasoundFrontendGraphClass& SubgraphClass : InSubgraphs)
-		{
-			ProxyDataCache.CreateAndCacheProxies(SubgraphClass);
-		}
-
-		for (const FMetasoundFrontendClass& DependencyClass : InDependencies)
-		{
-			ProxyDataCache.CreateAndCacheProxies(DependencyClass);
-		}
-
-		return CreateGraph(InGraph, InSubgraphs, InDependencies, ProxyDataCache, InDebugAssetName, Frontend::CreateLocallyUniqueId());
-	}
-	
-	TUniquePtr<FFrontendGraph> FFrontendGraphBuilder::CreateGraph(const FMetasoundFrontendDocument& InDocument, const TSet<FName>& InTransmittableInputNames, const FString& InDebugAssetName)
-	{
-		return CreateGraph(InDocument, InDebugAssetName);
-	}
-
-	TUniquePtr<FFrontendGraph> FFrontendGraphBuilder::CreateGraph(const FMetasoundFrontendDocument& InDocument, const Frontend::FProxyDataCache& InProxyDataCache, const FString& InDebugAssetName, const FGuid InGraphId)
-	{
-		return CreateGraph(InDocument.RootGraph, InDocument.Subgraphs, InDocument.Dependencies, InProxyDataCache, InDebugAssetName, InGraphId);
-	}
 }
diff --git a/Engine/Plugins/Runtime/Metasound/Source/MetasoundFrontend/Private/MetasoundFrontendGraphBuilder.cpp b/Engine/Plugins/Runtime/Metasound/Source/MetasoundFrontend/Private/MetasoundFrontendGraphBuilder.cpp
new file mode 100644
--- /dev/null
+++ b/Engine/Plugins/Runtime/Metasound/Source/MetasoundFrontend/Private/MetasoundFrontendGraphBuilder.cpp
@@ -0,0 +1,811 @@
+// Copyright Epic Games, Inc. All Rights Reserved.
+
+#include "MetasoundFrontendGraphBuilder.h"
+
+#include "MetasoundFrontendGraph.h"
+#include "Algo/AllOf.h"
+#include "Algo/AnyOf.h"
+#include "Algo/TopologicalSort.h"
+#include "Algo/Transform.h"
+#include "CoreMinimal.h"
+#include "CoreGlobals.h"
+#include "MetasoundDocumentInterface.h"
+#include "MetasoundFrontend.h"
+#include "MetasoundFrontendDataTypeRegistry.h"
+#include "MetasoundFrontendNodeTemplateRegistry.h"
+#include "MetasoundFrontendProxyDataCache.h"
+#include "MetasoundFrontendRegistries.h"
+#include "MetasoundGraph.h"
+#include "MetasoundLiteralNode.h"
+#include "MetasoundLog.h"
+#include "MetasoundNodeInterface.h"
+
+namespace Metasound::Frontend
+{
+	namespace GraphBuilderPrivate
+	{
+		FNodeInitData CreateNodeInitData(const FMetasoundFrontendNode& InNode)
+		{
+			FNodeInitData InitData;
+
+			InitData.InstanceName = InNode.Name;
+			InitData.InstanceID = InNode.GetID();
+
+			return InitData;
+		}
+
+		template <typename ResolveType>
+		FGuid ResolveTargetPageID(const ResolveType& InToResolve)
+		{
+			// Registry is not available in tests, so for now resolution is considered successful at this level
+			// if registry is not initialized and providing a resolved page ID. TODO: Add a test implementation
+			// that returns the default page (or whatever page behavior is desired for testing).
+			if (IDocumentBuilderRegistry* BuilderRegistry = IDocumentBuilderRegistry::Get())
+			{
+				return IDocumentBuilderRegistry::GetChecked().ResolveTargetPageID(InToResolve);
+			}
+
+			return Frontend::DefaultPageID;
+		}
+	} // namespace GraphBuilderPrivate
+
+
+	TUniquePtr<INode> FGraphBuilder::CreateVariableNode(const FBuildContext& InContext, const FMetasoundFrontendNode& InNode, const FMetasoundFrontendGraph& InGraph)
+	{
+		using namespace Metasound::Frontend;
+
+		const FMetasoundFrontendVariable* FrontendVariable = FindVariableForVariableNode(InNode, InGraph);
+
+		if (nullptr != FrontendVariable)
+		{
+			IDataTypeRegistry& DataTypeRegistry = IDataTypeRegistry::Get();
+			const bool IsLiteralParsableByDataType = DataTypeRegistry.IsLiteralTypeSupported(FrontendVariable->TypeName, FrontendVariable->Literal.GetType());
+
+			if (IsLiteralParsableByDataType)
+			{
+				FLiteral Literal = FrontendVariable->Literal.ToLiteral(FrontendVariable->TypeName, &InContext.DataTypeRegistry, &InContext.ProxyDataCache);
+
+				FVariableNodeConstructorParams InitParams =
+				{
+					InNode.Name,
+					InNode.GetID(),
+					MoveTemp(Literal)
+				};
+
+				return DataTypeRegistry.CreateVariableNode(FrontendVariable->TypeName, MoveTemp(InitParams));
+			}
+			else
+			{
+				UE_LOG(LogMetaSound, Error, TEXT("Cannot create variable node [NodeID:%s]. [Variable:%s] cannot be constructed with the provided literal type."), *InNode.GetID().ToString(), *FrontendVariable->Name.ToString());
+			}
+		}
+		else
+		{
+			UE_LOG(LogMetaSound, Error, TEXT("Cannot create variable node [NodeID:%s]. No variable found for variable node."), *InNode.GetID().ToString());
+		}
+
+		return TUniquePtr<INode>(nullptr);
+	}
+
+	TUniquePtr<INode> FGraphBuilder::CreateInputNode(const FBuildContext& InContext, const FMetasoundFrontendNode& InNode, const FMetasoundFrontendClass& InClass, const FMetasoundFrontendClassInput& InOwningGraphClassInput)
+	{
+		using namespace Metasound::Frontend;
+
+		const FMetasoundFrontendLiteral* FrontendLiteral = FindInputLiteralForInputNode(InNode, InClass, InOwningGraphClassInput);
+
+		if (nullptr != FrontendLiteral)
+		{
+			if (ensure(InNode.Interface.Inputs.Num() == 1))
+			{
+				IDataTypeRegistry& DataTypeRegistry = IDataTypeRegistry::Get();
+				const FMetasoundFrontendVertex& InputVertex = InNode.Interface.Inputs[0];
+
+				const bool IsLiteralParsableByDataType = DataTypeRegistry.IsLiteralTypeSupported(InputVertex.TypeName, FrontendLiteral->GetType());
+
+				if (IsLiteralParsableByDataType)
+				{
+					FLiteral Literal = FrontendLiteral->ToLiteral(InputVertex.TypeName, &InContext.DataTypeRegistry, &InContext.ProxyDataCache);
+
+					FInputNodeConstructorParams InitParams =
+					{
+						InNode.Name,
+						InNode.GetID(),
+						InputVertex.Name,
+						MoveTemp(Literal)
+					};
+
+					ensureAlwaysMsgf(InOwningGraphClassInput.AccessType != EMetasoundFrontendVertexAccessType::Unset, TEXT("Graph Class Input cannot be set to access type of 'Unset'"));
+					if (InOwningGraphClassInput.AccessType == EMetasoundFrontendVertexAccessType::Reference)
+					{
+						return DataTypeRegistry.CreateInputNode(InputVertex.TypeName, MoveTemp(InitParams));
+					}
+					else // InOwningGraphClassInput.AccessType == EMetasoundFrontendVertexAccessType::Value
+					{
+						return DataTypeRegistry.CreateConstructorInputNode(InputVertex.TypeName, MoveTemp(InitParams));
+					}
+				}
+				else
+				{
+					UE_LOG(LogMetaSound, Error, TEXT("Cannot create input node [NodeID:%s]. [Vertex:%s] cannot be constructed with the provided literal type."), *InNode.GetID().ToString(), *InputVertex.Name.ToString());
+				}
+			}
+		}
+		else
+		{
+			UE_LOG(LogMetaSound, Error, TEXT("Cannot create input node [NodeID:%s]. No default literal set for input node."), *InNode.GetID().ToString());
+		}
+
+		return TUniquePtr<INode>(nullptr);
+	}
+
+	TUniquePtr<INode> FGraphBuilder::CreateOutputNode(const FBuildContext& InContext, const FMetasoundFrontendNode& InNode, const FMetasoundFrontendClass& InClass, FBuildGraphContext& InGraphContext, const TSet<FNodeIDVertexID>& InEdgeDestinations)
+	{
+		using namespace Metasound::Frontend;
+
+		check(InClass.Metadata.GetType() == EMetasoundFrontendClassType::Output);
+		check(InNode.ClassID == InClass.ID);
+
+		if (ensure(InNode.Interface.Outputs.Num() == 1))
+		{
+			const FMetasoundFrontendVertex& OutputVertex = InNode.Interface.Outputs[0];
+
+			FOutputNodeConstructorParams InitParams =
+			{
+				InNode.Name,
+				InNode.GetID(),
+				OutputVertex.Name
+			};
+
+			{
+				const FNodeInitData InitData = GraphBuilderPrivate::CreateNodeInitData(InNode);
+				TArray<FDefaultLiteralData> DefaultLiteralData = GetInputDefaultLiteralData(InContext, InNode, InitData, InEdgeDestinations);
+				for (FDefaultLiteralData& Data : DefaultLiteralData)
+				{
+					InGraphContext.DefaultInputs.Emplace(FNodeIDVertexID { InNode.GetID(), Data.DestinationVertexID }, MoveTemp(Data));
+				}
+			}
+
+			ensure(InClass.Interface.Outputs.Num() == 1);
+			if (InClass.Interface.Outputs[0].AccessType == EMetasoundFrontendVertexAccessType::Reference)
+			{
+				return IDataTypeRegistry::Get().CreateOutputNode(OutputVertex.TypeName, MoveTemp(InitParams));
+			}
+			else // InClass.Interface.Outputs[0].AccessType == EMetasoundFrontendVertexAccessType::Value
+			{
+				return IDataTypeRegistry::Get().CreateConstructorOutputNode(OutputVertex.TypeName, MoveTemp(InitParams));
+			}
+		}
+		return TUniquePtr<INode>(nullptr);
+	}
+
+	TUniquePtr<INode> FGraphBuilder::CreateExternalNode(const FBuildContext& InContext, const FMetasoundFrontendNode& InNode, const FMetasoundFrontendClass& InClass, FBuildGraphContext& InGraphContext, const TSet<FNodeIDVertexID>& InEdgeDestinations)
+	{
+		using namespace Frontend;
+
+		check(InNode.ClassID == InClass.ID);
+
+		const FNodeInitData InitData = GraphBuilderPrivate::CreateNodeInitData(InNode);
+		{
+			TArray<FDefaultLiteralData> DefaultLiteralData = GetInputDefaultLiteralData(InContext, InNode, InitData, InEdgeDestinations);
+			for (FDefaultLiteralData& Data : DefaultLiteralData)
+			{
+				InGraphContext.DefaultInputs.Emplace(FNodeIDVertexID{ InNode.GetID(), Data.DestinationVertexID }, MoveTemp(Data));
+			}
+		}
+
+		// TODO: handle check to see if node interface conforms to class interface here. 
+		// TODO: check to see if external object supports class interface.
+
+		const FNodeRegistryKey Key = FNodeRegistryKey(InClass.Metadata);
+		return FMetasoundFrontendRegistryContainer::Get()->CreateNode(Key, InitData);
+	}
+
+	const FMetasoundFrontendClassInput* FGraphBuilder::FindClassInputForInputNode(const FMetasoundFrontendGraphClass& InOwningGraph, const FMetasoundFrontendNode& InInputNode)
+	{
+		// Input nodes should have exactly one input.
+		if (ensure(InInputNode.Interface.Inputs.Num() == 1))
+		{
+			const FName& TypeName = InInputNode.Interface.Inputs[0].TypeName;
+
+			auto IsMatchingInput = [&](const FMetasoundFrontendClassInput& GraphInput)
+			{
+				return (InInputNode.GetID() == GraphInput.NodeID);
+			};
+
+			return InOwningGraph.Interface.Inputs.FindByPredicate(IsMatchingInput);
+		}
+		return nullptr;
+	}
+
+	const FMetasoundFrontendClassOutput* FGraphBuilder::FindClassOutputForOutputNode(const FMetasoundFrontendGraphClass& InOwningGraph, const FMetasoundFrontendNode& InOutputNode)
+	{
+		// Output nodes should have exactly one output
+		if (ensure(InOutputNode.Interface.Outputs.Num() == 1))
+		{
+			const FName& TypeName = InOutputNode.Interface.Outputs[0].TypeName;
+
+			auto IsMatchingOutput = [&](const FMetasoundFrontendClassOutput& GraphOutput)
+			{
+				return (InOutputNode.GetID() == GraphOutput.NodeID);
+			};
+
+			return InOwningGraph.Interface.Outputs.FindByPredicate(IsMatchingOutput);
+		}
+		return nullptr;
+	}
+
+	const FMetasoundFrontendVariable* FGraphBuilder::FindVariableForVariableNode(const FMetasoundFrontendNode& InVariableNode, const FMetasoundFrontendGraph& InGraph)
+	{
+		const FGuid& DesiredID = InVariableNode.GetID();
+		return InGraph.Variables.FindByPredicate([&](const FMetasoundFrontendVariable& InVar) { return InVar.VariableNodeID == DesiredID; });
+	}
+
+	const FMetasoundFrontendLiteral* FGraphBuilder::FindInputLiteralForInputNode(const FMetasoundFrontendNode& InInputNode, const FMetasoundFrontendClass& InInputNodeClass, const FMetasoundFrontendClassInput& InOwningGraphClassInput)
+	{
+		using namespace Frontend;
+
+		// Default value priority is:
+		// 1. A value set directly on the node
+		// 2. A default value of the owning graph
+		// 3. A default value on the input node class.
+
+		const FMetasoundFrontendLiteral* Literal = nullptr;
+
+		// Check for default value directly on node.
+		if (ensure(InInputNode.Interface.Inputs.Num() == 1))
+		{
+			const FMetasoundFrontendVertex& InputVertex = InInputNode.Interface.Inputs[0];
+
+			// Find input literal matching VertexID
+			const FMetasoundFrontendVertexLiteral* VertexLiteral = InInputNode.InputLiterals.FindByPredicate(
+				[&](const FMetasoundFrontendVertexLiteral& InVertexLiteral)
+				{
+					return InVertexLiteral.VertexID == InputVertex.VertexID;
+				}
+			);
+
+			if (nullptr != VertexLiteral)
+			{
+				Literal = &VertexLiteral->Value;
+			}
+		}
+
+		// Check for default value on owning graph.
+		if (nullptr == Literal)
+		{
+			// Find Class Default that is not invalid
+			const FGuid PageID = GraphBuilderPrivate::ResolveTargetPageID(InOwningGraphClassInput);
+			const FMetasoundFrontendLiteral& DefaultLiteral = InOwningGraphClassInput.FindConstDefaultChecked(PageID);
+			if (DefaultLiteral.IsValid())
+			{
+				Literal = &DefaultLiteral;
+			}
+		}
+
+		// Check for default value on input node class
+		if (nullptr == Literal && ensure(InInputNodeClass.Interface.Inputs.Num() == 1))
+		{
+			const FMetasoundFrontendClassInput& InputNodeClassInput = InInputNodeClass.Interface.Inputs.Last();
+			const FGuid PageID = GraphBuilderPrivate::ResolveTargetPageID(InputNodeClassInput);
+			const FMetasoundFrontendLiteral& DefaultLiteral = InputNodeClassInput.FindConstDefaultChecked(PageID);
+			if (DefaultLiteral.IsValid())
+			{
+				Literal = &DefaultLiteral;
+			}
+		}
+
+		return Literal;
+	}
+
+	bool FGraphBuilder::AddNodesToGraph(FBuildGraphContext& InGraphContext)
+	{
+		TSet<FNodeIDVertexID> GraphEdgeDestinations;
+		const TArray<FMetasoundFrontendEdge>& GraphEdges = InGraphContext.PagedGraph.Edges;
+		Algo::Transform(GraphEdges, GraphEdgeDestinations, [](const FMetasoundFrontendEdge& Edge) 
+		{
+			return FNodeIDVertexID{Edge.ToNodeID, Edge.ToVertexID};
+		});
+
+		for (const FMetasoundFrontendNode& Node : InGraphContext.PagedGraph.Nodes)
+		{
+			const FMetasoundFrontendClass* NodeClass = InGraphContext.BuildContext.FrontendClasses.FindRef(Node.ClassID);
+
+			if (ensure(nullptr != NodeClass))
+			{
+				switch (NodeClass->Metadata.GetType())
+				{
+					case EMetasoundFrontendClassType::Input:
+					{
+						const FMetasoundFrontendClassInput* ClassInput = FindClassInputForInputNode(InGraphContext.GraphClass, Node);
+
+						if (nullptr != ClassInput)
+						{
+							TSharedPtr<const INode> InputNode(CreateInputNode(InGraphContext.BuildContext, Node, *NodeClass, *ClassInput).Release());
+							InGraphContext.Graph->AddInputNode(Node.GetID(), ClassInput->Name, InputNode);
+						}
+						else
+						{
+							const FString GraphClassIDString = InGraphContext.GraphClass.ID.ToString();
+							UE_LOG(LogMetaSound, Error, TEXT("MetaSound '%s': Failed to match input node [NodeID:%s, NodeName:%s] to owning graph [ClassID:%s] input."), *InGraphContext.BuildContext.DebugAssetName, *Node.GetID().ToString(), *Node.Name.ToString(), *GraphClassIDString);
+							return false;
+						}
+					}
+					break;
+
+					case EMetasoundFrontendClassType::Output:
+					{
+						const FMetasoundFrontendClassOutput* ClassOutput = FindClassOutputForOutputNode(InGraphContext.GraphClass, Node);
+						if (nullptr != ClassOutput)
+						{
+							TSharedPtr<const INode> OutputNode(CreateOutputNode(InGraphContext.BuildContext, Node, *NodeClass, InGraphContext, GraphEdgeDestinations).Release());
+							InGraphContext.Graph->AddOutputNode(Node.GetID(), ClassOutput->Name, OutputNode);
+						}
+						else
+						{
+							const FString GraphClassIDString = InGraphContext.GraphClass.ID.ToString();
+							UE_LOG(LogMetaSound, Error, TEXT("MetaSound '%s': Failed to match output node [NodeID:%s, NodeName:%s] to owning graph [ClassID:%s] output."), *InGraphContext.BuildContext.DebugAssetName, *Node.GetID().ToString(), *Node.Name.ToString(), *GraphClassIDString);
+							return false;
+						}
+					}
+					break;
+
+					case EMetasoundFrontendClassType::Graph:
+					{
+						const TSharedPtr<const INode> SubgraphPtr = InGraphContext.BuildContext.Graphs.FindRef(Node.ClassID);
+
+						if (SubgraphPtr.IsValid())
+						{
+							InGraphContext.Graph->AddNode(Node.GetID(), SubgraphPtr);
+						}
+						else
+						{
+							UE_LOG(LogMetaSound, Error, TEXT("MetaSound '%s': Failed to find subgraph for node [NodeID:%s, NodeName:%s, ClassID:%s]"), *InGraphContext.BuildContext.DebugAssetName, *Node.GetID().ToString(), *Node.Name.ToString(), *Node.ClassID.ToString());
+							return false;
+						}
+					}
+					break;
+
+					case EMetasoundFrontendClassType::Literal:
+					{
+						checkNoEntry(); // Unsupported.
+						return false;
+					}
+
+					case EMetasoundFrontendClassType::Variable:
+					{
+						TSharedPtr<const INode> VariableNode(CreateVariableNode(InGraphContext.BuildContext, Node, InGraphContext.PagedGraph).Release());
+						InGraphContext.Graph->AddNode(Node.GetID(), VariableNode);
+					}
+					break;
+
+					// Templates, variable accessors and mutators are
+					// constructed with the same parameters as external nodes.
+					case EMetasoundFrontendClassType::Template:
+					case EMetasoundFrontendClassType::VariableAccessor:
+					case EMetasoundFrontendClassType::VariableDeferredAccessor:
+					case EMetasoundFrontendClassType::VariableMutator:
+					case EMetasoundFrontendClassType::External:
+					default:
+					{
+						TSharedPtr<const INode> ExternalNode(CreateExternalNode(InGraphContext.BuildContext, Node, *NodeClass, InGraphContext, GraphEdgeDestinations).Release());
+						InGraphContext.Graph->AddNode(Node.GetID(), ExternalNode);
+					}
+					break;
+				}
+			}
+		}
+
+		return true;
+	}
+
+	bool FGraphBuilder::AddEdgesToGraph(FBuildGraphContext& InGraphContext)
+	{
+		// Pair of frontend node and core node. The frontend node can be one of
+		// several types.
+		struct FCoreNodeAndFrontendVertex
+		{
+			const INode* Node = nullptr;
+			const FMetasoundFrontendVertex* Vertex = nullptr;
+		};
+
+		TMap<FNodeIDVertexID, FCoreNodeAndFrontendVertex> NodeSourcesByID;
+		TMap<FNodeIDVertexID, FCoreNodeAndFrontendVertex> NodeDestinationsByID;
+
+		// Add nodes to NodeID/VertexID map
+		for (const FMetasoundFrontendNode& Node : InGraphContext.PagedGraph.Nodes)
+		{
+			const INode* CoreNode = InGraphContext.Graph->FindNode(Node.GetID());
+			if (nullptr == CoreNode)
+			{
+				UE_LOG(LogMetaSound, Warning, TEXT("Metasound '%s': Could not find referenced node [Name:%s, NodeID:%s]"), *InGraphContext.BuildContext.DebugAssetName, *Node.Name.ToString(), *Node.GetID().ToString());
+				return false;
+			}
+
+			for (const FMetasoundFrontendVertex& Vertex : Node.Interface.Inputs)
+			{
+				NodeDestinationsByID.Add(FNodeIDVertexID(Node.GetID(), Vertex.VertexID), FCoreNodeAndFrontendVertex({CoreNode, &Vertex}));
+			}
+
+			for (const FMetasoundFrontendVertex& Vertex : Node.Interface.Outputs)
+			{
+				NodeSourcesByID.Add(FNodeIDVertexID(Node.GetID(), Vertex.VertexID), FCoreNodeAndFrontendVertex({CoreNode, &Vertex}));
+			}
+		};
+
+		for (const FMetasoundFrontendEdge& Edge : InGraphContext.PagedGraph.Edges)
+		{
+			const FNodeIDVertexID DestinationKey(Edge.ToNodeID, Edge.ToVertexID);
+			const FCoreNodeAndFrontendVertex* DestinationNodeAndVertex = NodeDestinationsByID.Find(DestinationKey);
+
+			if (nullptr == DestinationNodeAndVertex)
+			{
+				UE_LOG(LogMetaSound, Error, TEXT("MetaSound '%s': Failed to add edge. Could not find destination [NodeID:%s, VertexID:%s]"), *InGraphContext.BuildContext.DebugAssetName, *Edge.ToNodeID.ToString(), *Edge.ToVertexID.ToString());
+				return false;
+			}
+
+			if (nullptr == DestinationNodeAndVertex->Node)
+			{
+				UE_LOG(LogMetaSound, Warning, TEXT("MetaSound '%s': 'Failed to add edge. Null destination node [NodeID:%s]"), *InGraphContext.BuildContext.DebugAssetName, *Edge.ToNodeID.ToString());
+				return false;
+			}
+
+			const FNodeIDVertexID SourceKey(Edge.FromNodeID, Edge.FromVertexID);
+			const FCoreNodeAndFrontendVertex* SourceNodeAndVertex = NodeSourcesByID.Find(SourceKey);
+
+			if (nullptr == SourceNodeAndVertex)
+			{
+				UE_LOG(LogMetaSound, Error, TEXT("MetaSound '%s': Failed to add edge. Could not find source [NodeID:%s, VertexID:%s]"), *InGraphContext.BuildContext.DebugAssetName, *Edge.FromNodeID.ToString(), *Edge.FromVertexID.ToString());
+				return false;
+			}
+
+			if (nullptr == SourceNodeAndVertex->Node)
+			{
+				UE_LOG(LogMetaSound, Warning, TEXT("MetaSound '%s': Skipping edge. Null source node [NodeID:%s]"), *InGraphContext.BuildContext.DebugAssetName, *Edge.FromNodeID.ToString());
+				return false;
+			}
+
+			const INode* FromNode = SourceNodeAndVertex->Node;
+			const FVertexName FromVertexKey = SourceNodeAndVertex->Vertex->Name;
+
+			const INode* ToNode = DestinationNodeAndVertex->Node;
+			const FVertexName ToVertexKey = DestinationNodeAndVertex->Vertex->Name;
+
+			bool bSuccess = InGraphContext.Graph->AddDataEdge(*FromNode, FromVertexKey,  *ToNode, ToVertexKey);
+
+			// If succeeded, remove input as viable vertex to construct default variable, as it has been superceded by a connection.
+			if (bSuccess)
+			{
+				FNodeIDVertexID DestinationPair { ToNode->GetInstanceID(), DestinationNodeAndVertex->Vertex->VertexID };
+				InGraphContext.DefaultInputs.Remove(DestinationPair);
+			}
+			else
+			{
+				UE_LOG(LogMetaSound, Error, TEXT("MetaSound '%s': Failed to connect edge from [NodeID:%s, VertexID:%s] to [NodeID:%s, VertexID:%s]"), *InGraphContext.BuildContext.DebugAssetName, *Edge.FromNodeID.ToString(), *Edge.FromVertexID.ToString(), *Edge.ToNodeID.ToString(), *Edge.ToVertexID.ToString());
+				return false;
+			}
+		}
+
+		return true;
+	}
+
+	bool FGraphBuilder::AddDefaultInputLiterals(FBuildGraphContext& InGraphContext)
+	{
+		using namespace Metasound::Frontend;
+		using namespace LiteralNodeNames; 
+
+		using FIterator = FDefaultInputByIDMap::TIterator;
+
+		for (FDefaultInputByIDMap::ElementType& Pair : InGraphContext.DefaultInputs)
+		{
+			FDefaultLiteralData& LiteralData = Pair.Value;
+			const FGuid LiteralNodeID = Frontend::CreateLocallyUniqueId();
+
+			// 1. Construct and add the default variable to the graph
+			{
+				TUniquePtr<const INode> DefaultLiteral = IDataTypeRegistry::Get().CreateLiteralNode(LiteralData.TypeName, MoveTemp(LiteralData.InitParams));
+				InGraphContext.Graph->AddNode(LiteralNodeID, TSharedPtr<const INode>(DefaultLiteral.Release()));
+			}
+
+			// 2. Connect the default variable to the expected input
+			const INode* FromNode = InGraphContext.Graph->FindNode(LiteralNodeID);
+			if (nullptr == FromNode)
+			{
+				UE_LOG(LogMetaSound, Error, TEXT("MetaSound '%s': Failed to find node in graph [NodeID:%s]"), *InGraphContext.BuildContext.DebugAssetName, *LiteralNodeID.ToString());
+				return false;
+			}
+			const FVertexName& FromVertexName = METASOUND_GET_PARAM_NAME(OutputValue);
+
+			const INode* ToNode = InGraphContext.Graph->FindNode(LiteralData.DestinationNodeID);
+			if (nullptr == ToNode)
+			{
+				UE_LOG(LogMetaSound, Error, TEXT("MetaSound '%s': Failed to node in graph [NodeID:%s]"), *InGraphContext.BuildContext.DebugAssetName, *LiteralData.DestinationNodeID.ToString());
+				return false;
+			}
+			const FVertexName& ToVertexName = LiteralData.DestinationVertexKey;
+
+			bool bSuccess = InGraphContext.Graph->AddDataEdge(*FromNode, FromVertexName, *ToNode, ToVertexName);
+			if (!bSuccess)
+			{
+				UE_LOG(LogMetaSound, Error, TEXT("MetaSound '%s': Failed to connect default variable edge: from '%s' to '%s'"), *InGraphContext.BuildContext.DebugAssetName, *FromVertexName.ToString(), *ToVertexName.ToString());
+				return false;
+			}
+		}
+
+		// Clear default inputs because literals have been moved out of default input map.
+		InGraphContext.DefaultInputs.Reset();
+		return true;
+	}
+
+	TArray<FGraphBuilder::FDefaultLiteralData> FGraphBuilder::GetInputDefaultLiteralData(const FBuildContext& InContext, const FMetasoundFrontendNode& InNode, const FNodeInitData& InInitData, const TSet<FNodeIDVertexID>& InEdgeDestinations)
+	{
+		TArray<FDefaultLiteralData> DefaultLiteralData;
+
+		TArray<FMetasoundFrontendVertex> InputVertices = InNode.Interface.Inputs;
+		for (const FMetasoundFrontendVertexLiteral& Literal : InNode.InputLiterals)
+		{
+			FLiteralNodeConstructorParams InitParams;
+			FName TypeName;
+
+			FGuid VertexID = Literal.VertexID;
+			FVertexName DestinationVertexName;
+			bool bRequiresDefault = false;
+			auto RemoveAndBuildParams = [&](const FMetasoundFrontendVertex& Vertex)
+			{
+				if (Vertex.VertexID == VertexID)
+				{
+					// Only build params and forward along to connect to dynamically generated
+					// literal node if edge is not explicitly connected to vertex as destination.
+					bRequiresDefault = !InEdgeDestinations.Contains({InNode.GetID(), Vertex.VertexID});
+					if (bRequiresDefault)
+					{
+						InitParams.Literal = Literal.Value.ToLiteral(Vertex.TypeName, &InContext.DataTypeRegistry, &InContext.ProxyDataCache);
+						InitParams.InstanceID = Frontend::CreateLocallyUniqueId();
+						InitParams.NodeName = "Literal";
+						TypeName = Vertex.TypeName;
+
+						DestinationVertexName = Vertex.Name;
+					}
+
+					return true;
+				}
+
+				return false;
+			};
+
+			const bool bRemoved = InputVertices.RemoveAllSwap(RemoveAndBuildParams, EAllowShrinking::No) > 0;
+			if (ensure(bRemoved))
+			{
+				if (bRequiresDefault)
+				{
+					DefaultLiteralData.Emplace(FDefaultLiteralData
+						{
+							InNode.GetID(),
+							VertexID,
+							DestinationVertexName,
+							TypeName,
+							MoveTemp(InitParams)
+						});
+				}
+			}
+		}
+
+		return DefaultLiteralData;
+	}
+
+	/** Check that all dependencies are C++ class dependencies. */
+	bool FGraphBuilder::IsFlat(const FMetasoundFrontendDocument& InDocument)
+	{
+		if (InDocument.Subgraphs.Num() > 0)
+		{
+			return false;
+		}
+
+PRAGMA_DISABLE_DEPRECATION_WARNINGS
+		return IsFlat(InDocument.RootGraph, InDocument.Dependencies);
+PRAGMA_ENABLE_DEPRECATION_WARNINGS
+	}
+
+	bool FGraphBuilder::IsFlat(const FMetasoundFrontendGraphClass& InRoot, const TArray<FMetasoundFrontendClass>& InDependencies)
+	{
+		// All dependencies are external dependencies in a flat graph
+		auto IsClassExternal = [&](const FMetasoundFrontendClass& InDesc) 
+		{
+			const bool bIsExternal = (InDesc.Metadata.GetType() == EMetasoundFrontendClassType::External) ||
+				(InDesc.Metadata.GetType() == EMetasoundFrontendClassType::Template) ||
+				(InDesc.Metadata.GetType() == EMetasoundFrontendClassType::Input) ||
+				(InDesc.Metadata.GetType() == EMetasoundFrontendClassType::Output);
+			return bIsExternal;
+		};
+		const bool bIsEveryDependencyExternal = Algo::AllOf(InDependencies, IsClassExternal);
+
+		if (!bIsEveryDependencyExternal)
+		{
+			return false;
+		}
+
+		// All the dependencies are met 
+		TSet<FGuid> AvailableDependencies;
+		Algo::Transform(InDependencies, AvailableDependencies, [](const FMetasoundFrontendClass& InDesc) { return InDesc.ID; });
+
+		auto IsDependencyMet = [&](const FMetasoundFrontendNode& InNode) 
+		{ 
+			return AvailableDependencies.Contains(InNode.ClassID);
+		};
+
+		const bool bIsEveryDependencyMet = Algo::AllOf(InRoot.GetConstDefaultGraph().Nodes, IsDependencyMet);
+
+		return bIsEveryDependencyMet;
+	}
+
+	bool FGraphBuilder::SortSubgraphDependencies(TArray<const FMetasoundFrontendGraphClass*>& Subgraphs)
+	{
+		// Helper for caching and querying subgraph dependencies
+		struct FSubgraphDependencyLookup
+		{
+			FSubgraphDependencyLookup(TArrayView<const FMetasoundFrontendGraphClass*> InGraphs)
+			{
+				// Map ClassID to graph pointer. 
+				TMap<FGuid, const FMetasoundFrontendGraphClass*> ClassIDAndGraph;
+				for (const FMetasoundFrontendGraphClass* Graph: InGraphs)
+				{
+					ClassIDAndGraph.Add(Graph->ID, Graph);
+				}
+
+				// Cache subgraph dependencies.
+				for (const FMetasoundFrontendGraphClass* GraphClass : InGraphs)
+				{
+					for (const FMetasoundFrontendNode& Node : GraphClass->GetConstDefaultGraph().Nodes)
+					{
+						if (ClassIDAndGraph.Contains(Node.ClassID))
+						{
+							DependencyMap.Add(GraphClass, ClassIDAndGraph[Node.ClassID]);
+						}
+					}
+				}
+			}
+
+			TArray<const FMetasoundFrontendGraphClass*> operator()(const FMetasoundFrontendGraphClass* InParent) const
+			{
+				TArray<const FMetasoundFrontendGraphClass*> Dependencies;
+				DependencyMap.MultiFind(InParent, Dependencies);
+				return Dependencies;
+			}
+
+		private:
+
+			TMultiMap<const FMetasoundFrontendGraphClass*, const FMetasoundFrontendGraphClass*> DependencyMap;
+		};
+
+		bool bSuccess = Algo::TopologicalSort(Subgraphs, FSubgraphDependencyLookup(Subgraphs));
+		if (!bSuccess)
+		{
+			UE_LOG(LogMetaSound, Error, TEXT("Failed to topologically sort subgraphs. Possible recursive subgraph dependency"));
+		}
+
+		return bSuccess;
+	}
+
+	TUniquePtr<FFrontendGraph> FGraphBuilder::CreateGraph(FBuildContext& InContext, const FMetasoundFrontendGraphClass& InGraphClass)
+	{
+		using namespace Frontend;
+
+		const FString GraphName = InContext.DebugAssetName;
+
+		const FGuid PageID = GraphBuilderPrivate::ResolveTargetPageID(InGraphClass);
+		const FMetasoundFrontendGraph& PageGraph = InGraphClass.FindConstGraphChecked(PageID);
+		FBuildGraphContext BuildGraphContext
+		{
+			MakeUnique<FFrontendGraph>(GraphName, InContext.GraphId),
+			InGraphClass,
+			PageGraph,
+			InContext
+		};
+
+		bool bSuccess = AddNodesToGraph(BuildGraphContext);
+
+		if (bSuccess)
+		{
+			bSuccess = AddEdgesToGraph(BuildGraphContext);
+		}
+
+		if (bSuccess)
+		{
+			bSuccess = AddDefaultInputLiterals(BuildGraphContext);
+		}
+
+		if (bSuccess)
+		{
+			return MoveTemp(BuildGraphContext.Graph);
+		}
+		else
+		{
+			return TUniquePtr<FFrontendGraph>(nullptr);
+		}
+	}
+
+	TUniquePtr<FFrontendGraph> FGraphBuilder::CreateGraph(const FMetasoundFrontendGraphClass& InGraph, const TArray<FMetasoundFrontendGraphClass>& InSubgraphs, const TArray<FMetasoundFrontendClass>& InDependencies, const Frontend::FProxyDataCache& InProxyDataCache, const FString& InDebugAssetName, const FGuid InGraphId)
+	{
+		FBuildContext Context
+		{
+			InDebugAssetName, 					// DebugAssetName
+			{}, 								// FrontendClasses
+			{}, 								// Graphs
+			Frontend::IDataTypeRegistry::Get(), // DataTypeRegistry
+			InProxyDataCache, 					// ProxyDataCache
+			InGraphId							// GraphId
+		};
+
+		// Gather all references to node classes from external dependencies and subgraphs.
+		for (const FMetasoundFrontendClass& ExtClass : InDependencies)
+		{
+			Context.FrontendClasses.Add(ExtClass.ID, &ExtClass);
+		}
+		for (const FMetasoundFrontendClass& ExtClass : InSubgraphs)
+		{
+			Context.FrontendClasses.Add(ExtClass.ID, &ExtClass);
+		}
+
+		// Sort subgraphs so that dependent subgraphs are created in correct order.
+		TArray<const FMetasoundFrontendGraphClass*> FrontendSubgraphPtrs;
+		Algo::Transform(InSubgraphs, FrontendSubgraphPtrs, [](const FMetasoundFrontendGraphClass& InClass) { return &InClass; });
+
+		bool bSuccess = SortSubgraphDependencies(FrontendSubgraphPtrs);
+		if (!bSuccess)
+		{
+			UE_LOG(LogMetaSound, Error, TEXT("Failed to create graph due to failed subgraph ordering in asset '%s'."), *InDebugAssetName);
+			return TUniquePtr<FFrontendGraph>(nullptr);
+		}
+
+		// Create each subgraph.
+		for (const FMetasoundFrontendGraphClass* FrontendSubgraphPtr : FrontendSubgraphPtrs)
+		{
+			TSharedPtr<const INode> Subgraph(CreateGraph(Context, *FrontendSubgraphPtr).Release());
+			if (!Subgraph.IsValid())
+			{
+				UE_LOG(LogMetaSound, Warning, TEXT("Failed to create subgraph [SubgraphName: %s] in asset '%s'"), *FrontendSubgraphPtr->Metadata.GetClassName().ToString(), *InDebugAssetName);
+			}
+			else
+			{
+				// Add subgraphs to context so they are accessible for subsequent graphs.
+				Context.Graphs.Add(FrontendSubgraphPtr->ID, Subgraph);
+			}
+		}
+
+		// Create parent graph.
+		return CreateGraph(Context, InGraph);
+	}
+
+	TUniquePtr<FFrontendGraph> FGraphBuilder::CreateGraph(const FMetasoundFrontendDocument& InDocument, const FString& InDebugAssetName)
+	{
+		// Create proxies before creating graph.
+		Frontend::FProxyDataCache ProxyDataCache;
+		ProxyDataCache.CreateAndCacheProxies(InDocument);
+		
+		return CreateGraph(InDocument, ProxyDataCache, InDebugAssetName, Frontend::CreateLocallyUniqueId());
+	}
+
+	TUniquePtr<FFrontendGraph> FGraphBuilder::CreateGraph(const FMetasoundFrontendGraphClass& InGraph, const TArray<FMetasoundFrontendGraphClass>& InSubgraphs, const TArray<FMetasoundFrontendClass>& InDependencies, const FString& InDebugAssetName)
+	{
+		// Create proxies before building graph
+		Frontend::FProxyDataCache ProxyDataCache;
+		ProxyDataCache.CreateAndCacheProxies(InGraph);
+
+		for (const FMetasoundFrontendGraphClass& SubgraphClass : InSubgraphs)
+		{
+			ProxyDataCache.CreateAndCacheProxies(SubgraphClass);
+		}
+
+		for (const FMetasoundFrontendClass& DependencyClass : InDependencies)
+		{
+			ProxyDataCache.CreateAndCacheProxies(DependencyClass);
+		}
+
+		return CreateGraph(InGraph, InSubgraphs, InDependencies, ProxyDataCache, InDebugAssetName, Frontend::CreateLocallyUniqueId());
+	}
+
+	TUniquePtr<FFrontendGraph> FGraphBuilder::CreateGraph(const FMetasoundFrontendDocument& InDocument, const Frontend::FProxyDataCache& InProxyDataCache, const FString& InDebugAssetName, const FGuid InGraphId)
+	{
+		return CreateGraph(InDocument.RootGraph, InDocument.Subgraphs, InDocument.Dependencies, InProxyDataCache, InDebugAssetName, InGraphId);
+	}
+}
+
diff --git a/Engine/Plugins/Runtime/Metasound/Source/MetasoundFrontend/Private/MetasoundFrontendGraphController.cpp b/Engine/Plugins/Runtime/Metasound/Source/MetasoundFrontend/Private/MetasoundFrontendGraphController.cpp
--- a/Engine/Plugins/Runtime/Metasound/Source/MetasoundFrontend/Private/MetasoundFrontendGraphController.cpp
+++ b/Engine/Plugins/Runtime/Metasound/Source/MetasoundFrontend/Private/MetasoundFrontendGraphController.cpp
@@ -1651,7 +1651,7 @@ namespace Metasound
 				const FMetasoundFrontendDocument& Doc = *(OwningDocument->GetDocumentPtr().Get());
 				ProxyCache.CreateAndCacheProxies(Doc);
 
-				TUniquePtr<FFrontendGraph> Graph = FFrontendGraphBuilder::CreateGraph(*GraphClass, Subgraphs, Dependencies, ProxyCache, UnknownAsset, Frontend::CreateLocallyUniqueId());
+				TUniquePtr<FFrontendGraph> Graph = FGraphBuilder::CreateGraph(*GraphClass, Subgraphs, Dependencies, ProxyCache, UnknownAsset, Frontend::CreateLocallyUniqueId());
 
 				if (!Graph.IsValid())
 				{
diff --git a/Engine/Plugins/Runtime/Metasound/Source/MetasoundFrontend/Private/MetasoundFrontendNodeClassRegistry.cpp b/Engine/Plugins/Runtime/Metasound/Source/MetasoundFrontend/Private/MetasoundFrontendNodeClassRegistry.cpp
--- a/Engine/Plugins/Runtime/Metasound/Source/MetasoundFrontend/Private/MetasoundFrontendNodeClassRegistry.cpp
+++ b/Engine/Plugins/Runtime/Metasound/Source/MetasoundFrontend/Private/MetasoundFrontendNodeClassRegistry.cpp
@@ -385,7 +385,7 @@ namespace Metasound::Frontend
 		METASOUND_TRACE_CPUPROFILER_EVENT_SCOPE_TEXT(*FString::Printf(TEXT("Metasound::FNodeClassRegistry::BuildAndRegisterGraphFromDocument asset %s"), *InNodeClassInfo.AssetPath.ToString()));
 
 		// Use the asset class id for the graph id because it should be locally unique. Unlike other ids, it is regenerated on asset duplicate. 
-		TUniquePtr<FFrontendGraph> FrontendGraph = FFrontendGraphBuilder::CreateGraph(InDocument, InProxyDataCache, InNodeClassInfo.AssetPath.ToString(), /*GraphId=*/InNodeClassInfo.AssetClassID);
+		TUniquePtr<FFrontendGraph> FrontendGraph = FGraphBuilder::CreateGraph(InDocument, InProxyDataCache, InNodeClassInfo.AssetPath.ToString(), /*GraphId=*/InNodeClassInfo.AssetClassID);
 		if (!FrontendGraph.IsValid())
 		{
 			UE_LOG(LogMetaSound, Error, TEXT("Failed to build MetaSound graph in asset '%s'"), *InNodeClassInfo.AssetPath.ToString());
diff --git a/Engine/Plugins/Runtime/Metasound/Source/MetasoundFrontend/Public/MetasoundFrontendGraph.h b/Engine/Plugins/Runtime/Metasound/Source/MetasoundFrontend/Public/MetasoundFrontendGraph.h
--- a/Engine/Plugins/Runtime/Metasound/Source/MetasoundFrontend/Public/MetasoundFrontendGraph.h
+++ b/Engine/Plugins/Runtime/Metasound/Source/MetasoundFrontend/Public/MetasoundFrontendGraph.h
@@ -6,6 +6,7 @@
 #include "Containers/Map.h"
 #include "Containers/Set.h"
 #include "Containers/UnrealString.h"
+#include "MetasoundFrontendGraphBuilder.h"
 #include "MetasoundFrontendDocument.h"
 #include "MetasoundFrontendDocumentIdGenerator.h"
 #include "MetasoundGraph.h"
@@ -41,8 +42,12 @@ namespace Metasound
 			 * @param InVertexName - The key for the graph input vertex.
 			 * @param InNode - A shared pointer to an input node. 
 			 */
+			UE_DEPRECATED(5.6, "Use AddInputNode(...) version which does not take an InIndex")
 			void AddInputNode(FGuid InNodeID, int32 InIndex, const FVertexName& InVertexName, TSharedPtr<const INode> InNode);
 
+			// Unhide AddInputNode function on base class
+			using FGraph::AddInputNode;
+
 			/** Add an output node to this graph.
 			 *
 			 * @param InNodeID - The NodeID related to the parent FMetasoundFrontendClass.
@@ -50,14 +55,19 @@ namespace Metasound
 			 * @param InVertexName - The key for the graph output vertex.
 			 * @param InNode - A shared pointer to an output node. 
 			 */
+			UE_DEPRECATED(5.6, "Use AddOutputNode(...) version which does not take an InIndex")
 			void AddOutputNode(FGuid InNodeID, int32 InIndex, const FVertexName& InVertexName, TSharedPtr<const INode> InNode);
 
+			// Unhide AddOutputNode function on base class
+			using FGraph::AddOutputNode;
+
 			/** Retrieve node by input index.
 			 *
 			 * @param InIndex - The index of the requested input.
 			 *
 			 * @return Pointer to the Node if it is stored on this graph. nullptr otherwise. 
 			 */
+			UE_DEPRECATED(5.6, "Retrieving input nodes nodes by index will no longer supported on this object")
 			const INode* FindInputNode(int32 InIndex) const;
 
 			/** Retrieve node by output index.
@@ -66,123 +76,15 @@ namespace Metasound
 			 *
 			 * @return Pointer to the Node if it is stored on this graph. nullptr otherwise. 
 			 */
+			UE_DEPRECATED(5.6, "Retrieving output nodes nodes by index will no longer supported on this object")
 			const INode* FindOutputNode(int32 InIndex) const;
 
-			UE_DEPRECATED(5.3, "This function is no longer analyzes node ownership and will always return true.")
-			bool OwnsAllReferencedNodes() const;
-
 		private:
 
 			TMap<int32, const INode*> InputNodes;
 			TMap<int32, const INode*> OutputNodes;
 	};
 
-	/** FFrontendGraphBuilder builds a FFrontendGraph from a FMetasoundDocument
-	 * or FMetasoundFrontendClass.
-	 */
-	class METASOUNDFRONTEND_API FFrontendGraphBuilder
-	{
-	public:
-		/** Check that all dependencies are C++ class dependencies. 
-		 * 
-		 * @param InDocument - Document containing dependencies.
-		 *
-		 * @return True if all dependencies are C++ classes. False otherwise.
-		 */
-		UE_DEPRECATED(5.5, "No longer supported (not useful nor maintainable with paged graph support")
-		static bool IsFlat(const FMetasoundFrontendDocument& InDocument);
-
-		UE_DEPRECATED(5.5, "No longer supported (not useful nor maintainable with paged graph support")
-		static bool IsFlat(const FMetasoundFrontendGraphClass& InRoot, const TArray<FMetasoundFrontendClass>& InDependencies);
-
-		/* Create a FFrontendGraph from a FMetasoundFrontendDocument.*/
-		UE_DEPRECATED(5.4, "Use the version of CreateGraph(...) which does not include \"TransmittableInputNames\".")
-		static TUniquePtr<FFrontendGraph> CreateGraph(const FMetasoundFrontendDocument& InDocument, const TSet<FName>& TransmittableInputNames, const FString& InDebugAssetName);
-
-		/* Create a FFrontendGraph from a FMetasoundFrontendDocument subobjects.*/
-		UE_DEPRECATED(5.4, "Use the version of CreateGraph(...) which does not include \"TransmittableInputNames\".")
-		static TUniquePtr<FFrontendGraph> CreateGraph(const FMetasoundFrontendGraphClass& InGraphClass, const TArray<FMetasoundFrontendGraphClass>& InSubgraphs, const TArray<FMetasoundFrontendClass>& InDependencies, const TSet<FName>& TransmittableInputNames, const FString& InDebugAssetName);
-
-		/* Create a FFrontendGraph from a FMetasoundFrontendDocument.*/
-		static TUniquePtr<FFrontendGraph> CreateGraph(const FMetasoundFrontendDocument& InDocument, const FString& InDebugAssetName);
-
-		/* Create a FFrontendGraph from a FMetasoundFrontendDocument retrieving proxies from a FProxyDataCache.*/
-		static TUniquePtr<FFrontendGraph> CreateGraph(const FMetasoundFrontendDocument& InDocument, const Frontend::FProxyDataCache& InProxies, const FString& InDebugAssetName, const FGuid InGraphId = Frontend::CreateLocallyUniqueId());
-
-		/* Create a FFrontendGraph from a FMetasoundFrontendDocument and subobjects.*/
-		static TUniquePtr<FFrontendGraph> CreateGraph(const FMetasoundFrontendGraphClass& InGraph, const TArray<FMetasoundFrontendGraphClass>& InSubgraphs, const TArray<FMetasoundFrontendClass>& InDependencies, const FString& InDebugAssetName);
-
-		/* Create a FFrontendGraph from a FMetasoundFrontendDocument and subobjects, retrieving proxies from a FProxyDataCache.*/
-		static TUniquePtr<FFrontendGraph> CreateGraph(
-			const FMetasoundFrontendGraphClass& InGraph,
-			const TArray<FMetasoundFrontendGraphClass>& InSubgraphs,
-			const TArray<FMetasoundFrontendClass>& InDependencies,
-			const Frontend::FProxyDataCache& InProxyDataCache,
-			const FString& InDebugAssetName,
-			const FGuid InGraphId = Frontend::CreateLocallyUniqueId());
-
-	private:
-		struct FDefaultLiteralData
-		{
-			FGuid DestinationNodeID;
-			FGuid DestinationVertexID;
-			FVertexName DestinationVertexKey;
-			FName TypeName;
-			FLiteralNodeConstructorParams InitParams;
-		};
+	using FFrontendGraphBuilder UE_DEPRECATED(5.6, "FFrontendGraphBuilder is deprecated. Please use Frontend::FGraphBuilder instead.") = Frontend::FGraphBuilder;
 
-		// Map of Input VertexID to variable data required to construct and connect default variable
-		using FNodeIDVertexID = TTuple<FGuid, FGuid>;
-		using FDependencyByIDMap = TMap<FGuid, const FMetasoundFrontendClass*>;
-		using FSharedNodeByIDMap = TMap<FGuid, TSharedPtr<const INode>>;
-		using FDefaultInputByIDMap = TMap<FNodeIDVertexID, FDefaultLiteralData>;
-
-		// Context used throughout entire graph build process
-		// (for both a root and nested subgraphs)
-		struct FBuildContext
-		{
-			FString DebugAssetName;
-			FDependencyByIDMap FrontendClasses;
-			FSharedNodeByIDMap Graphs;
-			const Frontend::IDataTypeRegistry& DataTypeRegistry;
-			const Frontend::FProxyDataCache& ProxyDataCache;
-			FGuid GraphId;
-		};
-
-		// Transient context used for building a specific graph
-		struct FBuildGraphContext
-		{
-			TUniquePtr<FFrontendGraph> Graph;
-			const FMetasoundFrontendGraphClass& GraphClass;
-			const FMetasoundFrontendGraph& PagedGraph;
-			FBuildContext& BuildContext;
-			FDefaultInputByIDMap DefaultInputs;
-		};
-
-
-		static TArray<FDefaultLiteralData> GetInputDefaultLiteralData(const FBuildContext& InContext, const FMetasoundFrontendNode& InNode, const FNodeInitData& InInitData, const TSet<FNodeIDVertexID>& InEdgeDestinations);
-
-		static bool SortSubgraphDependencies(TArray<const FMetasoundFrontendGraphClass*>& Subgraphs);
-
-		static TUniquePtr<FFrontendGraph> CreateGraph(FBuildContext& InContext, const FMetasoundFrontendGraphClass& InSubgraph);
-
-		static const FMetasoundFrontendClassInput* FindClassInputForInputNode(const FMetasoundFrontendGraphClass& InOwningGraph, const FMetasoundFrontendNode& InInputNode, int32& OutClassInputIndex);
-		static const FMetasoundFrontendClassOutput* FindClassOutputForOutputNode(const FMetasoundFrontendGraphClass& InOwningGraph, const FMetasoundFrontendNode& InOutputNode, int32& OutClassOutputIndex);
-		static const FMetasoundFrontendLiteral* FindInputLiteralForInputNode(const FMetasoundFrontendNode& InInputNode, const FMetasoundFrontendClass& InInputNodeClass, const FMetasoundFrontendClassInput& InOwningGraphClassInput);
-		static const FMetasoundFrontendVariable* FindVariableForVariableNode(const FMetasoundFrontendNode& InVariableNode, const FMetasoundFrontendGraph& InGraph);
-
-		static TUniquePtr<INode> CreateInputNode(const FBuildContext& InContext, const FMetasoundFrontendNode& InNode, const FMetasoundFrontendClass& InClass, const FMetasoundFrontendClassInput& InOwningGraphClassInput);
-		static TUniquePtr<INode> CreateOutputNode(const FBuildContext& InContext, const FMetasoundFrontendNode& InNode, const FMetasoundFrontendClass& InClass, FBuildGraphContext& InGraphContext, const TSet<FNodeIDVertexID>& InEdgeDestinations);
-		static TUniquePtr<INode> CreateVariableNode(const FBuildContext& InContext, const FMetasoundFrontendNode& InNode, const FMetasoundFrontendGraph& InGraph);
-		static TUniquePtr<INode> CreateExternalNode(const FBuildContext& InContext, const FMetasoundFrontendNode& InNode, const FMetasoundFrontendClass& InClass, FBuildGraphContext& InGraphContext, const TSet<FNodeIDVertexID>& InEdgeDestinations);
-
-		// Returns false on error
-		static bool AddNodesToGraph(FBuildGraphContext& InGraphContext);
-
-		// Returns false on error
-		static bool AddEdgesToGraph(FBuildGraphContext& InGraphContext);
-
-		// Returns false on error
-		static bool AddDefaultInputLiterals(FBuildGraphContext& InGraphContext);
-	};
 }
diff --git a/Engine/Plugins/Runtime/Metasound/Source/MetasoundFrontend/Public/MetasoundFrontendGraphBuilder.h b/Engine/Plugins/Runtime/Metasound/Source/MetasoundFrontend/Public/MetasoundFrontendGraphBuilder.h
new file mode 100644
--- /dev/null
+++ b/Engine/Plugins/Runtime/Metasound/Source/MetasoundFrontend/Public/MetasoundFrontendGraphBuilder.h
@@ -0,0 +1,129 @@
+// Copyright Epic Games, Inc. All Rights Reserved.
+
+#pragma once
+
+
+#include "Containers/Array.h"
+#include "Containers/Map.h"
+#include "Containers/Set.h"
+#include "Containers/UnrealString.h"
+#include "MetasoundFrontendDocument.h"
+#include "MetasoundFrontendDocumentIdGenerator.h"
+#include "MetasoundGraph.h"
+#include "MetasoundNodeConstructorParams.h"
+#include "MetasoundNodeInterface.h"
+#include "MetasoundVertex.h"
+#include "Misc/Guid.h"
+#include "Templates/SharedPointer.h"
+#include "Templates/UniquePtr.h"
+
+namespace Metasound
+{
+	// Forward declare
+	class FFrontendGraph;
+}
+
+namespace Metasound::Frontend
+{
+
+	/** FGraphBuilder builds a FFrontendGraph from a FMetasoundDocument
+	 * or FMetasoundFrontendClass.
+	 */
+	class METASOUNDFRONTEND_API FGraphBuilder
+	{
+	public:
+		/** Check that all dependencies are C++ class dependencies. 
+		 * 
+		 * @param InDocument - Document containing dependencies.
+		 *
+		 * @return True if all dependencies are C++ classes. False otherwise.
+		 */
+		UE_DEPRECATED(5.5, "No longer supported (not useful nor maintainable with paged graph support")
+		static bool IsFlat(const FMetasoundFrontendDocument& InDocument);
+
+		UE_DEPRECATED(5.5, "No longer supported (not useful nor maintainable with paged graph support")
+		static bool IsFlat(const FMetasoundFrontendGraphClass& InRoot, const TArray<FMetasoundFrontendClass>& InDependencies);
+
+		/* Create a FFrontendGraph from a FMetasoundFrontendDocument.*/
+		static TUniquePtr<FFrontendGraph> CreateGraph(const FMetasoundFrontendDocument& InDocument, const FString& InDebugAssetName);
+
+		/* Create a FFrontendGraph from a FMetasoundFrontendDocument retrieving proxies from a FProxyDataCache.*/
+		static TUniquePtr<FFrontendGraph> CreateGraph(const FMetasoundFrontendDocument& InDocument, const Frontend::FProxyDataCache& InProxies, const FString& InDebugAssetName, const FGuid InGraphId = Frontend::CreateLocallyUniqueId());
+
+		/* Create a FFrontendGraph from a FMetasoundFrontendDocument and subobjects.*/
+		static TUniquePtr<FFrontendGraph> CreateGraph(const FMetasoundFrontendGraphClass& InGraph, const TArray<FMetasoundFrontendGraphClass>& InSubgraphs, const TArray<FMetasoundFrontendClass>& InDependencies, const FString& InDebugAssetName);
+
+		/* Create a FFrontendGraph from a FMetasoundFrontendDocument and subobjects, retrieving proxies from a FProxyDataCache.*/
+		static TUniquePtr<FFrontendGraph> CreateGraph(
+			const FMetasoundFrontendGraphClass& InGraph,
+			const TArray<FMetasoundFrontendGraphClass>& InSubgraphs,
+			const TArray<FMetasoundFrontendClass>& InDependencies,
+			const Frontend::FProxyDataCache& InProxyDataCache,
+			const FString& InDebugAssetName,
+			const FGuid InGraphId = Frontend::CreateLocallyUniqueId());
+
+	private:
+		struct FDefaultLiteralData
+		{
+			FGuid DestinationNodeID;
+			FGuid DestinationVertexID;
+			FVertexName DestinationVertexKey;
+			FName TypeName;
+			FLiteralNodeConstructorParams InitParams;
+		};
+
+		// Map of Input VertexID to variable data required to construct and connect default variable
+		using FNodeIDVertexID = TTuple<FGuid, FGuid>;
+		using FDependencyByIDMap = TMap<FGuid, const FMetasoundFrontendClass*>;
+		using FSharedNodeByIDMap = TMap<FGuid, TSharedPtr<const INode>>;
+		using FDefaultInputByIDMap = TMap<FNodeIDVertexID, FDefaultLiteralData>;
+
+		// Context used throughout entire graph build process
+		// (for both a root and nested subgraphs)
+		struct FBuildContext
+		{
+			FString DebugAssetName;
+			FDependencyByIDMap FrontendClasses;
+			FSharedNodeByIDMap Graphs;
+			const Frontend::IDataTypeRegistry& DataTypeRegistry;
+			const Frontend::FProxyDataCache& ProxyDataCache;
+			FGuid GraphId;
+		};
+
+		// Transient context used for building a specific graph
+		struct FBuildGraphContext
+		{
+			TUniquePtr<FFrontendGraph> Graph;
+			const FMetasoundFrontendGraphClass& GraphClass;
+			const FMetasoundFrontendGraph& PagedGraph;
+			FBuildContext& BuildContext;
+			FDefaultInputByIDMap DefaultInputs;
+		};
+
+
+		static TArray<FDefaultLiteralData> GetInputDefaultLiteralData(const FBuildContext& InContext, const FMetasoundFrontendNode& InNode, const FNodeInitData& InInitData, const TSet<FNodeIDVertexID>& InEdgeDestinations);
+
+		static bool SortSubgraphDependencies(TArray<const FMetasoundFrontendGraphClass*>& Subgraphs);
+
+		static TUniquePtr<FFrontendGraph> CreateGraph(FBuildContext& InContext, const FMetasoundFrontendGraphClass& InSubgraph);
+
+		static const FMetasoundFrontendClassInput* FindClassInputForInputNode(const FMetasoundFrontendGraphClass& InOwningGraph, const FMetasoundFrontendNode& InInputNode);
+		static const FMetasoundFrontendClassOutput* FindClassOutputForOutputNode(const FMetasoundFrontendGraphClass& InOwningGraph, const FMetasoundFrontendNode& InOutputNode);
+		static const FMetasoundFrontendLiteral* FindInputLiteralForInputNode(const FMetasoundFrontendNode& InInputNode, const FMetasoundFrontendClass& InInputNodeClass, const FMetasoundFrontendClassInput& InOwningGraphClassInput);
+		static const FMetasoundFrontendVariable* FindVariableForVariableNode(const FMetasoundFrontendNode& InVariableNode, const FMetasoundFrontendGraph& InGraph);
+
+		static TUniquePtr<INode> CreateInputNode(const FBuildContext& InContext, const FMetasoundFrontendNode& InNode, const FMetasoundFrontendClass& InClass, const FMetasoundFrontendClassInput& InOwningGraphClassInput);
+		static TUniquePtr<INode> CreateOutputNode(const FBuildContext& InContext, const FMetasoundFrontendNode& InNode, const FMetasoundFrontendClass& InClass, FBuildGraphContext& InGraphContext, const TSet<FNodeIDVertexID>& InEdgeDestinations);
+		static TUniquePtr<INode> CreateVariableNode(const FBuildContext& InContext, const FMetasoundFrontendNode& InNode, const FMetasoundFrontendGraph& InGraph);
+		static TUniquePtr<INode> CreateExternalNode(const FBuildContext& InContext, const FMetasoundFrontendNode& InNode, const FMetasoundFrontendClass& InClass, FBuildGraphContext& InGraphContext, const TSet<FNodeIDVertexID>& InEdgeDestinations);
+
+		// Returns false on error
+		static bool AddNodesToGraph(FBuildGraphContext& InGraphContext);
+
+		// Returns false on error
+		static bool AddEdgesToGraph(FBuildGraphContext& InGraphContext);
+
+		// Returns false on error
+		static bool AddDefaultInputLiterals(FBuildGraphContext& InGraphContext);
+	};
+}
diff --git a/Engine/Plugins/Runtime/Metasound/Source/MetasoundGraphCore/Private/MetasoundGraph.cpp b/Engine/Plugins/Runtime/Metasound/Source/MetasoundGraphCore/Private/MetasoundGraph.cpp
--- a/Engine/Plugins/Runtime/Metasound/Source/MetasoundGraphCore/Private/MetasoundGraph.cpp
+++ b/Engine/Plugins/Runtime/Metasound/Source/MetasoundGraphCore/Private/MetasoundGraph.cpp
@@ -179,6 +179,24 @@ namespace Metasound
 		Nodes.Add(InNodeID, InNode);
 	}
 
+	void FGraph::AddInputNode(FGuid InNodeId, const FVertexName& InVertexName, TSharedPtr<const INode> InNode)
+	{
+		if (InNode.IsValid())
+		{
+			AddInputDataDestination(*InNode, InVertexName);
+			AddNode(InNodeId, InNode);
+		}
+	}
+
+	void FGraph::AddOutputNode(FGuid InNodeID, const FVertexName& InVertexName, TSharedPtr<const INode> InNode)
+	{
+		if (InNode.IsValid())
+		{
+			AddOutputDataSource(*InNode, InVertexName);
+			AddNode(InNodeID, InNode);
+		}
+	}
+
 	const INode* FGraph::FindNode(const FGuid& InNodeID) const
 	{
 		if (const TSharedPtr<const INode>* NodePtr = Nodes.Find(InNodeID))
diff --git a/Engine/Plugins/Runtime/Metasound/Source/MetasoundGraphCore/Public/MetasoundGraph.h b/Engine/Plugins/Runtime/Metasound/Source/MetasoundGraphCore/Public/MetasoundGraph.h
--- a/Engine/Plugins/Runtime/Metasound/Source/MetasoundGraphCore/Public/MetasoundGraph.h
+++ b/Engine/Plugins/Runtime/Metasound/Source/MetasoundGraphCore/Public/MetasoundGraph.h
@@ -97,11 +97,27 @@ namespace Metasound
 
 			/** Store a node on this graph. 
 			 *
-			 * @param InNodeID - The NodeID related to the parent FMetasoundFrontendClass.
+			 * @param InNodeID - A unique ID associated with the node.
 			 * @param InNode - A shared pointer to a node. 
 			 */
 			void AddNode(const FGuid& InNodeID, TSharedPtr<const INode> InNode);
 
+			/** Add an input node to this graph.
+			 *
+			 * @param InNodeID - A unique ID associated with the node.
+			 * @param InVertexName - The key for the graph input vertex.
+			 * @param InNode - A shared pointer to an input node. 
+			 */
+			void AddInputNode(FGuid InNodeID, const FVertexName& InVertexName, TSharedPtr<const INode> InNode);
+
+			/** Add an output node to this graph.
+			 *
+			 * @param InNodeID - A unique ID associated with the node.
+			 * @param InVertexName - The key for the graph output vertex.
+			 * @param InNode - A shared pointer to an output node. 
+			 */
+			void AddOutputNode(FGuid InNodeID, const FVertexName& InVertexName, TSharedPtr<const INode> InNode);
+
 			/** Retrieve node by node ID.
 			 *
 			 * @param InNodeID - The NodeID of the requested Node.
