[new workspace] Move testing of wave writer to seperate node node derived unit tests. Break out NumberedFileCache into its own module and add tests around it.

#rnx

diff --git a/Engine/Plugins/Runtime/Metasound/Source/MetasoundStandardNodes/Private/MetasoundWaveWriterNode.cpp b/Engine/Plugins/Runtime/Metasound/Source/MetasoundStandardNodes/Private/MetasoundWaveWriterNode.cpp
--- a/Engine/Plugins/Runtime/Metasound/Source/MetasoundStandardNodes/Private/MetasoundWaveWriterNode.cpp
+++ b/Engine/Plugins/Runtime/Metasound/Source/MetasoundStandardNodes/Private/MetasoundWaveWriterNode.cpp
@@ -1,6 +1,7 @@
 // Copyright Epic Games, Inc. All Rights Reserved.
 
 #include "Audio/SimpleWaveWriter.h"
+#include "Audio/NumberedFileCache.h"
 #include "HAL/FileManager.h"
 #include "MetasoundBuildError.h"
 #include "MetasoundEnumRegistrationMacro.h"
@@ -19,6 +20,10 @@
 
 namespace Metasound
 {
+	class FNode;
+	class FBuildErrorBase;
+	class IOperator;
+
 	namespace WaveWriterVertexNames
 	{
 		METASOUND_PARAM(InEnabledPin, "Enabled", "If this wave writer is enabled or not.");
@@ -44,74 +49,17 @@ namespace Metasound
 	};
 	const FName FFileWriteError::ErrorType = FName(TEXT("MetasoundFileWriterError"));
 
-	class FNumberedFileCache
-	{
-	public:
-		static const FString Seperator;
-
-		FNumberedFileCache(const FString& InRootPath, const FString& InExt)
-			: RootPath{ InRootPath }, FileExtention{ InExt }
-		{
-			CacheFilenames();
-		}
-
-		FString GenerateNextNumberedFilename(const FString& InPrefix)
-		{
-			FScopeLock Lock{ &Cs };
-			uint32& CurrentMax = FileIndexMap.FindOrAdd(InPrefix.ToUpper());
-			FString Filename{ InPrefix };
-			Filename.Append(*Seperator);
-			Filename.AppendInt(++CurrentMax);
-			Filename.Append(*FileExtention);
-			return RootPath / Filename;
-		}
-	private:
-
-		// Slow directory search of the root path for filenames.
-		void CacheFilenames()
-		{
-			FScopeLock Lock{ &Cs };
-			
-			// Find all files, split filenames into prefix + number, saving max number we find.
-			TArray<FString> Files;
-			IFileManager::Get().FindFiles(Files , *RootPath, *FileExtention);
-			for (const FString& i : Files)
-			{
-				FString Prefix, Postfix;
-				if (i.Split(Seperator, &Prefix, &Postfix, ESearchCase::IgnoreCase, ESearchDir::FromEnd))
-				{
-					FString NumberString = FPaths::GetBaseFilename(Postfix);
-					if (FCString::IsNumeric(*NumberString))
-					{
-						int32 Number = FCString::Atoi(*NumberString);
-						if (Number >= 0)
-						{
-							uint32& CurrentMax = FileIndexMap.FindOrAdd(*Prefix.ToUpper());
-							if (static_cast<uint32>(Number) > CurrentMax)
-							{
-								CurrentMax = static_cast<uint32>(Number);
-							}
-						}
-					}
-				}
-			}
-		}
-		FCriticalSection Cs;
-		FString RootPath;
-		FString FileExtention;
-		TMap<FString, uint32> FileIndexMap;
-	};
-	const FString FNumberedFileCache::Seperator{ TEXT("_") };
-
 	namespace WaveWriterOperatorPrivate
 	{
 		// Need to keep this outside the template so there's only 1
-		TSharedPtr<FNumberedFileCache> GetNameCache()
+		TSharedPtr<Audio::FNumberedFileCache> GetNameCache()
 		{
 			static const TCHAR* WaveExt = TEXT(".wav");
 
 			// Build cache of numbered files (do this once only).
-			static TSharedPtr<FNumberedFileCache> NumberedFileCacheSP = MakeShared<FNumberedFileCache>(*FPaths::AudioCaptureDir(), WaveExt);
+			static TSharedPtr<Audio::FNumberedFileCache> NumberedFileCacheSP = MakeShared<Audio::FNumberedFileCache>(
+				*FPaths::AudioCaptureDir(), WaveExt, IFileManager::Get());
+			
 			return NumberedFileCacheSP;
 		}
 
@@ -129,7 +77,7 @@ namespace Metasound
 		// Theoretical limit of .WAV files.
 		static_assert(NumInputChannels > 0 && NumInputChannels <= 65535, "Num Channels > 0 and <= 65535");
 
-		TWaveWriterOperator(const FOperatorSettings& InSettings, TArray<FAudioBufferReadRef>&& InAudioBuffers, FBoolReadRef&& InEnabled, const TSharedPtr<FNumberedFileCache, ESPMode::ThreadSafe>& InNumberedFileCache, FStringReadRef&& InFilenamePrefix)
+		TWaveWriterOperator(const FOperatorSettings& InSettings, TArray<FAudioBufferReadRef>&& InAudioBuffers, FBoolReadRef&& InEnabled, const TSharedPtr<Audio::FNumberedFileCache, ESPMode::ThreadSafe>& InNumberedFileCache, FStringReadRef&& InFilenamePrefix)
 			: AudioInputs{ MoveTemp(InAudioBuffers) }
 			, Enabled{ MoveTemp(InEnabled) }
 			, NumberedFileCacheSP{ InNumberedFileCache }
@@ -404,7 +352,7 @@ namespace Metasound
 		TArray<float> InterleaveBuffer;
 		FBoolReadRef Enabled;
 		TUniquePtr<Audio::FSimpleWaveWriter> Writer;
-		TSharedPtr<FNumberedFileCache, ESPMode::ThreadSafe> NumberedFileCacheSP;
+		TSharedPtr<Audio::FNumberedFileCache, ESPMode::ThreadSafe> NumberedFileCacheSP;
 		FStringReadRef FileNamePrefix;
 		float SampleRate = 0.f;
 		bool bIsEnabled = false;
diff --git a/Engine/Source/Runtime/Core/Private/Audio/NumberedFileCache.cpp b/Engine/Source/Runtime/Core/Private/Audio/NumberedFileCache.cpp
new file mode 100644
--- /dev/null
+++ b/Engine/Source/Runtime/Core/Private/Audio/NumberedFileCache.cpp
@@ -0,0 +1,9 @@
+﻿// Copyright Epic Games, Inc. All Rights Reserved.
+#include "Audio/NumberedFileCache.h"
+
+namespace Audio
+{
+	const FString FNumberedFileCache::Separator{ TEXT("_") };
+} //namespace Audio
+
+
diff --git a/Engine/Source/Runtime/Core/Public/Audio/NumberedFileCache.h b/Engine/Source/Runtime/Core/Public/Audio/NumberedFileCache.h
new file mode 100644
--- /dev/null
+++ b/Engine/Source/Runtime/Core/Public/Audio/NumberedFileCache.h
@@ -0,0 +1,71 @@
+// Copyright Epic Games, Inc. All Rights Reserved.
+#pragma once
+
+#include "HAL/FileManager.h"
+#include "HAL/CriticalSection.h"
+#include "Misc/Paths.h"
+#include "Misc/ScopeLock.h"
+#include "Containers/Map.h"
+
+namespace Audio
+{
+	class FNumberedFileCache
+	{
+	public:
+		static const CORE_API FString Separator;
+
+		FNumberedFileCache(const FString& InRootPath, const FString& InExt, IFileManager& InFileSystem)
+			: RootPath{ InRootPath }, FileExtention{ InExt }, FileSystem {InFileSystem}
+		{
+			CacheFilenames();
+		}
+
+		FString GenerateNextNumberedFilename(const FString& InPrefix)
+		{
+			FScopeLock Lock{ &Cs };
+			uint32& CurrentMax = FileIndexMap.FindOrAdd(InPrefix.ToUpper());
+			FString Filename{ InPrefix };
+			Filename.Append(*Separator);
+			Filename.AppendInt(++CurrentMax);
+			Filename.Append(*FileExtention);
+			return RootPath / Filename;
+		}
+	private:
+
+		// Slow directory search of the root path for filenames.
+		void CacheFilenames()
+		{
+			FScopeLock Lock{ &Cs };
+					
+			// Find all files, split filenames into prefix + number, saving max number we find.
+			TArray<FString> Files;
+			FileSystem.FindFiles(Files , *RootPath, *FileExtention);
+			for (const FString& i : Files)
+			{
+				FString Prefix, Postfix;
+				if (i.Split(Separator, &Prefix, &Postfix, ESearchCase::IgnoreCase, ESearchDir::FromEnd))
+				{
+					FString NumberString = FPaths::GetBaseFilename(Postfix);
+					if (FCString::IsNumeric(*NumberString))
+					{
+						int32 Number = FCString::Atoi(*NumberString);
+						if (Number >= 0)
+						{
+							uint32& CurrentMax = FileIndexMap.FindOrAdd(*Prefix.ToUpper());
+							if (static_cast<uint32>(Number) > CurrentMax)
+							{
+								CurrentMax = static_cast<uint32>(Number);
+							}
+						}
+					}
+				}
+			}
+		}
+		FCriticalSection Cs;
+		FString RootPath;
+		FString FileExtention;
+		TMap<FString, uint32> FileIndexMap;
+		IFileManager& FileSystem;
+	};
+
+}//namespace Audio
\ No newline at end of file
