[Audio][MS] Add new keywords to common array MS nodes:
- "length" and "size" to the TArrayNum nodes
- "append" to the TArrayConcat nodes
- "slice" to the TArraySubset nodes

diff --git a/Engine/Plugins/Runtime/Metasound/Source/MetasoundFrontend/Public/MetasoundArrayNodes.h b/Engine/Plugins/Runtime/Metasound/Source/MetasoundFrontend/Public/MetasoundArrayNodes.h
--- a/Engine/Plugins/Runtime/Metasound/Source/MetasoundFrontend/Public/MetasoundArrayNodes.h
+++ b/Engine/Plugins/Runtime/Metasound/Source/MetasoundFrontend/Public/MetasoundArrayNodes.h
@@ -107,8 +107,13 @@ namespace Metasound
 				const FText NodeDisplayName = METASOUND_LOCTEXT_FORMAT("ArrayOpArrayNumDisplayNamePattern", "Num ({0})", GetMetasoundDataTypeDisplayText<ArrayType>());
 				const FText NodeDescription = METASOUND_LOCTEXT("ArrayOpArrayNumDescription", "Number of elements in the array");
 				const FVertexInterface NodeInterface = GetDefaultInterface();
-			
-				return MetasoundArrayNodesPrivate::CreateArrayNodeClassMetadata(DataTypeName, OperatorName, NodeDisplayName, NodeDescription, NodeInterface);
+
+				FNodeClassMetadata NodeClassMetadata = MetasoundArrayNodesPrivate::CreateArrayNodeClassMetadata(DataTypeName, OperatorName, NodeDisplayName, NodeDescription, NodeInterface);
+				NodeClassMetadata.Keywords.Append({
+					METASOUND_LOCTEXT("ArrayNumKeyword_Length", "length"),
+					METASOUND_LOCTEXT("ArrayNumKeyword_Size", "size"),
+				});
+				return NodeClassMetadata;
 			};
 			
 			static const FNodeClassMetadata Metadata = CreateNodeClassMetadata();
@@ -623,7 +628,11 @@ namespace Metasound
 				const FText NodeDescription = METASOUND_LOCTEXT("ArrayOpArrayConcatDescription", "Concatenates two arrays on trigger.");
 				const FVertexInterface NodeInterface = GetDefaultInterface();
 			
-				return MetasoundArrayNodesPrivate::CreateArrayNodeClassMetadata(DataTypeName, OperatorName, NodeDisplayName, NodeDescription, NodeInterface);
+				FNodeClassMetadata NodeClassMetadata = MetasoundArrayNodesPrivate::CreateArrayNodeClassMetadata(DataTypeName, OperatorName, NodeDisplayName, NodeDescription, NodeInterface);
+				NodeClassMetadata.Keywords.Append({
+					METASOUND_LOCTEXT("ArrayConcatKeyword_Append", "append"),
+				});
+				return NodeClassMetadata;
 			};
 			
 			static const FNodeClassMetadata Metadata = CreateNodeClassMetadata();
@@ -758,7 +767,11 @@ namespace Metasound
 				const FText NodeDescription = METASOUND_LOCTEXT("ArrayOpArraySubsetDescription", "Subset array on trigger.");
 				const FVertexInterface NodeInterface = GetDefaultInterface();
 			
-				return MetasoundArrayNodesPrivate::CreateArrayNodeClassMetadata(DataTypeName, OperatorName, NodeDisplayName, NodeDescription, NodeInterface);
+				FNodeClassMetadata NodeClassMetadata = MetasoundArrayNodesPrivate::CreateArrayNodeClassMetadata(DataTypeName, OperatorName, NodeDisplayName, NodeDescription, NodeInterface);
+				NodeClassMetadata.Keywords.Append({
+					METASOUND_LOCTEXT("ArraySubsetKeyword_Slice", "slice"),
+				});
+				return NodeClassMetadata;
 			};
 			
 			static const FNodeClassMetadata Metadata = CreateNodeClassMetadata();
