Fix being unable to un-override a MetaSound preset input

diff --git a/Engine/Plugins/Runtime/Metasound/Source/MetasoundEditor/Private/MetasoundEditorGraph.cpp b/Engine/Plugins/Runtime/Metasound/Source/MetasoundEditor/Private/MetasoundEditorGraph.cpp
--- a/Engine/Plugins/Runtime/Metasound/Source/MetasoundEditor/Private/MetasoundEditorGraph.cpp
+++ b/Engine/Plugins/Runtime/Metasound/Source/MetasoundEditor/Private/MetasoundEditorGraph.cpp
@@ -1110,17 +1110,32 @@ void UMetasoundEditorGraphInput::UpdateFrontendDefaultLiteral(bool bPostTransact
 			FMetasoundFrontendLiteral Default;
 			if (Literal->TryFindDefault(Default, InPageID))
 			{
-				Builder.SetGraphInputDefault(MemberName, Default, InPageID);
+				const FMetasoundFrontendLiteral* ExistingDefault = Builder.GetGraphInputDefault(MemberName, InPageID);
+				if (!ExistingDefault || *ExistingDefault != Default)
+				{
+					Builder.SetGraphInputDefault(MemberName, Default, InPageID);
+				}
 			}
 		}
 		else
 		{
+			bool bDefaultsModified = false; 
 			TArray<FMetasoundFrontendClassInputDefault> NewDefaults;
-			Literal->IterateDefaults([&NewDefaults](const FGuid& PageID, FMetasoundFrontendLiteral Default)
+			Literal->IterateDefaults([&NewDefaults, &MemberName, &Builder, &bDefaultsModified](const FGuid& PageID, FMetasoundFrontendLiteral Default)
 			{
+				const FMetasoundFrontendLiteral* ExistingDefault = Builder.GetGraphInputDefault(MemberName, &PageID);
+				if (!ExistingDefault || *ExistingDefault != Default)
+				{
+					bDefaultsModified |= true;
+				}
 				NewDefaults.Add(FMetasoundFrontendClassInputDefault(PageID, MoveTemp(Default)));
 			});
-			Builder.SetGraphInputDefaults(MemberName, MoveTemp(NewDefaults));
+			
+			const FMetasoundFrontendClassInput* GraphInput = Builder.FindGraphInput(MemberName);
+			if (!GraphInput || bDefaultsModified || NewDefaults.Num() != GraphInput->GetDefaults().Num())
+			{
+				Builder.SetGraphInputDefaults(MemberName, MoveTemp(NewDefaults));
+			}
 		}
 	}
 }
diff --git a/Engine/Plugins/Runtime/Metasound/Source/MetasoundEditor/Private/MetasoundNodeDetailCustomization.cpp b/Engine/Plugins/Runtime/Metasound/Source/MetasoundEditor/Private/MetasoundNodeDetailCustomization.cpp
--- a/Engine/Plugins/Runtime/Metasound/Source/MetasoundEditor/Private/MetasoundNodeDetailCustomization.cpp
+++ b/Engine/Plugins/Runtime/Metasound/Source/MetasoundEditor/Private/MetasoundNodeDetailCustomization.cpp
@@ -1553,6 +1553,7 @@ PRAGMA_ENABLE_DEPRECATION_WARNINGS
 					FScopedTransaction Transaction(LOCTEXT("SetPresetInputOverrideTransaction", "Set MetaSound Preset Input Overridden"));
 
 					Input->GetOutermost()->Modify();
+					Input->GetOutermostObject()->Modify();
 					Input->Modify();
 					MemberDefaultLiteral->Modify();
 
@@ -1563,11 +1564,8 @@ PRAGMA_ENABLE_DEPRECATION_WARNINGS
 
 					Input->UpdateFrontendDefaultLiteral(false /* bPostTransaction */);
 
-					if (UMetasoundEditorGraphMemberDefaultLiteral* Literal = Input->GetLiteral())
-					{
-						Literal->ForceRefresh();
-					}
-
+					MemberDefaultLiteral->ForceRefresh();
+					
 					if (UObject* Metasound = Input->GetOutermostObject())
 					{
 						FGraphBuilder::RegisterGraphWithFrontend(*Metasound);
@@ -1585,6 +1583,7 @@ PRAGMA_ENABLE_DEPRECATION_WARNINGS
 					FScopedTransaction Transaction(LOCTEXT("ClearPresetInputOverrideTransaction", "Clear MetaSound Preset Input Overridden"));
 
 					Input->GetOutermost()->Modify();
+					Input->GetOutermostObject()->Modify();
 					Input->Modify();
 					MemberDefaultLiteral->Modify();
 
