Rework MetaSound node registration so that nodes can be created with a provided FVertexInterface
#rnx

diff --git a/Engine/Plugins/Runtime/Metasound/Source/MetasoundFrontend/Private/MetasoundAssetBase.cpp b/Engine/Plugins/Runtime/Metasound/Source/MetasoundFrontend/Private/MetasoundAssetBase.cpp
--- a/Engine/Plugins/Runtime/Metasound/Source/MetasoundFrontend/Private/MetasoundAssetBase.cpp
+++ b/Engine/Plugins/Runtime/Metasound/Source/MetasoundFrontend/Private/MetasoundAssetBase.cpp
@@ -76,7 +76,7 @@ namespace Metasound
 			// This registry entry does not support node creation as it is only intended to be
 			// used when serializing MetaSounds in contexts not requiring any runtime model to
 			// be generated (ex. cooking commandlets that don't play or are validating MetaSounds, etc.).
-			class FDocumentNodeRegistryEntryForSerialization : public INodeRegistryEntry
+			class FDocumentNodeRegistryEntryForSerialization : public INodeClassRegistryEntry
 			{
 			public:
 				FDocumentNodeRegistryEntryForSerialization(const FMetasoundFrontendDocument& InDocument, const FTopLevelAssetPath& InAssetPath)
@@ -98,9 +98,7 @@ namespace Metasound
 				}
 
 				virtual TUniquePtr<INode> CreateNode(const FNodeInitData&) const override { return nullptr; }
-				virtual TUniquePtr<INode> CreateNode(FDefaultLiteralNodeConstructorParams&&) const override { return nullptr; }
-				virtual TUniquePtr<INode> CreateNode(FDefaultNamedVertexNodeConstructorParams&&) const override { return nullptr; }
-				virtual TUniquePtr<INode> CreateNode(FDefaultNamedVertexWithLiteralNodeConstructorParams&&) const override { return nullptr; }
+				virtual TUniquePtr<INode> CreateNode(FNodeData) const override { return nullptr; }
 
 				virtual const FMetasoundFrontendClass& GetFrontendClass() const override
 				{
@@ -443,7 +441,7 @@ void FMetasoundAssetBase::UpdateAndRegisterForSerialization(FName CookPlatformNa
 		TScriptInterface<IMetaSoundDocumentInterface> DocInterface(Owner);
 		const FMetasoundFrontendDocument& Document = DocInterface->GetConstDocument();
 		const FTopLevelAssetPath AssetPath = DocInterface->GetAssetPathChecked();
-		TUniquePtr<INodeRegistryEntry> RegistryEntry = MakeUnique<AssetBasePrivate::FDocumentNodeRegistryEntryForSerialization>(Document, AssetPath);
+		TUniquePtr<INodeClassRegistryEntry> RegistryEntry = MakeUnique<AssetBasePrivate::FDocumentNodeRegistryEntryForSerialization>(Document, AssetPath);
 
 		const FNodeClassRegistryKey NodeKey = FNodeClassRegistry::Get().RegisterNode(MoveTemp(RegistryEntry));
 		GraphRegistryKey = FGraphRegistryKey { NodeKey, AssetPath };
diff --git a/Engine/Plugins/Runtime/Metasound/Source/MetasoundFrontend/Private/MetasoundFrontendDataTypeRegistry.cpp b/Engine/Plugins/Runtime/Metasound/Source/MetasoundFrontend/Private/MetasoundFrontendDataTypeRegistry.cpp
--- a/Engine/Plugins/Runtime/Metasound/Source/MetasoundFrontend/Private/MetasoundFrontendDataTypeRegistry.cpp
+++ b/Engine/Plugins/Runtime/Metasound/Source/MetasoundFrontend/Private/MetasoundFrontendDataTypeRegistry.cpp
@@ -110,7 +110,7 @@ namespace Metasound
 			}
 
 			// Base class for INodeRegistryEntrys that come from an IDataTypeRegistryEntry
-			class FDataTypeNodeRegistryEntry : public INodeRegistryEntry
+			class FDataTypeNodeRegistryEntry : public INodeClassRegistryEntry
 			{
 			public:
 				FDataTypeNodeRegistryEntry() = default;
@@ -127,21 +127,6 @@ namespace Metasound
 					return nullptr;
 				}
 
-				virtual TUniquePtr<INode> CreateNode(FDefaultLiteralNodeConstructorParams&& InParams) const override
-				{
-					return nullptr;
-				}
-
-				virtual TUniquePtr<INode> CreateNode(FDefaultNamedVertexNodeConstructorParams&&) const override
-				{
-					return nullptr;
-				}
-
-				virtual TUniquePtr<INode> CreateNode(FDefaultNamedVertexWithLiteralNodeConstructorParams&& InParams) const override
-				{
-					return nullptr;
-				}
-
 				virtual const TSet<FMetasoundFrontendVersion>* GetImplementedInterfaces() const override
 				{
 					return nullptr;
@@ -183,6 +168,9 @@ namespace Metasound
 					return DataTypeEntry->GetFrontendInputClass();
 				}
 
+				using FDataTypeNodeRegistryEntry::CreateNode;
+
+				UE_DEPRECATED(5.6, "Node classes should be constructed with FNodeData")
 				virtual TUniquePtr<INode> CreateNode(FDefaultNamedVertexWithLiteralNodeConstructorParams&& InParams) const override
 				{
 					FInputNodeConstructorParams InputParams;
@@ -218,6 +206,9 @@ namespace Metasound
 					return DataTypeEntry->GetFrontendConstructorInputClass();
 				}
 
+				using FDataTypeNodeRegistryEntry::CreateNode;
+
+				UE_DEPRECATED(5.6, "Node classes should be constructed with FNodeData")
 				virtual TUniquePtr<INode> CreateNode(FDefaultNamedVertexWithLiteralNodeConstructorParams&& InParams) const override
 				{
 					FInputNodeConstructorParams InputParams;
@@ -253,6 +244,9 @@ namespace Metasound
 					return DataTypeEntry->GetFrontendOutputClass();
 				}
 
+				using FDataTypeNodeRegistryEntry::CreateNode;
+
+				UE_DEPRECATED(5.6, "Node classes should be constructed with FNodeData")
 				virtual TUniquePtr<INode> CreateNode(FDefaultNamedVertexNodeConstructorParams&& InParams) const override
 				{
 					return DataTypeEntry->CreateOutputNode(MoveTemp(InParams));
@@ -282,6 +276,9 @@ namespace Metasound
 					return DataTypeEntry->GetFrontendConstructorOutputClass();
 				}
 
+				using FDataTypeNodeRegistryEntry::CreateNode;
+
+				UE_DEPRECATED(5.6, "Node classes should be constructed with FNodeData")
 				virtual TUniquePtr<INode> CreateNode(FDefaultNamedVertexNodeConstructorParams&& InParams) const override
 				{
 					return DataTypeEntry->CreateConstructorOutputNode(MoveTemp(InParams));
@@ -312,6 +309,9 @@ namespace Metasound
 					return DataTypeEntry->GetFrontendLiteralClass();
 				}
 
+				using FDataTypeNodeRegistryEntry::CreateNode;
+
+				UE_DEPRECATED(5.6, "Node classes should be constructed with FNodeData")
 				virtual TUniquePtr<INode> CreateNode(FDefaultLiteralNodeConstructorParams&& InParams) const override
 				{
 					return DataTypeEntry->CreateLiteralNode(MoveTemp(InParams));
@@ -342,6 +342,9 @@ namespace Metasound
 					return DataTypeEntry->GetFrontendVariableClass();
 				}
 
+				using FDataTypeNodeRegistryEntry::CreateNode;
+
+				UE_DEPRECATED(5.6, "Node classes should be constructed with FNodeData")
 				virtual TUniquePtr<INode> CreateNode(FDefaultLiteralNodeConstructorParams&& InParams) const override
 				{
 					return DataTypeEntry->CreateVariableNode(MoveTemp(InParams));
diff --git a/Engine/Plugins/Runtime/Metasound/Source/MetasoundFrontend/Private/MetasoundFrontendDocument.cpp b/Engine/Plugins/Runtime/Metasound/Source/MetasoundFrontend/Private/MetasoundFrontendDocument.cpp
--- a/Engine/Plugins/Runtime/Metasound/Source/MetasoundFrontend/Private/MetasoundFrontendDocument.cpp
+++ b/Engine/Plugins/Runtime/Metasound/Source/MetasoundFrontend/Private/MetasoundFrontendDocument.cpp
@@ -92,6 +92,32 @@ namespace Metasound
 				{Invalid, EMetasoundFrontendClassType::Invalid}
 			};
 		}
+
+		EMetasoundFrontendVertexAccessType CoreVertexAccessTypeToFrontendVertexAccessType(Metasound::EVertexAccessType InAccessType)
+		{
+			switch (InAccessType)
+			{
+				case EVertexAccessType::Value:
+					return EMetasoundFrontendVertexAccessType::Value;
+
+				case EVertexAccessType::Reference:
+				default:
+					return EMetasoundFrontendVertexAccessType::Reference;
+			}
+		}
+
+		EVertexAccessType FrontendVertexAccessTypeToCoreVertexAccessType(EMetasoundFrontendVertexAccessType InAccessType)
+		{
+			switch (InAccessType)
+			{
+				case EMetasoundFrontendVertexAccessType::Value:
+					return EVertexAccessType::Value;
+
+				case EMetasoundFrontendVertexAccessType::Reference:
+				default:
+					return EVertexAccessType::Reference;
+			}
+		}
 	} // namespace Frontend
 
 	namespace DocumentPrivate
@@ -149,18 +175,6 @@ namespace Metasound
 			return false;
 		}
 
-		EMetasoundFrontendVertexAccessType CoreVertexAccessTypeToFrontendVertexAccessType(Metasound::EVertexAccessType InAccessType)
-		{
-			switch (InAccessType)
-			{
-				case EVertexAccessType::Value:
-					return EMetasoundFrontendVertexAccessType::Value;
-
-				case EVertexAccessType::Reference:
-				default:
-					return EMetasoundFrontendVertexAccessType::Reference;
-			}
-		}
 
 		FName ResolveMemberDataType(FName DataType, EAudioParameterType ParamType)
 		{
@@ -526,7 +540,7 @@ FMetasoundFrontendClassInterface FMetasoundFrontendClassInterface::GenerateClass
 			FMetasoundFrontendClassInput ClassInput;
 			ClassInput.Name = InputVertex.VertexName;
 			ClassInput.TypeName = InputVertex.DataTypeName;
-			ClassInput.AccessType = Metasound::DocumentPrivate::CoreVertexAccessTypeToFrontendVertexAccessType(InputVertex.AccessType);
+			ClassInput.AccessType = CoreVertexAccessTypeToFrontendVertexAccessType(InputVertex.AccessType);
 			ClassInput.VertexID = FClassIDGenerator::Get().CreateInputID(ClassInput);
 
 #if WITH_EDITOR
@@ -581,7 +595,7 @@ FMetasoundFrontendClassInterface FMetasoundFrontendClassInterface::GenerateClass
 
 			ClassOutput.Name = OutputVertex.VertexName;
 			ClassOutput.TypeName = OutputVertex.DataTypeName;
-			ClassOutput.AccessType = Metasound::DocumentPrivate::CoreVertexAccessTypeToFrontendVertexAccessType(OutputVertex.AccessType);
+			ClassOutput.AccessType = CoreVertexAccessTypeToFrontendVertexAccessType(OutputVertex.AccessType);
 			ClassOutput.VertexID = FClassIDGenerator::Get().CreateOutputID(ClassOutput);
 #if WITH_EDITOR
 			const FDataVertexMetadata& VertexMetadata = OutputVertex.Metadata;
diff --git a/Engine/Plugins/Runtime/Metasound/Source/MetasoundFrontend/Private/MetasoundFrontendNodeClassRegistry.cpp b/Engine/Plugins/Runtime/Metasound/Source/MetasoundFrontend/Private/MetasoundFrontendNodeClassRegistry.cpp
--- a/Engine/Plugins/Runtime/Metasound/Source/MetasoundFrontend/Private/MetasoundFrontendNodeClassRegistry.cpp
+++ b/Engine/Plugins/Runtime/Metasound/Source/MetasoundFrontend/Private/MetasoundFrontendNodeClassRegistry.cpp
@@ -176,16 +176,22 @@ namespace Metasound::Frontend
 				return Graph->GetVertexInterface();
 			}
 
+			UE_DEPRECATED(5.6, "Mutable vertex interfaces are no longer part of the node interface.")
 			virtual bool SetVertexInterface(const FVertexInterface& InInterface) override
 			{
+				PRAGMA_DISABLE_DEPRECATION_WARNINGS
 				// Cannot set vertex interface because of const reference to IGraph
 				// Return true if the supplied interface is the same as the existing interface.
 				return Graph->GetVertexInterface() == InInterface;
+				PRAGMA_ENABLE_DEPRECATION_WARNINGS
 			}
 
+			UE_DEPRECATED(5.6, "Mutable vertex interfaces are no longer part of the node interface.")
 			virtual bool IsVertexInterfaceSupported(const FVertexInterface& InInterface) const override
 			{
+				PRAGMA_DISABLE_DEPRECATION_WARNINGS
 				return Graph->IsVertexInterfaceSupported(InInterface);
+				PRAGMA_ENABLE_DEPRECATION_WARNINGS
 			}
 
 			virtual FOperatorFactorySharedRef GetDefaultOperatorFactory() const override
@@ -236,10 +242,6 @@ namespace Metasound::Frontend
 				}
 			}
 
-			virtual TUniquePtr<INode> CreateNode(FDefaultLiteralNodeConstructorParams&&) const override { return nullptr; }
-			virtual TUniquePtr<INode> CreateNode(FDefaultNamedVertexNodeConstructorParams&&) const override { return nullptr; }
-			virtual TUniquePtr<INode> CreateNode(FDefaultNamedVertexWithLiteralNodeConstructorParams&&) const override { return nullptr; }
-
 			virtual const FMetasoundFrontendClass& GetFrontendClass() const override
 			{
 				return FrontendClass;
@@ -262,11 +264,108 @@ namespace Metasound::Frontend
 			FNodeClassInfo ClassInfo;
 			TSharedPtr<const IGraph> Graph;
 		};
-	} // namespace RegistryPrivate
 
+		// This method is here to support back compatability for node registration.
+		// Originally, nodes did not take in a FVertexInterface on construction,
+		// but with the introduction of node configuration in 5.6, they now do take
+		// in a FVertexInterface. 
+		FDataVertexMetadata CreateVertexMetadata(const FMetasoundFrontendClassVertex& InClassVertex)
+		{
+#if WITH_EDITORONLY_DATA
+			return FDataVertexMetadata{InClassVertex.Metadata.GetDescription(), InClassVertex.Metadata.GetDisplayName(), InClassVertex.Metadata.bIsAdvancedDisplay};
+#else
+			return {};
+#endif
+		}
 
+		// This method is here to support back compatability for node registration.
+		// Originally, nodes did not take in a FVertexInterface on construction,
+		// but with the introduction of node configuration in 5.6, they now do take
+		// in a FVertexInterface. 
+		FVertexInterface CreateVertexInterfaceFromClass(const FMetasoundFrontendClass& InNodeClass)
+		{
+			FVertexInterface Interface;
 
+			FInputVertexInterface& Inputs = Interface.GetInputInterface();
+			for (const FMetasoundFrontendClassInput& ClassInput : InNodeClass.Interface.Inputs)
+			{
+				const FMetasoundFrontendLiteral* DefaultLiteral = ClassInput.FindConstDefault(::Metasound::Frontend::DefaultPageID);
+				if (DefaultLiteral)
+				{
+					EMetasoundFrontendLiteralType LiteralType = DefaultLiteral->GetType();
+					if ((EMetasoundFrontendLiteralType::UObject == LiteralType) || (EMetasoundFrontendLiteralType::UObjectArray == LiteralType))
+					{
+						UE_LOG(LogMetaSound, Warning, TEXT("Ignoring default literal set on vertex %s of node %s. Please update construct of node to use FNodeData"), *ClassInput.Name.ToString(), *InNodeClass.Metadata.GetClassName().ToString());
+						DefaultLiteral = nullptr; 
+					}
+				}
+
+				if (DefaultLiteral)
+				{
+					Inputs.Add(
+						FInputDataVertex
+						{
+							ClassInput.Name, 
+							ClassInput.TypeName, 
+							CreateVertexMetadata(ClassInput),
+							FrontendVertexAccessTypeToCoreVertexAccessType(ClassInput.AccessType),
+							DefaultLiteral->ToLiteralNoProxy()
+						});
+				}
+				else
+				{
+					Inputs.Add(
+						FInputDataVertex
+						{
+							ClassInput.Name, 
+							ClassInput.TypeName, 
+							CreateVertexMetadata(ClassInput),
+							FrontendVertexAccessTypeToCoreVertexAccessType(ClassInput.AccessType)
+						});
+				}
+			}
+
+			FOutputVertexInterface& Outputs = Interface.GetOutputInterface();
+			for (const FMetasoundFrontendClassOutput& ClassOutput : InNodeClass.Interface.Outputs)
+			{
+				Outputs.Add(
+					FOutputDataVertex
+					{
+						ClassOutput.Name, 
+						ClassOutput.TypeName, 
+						CreateVertexMetadata(ClassOutput),
+						FrontendVertexAccessTypeToCoreVertexAccessType(ClassOutput.AccessType)
+					});
+			}
+
+			FEnvironmentVertexInterface& Environments = Interface.GetEnvironmentInterface();
+			for (const FMetasoundFrontendClassEnvironmentVariable& ClassEnvironment : InNodeClass.Interface.Environment)
+			{
+				Environments.Add(FEnvironmentVertex{ ClassEnvironment.Name, FText::GetEmpty() });
+			}
+
+			return Interface;
+		}
+	} // namespace RegistryPrivate
+	  
+	TUniquePtr<INode> INodeClassRegistryEntry::CreateNode(FNodeData InNodeData) const
+	{
+#if UE_METASOUND_ENABLE_5_6_NODE_REGISTRATION_DEPRECATION_WARNINGS
+		static bool bDidLogError = false;
+
+		if (!bDidLogError)
+		{
+			bDidLogError = true;
 
+			FNodeClassInfo ClassInfo = GetClassInfo();
+			UE_LOG(LogMetaSound, Warning, TEXT("Use of deprecated code path for node registration. First occurrence on node %s (asset path: %s). Please implement INodeClassRegistryEntry::CreateNode(FNodeData) const"), *ClassInfo.ClassName.ToString(), *ClassInfo.AssetPath.ToString());
+		}
+#endif
+PRAGMA_DISABLE_DEPRECATION_WARNINGS
+		// Use old path for node creation. 
+		return CreateNode(FNodeInitData{InNodeData.Name, InNodeData.ID});
+PRAGMA_ENABLE_DEPRECATION_WARNINGS
+	}
 
 	FNodeClassRegistryTransaction::FNodeClassRegistryTransaction(ETransactionType InType, const FNodeClassInfo& InNodeClassInfo, FNodeClassRegistryTransaction::FTimeType InTimestamp)
 	: Type(InType)
@@ -476,9 +575,17 @@ namespace Metasound::Frontend
 	TUniquePtr<Metasound::INode> FNodeClassRegistry::CreateNode(const FNodeClassRegistryKey& InKey, const Metasound::FNodeInitData& InInitData) const
 	{
 		TUniquePtr<INode> Node;
-		auto CreateNodeLambda = [&Node, &InInitData](const INodeClassRegistryEntry& Entry)
+		auto CreateNodeLambda = [&Node, &InInitData](const INodeClassRegistryEntry& Entry) mutable
 		{ 
-			Node = Entry.CreateNode(InInitData); 
+			const FMetasoundFrontendClass& FrontendClass = Entry.GetFrontendClass();
+
+			FNodeData NodeData
+			{
+				InInitData.InstanceName,
+				InInitData.InstanceID,
+				RegistryPrivate::CreateVertexInterfaceFromClass(FrontendClass) // This is an attempt to maintain backwards compatibility. It creates an FVertexInterface from the class definition because none was provided. 
+			};
+			Node = Entry.CreateNode(MoveTemp(NodeData)); 
 		};
 
 		if (!AccessNodeEntryThreadSafe(InKey, CreateNodeLambda))
@@ -491,48 +598,18 @@ namespace Metasound::Frontend
 		return MoveTemp(Node);
 	}
 
-	TUniquePtr<Metasound::INode> FNodeClassRegistry::CreateNode(const FNodeClassRegistryKey& InKey, FDefaultLiteralNodeConstructorParams&& InParams) const
+	TUniquePtr<Metasound::INode> FNodeClassRegistry::CreateNode(const FNodeClassRegistryKey& InKey, Metasound::FNodeData InNodeData) const
 	{
 		TUniquePtr<INode> Node;
-		auto CreateNodeLambda = [&Node, &InParams](const INodeClassRegistryEntry& Entry)
+		auto CreateNodeLambda = [&Node, &InNodeData](const INodeClassRegistryEntry& Entry) mutable
 		{ 
-			Node = Entry.CreateNode(MoveTemp(InParams)); 
-		};
-
-		if (!AccessNodeEntryThreadSafe(InKey, CreateNodeLambda))
-		{
-			UE_LOG(LogMetaSound, Error, TEXT("Could not find node [RegistryKey:%s]"), *InKey.ToString());
-		}
-
-		return MoveTemp(Node);
-	}
-
-	TUniquePtr<Metasound::INode> FNodeClassRegistry::CreateNode(const FNodeClassRegistryKey& InKey, FDefaultNamedVertexNodeConstructorParams&& InParams) const
-	{
-		TUniquePtr<INode> Node;
-		auto CreateNodeLambda = [&Node, &InParams](const INodeClassRegistryEntry& Entry) 
-		{ 
-			Node = Entry.CreateNode(MoveTemp(InParams)); 
-		};
-
-		if (!AccessNodeEntryThreadSafe(InKey, CreateNodeLambda))
-		{
-			UE_LOG(LogMetaSound, Error, TEXT("Could not find node [RegistryKey:%s]"), *InKey.ToString());
-		}
-
-		return MoveTemp(Node);
-	}
-
-	TUniquePtr<Metasound::INode> FNodeClassRegistry::CreateNode(const FNodeClassRegistryKey& InKey, FDefaultNamedVertexWithLiteralNodeConstructorParams&& InParams) const
-	{
-		TUniquePtr<INode> Node;
-		auto CreateNodeLambda = [&Node, &InParams](const INodeClassRegistryEntry& Entry)
-		{ 
-			Node = Entry.CreateNode(MoveTemp(InParams));
+			Node = Entry.CreateNode(MoveTemp(InNodeData)); 
 		};
 
 		if (!AccessNodeEntryThreadSafe(InKey, CreateNodeLambda))
 		{
+			// Creation of external nodes can rely on assets being unavailable due to errors in loading order, asset(s)
+			// missing, etc. 
 			UE_LOG(LogMetaSound, Error, TEXT("Could not find node [RegistryKey:%s]"), *InKey.ToString());
 		}
 
@@ -573,7 +650,7 @@ namespace Metasound::Frontend
 		if (!RegistryKey.IsValid())
 		{
 			// Do not attempt to build and register a MetaSound with an invalid registry key
-			UE_LOG(LogMetaSound, Warning, TEXT("Registry key is invalid when attemping to register graph for asset %s"), *AssetPath.ToString());
+			UE_LOG(LogMetaSound, Warning, TEXT("Registry key is invalid when attempting to register graph for asset %s"), *AssetPath.ToString());
 			return RegistryKey;
 		}
 
diff --git a/Engine/Plugins/Runtime/Metasound/Source/MetasoundFrontend/Private/MetasoundFrontendNodeClassRegistryPrivate.h b/Engine/Plugins/Runtime/Metasound/Source/MetasoundFrontend/Private/MetasoundFrontendNodeClassRegistryPrivate.h
--- a/Engine/Plugins/Runtime/Metasound/Source/MetasoundFrontend/Private/MetasoundFrontendNodeClassRegistryPrivate.h
+++ b/Engine/Plugins/Runtime/Metasound/Source/MetasoundFrontend/Private/MetasoundFrontendNodeClassRegistryPrivate.h
@@ -105,10 +105,13 @@ namespace Metasound::Frontend
 		virtual bool FindOutputNodeRegistryKeyForDataType(const FName& InDataTypeName, const EMetasoundFrontendVertexAccessType InAccessType, FNodeClassRegistryKey& OutKey) override;
 
 		// Create a new instance of a C++ implemented node from the registry.
+#if UE_METASOUND_ENABLE_5_6_NODE_REGISTRATION_DEPRECATION_WARNINGS
+		UE_DEPRECATED(5.6, "Node classes should be constructed with FNodeData")
+#endif
 		virtual TUniquePtr<Metasound::INode> CreateNode(const FNodeClassRegistryKey& InKey, const Metasound::FNodeInitData& InInitData) const override;
-		virtual TUniquePtr<INode> CreateNode(const FNodeClassRegistryKey& InKey, FDefaultLiteralNodeConstructorParams&&) const override;
-		virtual TUniquePtr<INode> CreateNode(const FNodeClassRegistryKey& InKey, FDefaultNamedVertexNodeConstructorParams&&) const override;
-		virtual TUniquePtr<INode> CreateNode(const FNodeClassRegistryKey& InKey, FDefaultNamedVertexWithLiteralNodeConstructorParams&&) const override;
+
+		// Create a new instance of a C++ implemented node from the registry.
+		virtual TUniquePtr<Metasound::INode> CreateNode(const FNodeClassRegistryKey& InKey, Metasound::FNodeData InNodeData) const override;
 
 		// Returns a list of possible nodes to use to convert from FromDataType to ToDataType.
 		// Returns an empty array if none are available.
diff --git a/Engine/Plugins/Runtime/Metasound/Source/MetasoundFrontend/Private/MetasoundNodeRegistrationMacro.cpp b/Engine/Plugins/Runtime/Metasound/Source/MetasoundFrontend/Private/MetasoundNodeRegistrationMacro.cpp
new file mode 100644
--- /dev/null
+++ b/Engine/Plugins/Runtime/Metasound/Source/MetasoundFrontend/Private/MetasoundNodeRegistrationMacro.cpp
@@ -0,0 +1,94 @@
+// Copyright Epic Games, Inc. All Rights Reserved.
+
+#include "MetasoundNodeRegistrationMacro.h"
+
+#include "MetasoundFrontendDocument.h"
+#include "MetasoundNodeInterface.h"
+#include "MetasoundFrontend.h"
+
+namespace Metasound::Frontend
+{
+	namespace NodeRegistrationPrivate
+	{
+		FNodeRegistryEntryBase::FNodeRegistryEntryBase(const FNodeClassMetadata& InMetadata)
+		: ClassMetadata(MakeShared<FNodeClassMetadata>(InMetadata))
+		, ClassInfo(FMetasoundFrontendClassMetadata::GenerateClassMetadata(InMetadata, EMetasoundFrontendClassType::External))
+		, FrontendClass(GenerateClass(InMetadata))
+		{
+		}
+
+		const FNodeClassInfo& FNodeRegistryEntryBase::GetClassInfo() const
+		{
+			return ClassInfo;
+		}
+
+		const FMetasoundFrontendClass& FNodeRegistryEntryBase::GetFrontendClass() const
+		{
+			return FrontendClass;
+		}
+
+		TUniquePtr<INode> FNodeRegistryEntryBase::CreateNode(const FNodeInitData& InInitData) const 
+		{
+			return CreateNode(FNodeData{InInitData.InstanceName, InInitData.InstanceID, ClassMetadata->DefaultInterface});
+		}
+
+		const TSet<FMetasoundFrontendVersion>* FNodeRegistryEntryBase::GetImplementedInterfaces() const
+		{
+			return nullptr;
+		}
+
+		bool FNodeRegistryEntryBase::IsNative() const
+		{
+			return true;
+		}
+
+		TSharedRef<const FNodeClassMetadata> FNodeRegistryEntryBase::GetNodeClassMetadata() const
+		{
+			return ClassMetadata;
+		}
+
+		// This adapter class forwards the correct FBuilderOperatorParams
+		// to the node's operator creation method. Many operator creation
+		// methods downcast the supplied INode in `FBuilderOperatorParams`
+		// and so it is required that it point to the correct runtime instance
+		// when calling CreateOperator(...)
+		class FDeprecatedNodeAPIFactory : public IOperatorFactory 
+		{
+		public:
+			FDeprecatedNodeAPIFactory(TUniquePtr<INode> InNode)
+			: Node(MoveTemp(InNode))
+			{
+				check(Node);
+			}
+
+			virtual ~FDeprecatedNodeAPIFactory() = default;
+
+			virtual TUniquePtr<IOperator> CreateOperator(const FBuildOperatorParams& InParams, FBuildResults& OutResults) override
+			{
+				FBuildOperatorParams ForwardParams 
+				{
+					*Node,  // Point to correct INode instance
+					InParams.OperatorSettings,
+					InParams.InputData,
+					InParams.Environment,
+					InParams.Builder,
+					InParams.GraphRenderCost
+				};
+				return Node->GetDefaultOperatorFactory()->CreateOperator(ForwardParams, OutResults);
+			}
+		private:
+			TUniquePtr<INode> Node;
+		};
+
+		FDeprecatedNodeAPIAdapterBase::FDeprecatedNodeAPIAdapterBase(FNodeData InNodeData, TSharedRef<const FNodeClassMetadata> InClassMetadata, TUniquePtr<INode> InNode)
+		: FBasicNode(MoveTemp(InNodeData), MoveTemp(InClassMetadata))
+		, Factory(MakeShared<FDeprecatedNodeAPIFactory>(MoveTemp(InNode)))
+		{
+		}
+
+		FOperatorFactorySharedRef FDeprecatedNodeAPIAdapterBase::GetDefaultOperatorFactory() const 
+		{
+			return Factory;
+		}
+	}
+}
diff --git a/Engine/Plugins/Runtime/Metasound/Source/MetasoundFrontend/Public/MetasoundArrayNodesRegistration.h b/Engine/Plugins/Runtime/Metasound/Source/MetasoundFrontend/Public/MetasoundArrayNodesRegistration.h
--- a/Engine/Plugins/Runtime/Metasound/Source/MetasoundFrontend/Public/MetasoundArrayNodesRegistration.h
+++ b/Engine/Plugins/Runtime/Metasound/Source/MetasoundFrontend/Public/MetasoundArrayNodesRegistration.h
@@ -64,7 +64,7 @@ namespace Metasound
 		bool RegisterArrayGetNode()
 		{
 			using FNodeType = typename Metasound::TArrayGetNode<ArrayType>;
-			return RegisterNodeWithFrontend<FNodeType>();
+			return Frontend::RegisterNode<FNodeType>();
 		}
 
 		template<typename ArrayType, typename std::enable_if<!TArrayNodeSupport<ArrayType>::bIsArrayGetSupported, bool>::type = true>
@@ -81,7 +81,7 @@ namespace Metasound
 
 			static_assert(TArrayNodeSupport<ArrayType>::bIsArraySetSupported, "TArraySetNode<> is not supported by array type");
 
-			return RegisterNodeWithFrontend<FNodeType>();
+			return Frontend::RegisterNode<FNodeType>();
 		}
 
 		template<typename ArrayType, typename std::enable_if<!TArrayNodeSupport<ArrayType>::bIsArraySetSupported, bool>::type = true>
@@ -98,7 +98,7 @@ namespace Metasound
 
 			static_assert(TArrayNodeSupport<ArrayType>::bIsArraySubsetSupported, "TArraySubsetNode<> is not supported by array type");
 
-			return RegisterNodeWithFrontend<FNodeType>();
+			return Frontend::RegisterNode<FNodeType>();
 		}
 
 		template<typename ArrayType, typename std::enable_if<!TArrayNodeSupport<ArrayType>::bIsArraySubsetSupported, bool>::type = true>
@@ -115,7 +115,7 @@ namespace Metasound
 
 			static_assert(TArrayNodeSupport<ArrayType>::bIsArrayConcatSupported, "TArrayConcatNode<> is not supported by array type");
 
-			return RegisterNodeWithFrontend<FNodeType>();
+			return Frontend::RegisterNode<FNodeType>();
 		}
 
 		template<typename ArrayType, typename std::enable_if<!TArrayNodeSupport<ArrayType>::bIsArrayConcatSupported, bool>::type = true>
@@ -128,7 +128,7 @@ namespace Metasound
 		template<typename ArrayType, typename std::enable_if<TArrayNodeSupport<ArrayType>::bIsArrayNumSupported, bool>::type = true>
 		bool RegisterArrayNumNode()
 		{
-			return ensureAlways(RegisterNodeWithFrontend<Metasound::TArrayNumNode<ArrayType>>());
+			return Frontend::RegisterNode<Metasound::TArrayNumNode<ArrayType>>();
 		}
 
 		template<typename ArrayType, typename std::enable_if<!TArrayNodeSupport<ArrayType>::bIsArrayNumSupported, bool>::type = true>
@@ -142,7 +142,7 @@ namespace Metasound
 		bool RegisterArrayShuffleNode()
 		{
 			using FNodeType = typename Metasound::TArrayShuffleNode<ArrayType>;
-			return RegisterNodeWithFrontend<FNodeType>();
+			return Frontend::RegisterNode<FNodeType>();
 		}
 
 		template<typename ArrayType, typename std::enable_if<!TArrayNodeSupport<ArrayType>::bIsArrayShuffleSupported, bool>::type = true>
@@ -156,7 +156,7 @@ namespace Metasound
 		bool RegisterArrayRandomGetNode()
 		{
 			using FNodeType = typename Metasound::TArrayRandomGetNode<ArrayType>;
-			return RegisterNodeWithFrontend<FNodeType>();
+			return Frontend::RegisterNode<FNodeType>();
 		}
 
 		template<typename ArrayType, typename std::enable_if<!TArrayNodeSupport<ArrayType>::bIsArrayRandomGetSupported, bool>::type = true>
@@ -169,7 +169,7 @@ namespace Metasound
 		template<typename ArrayType, typename std::enable_if<TArrayNodeSupport<ArrayType>::bIsArrayNumSupported, bool>::type = true>
 		bool RegisterArrayLastIndexNode()
 		{
-			return ensureAlways(RegisterNodeWithFrontend<Metasound::TArrayLastIndexNode<ArrayType>>());
+			return Frontend::RegisterNode<Metasound::TArrayLastIndexNode<ArrayType>>();
 		}
 
 		template<typename ArrayType, typename std::enable_if<!TArrayNodeSupport<ArrayType>::bIsArrayNumSupported, bool>::type = true>
@@ -180,23 +180,34 @@ namespace Metasound
 		}
 	}
 
-	/** Registers all available array nodes which can be instantiated for the given
-	 * ArrayType. Some nodes cannot be instantiated due to limitations of the 
-	 * array elements.
-	 */
+	namespace Frontend
+	{
+		/** Registers all available array nodes which can be instantiated for the given
+		 * ArrayType. Some nodes cannot be instantiated due to limitations of the 
+		 * array elements.
+		 */
+		template<typename ArrayType>
+		bool RegisterArrayNodes()
+		{
+			using namespace MetasoundArrayNodesPrivate;
+
+			bool bSuccess = RegisterArrayNumNode<ArrayType>();
+			bSuccess = bSuccess && RegisterArrayGetNode<ArrayType>();
+			bSuccess = bSuccess && RegisterArraySetNode<ArrayType>();
+			bSuccess = bSuccess && RegisterArraySubsetNode<ArrayType>();
+			bSuccess = bSuccess && RegisterArrayConcatNode<ArrayType>();
+			bSuccess = bSuccess && RegisterArrayShuffleNode<ArrayType>();
+			bSuccess = bSuccess && RegisterArrayRandomGetNode<ArrayType>();
+			bSuccess = bSuccess && RegisterArrayLastIndexNode<ArrayType>();
+			return bSuccess;
+		}
+	}
+
 	template<typename ArrayType>
+	UE_DEPRECATED(5.6, "Use Frontend::RegisterArrayNodes() instead")
 	bool RegisterArrayNodes()
 	{
-		using namespace MetasoundArrayNodesPrivate;
-
-		bool bSuccess = RegisterArrayNumNode<ArrayType>();
-		bSuccess = bSuccess && RegisterArrayGetNode<ArrayType>();
-		bSuccess = bSuccess && RegisterArraySetNode<ArrayType>();
-		bSuccess = bSuccess && RegisterArraySubsetNode<ArrayType>();
-		bSuccess = bSuccess && RegisterArrayConcatNode<ArrayType>();
-		bSuccess = bSuccess && RegisterArrayShuffleNode<ArrayType>();
-		bSuccess = bSuccess && RegisterArrayRandomGetNode<ArrayType>();
-		bSuccess = bSuccess && RegisterArrayLastIndexNode<ArrayType>();
-		return bSuccess;
+		return Frontend::RegisterArrayNodes<ArrayType>();
 	}
 }
+
diff --git a/Engine/Plugins/Runtime/Metasound/Source/MetasoundFrontend/Public/MetasoundDataTypeRegistrationMacro.h b/Engine/Plugins/Runtime/Metasound/Source/MetasoundFrontend/Public/MetasoundDataTypeRegistrationMacro.h
--- a/Engine/Plugins/Runtime/Metasound/Source/MetasoundFrontend/Public/MetasoundDataTypeRegistrationMacro.h
+++ b/Engine/Plugins/Runtime/Metasound/Source/MetasoundFrontend/Public/MetasoundDataTypeRegistrationMacro.h
@@ -2,7 +2,6 @@
 
 #pragma once
 
-
 #include "IAudioProxyInitializer.h"
 #include "MetasoundArrayNodesRegistration.h"
 #include "MetasoundAutoConverterNode.h"
@@ -30,7 +29,7 @@
 
 #include <type_traits>
 
-namespace Metasound
+namespace Metasound::Frontend
 {
 	namespace MetasoundDataTypeRegistrationPrivate
 	{
@@ -121,8 +120,8 @@ namespace Metasound
 		{
 			if (TEnableTransmissionNodeRegistration<TDataType>::Value)
 			{
-				ensureAlways(RegisterNodeWithFrontend<Metasound::TSendNode<TDataType>>());
-				ensureAlways(RegisterNodeWithFrontend<Metasound::TReceiveNode<TDataType>>());
+				ensureAlways(RegisterNode<Metasound::TSendNode<TDataType>>());
+				ensureAlways(RegisterNode<Metasound::TReceiveNode<TDataType>>());
 			}
 		}
 
@@ -147,7 +146,7 @@ namespace Metasound
 
 				if (!std::is_same<TFromDataType, TToDataType>::value && !FMetasoundFrontendRegistryContainer::Get()->IsNodeRegistered(Key))
 				{
-					ensureAlways(RegisterNodeWithFrontend<FConverterNode>(Metadata));
+					ensureAlways(RegisterNode<FConverterNode>(Metadata));
 					
 					bool bSucessfullyRegisteredConversionNode = RegisterConversionNode<FConverterNode, TFromDataType, TToDataType>(FConverterNode::GetInputName(), FConverterNode::GetOutputName(), Metadata);
 					ensureAlways(bSucessfullyRegisteredConversionNode);
@@ -768,7 +767,7 @@ namespace Metasound
 	 * @return True on success, false on failure.
 	 */
 	template<typename TDataType, ELiteralType PreferredArgType = ELiteralType::None, typename UClassToUse = void>
-	bool RegisterDataTypeWithFrontend()
+	bool RegisterDataType()
 	{
 		using namespace MetasoundDataTypeRegistrationPrivate;
 
@@ -803,6 +802,19 @@ namespace Metasound
 	};
 }
 
+namespace Metasound
+{
+	template<typename DataType>
+	using TMetasoundDataTypeRegistration UE_DEPRECATED(5.6, "Use Frontend::TMetasoundDataTypeRegistration") = Frontend::TMetasoundDataTypeRegistration<DataType>;
+
+	template<typename TDataType, ELiteralType PreferredArgType = ELiteralType::None, typename UClassToUse = void>
+	UE_DEPRECATED(5.6, "Use Frontend::RegisterDataType instead")
+	bool RegisterDataTypeWithFrontend()
+	{
+		return Frontend::RegisterDataType<TDataType, PreferredArgType, UClassToUse>();
+	}
+}
+
 
 #define CANNOT_REGISTER_METASOUND_DATA_TYPE_ASSERT_STRING(DataType) \
 "To register " #DataType " to be used as a Metasounds input or output type, it needs a default constructor or one of the following constructors must be implemented:  " \
@@ -822,8 +834,8 @@ namespace Metasound
 #DataType "(const ::Metasound::FOperatorSettings& InSettings, const TArray<Audio::IProxyDataPtr>& InProxyArray)."
 
 #define ENQUEUE_METASOUND_DATATYPE_REGISTRATION_COMMAND(DataType, DataTypeName, ...) \
-	static_assert(::Metasound::TMetasoundDataTypeRegistration<DataType>::bCanRegister, CANNOT_REGISTER_METASOUND_DATA_TYPE_ASSERT_STRING(DataType)); \
-	template<> const bool ::Metasound::TMetasoundDataTypeRegistration<DataType>::bSuccessfullyRegistered = ::FMetasoundFrontendRegistryContainer::Get()->EnqueueInitCommand([](){ ::Metasound::RegisterDataTypeWithFrontend<DataType, ##__VA_ARGS__>(); }); // This static bool is useful for debugging, but also is the only way the compiler will let us call this function outside of an expression.
+	static_assert(::Metasound::Frontend::TMetasoundDataTypeRegistration<DataType>::bCanRegister, CANNOT_REGISTER_METASOUND_DATA_TYPE_ASSERT_STRING(DataType)); \
+	template<> const bool ::Metasound::Frontend::TMetasoundDataTypeRegistration<DataType>::bSuccessfullyRegistered = ::Metasound::Frontend::INodeClassRegistry::Get()->EnqueueInitCommand([](){ ::Metasound::Frontend::RegisterDataType<DataType, ##__VA_ARGS__>(); }); // This static bool is useful for debugging, but also is the only way the compiler will let us call this function outside of an expression.
 
 // This should be used to expose a datatype as a potential input or output for a metasound graph.
 // The first argument to the macro is the class to expose.
diff --git a/Engine/Plugins/Runtime/Metasound/Source/MetasoundFrontend/Public/MetasoundFrontendDocument.h b/Engine/Plugins/Runtime/Metasound/Source/MetasoundFrontend/Public/MetasoundFrontendDocument.h
--- a/Engine/Plugins/Runtime/Metasound/Source/MetasoundFrontend/Public/MetasoundFrontendDocument.h
+++ b/Engine/Plugins/Runtime/Metasound/Source/MetasoundFrontend/Public/MetasoundFrontendDocument.h
@@ -1895,6 +1895,12 @@ METASOUNDFRONTEND_API const TCHAR* LexToString(EMetasoundFrontendVertexAccessTyp
 
 namespace Metasound::Frontend
 {
+	// Convert access type between enums
+	METASOUNDFRONTEND_API EMetasoundFrontendVertexAccessType CoreVertexAccessTypeToFrontendVertexAccessType(Metasound::EVertexAccessType InAccessType);
+
+	// Convert access type between enums
+	METASOUNDFRONTEND_API EVertexAccessType FrontendVertexAccessTypeToCoreVertexAccessType(EMetasoundFrontendVertexAccessType InAccessType);
+
 	METASOUNDFRONTEND_API bool StringToClassType(const FString& InString, EMetasoundFrontendClassType& OutClassType);
 
 	/** Signature of function called for each found literal. */
diff --git a/Engine/Plugins/Runtime/Metasound/Source/MetasoundFrontend/Public/MetasoundFrontendNodeClassRegistry.h b/Engine/Plugins/Runtime/Metasound/Source/MetasoundFrontend/Public/MetasoundFrontendNodeClassRegistry.h
--- a/Engine/Plugins/Runtime/Metasound/Source/MetasoundFrontend/Public/MetasoundFrontendNodeClassRegistry.h
+++ b/Engine/Plugins/Runtime/Metasound/Source/MetasoundFrontend/Public/MetasoundFrontendNodeClassRegistry.h
@@ -8,6 +8,13 @@
 #include "MetasoundVertex.h"
 #include "Templates/Function.h"
 
+// In UE 5.6, registered node are expected to support the constructor signature Constructor(FNodeData, TSharedRef<const FNodeClassMetadata>)
+// Because there are many existing nodes, it may take time to update them. For convenience, the deprecations related to this change are
+// configurable via a preprocessor macro so that the deprecation warnings do not drown out other compiler errors and warnings.
+#ifndef UE_METASOUND_ENABLE_5_6_NODE_REGISTRATION_DEPRECATION_WARNINGS
+#define UE_METASOUND_ENABLE_5_6_NODE_REGISTRATION_DEPRECATION_WARNINGS (0)
+#endif
+
 // Forward declare
 class UObject;
 
@@ -47,29 +54,42 @@ namespace Metasound::Frontend
 			* If a node can be created with FDefaultNamedVertexNodeConstructorParams, this function
 			* should return a valid node pointer.
 			*/
-		virtual TUniquePtr<INode> CreateNode(FDefaultNamedVertexNodeConstructorParams&&) const = 0;
+		UE_DEPRECATED(5.6, "Node classes should be constructed with FNodeData")
+		virtual TUniquePtr<INode> CreateNode(FDefaultNamedVertexNodeConstructorParams&&) const { return nullptr; }
 
 		/** Create a node given FDefaultNamedVertexWithLiteralNodeConstructorParams.
 			*
 			* If a node can be created with FDefaultNamedVertexWithLiteralNodeConstructorParams, this function
 			* should return a valid node pointer.
 			*/
-		virtual TUniquePtr<INode> CreateNode(FDefaultNamedVertexWithLiteralNodeConstructorParams&&) const = 0;
+		UE_DEPRECATED(5.6, "Node classes should be constructed with FNodeData")
+		virtual TUniquePtr<INode> CreateNode(FDefaultNamedVertexWithLiteralNodeConstructorParams&&) const { return nullptr; }
 
 		/** Create a node given FDefaultLiteralNodeConstructorParams.
 			*
 			* If a node can be created with FDefaultLiteralNodeConstructorParams, this function
 			* should return a valid node pointer.
 			*/
-		virtual TUniquePtr<INode> CreateNode(FDefaultLiteralNodeConstructorParams&&) const = 0;
+		UE_DEPRECATED(5.6, "Node classes should be constructed with FNodeData")
+		virtual TUniquePtr<INode> CreateNode(FDefaultLiteralNodeConstructorParams&&) const { return nullptr; }
 
 		/** Create a node given FNodeInitData.
 			*
 			* If a node can be created with FNodeInitData, this function
 			* should return a valid node pointer.
 			*/
+#if UE_METASOUND_ENABLE_5_6_NODE_REGISTRATION_DEPRECATION_WARNINGS
+		UE_DEPRECATED(5.6, "Node classes should be constructed with FNodeData")
+#endif
 		virtual TUniquePtr<INode> CreateNode(const FNodeInitData&) const = 0;
 
+		/** Create a node given FNodeData.
+			*
+			* If a node can be created with FNodeData, this function
+			* should return a valid node pointer.
+			*/
+		virtual TUniquePtr<INode> CreateNode(FNodeData) const;
+
 		/** Return a FMetasoundFrontendClass which describes the node. */
 		virtual const FMetasoundFrontendClass& GetFrontendClass() const = 0;
 
@@ -86,7 +106,8 @@ namespace Metasound::Frontend
 		/** Whether or not the node is natively defined */
 		virtual bool IsNative() const = 0;
 	};
-	using INodeRegistryEntry = INodeClassRegistryEntry;
+
+	using INodeRegistryEntry UE_DEPRECATED(5.6, "Use INodeClassRegistryEntry instead") = INodeClassRegistryEntry;
 
 	struct METASOUNDFRONTEND_API FConverterNodeClassRegistryKey
 	{
@@ -303,10 +324,30 @@ namespace Metasound::Frontend
 		virtual bool FindVariableNodeRegistryKeyForDataType(const FName& InDataTypeName, FNodeClassRegistryKey& OutKey) = 0;
 		virtual bool FindOutputNodeRegistryKeyForDataType(const FName& InDataTypeName, const EMetasoundFrontendVertexAccessType InAccessType, FNodeClassRegistryKey& OutKey) = 0;
 
+		/** Create a MetaSound Node with the given registration key and Node Data 
+		 * 
+		 * Returns an invalid TUniquePtr<> if the key is not in the registry. 
+		 */
+		virtual TUniquePtr<Metasound::INode> CreateNode(const FNodeClassRegistryKey& InKey, Metasound::FNodeData) const = 0;
+
+		/** Create a MetaSound Node with the given registration key and Init Data.  
+		 * The node will be created with a default interface.
+		 * 
+		 * Returns an invalid TUniquePtr<> if the key is not in the registry. 
+		 */
+#if UE_METASOUND_ENABLE_5_6_NODE_REGISTRATION_DEPRECATION_WARNINGS
+		UE_DEPRECATED(5.6, "Node classes should be constructed with FNodeData")
+#endif
 		virtual TUniquePtr<Metasound::INode> CreateNode(const FNodeClassRegistryKey& InKey, const Metasound::FNodeInitData&) const = 0;
-		virtual TUniquePtr<Metasound::INode> CreateNode(const FNodeClassRegistryKey& InKey, Metasound::FDefaultLiteralNodeConstructorParams&&) const = 0;
-		virtual TUniquePtr<Metasound::INode> CreateNode(const FNodeClassRegistryKey& InKey, Metasound::FDefaultNamedVertexNodeConstructorParams&&) const = 0;
-		virtual TUniquePtr<Metasound::INode> CreateNode(const FNodeClassRegistryKey& InKey, Metasound::FDefaultNamedVertexWithLiteralNodeConstructorParams&&) const = 0;
+
+		UE_DEPRECATED(5.6, "Node classes should be constructed with FNodeData")
+		virtual TUniquePtr<Metasound::INode> CreateNode(const FNodeClassRegistryKey& InKey, Metasound::FDefaultLiteralNodeConstructorParams&&) const { return nullptr; }
+
+		UE_DEPRECATED(5.6, "Node classes should be constructed with FNodeData")
+		virtual TUniquePtr<Metasound::INode> CreateNode(const FNodeClassRegistryKey& InKey, Metasound::FDefaultNamedVertexNodeConstructorParams&&) const { return nullptr; }
+
+		UE_DEPRECATED(5.6, "Node classes should be constructed with FNodeData")
+		virtual TUniquePtr<Metasound::INode> CreateNode(const FNodeClassRegistryKey& InKey, Metasound::FDefaultNamedVertexWithLiteralNodeConstructorParams&&) const { return nullptr; }
 
 		virtual bool RegisterConversionNode(const FConverterNodeClassRegistryKey& InNodeKey, const FConverterNodeClassInfo& InNodeInfo) = 0;
 
diff --git a/Engine/Plugins/Runtime/Metasound/Source/MetasoundFrontend/Public/MetasoundNodeRegistrationMacro.h b/Engine/Plugins/Runtime/Metasound/Source/MetasoundFrontend/Public/MetasoundNodeRegistrationMacro.h
--- a/Engine/Plugins/Runtime/Metasound/Source/MetasoundFrontend/Public/MetasoundNodeRegistrationMacro.h
+++ b/Engine/Plugins/Runtime/Metasound/Source/MetasoundFrontend/Public/MetasoundNodeRegistrationMacro.h
@@ -3,137 +3,228 @@
 #pragma once
 
 #include "CoreMinimal.h"
+#include "MetasoundBasicNode.h"
 #include "MetasoundFrontend.h"
 #include "MetasoundFrontendDocument.h"
 #include "MetasoundFrontendNodeClassRegistry.h"
 #include "MetasoundLog.h"
 #include "MetasoundNodeInterface.h"
 
-namespace Metasound
+#include "Traits/MetasoundNodeConstructorTraits.h"
+#include "Traits/MetasoundNodeStaticMemberTraits.h"
+
+namespace Metasound::Frontend
 {
-	// Utility class to ensure that a node class can use the constructor the frontend uses.
-	template <typename NodeClass>
-	struct ConstructorTakesNodeInitData {
+	namespace NodeRegistrationPrivate
+	{
+		// Utilize base class to reduce template bloat in TNodeRegistryEntry
+		class METASOUNDFRONTEND_API FNodeRegistryEntryBase : public INodeClassRegistryEntry
+		{
+		public:
+			FNodeRegistryEntryBase(const Metasound::FNodeClassMetadata& InMetadata);
 
-		// Use SFINAE trick to see if we have a valid constructor:
-		template<typename T>
-		static uint16 TestForConstructor(decltype(T(::Metasound::FNodeInitData()))*);
+			virtual ~FNodeRegistryEntryBase() = default;
 
-		template<typename T>
-		static uint8 TestForConstructor(...);
+			virtual const FNodeClassInfo& GetClassInfo() const override;
 
-		static const bool Value = sizeof(TestForConstructor<NodeClass>(nullptr)) == sizeof(uint16);
-	};
+			virtual const FMetasoundFrontendClass& GetFrontendClass() const override;
 
+			// Unhide CreateNode overloads which exist on INodeClassRegistryEntry
+			using INodeClassRegistryEntry::CreateNode;
 
-	template <typename FNodeType>
-	bool RegisterNodeWithFrontend(const Metasound::FNodeClassMetadata& InMetadata)
-	{
-		// if we reenter this code (because DECLARE_METASOUND_DATA_REFERENCE_TYPES was called twice with the same type),
-		// we catch it here.
-		static bool bAlreadyRegisteredThisDataType = false;
-		if (bAlreadyRegisteredThisDataType)
-		{
-			UE_LOG(LogMetaSound, Display, TEXT("Tried to call METASOUND_REGISTER_NODE twice with the same class. ignoring the second call. Likely because METASOUND_REGISTER_NODE is in a header that's used in multiple modules. Consider moving it to a private header or cpp file."));
-			return false;
-		}
+#if UE_METASOUND_ENABLE_5_6_NODE_REGISTRATION_DEPRECATION_WARNINGS
+		UE_DEPRECATED(5.6, "Node classes should be constructed with FNodeData")
+#endif
+			virtual TUniquePtr<INode> CreateNode(const FNodeInitData& InInitData) const override;
+
+			virtual const TSet<FMetasoundFrontendVersion>* GetImplementedInterfaces() const override;
+
+			virtual bool IsNative() const override;
+
+		protected:
 
-		bAlreadyRegisteredThisDataType = true;
+			TSharedRef<const FNodeClassMetadata> GetNodeClassMetadata() const;
+		private:
+
+			TSharedRef<FNodeClassMetadata> ClassMetadata;
+			FNodeClassInfo ClassInfo;
+			FMetasoundFrontendClass FrontendClass;
+		};
 
-		// Node registry entry specialized to FNodeType.
-		class FNodeRegistryEntryBase : public Frontend::INodeRegistryEntry
+		template<typename TNodeType>
+		class TNodeRegistryEntry : public FNodeRegistryEntryBase
 		{
 		public:
-			FNodeRegistryEntryBase(const Metasound::FNodeClassMetadata& InMetadata)
-			: ClassInfo(FMetasoundFrontendClassMetadata::GenerateClassMetadata(InMetadata, EMetasoundFrontendClassType::External))
-			, FrontendClass(Metasound::Frontend::GenerateClass(InMetadata))
-			{
-			}
+			// Expose FNodeRegistryEntryBase constructors. 
+			using FNodeRegistryEntryBase::FNodeRegistryEntryBase;
 
-			virtual ~FNodeRegistryEntryBase() = default;
+			// Unhide CreateNode overloads which exist on INodeClassRegistryEntry
+			using FNodeRegistryEntryBase::CreateNode;
 
-			virtual const Frontend::FNodeClassInfo& GetClassInfo() const override
+			virtual TUniquePtr<INode> CreateNode(FNodeData InNodeData) const override
 			{
-				return ClassInfo;
+				if constexpr(std::is_constructible_v<TNodeType, ::Metasound::FNodeData, TSharedRef<const FNodeClassMetadata>>)
+				{
+					// Prefer construction of nodes using (FNodeData, TShareRef<const FNodeClassMetadata>)
+					return MakeUnique<TNodeType>(MoveTemp(InNodeData), GetNodeClassMetadata());
+				}
+				else if constexpr(std::is_constructible_v<TNodeType, ::Metasound::FNodeData>)
+				{
+					// Some node classes have FNodeClassMetadata declared as static members on
+					// the node class and do not need a separate TSharedRef<const FNodeClassMetadata>. 
+					return MakeUnique<TNodeType>(MoveTemp(InNodeData));
+				}
+				else
+				{
+					checkNoEntry();
+					return nullptr;
+				}
 			}
+		};
 
-			virtual TUniquePtr<INode> CreateNode(const FNodeInitData& InParams) const = 0;
+#if UE_METASOUND_ENABLE_5_6_NODE_REGISTRATION_DEPRECATION_WARNINGS
+		template<typename T>
+		UE_DEPRECATED(5.6, "Update the node's constructor to be of the form FMyNode(FNodeData InNodeData, TSharedRef<const FNodeClassMetadata> InClassMetadata)")
+		void TriggerDeprecatedNodeConstructorWarning() {}
 
-			virtual TUniquePtr<INode> CreateNode(FDefaultLiteralNodeConstructorParams&& InParams) const override
-			{
-				return nullptr;
-			}
+		template<typename T>
+		UE_DEPRECATED(5.6, "Update the node class to include 'static FNodeClassMetadata CreateNodeClassMetadata()`)")
+		void TriggerMissingCreateNodeClassMetadataWarning() {}
+#else
+		// These only trigger deprecation if the macro UE_METASOUND_ENABLED_5_6_NODE_REGISTRATION_DEPRECATION_WARNINGS is true
+		template<typename T>
+		void TriggerDeprecatedNodeConstructorWarning() {}
 
-			virtual TUniquePtr<INode> CreateNode(FDefaultNamedVertexNodeConstructorParams&& InParams) const override
-			{
-				return nullptr;
-			}
-		
-			virtual TUniquePtr<INode> CreateNode(FDefaultNamedVertexWithLiteralNodeConstructorParams&& InParams) const override
-			{
-				return nullptr;
-			}
+		template<typename T>
+		void TriggerMissingCreateNodeClassMetadataWarning() {}
+#endif
 
+		// Forward declare
+		class FDeprecatedNodeAPIFactory;
 
-			virtual const FMetasoundFrontendClass& GetFrontendClass() const override
-			{
-				return FrontendClass;
-			}
+		// FDeprecatedNodeAPIAdapterBase is used as a back compatible shim
+		// for nodes which have not been updated to support the new APIs for node
+		// registration.
+		class METASOUNDFRONTEND_API FDeprecatedNodeAPIAdapterBase : public FBasicNode
+		{
+		protected:
+			FDeprecatedNodeAPIAdapterBase(FNodeData InNodeData, TSharedRef<const FNodeClassMetadata> InClassMetadata, TUniquePtr<INode> InNode);
 
-			virtual const TSet<FMetasoundFrontendVersion>* GetImplementedInterfaces() const override
-			{
-				return nullptr;
-			}
+		public:
 
-			virtual bool IsNative() const override
-			{
-				return true;
-			}
+			virtual ~FDeprecatedNodeAPIAdapterBase() = default;
+			virtual FOperatorFactorySharedRef GetDefaultOperatorFactory() const override;
 
 		private:
 
-			Frontend::FNodeClassInfo ClassInfo;
-			FMetasoundFrontendClass FrontendClass;
+			TSharedRef<FDeprecatedNodeAPIFactory> Factory;
 		};
 
-		class FNodeRegistryEntry : public FNodeRegistryEntryBase
+		// TDeprecatedNodeAPIAdapter is used as a back compatible shim
+		// for nodes which have not been updated to support the new APIs for node
+		// registration.
+		template<typename TNodeType>
+		class TDeprecatedNodeAPIAdapter : public FDeprecatedNodeAPIAdapterBase
 		{
 		public:
-			FNodeRegistryEntry(const Metasound::FNodeClassMetadata& InMetadata)
-				: FNodeRegistryEntryBase(InMetadata)
+			TDeprecatedNodeAPIAdapter(FNodeData InNodeData, TSharedRef<const FNodeClassMetadata> InClassMetadata)
+			: FDeprecatedNodeAPIAdapterBase(InNodeData, MoveTemp(InClassMetadata), MakeUnique<TNodeType>(FNodeInitData{InNodeData.Name, InNodeData.ID}))
 			{
+				if constexpr(TIsOnlyDeprecatedNodeConstructorProvided<TNodeType>::Value)
+				{
+					TriggerDeprecatedNodeConstructorWarning<TNodeType>();
+				}
 			}
+		};
 
-			virtual TUniquePtr<INode> CreateNode(const FNodeInitData& InParams) const override
+	}
+
+
+	template <typename TNodeType>
+	bool RegisterNode(const Metasound::FNodeClassMetadata& InMetadata)
+	{
+		using namespace NodeRegistrationPrivate;
+		if constexpr (TIsOnlyDeprecatedNodeConstructorProvided<TNodeType>::Value)
+		{
+			return RegisterNode<TDeprecatedNodeAPIAdapter<TNodeType>>(InMetadata);
+		}
+		else
+		{
+			// if we reenter this code (because DECLARE_METASOUND_DATA_REFERENCE_TYPES was called twice with the same type),
+			// we catch it here.
+			static bool bAlreadyRegisteredThisDataType = false;
+			if (bAlreadyRegisteredThisDataType)
 			{
-				return MakeUnique<FNodeType>(InParams);
+				UE_LOG(LogMetaSound, Display, TEXT("Tried to call METASOUND_REGISTER_NODE twice with the same class. ignoring the second call. Likely because METASOUND_REGISTER_NODE is in a header that's used in multiple modules. Consider moving it to a private header or cpp file."));
+				return false;
 			}
-		};
 
+			bAlreadyRegisteredThisDataType = true;
 
-		Frontend::FNodeRegistryKey Key = FMetasoundFrontendRegistryContainer::Get()->RegisterNode(MakeUnique<FNodeRegistryEntry>(InMetadata));
-		const bool bSuccessfullyRegisteredNode = Key.IsValid();
-		ensureAlwaysMsgf(bSuccessfullyRegisteredNode, TEXT("Registering node class failed. Please check the logs."));
 
-		return bSuccessfullyRegisteredNode;
+			Frontend::FNodeRegistryKey Key = FMetasoundFrontendRegistryContainer::Get()->RegisterNode(MakeUnique<TNodeRegistryEntry<TNodeType>>(InMetadata));
+			const bool bSuccessfullyRegisteredNode = Key.IsValid();
+			ensureAlwaysMsgf(bSuccessfullyRegisteredNode, TEXT("Registering node class failed. Please check the logs."));
+
+			return bSuccessfullyRegisteredNode;
+		}
+	}
+
+	template <typename TNodeType>
+	bool RegisterNode()
+	{
+		if constexpr(TIsCreateNodeClassMetadataDeclared<TNodeType>::Value)
+		{
+			return RegisterNode<TNodeType>(TNodeType::CreateNodeClassMetadata());
+		}
+		else
+		{
+			using namespace NodeRegistrationPrivate;
+			TriggerMissingCreateNodeClassMetadataWarning<TNodeType>();
+			// Register a node using a prototype node.
+			FNodeInitData InitData;
+			TUniquePtr<Metasound::INode> Node = MakeUnique<TNodeType>(InitData);
+			return RegisterNode<TNodeType>(Node->GetMetadata());
+		}
 	}
+}
+
+namespace Metasound
+{
+	// Utility class to ensure that a node class can use the constructor the frontend uses.
+	template <typename NodeClass>
+	struct UE_DEPRECATED(5.6, "Use std::is_constructible<T, FNodeInitData> instead") ConstructorTakesNodeInitData 
+	{
+
+		// Use SFINAE trick to see if we have a valid constructor:
+		template<typename T>
+		static uint16 TestForConstructor(decltype(T(::Metasound::FNodeInitData()))*);
+
+		template<typename T>
+		static uint8 TestForConstructor(...);
 
-	template <typename FNodeType>
+		static const bool Value = sizeof(TestForConstructor<NodeClass>(nullptr)) == sizeof(uint16);
+	};
+
+	template<typename T>
+	UE_DEPRECATED(5.6, "Use Frontend::RegisterNode()")
 	bool RegisterNodeWithFrontend()
 	{
-		// Register a node using a prototype node.
-		// TODO: may want to add a warning here since we don't want to use this registration pathway.
-		Metasound::FNodeInitData InitData;
-		TUniquePtr<Metasound::INode> Node = MakeUnique<FNodeType>(InitData);
-		return RegisterNodeWithFrontend<FNodeType>(Node->GetMetadata());
+		return Frontend::RegisterNode<T>();
+	}
+
+	template<typename T>
+	UE_DEPRECATED(5.6, "Use Frontend::RegisterNode(const Metasound::FNodeClassMetadata&)")
+	bool RegisterNodeWithFrontend(const Metasound::FNodeClassMetadata& InMetadata)
+	{
+		return Frontend::RegisterNode<T>(InMetadata);
 	}
 }
 
-// TODO: should remove this and force folks to use module startup/shutdown module system.
 #define METASOUND_REGISTER_NODE(NodeClass, ...) \
 	 static_assert(std::is_base_of<::Metasound::INodeBase, NodeClass>::value, "To be registered as a  Metasound Node," #NodeClass "need to be a derived class from Metasound::INodeBase, Metasound::INode, or Metasound::FNode."); \
-	 static_assert(::Metasound::ConstructorTakesNodeInitData<NodeClass>::Value, "In order to be registered as a Metasound Node, " #NodeClass " needs to implement the following public constructor: " #NodeClass "(const Metasound::FNodeInitData& InInitData);"); \
-	 static bool bSuccessfullyRegistered##NodeClass  = FMetasoundFrontendRegistryContainer::Get()->EnqueueInitCommand([](){ ::Metasound::RegisterNodeWithFrontend<NodeClass>(__VA_ARGS__); }); // This static bool is useful for debugging, but also is the only way the compiler will let us call this function outside of an expression.
+	 static_assert(::Metasound::TIsNodeConstructorSupported<NodeClass>::Value, "In order to be registered as a Metasound Node, " #NodeClass " needs to implement the following public constructor: " #NodeClass "(Metasound::FNodeData InNodeData);"); \
+	 static bool bSuccessfullyRegistered##NodeClass  = ::Metasound::Frontend::INodeClassRegistry::Get()->EnqueueInitCommand([](){ ::Metasound::Frontend::RegisterNode<NodeClass>(__VA_ARGS__); }); // This static bool is useful for debugging, but also is the only way the compiler will let us call this function outside of an expression.
 
 /*
 Macros to help define various FText node fields.
diff --git a/Engine/Plugins/Runtime/Metasound/Source/MetasoundFrontend/Public/Traits/MetasoundNodeConstructorTraits.h b/Engine/Plugins/Runtime/Metasound/Source/MetasoundFrontend/Public/Traits/MetasoundNodeConstructorTraits.h
new file mode 100644
--- /dev/null
+++ b/Engine/Plugins/Runtime/Metasound/Source/MetasoundFrontend/Public/Traits/MetasoundNodeConstructorTraits.h
@@ -0,0 +1,37 @@
+// Copyright Epic Games, Inc. All Rights Reserved.
+
+#include <type_traits>
+
+#include "MetasoundNodeInterface.h"
+#include "Templates/SharedPointer.h"
+
+namespace Metasound
+{
+	/** Determines whether the constructor provided by the node is usable in node 
+	 * registration. */
+	template<typename T>
+	struct TIsNodeConstructorSupported
+	{
+	private:
+		static constexpr bool bIsConstructibleWithFNodeInitData = std::is_constructible_v<T, ::Metasound::FNodeInitData>;
+		static constexpr bool bIsConstructibleWithFNodeData = std::is_constructible_v<T, ::Metasound::FNodeData>;
+		static constexpr bool bIsConstructibleWithFNodeDataAndClassMetadata = std::is_constructible_v<T, ::Metasound::FNodeData, TSharedRef<const FNodeClassMetadata>>;
+
+	public:
+		static constexpr bool Value = bIsConstructibleWithFNodeInitData || bIsConstructibleWithFNodeData || bIsConstructibleWithFNodeDataAndClassMetadata;
+	};
+
+	/** Determines a node only provides a deprecated constructor. */
+	template<typename T> 
+	struct TIsOnlyDeprecatedNodeConstructorProvided
+	{
+	private:
+		static constexpr bool bIsConstructibleWithFNodeInitData = std::is_constructible_v<T, ::Metasound::FNodeInitData>;
+		static constexpr bool bIsConstructibleWithFNodeData = std::is_constructible_v<T, ::Metasound::FNodeData>;
+		static constexpr bool bIsConstructibleWithFNodeDataAndClassMetadata = std::is_constructible_v<T, ::Metasound::FNodeData, TSharedRef<const FNodeClassMetadata>>;
+
+	public:
+		static constexpr bool Value = bIsConstructibleWithFNodeInitData && !(bIsConstructibleWithFNodeData || bIsConstructibleWithFNodeDataAndClassMetadata);
+	};
+}
+
diff --git a/Engine/Plugins/Runtime/Metasound/Source/MetasoundFrontend/Public/Traits/MetasoundNodeStaticMemberTraits.h b/Engine/Plugins/Runtime/Metasound/Source/MetasoundFrontend/Public/Traits/MetasoundNodeStaticMemberTraits.h
new file mode 100644
--- /dev/null
+++ b/Engine/Plugins/Runtime/Metasound/Source/MetasoundFrontend/Public/Traits/MetasoundNodeStaticMemberTraits.h
@@ -0,0 +1,26 @@
+// Copyright Epic Games, Inc. All Rights Reserved.
+
+#include <type_traits>
+
+#include "MetasoundNodeInterface.h"
+
+namespace Metasound
+{
+	template <typename U>
+	class TIsCreateNodeClassMetadataDeclared
+	{
+		private:
+			template<typename T, T> 
+			struct Helper;
+
+			template<typename T>
+			static uint8 Check(Helper<FNodeClassMetadata(*)(), &T::CreateNodeClassMetadata>*);
+
+			template<typename T> static uint16 Check(...);
+
+		public:
+
+			// If the function exists, then "Value" is true. Otherwise "Value" is false.
+			static constexpr bool Value = sizeof(Check<U>(0)) == sizeof(uint8);
+	};
+}
diff --git a/Engine/Plugins/Runtime/Metasound/Source/MetasoundGraphCore/Private/MetasoundBasicNode.cpp b/Engine/Plugins/Runtime/Metasound/Source/MetasoundGraphCore/Private/MetasoundBasicNode.cpp
new file mode 100644
--- /dev/null
+++ b/Engine/Plugins/Runtime/Metasound/Source/MetasoundGraphCore/Private/MetasoundBasicNode.cpp
@@ -0,0 +1,38 @@
+// Copyright Epic Games, Inc. All Rights Reserved.
+
+#include "MetasoundBasicNode.h"
+
+#include "MetasoundNodeInterface.h"
+#include "Templates/SharedPointer.h"
+
+namespace Metasound
+{
+	FBasicNode::FBasicNode(FNodeData InNodeData, TSharedRef<const FNodeClassMetadata> InClassMetadata)
+	: NodeData(MoveTemp(InNodeData))
+	, ClassMetadata(MoveTemp(InClassMetadata))
+	{
+	}
+
+	/** Return the name of this specific instance of the node class. */
+	const FName& FBasicNode::GetInstanceName() const
+	{
+		return NodeData.Name;
+	}
+
+	/** Return the ID of this specific instance of the node class. */
+	const FGuid& FBasicNode::GetInstanceID() const
+	{
+		return NodeData.ID;
+	}
+
+	/** Return the type name of this node. */
+	const FNodeClassMetadata& FBasicNode::GetMetadata() const
+	{
+		return *ClassMetadata;
+	}
+
+	const FVertexInterface& FBasicNode::GetVertexInterface() const
+	{
+		return NodeData.Interface;
+	}
+}
diff --git a/Engine/Plugins/Runtime/Metasound/Source/MetasoundGraphCore/Private/MetasoundFacade.cpp b/Engine/Plugins/Runtime/Metasound/Source/MetasoundGraphCore/Private/MetasoundFacade.cpp
--- a/Engine/Plugins/Runtime/Metasound/Source/MetasoundGraphCore/Private/MetasoundFacade.cpp
+++ b/Engine/Plugins/Runtime/Metasound/Source/MetasoundGraphCore/Private/MetasoundFacade.cpp
@@ -14,21 +14,6 @@ namespace Metasound
 		return CreateFunc(InParams, OutResults);
 	}
 
-	const FVertexInterface& FNodeFacade::GetVertexInterface() const
-	{
-		return VertexInterface;
-	}
-
-	bool FNodeFacade::SetVertexInterface(const FVertexInterface& InInterface)
-	{
-		return InInterface == VertexInterface;
-	}
-
-	bool FNodeFacade::IsVertexInterfaceSupported(const FVertexInterface& InInterface) const
-	{
-		return InInterface == VertexInterface;
-	}
-
 	/** Return a reference to the default operator factory. */
 	FOperatorFactorySharedRef FNodeFacade::GetDefaultOperatorFactory() const
 	{
diff --git a/Engine/Plugins/Runtime/Metasound/Source/MetasoundGraphCore/Private/MetasoundNodeInterface.cpp b/Engine/Plugins/Runtime/Metasound/Source/MetasoundGraphCore/Private/MetasoundNodeInterface.cpp
--- a/Engine/Plugins/Runtime/Metasound/Source/MetasoundGraphCore/Private/MetasoundNodeInterface.cpp
+++ b/Engine/Plugins/Runtime/Metasound/Source/MetasoundGraphCore/Private/MetasoundNodeInterface.cpp
@@ -4,6 +4,13 @@
 
 namespace Metasound
 {
+	FNodeData::FNodeData(const FName& InName, const FGuid& InID, FVertexInterface InInterface)
+	: Name(InName)
+	, ID(InID)
+	, Interface(MoveTemp(InInterface))
+	{
+	}
+
 	const FString PluginAuthor = TEXT("Epic Games, Inc.");
 #if WITH_EDITOR
 	const FText PluginNodeMissingPrompt = NSLOCTEXT("MetasoundGraphCore", "Metasound_DefaultMissingNodePrompt", "The node was likely removed, renamed, or the Metasound plugin is not loaded.");
@@ -50,18 +57,6 @@ namespace Metasound
 		return *Builder;
 	}
 
-	/** Namespace and name of the node class. */
-	const FName FNodeClassName::GetScopedName() const
-	{
-		return FormatScopedName(Namespace, Name);
-	}
-
-	/** Namespace, name and variant of the node class. */
-	const FName FNodeClassName::GetFullName() const
-	{
-		return FormatFullName(Namespace, Name, Variant);
-	}
-
 	FName FNodeClassName::FormatFullName(const FName& InNamespace, const FName& InName, const FName& InVariant)
 	{
 		FNameBuilder Builder;
diff --git a/Engine/Plugins/Runtime/Metasound/Source/MetasoundGraphCore/Public/MetasoundBasicNode.h b/Engine/Plugins/Runtime/Metasound/Source/MetasoundGraphCore/Public/MetasoundBasicNode.h
new file mode 100644
--- /dev/null
+++ b/Engine/Plugins/Runtime/Metasound/Source/MetasoundGraphCore/Public/MetasoundBasicNode.h
@@ -0,0 +1,40 @@
+// Copyright Epic Games, Inc. All Rights Reserved.
+
+#pragma once
+
+#include "MetasoundNodeInterface.h"
+#include "Templates/SharedPointer.h"
+
+namespace Metasound
+{
+	class METASOUNDGRAPHCORE_API FBasicNode : public INode
+	{
+		public:
+			/** Construct a basic node.
+			 *
+			 * @param InNodeData - Information unique to this node instance. 
+			 * @param InClassMetadata - Shared object representing the node class.
+			 */
+			FBasicNode(FNodeData InNodeData, TSharedRef<const FNodeClassMetadata> InClassMetadata);
+
+			virtual ~FBasicNode() = default;
+
+			/** Return the name of this specific instance of the node class. */
+			virtual const FName& GetInstanceName() const override;
+
+			/** Return the ID of this specific instance of the node class. */
+			virtual const FGuid& GetInstanceID() const override;
+
+			/** Return metadata associated with this node. */
+			virtual const FNodeClassMetadata& GetMetadata() const override;
+
+			/** Return the interface associated with this node instance. */
+			virtual const FVertexInterface& GetVertexInterface() const override;
+
+		private:
+
+			FNodeData NodeData;
+			TSharedRef<const FNodeClassMetadata> ClassMetadata;
+	};
+}
+
diff --git a/Engine/Plugins/Runtime/Metasound/Source/MetasoundGraphCore/Public/MetasoundFacade.h b/Engine/Plugins/Runtime/Metasound/Source/MetasoundGraphCore/Public/MetasoundFacade.h
--- a/Engine/Plugins/Runtime/Metasound/Source/MetasoundGraphCore/Public/MetasoundFacade.h
+++ b/Engine/Plugins/Runtime/Metasound/Source/MetasoundGraphCore/Public/MetasoundFacade.h
@@ -2,8 +2,8 @@
 
 #pragma once
 
+#include "MetasoundBasicNode.h"
 #include "MetasoundBuilderInterface.h"
-#include "MetasoundNode.h"
 #include "MetasoundNodeInterface.h"
 #include "MetasoundVertex.h"
 #include <type_traits>
@@ -82,74 +82,63 @@ namespace Metasound
 	 * The type of the concrete IOperator class to instantiate is defined in the
 	 * FNodeFacade constructors TFacadeOperatorClass<>.
 	 */
-	class METASOUNDGRAPHCORE_API FNodeFacade : public FNode
+	class METASOUNDGRAPHCORE_API FNodeFacade : public FBasicNode
 	{
-			// Factory class for create an IOperator by using a TFunction matching
-			// the signature of the CreateOperator function.
-			class METASOUNDGRAPHCORE_API FFactory : public IOperatorFactory
-			{
-				using FCreateOperatorFunction = TFunction<TUniquePtr<IOperator>(const FBuildOperatorParams& InParams, FBuildResults& OutResults)>;
-
-				public:
-					FFactory(FCreateOperatorFunction InCreateFunc);
-
-					virtual TUniquePtr<IOperator> CreateOperator(const FBuildOperatorParams& InParams, FBuildResults& OutResults) override;
-
-				private:
-
-					FCreateOperatorFunction CreateFunc;
-			};
-
-		public: 
+		// Factory class for create an IOperator by using a TFunction matching
+		// the signature of the CreateOperator function.
+		class METASOUNDGRAPHCORE_API FFactory : public IOperatorFactory
+		{
+			using FCreateOperatorFunction = TFunction<TUniquePtr<IOperator>(const FBuildOperatorParams& InParams, FBuildResults& OutResults)>;
 
-			FNodeFacade() = delete;
+			public:
+				FFactory(FCreateOperatorFunction InCreateFunc);
 
-			/** FNodeFacade constructor using the OperatorType template parameter
-			 * to get the node info, operator factory method and vertex interface. 
-			 *
-			 * @param InInstanceName - Instance name for the node.
-			 * @param InInstanceID - Instance ID for the node.
-			 * @param OperatorClass - Template class wrapper for the underlying 
-			 * 						  IOperator which is created by this node.
-			 */
-			template<typename OperatorType>
-			FNodeFacade(const FVertexName& InInstanceName, const FGuid& InInstanceID, TFacadeOperatorClass<OperatorType> OperatorClass)
-			:	FNode(InInstanceName, InInstanceID, OperatorType::GetNodeInfo())
-			, 	Factory(MakeShared<FFactory, ESPMode::ThreadSafe>(&OperatorType::CreateOperator))
-			{
-				VertexInterface = GetMetadata().DefaultInterface;
-			}
+				virtual TUniquePtr<IOperator> CreateOperator(const FBuildOperatorParams& InParams, FBuildResults& OutResults) override;
 
-			template<typename OperatorType>
-			FNodeFacade(FNodeInitData&& InInitData, TFacadeOperatorClass<OperatorType> OperatorClass)
-				: FNode(InInitData, OperatorType::GetNodeInfo())
-				, Factory(MakeShared<FFactory, ESPMode::ThreadSafe>(&OperatorType::CreateOperator))
-			{
-				VertexInterface = GetMetadata().DefaultInterface;
-			}
+			private:
 
-			virtual ~FNodeFacade() = default;
+				FCreateOperatorFunction CreateFunc;
+		};
 
-			virtual const FVertexInterface& GetVertexInterface() const override;
+	public: 
+
+		FNodeFacade() = delete;
+
+		/** FNodeFacade constructor using the OperatorType template parameter
+		 * to get the node info, operator factory method and vertex interface. 
+		 *
+		 * @param InInstanceName - Instance name for the node.
+		 * @param InInstanceID - Instance ID for the node.
+		 * @param OperatorClass - Template class wrapper for the underlying 
+		 * 						  IOperator which is created by this node.
+		 */
+		template<typename OperatorType>
+		FNodeFacade(const FVertexName& InInstanceName, const FGuid& InInstanceID, TFacadeOperatorClass<OperatorType> OperatorClass)
+		:	FNodeFacade(FNodeData(InInstanceName, InInstanceID, OperatorType::GetNodeInfo().DefaultInterface), MakeShared<FNodeClassMetadata>(OperatorType::GetNodeInfo()), OperatorClass)
+		{
+		}
 
-			/** FNodeFacade has a static vertex interface. This will return true
-			 * only if the input interface is equal to the default interface.
-			 */
-			virtual bool SetVertexInterface(const FVertexInterface& InInterface) override;
+		template<typename OperatorType>
+		FNodeFacade(FNodeInitData InInitData, TFacadeOperatorClass<OperatorType> OperatorClass)
+		:	FNodeFacade(FNodeData(InInitData.InstanceName, InInitData.InstanceID, OperatorType::GetNodeInfo().DefaultInterface), MakeShared<FNodeClassMetadata>(OperatorType::GetNodeInfo()), OperatorClass)
+		{
+		}
 
-			/** FNodeFacade has a static vertex interface. This will return true
-			 * only if the input interface is equal to the default interface.
-			 */
-			virtual bool IsVertexInterfaceSupported(const FVertexInterface& InInterface) const override;
+		template<typename OperatorType>
+		FNodeFacade(FNodeData InNodeData, TSharedRef<const FNodeClassMetadata> InClassMetadata, TFacadeOperatorClass<OperatorType> OperatorClass)
+		: FBasicNode(MoveTemp(InNodeData), MoveTemp(InClassMetadata))
+		, Factory(MakeShared<FFactory, ESPMode::ThreadSafe>(&OperatorType::CreateOperator))
+		{
+		}
 
-			/** Return a reference to the default operator factory. */
-			virtual FOperatorFactorySharedRef GetDefaultOperatorFactory() const override;
+		virtual ~FNodeFacade() = default;
 
-		private:
+		/** Return a reference to the default operator factory. */
+		virtual FOperatorFactorySharedRef GetDefaultOperatorFactory() const override;
 
-			TSharedRef<IOperatorFactory, ESPMode::ThreadSafe> Factory;
+	private:
 
-			FVertexInterface VertexInterface;
+		TSharedRef<IOperatorFactory, ESPMode::ThreadSafe> Factory;
 	};
 
 
@@ -161,7 +150,16 @@ namespace Metasound
 	{
 	public:
 		TNodeFacade(const FNodeInitData& InInitData)
-			: FNodeFacade(InInitData.InstanceName, InInitData.InstanceID, TFacadeOperatorClass<OperatorType>())
+			: FNodeFacade(FNodeData(InInitData.InstanceName, InInitData.InstanceID, OperatorType::GetNodeInfo().DefaultInterface), MakeShared<FNodeClassMetadata>(OperatorType::GetNodeInfo()), TFacadeOperatorClass<OperatorType>())
 		{}
+
+		TNodeFacade(FNodeData InNodeData, TSharedRef<const FNodeClassMetadata> InClassMetadata)
+			: FNodeFacade(MoveTemp(InNodeData), MoveTemp(InClassMetadata), TFacadeOperatorClass<OperatorType>())
+		{}
+
+		static FNodeClassMetadata CreateNodeClassMetadata()
+		{
+			return OperatorType::GetNodeInfo();
+		}
 	};
 }
diff --git a/Engine/Plugins/Runtime/Metasound/Source/MetasoundGraphCore/Public/MetasoundNode.h b/Engine/Plugins/Runtime/Metasound/Source/MetasoundGraphCore/Public/MetasoundNode.h
--- a/Engine/Plugins/Runtime/Metasound/Source/MetasoundGraphCore/Public/MetasoundNode.h
+++ b/Engine/Plugins/Runtime/Metasound/Source/MetasoundGraphCore/Public/MetasoundNode.h
@@ -26,8 +26,10 @@ namespace Metasound
 
 			virtual const FVertexInterface& GetVertexInterface() const override;
 
+			UE_DEPRECATED(5.6, "FNode will not support modification to the FVertexInterface")
 			virtual bool SetVertexInterface(const FVertexInterface& InInterface) override;
 
+			UE_DEPRECATED(5.6, "FNode will not support modification to the FVertexInterface")
 			virtual bool IsVertexInterfaceSupported(const FVertexInterface& InInterface) const override;
 
 		private:
diff --git a/Engine/Plugins/Runtime/Metasound/Source/MetasoundGraphCore/Public/MetasoundNodeInterface.h b/Engine/Plugins/Runtime/Metasound/Source/MetasoundGraphCore/Public/MetasoundNodeInterface.h
--- a/Engine/Plugins/Runtime/Metasound/Source/MetasoundGraphCore/Public/MetasoundNodeInterface.h
+++ b/Engine/Plugins/Runtime/Metasound/Source/MetasoundGraphCore/Public/MetasoundNodeInterface.h
@@ -16,6 +16,17 @@ namespace Metasound
 	extern const FString METASOUNDGRAPHCORE_API PluginAuthor;
 	extern const FText METASOUNDGRAPHCORE_API PluginNodeMissingPrompt;
 
+	/** Data used to construct a Node. */
+	struct METASOUNDGRAPHCORE_API FNodeData
+	{
+		FNodeData() = default;
+		FNodeData(const FName& InName, const FGuid& InID, FVertexInterface InInterface);
+
+		FName Name;
+		FGuid ID;
+		FVertexInterface Interface;
+	};
+
 	/**
 	 * Node style data
 	 */
@@ -66,14 +77,6 @@ namespace Metasound
 		/** The full name of the Node formatted Namespace.Name and optionally appended with .Variant if valid */
 		const FString ToString() const;
 
-		/** Namespace and name of the node class. */
-		UE_DEPRECATED(5.4, "Use FormatScopedName functions to generate")
-		const FName GetScopedName() const;
-
-		/** Namespace, name and variant of the node class. */
-		UE_DEPRECATED(5.4, "Use FormatFullName functions to generate")
-		const FName GetFullName() const;
-
 		/** Whether or not this instance of a node class name is a valid name.*/
 		bool IsValid() const;
 
@@ -168,7 +171,8 @@ namespace Metasound
 			 *
 			 * @return True on success, false otherwise.
 			 */
-			virtual bool SetVertexInterface(const FVertexInterface& InInterface) = 0;
+			UE_DEPRECATED(5.6, "INodeBase will not expose dynamic FVertexInterface operations")
+			virtual bool SetVertexInterface(const FVertexInterface& InInterface) { return false; }
 
 			/** Expresses whether a specific vertex interface is supported.
 			 *
@@ -176,7 +180,8 @@ namespace Metasound
 			 *
 			 * @return True if the interface is supported, false otherwise. 
 			 */
-			virtual bool IsVertexInterfaceSupported(const FVertexInterface& InInterface) const = 0;
+			UE_DEPRECATED(5.6, "INodeBase will not expose dynamic FVertexInterface operations")
+			virtual bool IsVertexInterfaceSupported(const FVertexInterface& InInterface) const { return false; }
 	};
 
 	// Forward declare
@@ -204,12 +209,6 @@ namespace Metasound
 
 			/** Return a reference to the default operator factory. */
 			virtual FOperatorFactorySharedRef GetDefaultOperatorFactory() const = 0;
-
-			/* Future implementations may support additional factory types and interfaces
-				virtual bool DoesSupportOperatorFactory(const FString& InFractoryName) const = 0;
-				virtual IOperatorFactory* GetOperatorFactory(const FString& InFactoryName) = 0;
-				virtual ISpecialFactory* GetSpecialFactory() { return nullptr; }
-			*/
 	};
 
 	/** FOutputDataSource describes the source of data which is produced within
