Fix CIS warning on asset registry.
#rnx

diff --git a/Engine/Plugins/Runtime/Metasound/Source/MetasoundFrontend/Private/MetasoundFrontendNodeClassRegistry.cpp b/Engine/Plugins/Runtime/Metasound/Source/MetasoundFrontend/Private/MetasoundFrontendNodeClassRegistry.cpp
--- a/Engine/Plugins/Runtime/Metasound/Source/MetasoundFrontend/Private/MetasoundFrontendNodeClassRegistry.cpp
+++ b/Engine/Plugins/Runtime/Metasound/Source/MetasoundFrontend/Private/MetasoundFrontendNodeClassRegistry.cpp
@@ -575,6 +575,8 @@ PRAGMA_ENABLE_DEPRECATION_WARNINGS
 	TUniquePtr<Metasound::INode> FNodeClassRegistry::CreateNode(const FNodeClassRegistryKey& InKey, const Metasound::FNodeInitData& InInitData) const
 	{
 		TUniquePtr<INode> Node;
+
+#if UE_METASOUND_ENABLE_5_6_NODE_REGISTRATION_DEPRECATION_WARNINGS
 		auto CreateNodeLambda = [&Node, &InInitData](const INodeClassRegistryEntry& Entry) mutable
 		{ 
 			const FMetasoundFrontendClass& FrontendClass = Entry.GetFrontendClass();
@@ -587,6 +589,12 @@ PRAGMA_ENABLE_DEPRECATION_WARNINGS
 			};
 			Node = Entry.CreateNode(MoveTemp(NodeData)); 
 		};
+#else 
+		auto CreateNodeLambda = [&Node, &InInitData](const INodeClassRegistryEntry& Entry) mutable
+		{ 
+			Node = Entry.CreateNode(InInitData); 
+		};
+#endif
 
 		if (!AccessNodeEntryThreadSafe(InKey, CreateNodeLambda))
 		{
