Fix ensure when canceling creating a MetaSound preset

#rnx

diff --git a/Engine/Plugins/Runtime/Metasound/Source/MetasoundEditor/Private/MetasoundAssetDefinitions.cpp b/Engine/Plugins/Runtime/Metasound/Source/MetasoundEditor/Private/MetasoundAssetDefinitions.cpp
--- a/Engine/Plugins/Runtime/Metasound/Source/MetasoundEditor/Private/MetasoundAssetDefinitions.cpp
+++ b/Engine/Plugins/Runtime/Metasound/Source/MetasoundEditor/Private/MetasoundAssetDefinitions.cpp
@@ -425,14 +425,14 @@ namespace MenuExtension_MetaSoundSourceTemplate
 
 				UObject* NewMetaSound = IAssetTools::Get().CreateAssetWithDialog(AssetName, FPackageName::GetLongPackagePath(PackagePath), Factory->GetSupportedClass(), Factory);
 
-				if (ensure(NewMetaSound))
+				if (NewMetaSound)
 				{
 					FGraphBuilder::RegisterGraphWithFrontend(*NewMetaSound);
 					ObjectsToSync.Add(NewMetaSound);
 				}
 				else
 				{
-					UE_LOG(LogMetaSound, Error, TEXT("Error creating new asset with factory '%s' when creating preset '%s'"), *Factory->GetName(), *AssetName);
+					UE_LOG(LogMetaSound, Display, TEXT("Error creating new asset with factory '%s' when creating preset '%s' or asset creation was canceled by user."), *Factory->GetName(), *AssetName);
 				}
 			}
 
