[AutoRTFM] Postpone Metasound enqueueing until the Verse transaction is committed.

We can't kick off work in another thread until we are sure that the Verse transaction can't be rolled back.

diff --git a/Engine/Plugins/Runtime/Metasound/Source/MetasoundGenerator/Private/MetasoundGenerator.cpp b/Engine/Plugins/Runtime/Metasound/Source/MetasoundGenerator/Private/MetasoundGenerator.cpp
--- a/Engine/Plugins/Runtime/Metasound/Source/MetasoundGenerator/Private/MetasoundGenerator.cpp
+++ b/Engine/Plugins/Runtime/Metasound/Source/MetasoundGenerator/Private/MetasoundGenerator.cpp
@@ -21,6 +21,7 @@
 #include "DSP/FloatArrayMath.h"
 #include "Interfaces/MetasoundFrontendSourceInterface.h"
 #include "Modules/ModuleManager.h"
+#include "AutoRTFM/AutoRTFM.h"
 
 #ifndef ENABLE_METASOUNDGENERATOR_INVALID_SAMPLE_VALUE_LOGGING 
 #define ENABLE_METASOUNDGENERATOR_INVALID_SAMPLE_VALUE_LOGGING !UE_BUILD_SHIPPING
@@ -502,7 +503,11 @@ PRAGMA_ENABLE_DEPRECATION_WARNINGS
 
 	void FMetasoundGenerator::QueueParameterPack(TSharedPtr<FMetasoundParameterPackStorage> ParameterPack)
 	{
-		ParameterPackQueue.Enqueue(ParameterPack);
+		// We can't start processing Metasound data until the Verse transaction completes, since it might be rolled back.
+		UE_AUTORTFM_ONCOMMIT(this, ParameterPack)
+		{
+			ParameterPackQueue.Enqueue(ParameterPack);
+		};
 	}
 
 	int32 FMetasoundGenerator::GetNumChannels() const
