Fixup metasound nodes to use the new node registration and isntantiation API

#rnx

diff --git a/Engine/Plugins/Runtime/Metasound/Source/MetasoundFrontend/Private/MetasoundAutoConverterNode.cpp b/Engine/Plugins/Runtime/Metasound/Source/MetasoundFrontend/Private/MetasoundAutoConverterNode.cpp
new file mode 100644
--- /dev/null
+++ b/Engine/Plugins/Runtime/Metasound/Source/MetasoundFrontend/Private/MetasoundAutoConverterNode.cpp
@@ -0,0 +1,76 @@
+// Copyright Epic Games, Inc. All Rights Reserved.
+
+#include "MetasoundAutoConverterNode.h"
+
+#define LOCTEXT_NAMESPACE "MetasoundFrontend"
+
+namespace Metasound
+{
+	namespace AutoConverterNodePrivate
+	{
+		FVertexName GetInputName(const FConvertDataTypeInfo& InInfo)
+		{
+			return InInfo.FromDataTypeName;
+		}
+
+		FVertexName GetOutputName(const FConvertDataTypeInfo& InInfo)
+		{
+			return InInfo.ToDataTypeName;
+		}
+
+		FVertexInterface CreateVertexInterface(const FConvertDataTypeInfo& InInfo)
+		{
+			const FText InputDesc = METASOUND_LOCTEXT_FORMAT("AutoConvDisplayNamePatternFrom", "Input {0} value.", InInfo.FromDataTypeText);
+			const FText OutputDesc = METASOUND_LOCTEXT_FORMAT("AutoConvDisplayNamePatternTo", "Output {0} value.", InInfo.ToDataTypeText);
+
+			return FVertexInterface(
+				FInputVertexInterface(
+					FInputDataVertex(GetInputName(InInfo), InInfo.FromDataTypeName, FDataVertexMetadata{ InputDesc }, EVertexAccessType::Reference)
+				),
+				FOutputVertexInterface(
+					FOutputDataVertex(GetOutputName(InInfo), InInfo.ToDataTypeName, FDataVertexMetadata{ OutputDesc }, EVertexAccessType::Reference)
+				)
+			);
+		}
+
+		FNodeClassMetadata CreateAutoConverterNodeMetadata(const FConvertDataTypeInfo& InInfo)
+		{
+			FNodeDisplayStyle DisplayStyle;
+			DisplayStyle.bShowName = false;
+			DisplayStyle.ImageName = TEXT("MetasoundEditor.Graph.Node.Conversion");
+			DisplayStyle.bShowInputNames = false;
+			DisplayStyle.bShowOutputNames = false;
+
+			const FText FromTypeText = InInfo.FromDataTypeText;
+			const FText ToTypeText = InInfo.ToDataTypeText;
+
+			FNodeClassMetadata Info;
+			Info.ClassName = { TEXT("Convert"), InInfo.ToDataTypeName, InInfo.FromDataTypeName };
+			Info.MajorVersion = 1;
+			Info.MinorVersion = 0;
+			Info.DisplayName = METASOUND_LOCTEXT_FORMAT("Metasound_AutoConverterNodeDisplayNameFormat", "{0} to {1}", InInfo.FromDataTypeText, InInfo.ToDataTypeText);
+			Info.Description = METASOUND_LOCTEXT_FORMAT("Metasound_AutoConverterNodeDescriptionNameFormat", "Converts from {0} to {1}.", InInfo.FromDataTypeText, InInfo.ToDataTypeText);
+			Info.Author = PluginAuthor;
+			Info.DisplayStyle = DisplayStyle;
+			Info.PromptIfMissing = PluginNodeMissingPrompt;
+			Info.DefaultInterface = CreateVertexInterface(InInfo);
+
+			Info.CategoryHierarchy.Emplace(NodeCategories::Conversions);
+			if (InInfo.bIsFromEnum || InInfo.bIsToEnum)
+			{
+				Info.CategoryHierarchy.Emplace(NodeCategories::EnumConversions);
+			}
+			
+			Info.Keywords =
+			{
+				METASOUND_LOCTEXT("MetasoundConvertKeyword", "Convert"),
+				InInfo.FromDataTypeText, 
+				InInfo.ToDataTypeText
+			};
+
+			return Info;
+		}
+	}
+}
+
+#undef LOCTEXT_NAMESPACE
diff --git a/Engine/Plugins/Runtime/Metasound/Source/MetasoundFrontend/Public/MetasoundArrayNodes.h b/Engine/Plugins/Runtime/Metasound/Source/MetasoundFrontend/Public/MetasoundArrayNodes.h
--- a/Engine/Plugins/Runtime/Metasound/Source/MetasoundFrontend/Public/MetasoundArrayNodes.h
+++ b/Engine/Plugins/Runtime/Metasound/Source/MetasoundFrontend/Public/MetasoundArrayNodes.h
@@ -186,16 +186,7 @@ namespace Metasound
 	};
 
 	template<typename ArrayType>
-	class TArrayNumNode : public FNodeFacade
-	{
-	public:
-		TArrayNumNode(const FNodeInitData& InInitData)
-		: FNodeFacade(InInitData.InstanceName, InInitData.InstanceID, TFacadeOperatorClass<TArrayNumOperator<ArrayType>>())
-		{
-		}
-
-		virtual ~TArrayNumNode() = default;
-	};
+	using TArrayNumNode = TNodeFacade<TArrayNumOperator<ArrayType>>;
 
 	/** TArrayGetOperator copies a value from the array to the output when
 	 * a trigger occurs. Initially, the output value is default constructed and
@@ -389,16 +380,7 @@ namespace Metasound
 	};
 
 	template<typename ArrayType>
-	class TArrayGetNode : public FNodeFacade
-	{
-	public:
-		TArrayGetNode(const FNodeInitData& InInitData)
-		: FNodeFacade(InInitData.InstanceName, InInitData.InstanceID, TFacadeOperatorClass<TArrayGetOperator<ArrayType>>())
-		{
-		}
-
-		virtual ~TArrayGetNode() = default;
-	};
+	using TArrayGetNode = TNodeFacade<TArrayGetOperator<ArrayType>>;
 
 	/** TArraySetOperator sets an element in an array to a specific value. */
 	template<typename ArrayType>
@@ -580,16 +562,7 @@ namespace Metasound
 	};
 
 	template<typename ArrayType>
-	class TArraySetNode : public FNodeFacade
-	{
-	public:
-		TArraySetNode(const FNodeInitData& InInitData)
-		: FNodeFacade(InInitData.InstanceName, InInitData.InstanceID, TFacadeOperatorClass<TArraySetOperator<ArrayType>>())
-		{
-		}
-
-		virtual ~TArraySetNode() = default;
-	};
+	using TArraySetNode = TNodeFacade<TArraySetOperator<ArrayType>>;
 
 	/** TArrayConcatOperator concatenates two arrays on trigger. */
 	template<typename ArrayType>
@@ -717,16 +690,7 @@ namespace Metasound
 	};
 
 	template<typename ArrayType>
-	class TArrayConcatNode : public FNodeFacade
-	{
-	public:
-		TArrayConcatNode(const FNodeInitData& InInitData)
-		: FNodeFacade(InInitData.InstanceName, InInitData.InstanceID, TFacadeOperatorClass<TArrayConcatOperator<ArrayType>>())
-		{
-		}
-
-		virtual ~TArrayConcatNode() = default;
-	};
+	using TArrayConcatNode = TNodeFacade<TArrayConcatOperator<ArrayType>>;
 
 	/** TArraySubsetOperator slices an array on trigger. */
 	template<typename ArrayType>
@@ -872,16 +836,7 @@ namespace Metasound
 	};
 
 	template<typename ArrayType>
-	class TArraySubsetNode : public FNodeFacade
-	{
-	public:
-		TArraySubsetNode(const FNodeInitData& InInitData)
-		: FNodeFacade(InInitData.InstanceName, InInitData.InstanceID, TFacadeOperatorClass<TArraySubsetOperator<ArrayType>>())
-		{
-		}
-
-		virtual ~TArraySubsetNode() = default;
-	};
+	using TArraySubsetNode = TNodeFacade<TArraySubsetOperator<ArrayType>>;
 
 	/** TArrayLastIndex gets last index of an array. */
 	template<typename ArrayType>
@@ -978,15 +933,7 @@ namespace Metasound
 	};
 
 	template<typename ArrayType>
-	class TArrayLastIndexNode : public FNodeFacade
-	{
-	public:
-		TArrayLastIndexNode(const FNodeInitData& InInitData)
-			: FNodeFacade(InInitData.InstanceName, InInitData.InstanceID, TFacadeOperatorClass<TArrayLastIndexOperator<ArrayType>>())
-		{
-		}
+	using TArrayLastIndexNode = TNodeFacade<TArrayLastIndexOperator<ArrayType>>;
 
-		virtual ~TArrayLastIndexNode() = default;
-	};
 } // namespace Metasound
 #undef LOCTEXT_NAMESPACE
diff --git a/Engine/Plugins/Runtime/Metasound/Source/MetasoundFrontend/Public/MetasoundArrayRandomNode.h b/Engine/Plugins/Runtime/Metasound/Source/MetasoundFrontend/Public/MetasoundArrayRandomNode.h
--- a/Engine/Plugins/Runtime/Metasound/Source/MetasoundFrontend/Public/MetasoundArrayRandomNode.h
+++ b/Engine/Plugins/Runtime/Metasound/Source/MetasoundFrontend/Public/MetasoundArrayRandomNode.h
@@ -559,15 +559,6 @@ namespace Metasound
 	};
 
 	template<typename ArrayType>
-	class TArrayRandomGetNode : public FNodeFacade
-	{
-	public:
-		TArrayRandomGetNode(const FNodeInitData& InInitData)
-			: FNodeFacade(InInitData.InstanceName, InInitData.InstanceID, TFacadeOperatorClass<TArrayRandomGetOperator<ArrayType>>())
-		{
-		}
-
-		virtual ~TArrayRandomGetNode() = default;
-	};
+	using TArrayRandomGetNode = TNodeFacade<TArrayRandomGetOperator<ArrayType>>;
 }
 #undef LOCTEXT_NAMESPACE
diff --git a/Engine/Plugins/Runtime/Metasound/Source/MetasoundFrontend/Public/MetasoundArrayShuffleNode.h b/Engine/Plugins/Runtime/Metasound/Source/MetasoundFrontend/Public/MetasoundArrayShuffleNode.h
--- a/Engine/Plugins/Runtime/Metasound/Source/MetasoundFrontend/Public/MetasoundArrayShuffleNode.h
+++ b/Engine/Plugins/Runtime/Metasound/Source/MetasoundFrontend/Public/MetasoundArrayShuffleNode.h
@@ -451,16 +451,7 @@ namespace Metasound
 	};
 
 	template<typename ArrayType>
-	class TArrayShuffleNode : public FNodeFacade
-	{
-	public:
-		TArrayShuffleNode(const FNodeInitData& InInitData)
-			: FNodeFacade(InInitData.InstanceName, InInitData.InstanceID, TFacadeOperatorClass<TArrayShuffleOperator<ArrayType>>())
-		{
-		}
-
-		virtual ~TArrayShuffleNode() = default;
-	};
+	using TArrayShuffleNode = TNodeFacade<TArrayShuffleOperator<ArrayType>>;
 }
 
 #undef LOCTEXT_NAMESPACE
diff --git a/Engine/Plugins/Runtime/Metasound/Source/MetasoundFrontend/Public/MetasoundAutoConverterNode.h b/Engine/Plugins/Runtime/Metasound/Source/MetasoundFrontend/Public/MetasoundAutoConverterNode.h
--- a/Engine/Plugins/Runtime/Metasound/Source/MetasoundFrontend/Public/MetasoundAutoConverterNode.h
+++ b/Engine/Plugins/Runtime/Metasound/Source/MetasoundFrontend/Public/MetasoundAutoConverterNode.h
@@ -2,9 +2,9 @@
 #pragma once
 
 #include "CoreMinimal.h"
+#include "MetasoundBasicNode.h"
 #include "MetasoundBuilderInterface.h"
 #include "MetasoundBuildError.h"
-#include "MetasoundNode.h"
 #include "MetasoundNodeInterface.h"
 #include "MetasoundNodeRegistrationMacro.h"
 #include "MetasoundOperatorInterface.h"
@@ -17,11 +17,29 @@
 
 #include <type_traits>
 
-#define LOCTEXT_NAMESPACE "MetasoundFrontend"
-
-
 namespace Metasound
 {
+	namespace AutoConverterNodePrivate
+	{
+		struct FConvertDataTypeInfo
+		{
+			FName FromDataTypeName;
+			FText FromDataTypeText;
+			FName ToDataTypeName;
+			FText ToDataTypeText;
+			bool bIsFromEnum;
+			bool bIsToEnum;
+		};
+
+		METASOUNDFRONTEND_API FVertexName GetInputName(const FConvertDataTypeInfo& InInfo);
+
+		METASOUNDFRONTEND_API FVertexName GetOutputName(const FConvertDataTypeInfo& InInfo);
+
+		METASOUNDFRONTEND_API FVertexInterface CreateVertexInterface(const FConvertDataTypeInfo& InInfo);
+
+		METASOUNDFRONTEND_API FNodeClassMetadata CreateAutoConverterNodeMetadata(const FConvertDataTypeInfo& InInfo);
+	}
+
 	// Determines whether an auto converter node will be registered to convert 
 	// between two types. 
 	template<typename TFromDataType, typename TToDataType>
@@ -39,10 +57,23 @@ namespace Metasound
 	// This convenience node can be registered and will invoke static_cast<ToDataType>(FromDataType) every time it is executed, 
 	// with a special case for enum <-> int32 conversions. 
 	template<typename FromDataType, typename ToDataType>
-	class TAutoConverterNode : public FNode
+	class TAutoConverterNode : public FBasicNode
 	{
 		static_assert(TIsAutoConvertible<FromDataType, ToDataType>::Value,
 		"Tried to create an auto converter node between two types we can't static_cast between.");
+
+		static AutoConverterNodePrivate::FConvertDataTypeInfo GetConverterDataTypeInfo()
+		{
+			return AutoConverterNodePrivate::FConvertDataTypeInfo
+			{
+				GetMetasoundDataTypeName<FromDataType>(),
+				GetMetasoundDataTypeDisplayText<FromDataType>(),
+				GetMetasoundDataTypeName<ToDataType>(),
+				GetMetasoundDataTypeDisplayText<ToDataType>(),
+				TEnumTraits<FromDataType>::bIsEnum,
+				TEnumTraits<ToDataType>::bIsEnum
+			};
+		}
 		
 	public:
 		static const FVertexName& GetInputName()
@@ -59,60 +90,14 @@ namespace Metasound
 
 		static FVertexInterface DeclareVertexInterface()
 		{
-			static const FText InputDesc = METASOUND_LOCTEXT_FORMAT("AutoConvDisplayNamePatternFrom", "Input {0} value.", GetMetasoundDataTypeDisplayText<FromDataType>());
-			static const FText OutputDesc = METASOUND_LOCTEXT_FORMAT("AutoConvDisplayNamePatternTo", "Output {0} value.", GetMetasoundDataTypeDisplayText<ToDataType>());
-
-			return FVertexInterface(
-				FInputVertexInterface(
-					TInputDataVertex<FromDataType>(GetInputName(), FDataVertexMetadata{ InputDesc })
-				),
-				FOutputVertexInterface(
-					TOutputDataVertex<ToDataType>(GetOutputName(), FDataVertexMetadata{ OutputDesc })
-				)
-			);
+			using namespace AutoConverterNodePrivate;
+			return CreateVertexInterface(GetConverterDataTypeInfo());
 		}
 
 		static const FNodeClassMetadata& GetAutoConverterNodeMetadata()
 		{
-			auto InitNodeInfo = []() -> FNodeClassMetadata
-			{
-				FNodeDisplayStyle DisplayStyle;
-				DisplayStyle.bShowName = false;
-				DisplayStyle.ImageName = TEXT("MetasoundEditor.Graph.Node.Conversion");
-				DisplayStyle.bShowInputNames = false;
-				DisplayStyle.bShowOutputNames = false;
-
-				const FText FromTypeText = GetMetasoundDataTypeDisplayText<FromDataType>();
-				const FText ToTypeText = GetMetasoundDataTypeDisplayText<ToDataType>();
-
-				FNodeClassMetadata Info;
-				Info.ClassName = { TEXT("Convert"), GetMetasoundDataTypeName<ToDataType>(), GetMetasoundDataTypeName<FromDataType>() };
-				Info.MajorVersion = 1;
-				Info.MinorVersion = 0;
-				Info.DisplayName = METASOUND_LOCTEXT_FORMAT("Metasound_AutoConverterNodeDisplayNameFormat", "{0} to {1}", FromTypeText, ToTypeText);
-				Info.Description = METASOUND_LOCTEXT_FORMAT("Metasound_AutoConverterNodeDescriptionNameFormat", "Converts from {0} to {1}.", FromTypeText, ToTypeText);
-				Info.Author = PluginAuthor;
-				Info.DisplayStyle = DisplayStyle;
-				Info.PromptIfMissing = PluginNodeMissingPrompt;
-				Info.DefaultInterface = DeclareVertexInterface();
-
-				Info.CategoryHierarchy.Emplace(NodeCategories::Conversions);
-				if (TEnumTraits<FromDataType>::bIsEnum || TEnumTraits<ToDataType>::bIsEnum)
-				{
-					Info.CategoryHierarchy.Emplace(NodeCategories::EnumConversions);
-				}
-				
-				Info.Keywords =
-				{
-					METASOUND_LOCTEXT("MetasoundConvertKeyword", "Convert"),
-					GetMetasoundDataTypeDisplayText<FromDataType>(),
-					GetMetasoundDataTypeDisplayText<ToDataType>()
-				};
-
-				return Info;
-			};
-
-			static const FNodeClassMetadata Info = InitNodeInfo();
+			using namespace AutoConverterNodePrivate;
+			static const FNodeClassMetadata Info = CreateAutoConverterNodeMetadata(GetConverterDataTypeInfo());
 
 			return Info;
 		}
@@ -134,18 +119,6 @@ namespace Metasound
 
 			virtual ~FConverterOperator() {}
 
-			virtual FDataReferenceCollection GetInputs() const override
-			{
-				checkNoEntry();
-				return {};
-			}
-
-			virtual FDataReferenceCollection GetOutputs() const override
-			{
-				checkNoEntry();
-				return {};
-			}
-
 			virtual void BindInputs(FInputVertexInterfaceData& InVertexData) override
 			{
 				InVertexData.BindReadVertex(GetInputName(), FromData);
@@ -261,28 +234,17 @@ namespace Metasound
 		public:
 
 			TAutoConverterNode(const FNodeInitData& InInitData)
-				: FNode(InInitData.InstanceName, InInitData.InstanceID, GetAutoConverterNodeMetadata())
-				, Interface(DeclareVertexInterface())
-				, Factory(MakeOperatorFactoryRef<FCoverterOperatorFactory>())
+				: TAutoConverterNode(FNodeData(InInitData.InstanceName, InInitData.InstanceID, DeclareVertexInterface()), MakeShared<const FNodeClassMetadata>(GetAutoConverterNodeMetadata()))
 			{
 			}
 
-			virtual ~TAutoConverterNode() = default;
-
-			virtual const FVertexInterface& GetVertexInterface() const override
-			{
-				return Interface;
-			}
-
-			virtual bool SetVertexInterface(const FVertexInterface& InInterface) override
+			TAutoConverterNode(FNodeData InNodeData, TSharedRef<const FNodeClassMetadata> InClassMetadata)
+				: FBasicNode(MoveTemp(InNodeData), MoveTemp(InClassMetadata))
+				, Factory(MakeOperatorFactoryRef<FCoverterOperatorFactory>())
 			{
-				return Interface == InInterface;
 			}
 
-			virtual bool IsVertexInterfaceSupported(const FVertexInterface& InInterface) const override
-			{
-				return Interface == InInterface;
-			}
+			virtual ~TAutoConverterNode() = default;
 
 			virtual FOperatorFactorySharedRef GetDefaultOperatorFactory() const override
 			{
@@ -290,8 +252,6 @@ namespace Metasound
 			}
 
 		private:
-			FVertexInterface Interface;
 			FOperatorFactorySharedRef Factory;
 	};
 } // namespace Metasound
-#undef LOCTEXT_NAMESPACE
diff --git a/Engine/Plugins/Runtime/Metasound/Source/MetasoundFrontend/Public/MetasoundReceiveNode.h b/Engine/Plugins/Runtime/Metasound/Source/MetasoundFrontend/Public/MetasoundReceiveNode.h
--- a/Engine/Plugins/Runtime/Metasound/Source/MetasoundFrontend/Public/MetasoundReceiveNode.h
+++ b/Engine/Plugins/Runtime/Metasound/Source/MetasoundFrontend/Public/MetasoundReceiveNode.h
@@ -9,6 +9,7 @@
 #include "MetasoundDataFactory.h"
 #include "MetasoundDataReference.h"
 #include "MetasoundExecutableOperator.h"
+#include "MetasoundFacade.h"
 #include "MetasoundParamHelper.h"
 #include "MetasoundRouter.h"
 
@@ -31,239 +32,215 @@ namespace Metasound
 		METASOUNDFRONTEND_API int32 GetCurrentMinorVersion();
 	};
 
-	template<typename TDataType>
-	class TReceiveNode : public FNode
+	namespace ReceiveNodePrivate
 	{
-	public:
-		static FVertexInterface DeclareVertexInterface()
+		template<typename TDataType>
+		class TReceiverOperator : public TExecutableOperator<TReceiverOperator<TDataType>>
 		{
-			using namespace ReceiveNodeInfo;
-			static const FDataVertexMetadata AddressInputMetadata
-			{
-				  FText::GetEmpty() // description
-				, METASOUND_GET_PARAM_DISPLAYNAME(AddressInput) // display name
-			};
-			static const FDataVertexMetadata DefaultDataInputMetadata
-			{
-				  FText::GetEmpty() // description
-				, METASOUND_GET_PARAM_DISPLAYNAME(DefaultDataInput) // display name
-			};
-			static const FDataVertexMetadata OutputMetadata
+			TReceiverOperator() = delete;
+
+		public:
+			static FVertexInterface DeclareVertexInterface()
 			{
-				  FText::GetEmpty() // description
-				, METASOUND_GET_PARAM_DISPLAYNAME(Output) // display name
-			};
-			return FVertexInterface(
-				FInputVertexInterface(
-					TInputDataVertex<FSendAddress>(METASOUND_GET_PARAM_NAME(AddressInput), AddressInputMetadata),
-					TInputDataVertex<TDataType>(METASOUND_GET_PARAM_NAME(DefaultDataInput), DefaultDataInputMetadata)
-				),
-				FOutputVertexInterface(
-					TOutputDataVertex<TDataType>(METASOUND_GET_PARAM_NAME(Output), OutputMetadata)
-				)
-			);
-		}
+				using namespace ReceiveNodeInfo;
+				static const FDataVertexMetadata AddressInputMetadata
+				{
+					  FText::GetEmpty() // description
+					, METASOUND_GET_PARAM_DISPLAYNAME(AddressInput) // display name
+				};
+				static const FDataVertexMetadata DefaultDataInputMetadata
+				{
+					  FText::GetEmpty() // description
+					, METASOUND_GET_PARAM_DISPLAYNAME(DefaultDataInput) // display name
+				};
+				static const FDataVertexMetadata OutputMetadata
+				{
+					  FText::GetEmpty() // description
+					, METASOUND_GET_PARAM_DISPLAYNAME(Output) // display name
+				};
+				return FVertexInterface(
+					FInputVertexInterface(
+						TInputDataVertex<FSendAddress>(METASOUND_GET_PARAM_NAME(AddressInput), AddressInputMetadata),
+						TInputDataVertex<TDataType>(METASOUND_GET_PARAM_NAME(DefaultDataInput), DefaultDataInputMetadata)
+					),
+					FOutputVertexInterface(
+						TOutputDataVertex<TDataType>(METASOUND_GET_PARAM_NAME(Output), OutputMetadata)
+					)
+				);
+			}
 
-		static const FNodeClassMetadata& GetNodeInfo()
-		{
-			auto InitNodeInfo = []() -> FNodeClassMetadata
+			static const FNodeClassMetadata& GetNodeInfo()
 			{
-				FNodeClassMetadata Info;
-				Info.ClassName = ReceiveNodeInfo::GetClassNameForDataType(GetMetasoundDataTypeName<TDataType>());
-				Info.MajorVersion = ReceiveNodeInfo::GetCurrentMajorVersion();
-				Info.MinorVersion = ReceiveNodeInfo::GetCurrentMinorVersion();
-				Info.DisplayName = METASOUND_LOCTEXT_FORMAT("Metasound_ReceiveNodeDisplayNameFormat", "Receive {0}", GetMetasoundDataTypeDisplayText<TDataType>());
-				Info.Description = METASOUND_LOCTEXT("Metasound_ReceiveNodeDescription", "Receives data from a send node with the same name.");
-				Info.Author = PluginAuthor;
-				Info.PromptIfMissing = PluginNodeMissingPrompt;
-				Info.DefaultInterface = DeclareVertexInterface();
-				Info.CategoryHierarchy = { METASOUND_LOCTEXT("Metasound_TransmissionNodeCategory", "Transmission") };
-				Info.Keywords = { };
-
-				// Then send & receive nodes do not work as expected, particularly 
-				// around multiple-consumer scenarios. Deprecate them to avoid
-				// metasound assets from relying on send & receive nodes. 
-				Info.bDeprecated = true; 
+				auto InitNodeInfo = []() -> FNodeClassMetadata
+				{
+					FNodeClassMetadata Info;
+					Info.ClassName = ReceiveNodeInfo::GetClassNameForDataType(GetMetasoundDataTypeName<TDataType>());
+					Info.MajorVersion = ReceiveNodeInfo::GetCurrentMajorVersion();
+					Info.MinorVersion = ReceiveNodeInfo::GetCurrentMinorVersion();
+					Info.DisplayName = METASOUND_LOCTEXT_FORMAT("Metasound_ReceiveNodeDisplayNameFormat", "Receive {0}", GetMetasoundDataTypeDisplayText<TDataType>());
+					Info.Description = METASOUND_LOCTEXT("Metasound_ReceiveNodeDescription", "Receives data from a send node with the same name.");
+					Info.Author = PluginAuthor;
+					Info.PromptIfMissing = PluginNodeMissingPrompt;
+					Info.DefaultInterface = DeclareVertexInterface();
+					Info.CategoryHierarchy = { METASOUND_LOCTEXT("Metasound_TransmissionNodeCategory", "Transmission") };
+					Info.Keywords = { };
+
+					// Then send & receive nodes do not work as expected, particularly 
+					// around multiple-consumer scenarios. Deprecate them to avoid
+					// metasound assets from relying on send & receive nodes. 
+					Info.bDeprecated = true; 
+
+					return Info;
+				};
+
+				static const FNodeClassMetadata Info = InitNodeInfo();
 
 				return Info;
-			};
+			}
 
-			static const FNodeClassMetadata Info = InitNodeInfo();
+			TReceiverOperator(TDataReadReference<TDataType> InInitDataRef, TDataWriteReference<TDataType> InOutDataRef, TDataReadReference<FSendAddress> InSendAddress, const FOperatorSettings& InOperatorSettings)
+				: bHasNotReceivedData(true)
+				, DefaultData(InInitDataRef)
+				, OutputData(InOutDataRef)
+				, SendAddress(InSendAddress)
+				, CachedSendAddress(*InSendAddress)
+				, CachedReceiverParams({InOperatorSettings})
+				, Receiver(nullptr)
+			{
+				Receiver = CreateNewReceiver();
+			}
 
-			return Info;
-		}
+			virtual ~TReceiverOperator() 
+			{
+				ResetReceiverAndCleanupChannel();
+			}
 
-	private:
-		class TReceiverOperator : public TExecutableOperator<TReceiverOperator>
-		{
-			TReceiverOperator() = delete;
+			virtual void BindInputs(FInputVertexInterfaceData& InOutVertexData) override
+			{
+				using namespace ReceiveNodeInfo; 
+				InOutVertexData.BindReadVertex<TDataType>(METASOUND_GET_PARAM_NAME(DefaultDataInput), DefaultData);
+				InOutVertexData.BindReadVertex<FSendAddress>(METASOUND_GET_PARAM_NAME(AddressInput), SendAddress);
+			}
 
-			public:
-				TReceiverOperator(TDataReadReference<TDataType> InInitDataRef, TDataWriteReference<TDataType> InOutDataRef, TDataReadReference<FSendAddress> InSendAddress, const FOperatorSettings& InOperatorSettings)
-					: bHasNotReceivedData(true)
-					, DefaultData(InInitDataRef)
-					, OutputData(InOutDataRef)
-					, SendAddress(InSendAddress)
-					, CachedSendAddress(*InSendAddress)
-					, CachedReceiverParams({InOperatorSettings})
-					, Receiver(nullptr)
-				{
-					Receiver = CreateNewReceiver();
-				}
+			virtual void BindOutputs(FOutputVertexInterfaceData& InOutVertexData) override
+			{
+				using namespace ReceiveNodeInfo;
+				InOutVertexData.BindReadVertex<TDataType>(METASOUND_GET_PARAM_NAME(Output), OutputData);
+			}
 
-				virtual ~TReceiverOperator() 
+			void Execute()
+			{
+				if (*SendAddress != CachedSendAddress)
 				{
 					ResetReceiverAndCleanupChannel();
-				}
-
-				virtual void BindInputs(FInputVertexInterfaceData& InOutVertexData) override
-				{
-					using namespace ReceiveNodeInfo; 
-					InOutVertexData.BindReadVertex<TDataType>(METASOUND_GET_PARAM_NAME(DefaultDataInput), DefaultData);
-					InOutVertexData.BindReadVertex<FSendAddress>(METASOUND_GET_PARAM_NAME(AddressInput), SendAddress);
-				}
+					CachedSendAddress = *SendAddress;
 
-				virtual void BindOutputs(FOutputVertexInterfaceData& InOutVertexData) override
-				{
-					using namespace ReceiveNodeInfo;
-					InOutVertexData.BindReadVertex<TDataType>(METASOUND_GET_PARAM_NAME(Output), OutputData);
+					Receiver = CreateNewReceiver();
 				}
 
-				void Execute()
+				bool bHasNewData = false;
+				if (ensure(Receiver.IsValid()))
 				{
-					if (*SendAddress != CachedSendAddress)
-					{
-						ResetReceiverAndCleanupChannel();
-						CachedSendAddress = *SendAddress;
-
-						Receiver = CreateNewReceiver();
-					}
-
-					bool bHasNewData = false;
-					if (ensure(Receiver.IsValid()))
-					{
-						bHasNewData = Receiver->CanPop();
-						if (bHasNewData)
-						{
-							Receiver->Pop(*OutputData);
-							bHasNotReceivedData = false;
-						}
-					}
-
-					if (bHasNotReceivedData)
-					{
-						*OutputData = *DefaultData;
-						bHasNewData = true;
-					}
-
-					if (TExecutableDataType<TDataType>::bIsExecutable)
+					bHasNewData = Receiver->CanPop();
+					if (bHasNewData)
 					{
-						TExecutableDataType<TDataType>::ExecuteInline(*OutputData, bHasNewData);
+						Receiver->Pop(*OutputData);
+						bHasNotReceivedData = false;
 					}
 				}
 
-				void Reset(const IOperator::FResetParams& InParams)
+				if (bHasNotReceivedData)
 				{
 					*OutputData = *DefaultData;
-					bHasNotReceivedData = true;
+					bHasNewData = true;
 				}
 
-			private:
-
-				FSendAddress GetSendAddressWithDataType(const FSendAddress& InAddress) const 
-				{
-					// The data type of a send address is inferred by the underlying
-					// data type of this node. A full send address, including the data type,
-					// cannot be constructed from a literal. 
-					return FSendAddress{ InAddress.GetChannelName(), GetMetasoundDataTypeName<TDataType>(), InAddress.GetInstanceID() };
-				}
-
-				TReceiverPtr<TDataType> CreateNewReceiver() const
-				{
-					if (ensure(SendAddress->GetDataType().IsNone() || (GetMetasoundDataTypeName<TDataType>() == SendAddress->GetDataType())))
-					{
-						return FDataTransmissionCenter::Get().RegisterNewReceiver<TDataType>(GetSendAddressWithDataType(*SendAddress), CachedReceiverParams);
-					}
-					return TReceiverPtr<TDataType>(nullptr);
-				}
-				
-				void ResetReceiverAndCleanupChannel()
+				if (TExecutableDataType<TDataType>::bIsExecutable)
 				{
-					Receiver.Reset();
-					FDataTransmissionCenter::Get().UnregisterDataChannelIfUnconnected(GetSendAddressWithDataType(CachedSendAddress));
+					TExecutableDataType<TDataType>::ExecuteInline(*OutputData, bHasNewData);
 				}
+			}
 
-				bool bHasNotReceivedData;
-
-				TDataReadReference<TDataType> DefaultData;
-				TDataWriteReference<TDataType> OutputData;
-				TDataReadReference<FSendAddress> SendAddress;
-
-				FSendAddress CachedSendAddress;
-				FReceiverInitParams CachedReceiverParams;
+			void Reset(const IOperator::FResetParams& InParams)
+			{
+				*OutputData = *DefaultData;
+				bHasNotReceivedData = true;
+			}
 
-				TReceiverPtr<TDataType> Receiver;
-		};
+			static TUniquePtr<IOperator> CreateOperator(const FBuildOperatorParams& InParams, FBuildResults& OutResults)
+			{
+				using namespace ReceiveNodeInfo; 
 
-		class FReceiverOperatorFactory : public IOperatorFactory
-		{
-			public:
-				FReceiverOperatorFactory() = default;
+				TDataReadReference<TDataType> DefaultReadRef = TDataReadReferenceFactory<TDataType>::CreateAny(InParams.OperatorSettings);
 
-				virtual TUniquePtr<IOperator> CreateOperator(const FBuildOperatorParams& InParams, FBuildResults& OutResults) override
+				if (InParams.InputData.IsVertexBound(METASOUND_GET_PARAM_NAME(DefaultDataInput)))
 				{
-					using namespace ReceiveNodeInfo; 
-
-					TDataReadReference<TDataType> DefaultReadRef = TDataReadReferenceFactory<TDataType>::CreateAny(InParams.OperatorSettings);
-
-					if (InParams.InputData.IsVertexBound(METASOUND_GET_PARAM_NAME(DefaultDataInput)))
-					{
-						DefaultReadRef = InParams.InputData.GetDataReadReference<TDataType>(METASOUND_GET_PARAM_NAME(DefaultDataInput));
-					}
-
-					return MakeUnique<TReceiverOperator>(
-						DefaultReadRef,
-						TDataWriteReferenceFactory<TDataType>::CreateAny(InParams.OperatorSettings, *DefaultReadRef),
-						InParams.InputData.GetOrConstructDataReadReference<FSendAddress>(METASOUND_GET_PARAM_NAME(AddressInput)),
-						InParams.OperatorSettings
-						);
+					DefaultReadRef = InParams.InputData.GetDataReadReference<TDataType>(METASOUND_GET_PARAM_NAME(DefaultDataInput));
 				}
-		};
-
-		public:
 
-			TReceiveNode(const FNodeInitData& InInitData)
-				: FNode(InInitData.InstanceName, InInitData.InstanceID, GetNodeInfo())
-				, Interface(DeclareVertexInterface())
-				, Factory(MakeOperatorFactoryRef<FReceiverOperatorFactory>())
-			{
+				return MakeUnique<TReceiverOperator>(
+					DefaultReadRef,
+					TDataWriteReferenceFactory<TDataType>::CreateAny(InParams.OperatorSettings, *DefaultReadRef),
+					InParams.InputData.GetOrConstructDataReadReference<FSendAddress>(METASOUND_GET_PARAM_NAME(AddressInput)),
+					InParams.OperatorSettings
+					);
 			}
 
-			virtual ~TReceiveNode() = default;
+		private:
 
-			virtual const FVertexInterface& GetVertexInterface() const override
+			FSendAddress GetSendAddressWithDataType(const FSendAddress& InAddress) const 
 			{
-				return Interface;
+				// The data type of a send address is inferred by the underlying
+				// data type of this node. A full send address, including the data type,
+				// cannot be constructed from a literal. 
+				return FSendAddress{ InAddress.GetChannelName(), GetMetasoundDataTypeName<TDataType>(), InAddress.GetInstanceID() };
 			}
 
-			virtual bool SetVertexInterface(const FVertexInterface& InInterface) override
+			TReceiverPtr<TDataType> CreateNewReceiver() const
 			{
-				return Interface == InInterface;
+				if (ensure(SendAddress->GetDataType().IsNone() || (GetMetasoundDataTypeName<TDataType>() == SendAddress->GetDataType())))
+				{
+					return FDataTransmissionCenter::Get().RegisterNewReceiver<TDataType>(GetSendAddressWithDataType(*SendAddress), CachedReceiverParams);
+				}
+				return TReceiverPtr<TDataType>(nullptr);
 			}
-
-			virtual bool IsVertexInterfaceSupported(const FVertexInterface& InInterface) const override
+			
+			void ResetReceiverAndCleanupChannel()
 			{
-				return Interface == InInterface;
+				Receiver.Reset();
+				FDataTransmissionCenter::Get().UnregisterDataChannelIfUnconnected(GetSendAddressWithDataType(CachedSendAddress));
 			}
 
-			virtual FOperatorFactorySharedRef GetDefaultOperatorFactory() const override
-			{
-				return Factory;
-			}
+			bool bHasNotReceivedData;
 
-		private:
-			FVertexInterface Interface;
-			FOperatorFactorySharedRef Factory;
+			TDataReadReference<TDataType> DefaultData;
+			TDataWriteReference<TDataType> OutputData;
+			TDataReadReference<FSendAddress> SendAddress;
+
+			FSendAddress CachedSendAddress;
+			FReceiverInitParams CachedReceiverParams;
+
+			TReceiverPtr<TDataType> Receiver;
+		};
+	}
+
+	template<typename TDataType>
+	class TReceiveNode : public TNodeFacade<ReceiveNodePrivate::TReceiverOperator<TDataType>>
+	{
+	public:
+		using FOperator = ReceiveNodePrivate::TReceiverOperator<TDataType>;
+		using FSuper = TNodeFacade<FOperator>;
+		using FSuper::FSuper;
+
+		static FVertexInterface DeclareVertexInterface()
+		{
+			return FOperator::DeclareVertexInterface();
+		}
+
+		static const FNodeClassMetadata& GetNodeInfo()
+		{
+			return FOperator::GetNodeInfo();
+		}
 	};
 }
 #undef LOCTEXT_NAMESPACE
diff --git a/Engine/Plugins/Runtime/Metasound/Source/MetasoundFrontend/Public/MetasoundSendNode.h b/Engine/Plugins/Runtime/Metasound/Source/MetasoundFrontend/Public/MetasoundSendNode.h
--- a/Engine/Plugins/Runtime/Metasound/Source/MetasoundFrontend/Public/MetasoundSendNode.h
+++ b/Engine/Plugins/Runtime/Metasound/Source/MetasoundFrontend/Public/MetasoundSendNode.h
@@ -2,6 +2,7 @@
 #pragma once
 
 #include "MetasoundBuilderInterface.h"
+#include "MetasoundFacade.h"
 #include "MetasoundNode.h"
 #include "MetasoundNodeInterface.h"
 #include "MetasoundNodeRegistrationMacro.h"
@@ -24,193 +25,190 @@ namespace Metasound
 		METASOUND_PARAM(AddressInput, "Address", "Address")
 	}
 
-	template<typename TDataType>
-	class TSendNode : public FNode
+	namespace SendNodePrivate
 	{
-	public:
-		static const FVertexName& GetSendInputName()
+		template<typename TDataType>
+		class TSendOperator : public TExecutableOperator<TSendOperator<TDataType>>
 		{
-			static const FVertexName& SendInput = GetMetasoundDataTypeName<TDataType>();
-			return SendInput;
-		}
+		public:
 
-		static FVertexInterface DeclareVertexInterface()
-		{
-			using namespace SendVertexNames; 
-			static const FDataVertexMetadata AddressInputMetadata
+			static const FVertexName& GetSendInputName()
 			{
-				  FText::GetEmpty() // description
-				, METASOUND_GET_PARAM_DISPLAYNAME(AddressInput) // display name
-			};
-
-			return FVertexInterface(
-				FInputVertexInterface(
-					TInputDataVertex<FSendAddress>(METASOUND_GET_PARAM_NAME(AddressInput), AddressInputMetadata),
-					TInputDataVertex<TDataType>(GetSendInputName(), FDataVertexMetadata{ FText::GetEmpty() })
-				),
-				FOutputVertexInterface(
-				)
-			);
-		}
+				static const FVertexName& SendInput = GetMetasoundDataTypeName<TDataType>();
+				return SendInput;
+			}
 
-		static const FNodeClassMetadata& GetNodeInfo()
-		{
-			auto InitNodeInfo = []() -> FNodeClassMetadata
+			TSendOperator(TDataReadReference<TDataType> InInputData, TDataReadReference<FSendAddress> InSendAddress, const FOperatorSettings& InOperatorSettings)
+				: InputData(InInputData)
+				, SendAddress(InSendAddress)
+				, CachedSendAddress(*InSendAddress)
+				, CachedSenderParams({InOperatorSettings, 0.0f})
+				, Sender(nullptr)
 			{
-				const FVertexName& InputName = GetSendInputName();
-				FNodeClassMetadata Info;
-
-				Info.ClassName = { "Send", GetMetasoundDataTypeName<TDataType>(), FName() };
-				Info.MajorVersion = 1;
-				Info.MinorVersion = 0;
-				Info.DisplayName = METASOUND_LOCTEXT_FORMAT("Metasound_SendNodeDisplayNameFormat", "Send {0}", GetMetasoundDataTypeDisplayText<TDataType>());
-				Info.Description = METASOUND_LOCTEXT("Metasound_SendNodeDescription", "Sends data from a send node with the same name.");
-				Info.Author = PluginAuthor;
-				Info.PromptIfMissing = PluginNodeMissingPrompt;
-				Info.DefaultInterface = DeclareVertexInterface();
-				Info.CategoryHierarchy = { METASOUND_LOCTEXT("Metasound_TransmissionNodeCategory", "Transmission") };
-				Info.Keywords = { };
-
-				// Then send & receive nodes do not work as expected, particularly 
-				// around multiple-consumer scenarios. Deprecate them to avoid
-				// metasound assets from relying on send & receive nodes. 
-				Info.bDeprecated = true;
-
-				return Info;
-			};
-
-			static const FNodeClassMetadata Info = InitNodeInfo();
+				Sender = CreateNewSender();
+			}
 
-			return Info;
-		}
+			virtual ~TSendOperator() 
+			{
+				ResetSenderAndCleanupChannel();
+			}
 
+			virtual void BindInputs(FInputVertexInterfaceData& InOutVertexData) override
+			{
+				using namespace SendVertexNames; 
+				InOutVertexData.BindReadVertex<FSendAddress>(METASOUND_GET_PARAM_NAME(AddressInput), SendAddress);
+				InOutVertexData.BindReadVertex<TDataType>(GetSendInputName(), InputData);
+			}
 
-	private:
-		class TSendOperator : public TExecutableOperator<TSendOperator>
-		{
-			public:
-
-				TSendOperator(TDataReadReference<TDataType> InInputData, TDataReadReference<FSendAddress> InSendAddress, const FOperatorSettings& InOperatorSettings)
-					: InputData(InInputData)
-					, SendAddress(InSendAddress)
-					, CachedSendAddress(*InSendAddress)
-					, CachedSenderParams({InOperatorSettings, 0.0f})
-					, Sender(nullptr)
-				{
-					Sender = CreateNewSender();
-				}
+			virtual void BindOutputs(FOutputVertexInterfaceData& InOutVertexData) override
+			{
+			}
 
-				virtual ~TSendOperator() 
+			void Execute()
+			{
+				if (*SendAddress != CachedSendAddress)
 				{
 					ResetSenderAndCleanupChannel();
+					CachedSendAddress = *SendAddress;
+					Sender = CreateNewSender();
+					check(Sender.IsValid());
 				}
 
-				virtual void BindInputs(FInputVertexInterfaceData& InOutVertexData) override
-				{
-					using namespace SendVertexNames; 
-					InOutVertexData.BindReadVertex<FSendAddress>(METASOUND_GET_PARAM_NAME(AddressInput), SendAddress);
-					InOutVertexData.BindReadVertex<TDataType>(GetSendInputName(), InputData);
-				}
+				Sender->Push(*InputData);
+			}
 
-				virtual void BindOutputs(FOutputVertexInterfaceData& InOutVertexData) override
-				{
-				}
+			void Reset(const IOperator::FResetParams& InParams)
+			{
+				ResetSenderAndCleanupChannel();
+				CachedSendAddress = *SendAddress;
+				Sender = CreateNewSender();
+				check(Sender.IsValid());
+			}
 
-				void Execute()
+			static FVertexInterface DeclareVertexInterface()
+			{
+				using namespace SendVertexNames; 
+				static const FDataVertexMetadata AddressInputMetadata
 				{
-					if (*SendAddress != CachedSendAddress)
-					{
-						ResetSenderAndCleanupChannel();
-						CachedSendAddress = *SendAddress;
-						Sender = CreateNewSender();
-						check(Sender.IsValid());
-					}
-
-					Sender->Push(*InputData);
-				}
+					  FText::GetEmpty() // description
+					, METASOUND_GET_PARAM_DISPLAYNAME(AddressInput) // display name
+				};
+
+				return FVertexInterface(
+					FInputVertexInterface(
+						TInputDataVertex<FSendAddress>(METASOUND_GET_PARAM_NAME(AddressInput), AddressInputMetadata),
+						TInputDataVertex<TDataType>(GetSendInputName(), FDataVertexMetadata{ FText::GetEmpty() })
+					),
+					FOutputVertexInterface(
+					)
+				);
+			}
 
-				void Reset(const IOperator::FResetParams& InParams)
+			static const FNodeClassMetadata& GetNodeInfo()
+			{
+				auto InitNodeInfo = []() -> FNodeClassMetadata
 				{
-					ResetSenderAndCleanupChannel();
-					CachedSendAddress = *SendAddress;
-					Sender = CreateNewSender();
-					check(Sender.IsValid());
-				}
+					const FVertexName& InputName = GetSendInputName();
+					FNodeClassMetadata Info;
+
+					Info.ClassName = { "Send", GetMetasoundDataTypeName<TDataType>(), FName() };
+					Info.MajorVersion = 1;
+					Info.MinorVersion = 0;
+					Info.DisplayName = METASOUND_LOCTEXT_FORMAT("Metasound_SendNodeDisplayNameFormat", "Send {0}", GetMetasoundDataTypeDisplayText<TDataType>());
+					Info.Description = METASOUND_LOCTEXT("Metasound_SendNodeDescription", "Sends data from a send node with the same name.");
+					Info.Author = PluginAuthor;
+					Info.PromptIfMissing = PluginNodeMissingPrompt;
+					Info.DefaultInterface = DeclareVertexInterface();
+					Info.CategoryHierarchy = { METASOUND_LOCTEXT("Metasound_TransmissionNodeCategory", "Transmission") };
+					Info.Keywords = { };
+
+					// Then send & receive nodes do not work as expected, particularly 
+					// around multiple-consumer scenarios. Deprecate them to avoid
+					// metasound assets from relying on send & receive nodes. 
+					Info.bDeprecated = true;
+
+					return Info;
+				};
+
+				static const FNodeClassMetadata Info = InitNodeInfo();
 
+				return Info;
+			}
 
-			private:
-				FSendAddress GetSendAddressWithDataType(const FSendAddress& InAddress) const 
-				{
-					// The data type of a send address is inferred by the underlying
-					// data type of this node. A full send address, including the data type,
-					// cannot be constructed from a literal. 
-					return FSendAddress{ InAddress.GetChannelName(), GetMetasoundDataTypeName<TDataType>(), InAddress.GetInstanceID() };
-				}
+			static TUniquePtr<IOperator> CreateOperator(const FBuildOperatorParams& InParams, FBuildResults& OutResults)
+			{
+				using namespace SendVertexNames;
 
-				TSenderPtr<TDataType> CreateNewSender() const
+				if (InParams.InputData.IsVertexBound(GetSendInputName()))
 				{
-					if (ensure(SendAddress->GetDataType().IsNone() || (GetMetasoundDataTypeName<TDataType>() == SendAddress->GetDataType())))
-					{
-						return FDataTransmissionCenter::Get().RegisterNewSender<TDataType>(GetSendAddressWithDataType(*SendAddress), CachedSenderParams);
-					}
-					return TSenderPtr<TDataType>(nullptr);
+					return MakeUnique<SendNodePrivate::TSendOperator<TDataType>>(InParams.InputData.GetDataReadReference<TDataType>(GetSendInputName()),
+						InParams.InputData.GetOrConstructDataReadReference<FSendAddress>(METASOUND_GET_PARAM_NAME(AddressInput)),
+						InParams.OperatorSettings
+					);
 				}
-
-				void ResetSenderAndCleanupChannel()
+				else
 				{
-					Sender.Reset();
-					FDataTransmissionCenter::Get().UnregisterDataChannelIfUnconnected(GetSendAddressWithDataType(CachedSendAddress));
+					// No input hook up to send, so this node can no-op
+					return MakeUnique<FNoOpOperator>();
 				}
+			}
 
-				TDataReadReference<TDataType> InputData;
-				TDataReadReference<FSendAddress> SendAddress;
-				FSendAddress CachedSendAddress;
-				FSenderInitParams CachedSenderParams;
-
-				TSenderPtr<TDataType> Sender;
-		};
 
-		class FSendOperatorFactory : public IOperatorFactory
-		{
-			public:
-				FSendOperatorFactory() = default;
+		private:
+			FSendAddress GetSendAddressWithDataType(const FSendAddress& InAddress) const 
+			{
+				// The data type of a send address is inferred by the underlying
+				// data type of this node. A full send address, including the data type,
+				// cannot be constructed from a literal. 
+				return FSendAddress{ InAddress.GetChannelName(), GetMetasoundDataTypeName<TDataType>(), InAddress.GetInstanceID() };
+			}
 
-				virtual TUniquePtr<IOperator> CreateOperator(const FBuildOperatorParams& InParams, FBuildResults& OutResults) override
+			TSenderPtr<TDataType> CreateNewSender() const
+			{
+				if (ensure(SendAddress->GetDataType().IsNone() || (GetMetasoundDataTypeName<TDataType>() == SendAddress->GetDataType())))
 				{
-					using namespace SendVertexNames;
-
-					if (InParams.InputData.IsVertexBound(GetSendInputName()))
-					{
-						return MakeUnique<TSendOperator>(InParams.InputData.GetDataReadReference<TDataType>(GetSendInputName()),
-							InParams.InputData.GetOrConstructDataReadReference<FSendAddress>(METASOUND_GET_PARAM_NAME(AddressInput)),
-							InParams.OperatorSettings
-						);
-					}
-					else
-					{
-						// No input hook up to send, so this node can no-op
-						return MakeUnique<FNoOpOperator>();
-					}
+					return FDataTransmissionCenter::Get().RegisterNewSender<TDataType>(GetSendAddressWithDataType(*SendAddress), CachedSenderParams);
 				}
-		};
-
-		public:
+				return TSenderPtr<TDataType>(nullptr);
+			}
 
-			TSendNode(const FNodeInitData& InInitData)
-				: FNode(InInitData.InstanceName, InInitData.InstanceID, GetNodeInfo())
-				, Factory(MakeOperatorFactoryRef<FSendOperatorFactory>())
+			void ResetSenderAndCleanupChannel()
 			{
+				Sender.Reset();
+				FDataTransmissionCenter::Get().UnregisterDataChannelIfUnconnected(GetSendAddressWithDataType(CachedSendAddress));
 			}
 
-			virtual ~TSendNode() = default;
+			TDataReadReference<TDataType> InputData;
+			TDataReadReference<FSendAddress> SendAddress;
+			FSendAddress CachedSendAddress;
+			FSenderInitParams CachedSenderParams;
 
-			virtual FOperatorFactorySharedRef GetDefaultOperatorFactory() const override
-			{
-				return Factory;
-			}
+			TSenderPtr<TDataType> Sender;
+		};
+	}
 
-		private:
-			FOperatorFactorySharedRef Factory;
+	template<typename TDataType>
+	class TSendNode : public TNodeFacade<SendNodePrivate::TSendOperator<TDataType>>
+	{
+	public:
+		using FOperator = SendNodePrivate::TSendOperator<TDataType>;
+		using FSuper = TNodeFacade<SendNodePrivate::TSendOperator<TDataType>>;
+		using FSuper::FSuper;
+
+		static const FVertexName& GetSendInputName()
+		{
+			return FOperator::GetSendInputName();
+		}
+
+		static FVertexInterface DeclareVertexInterface()
+		{
+			return FOperator::DeclareVertexInterface();
+		}
+
+		static const FNodeClassMetadata& GetNodeInfo()
+		{
+			return FOperator::GetNodeInfo();
+		}
 	};
 }
 #undef LOCTEXT_NAMESPACE
diff --git a/Engine/Plugins/Runtime/Metasound/Source/MetasoundStandardNodes/Private/MetasoundADEnvelopeNode.cpp b/Engine/Plugins/Runtime/Metasound/Source/MetasoundStandardNodes/Private/MetasoundADEnvelopeNode.cpp
--- a/Engine/Plugins/Runtime/Metasound/Source/MetasoundStandardNodes/Private/MetasoundADEnvelopeNode.cpp
+++ b/Engine/Plugins/Runtime/Metasound/Source/MetasoundStandardNodes/Private/MetasoundADEnvelopeNode.cpp
@@ -349,22 +349,6 @@ namespace Metasound
 			InOutVertexData.BindReadVertex(METASOUND_GET_PARAM_NAME(OutputEnvelopeValue), OutputEnvelope);
 		}
 
-		virtual FDataReferenceCollection GetInputs() const override
-		{
-			// This should never be called. Bind(...) is called instead. This method
-			// exists as a stop-gap until the API can be deprecated and removed.
-			checkNoEntry();
-			return {};
-		}
-
-		virtual FDataReferenceCollection GetOutputs() const override
-		{
-			// This should never be called. Bind(...) is called instead. This method
-			// exists as a stop-gap until the API can be deprecated and removed.
-			checkNoEntry();
-			return {};
-		}
-
 		void UpdateParams()
 		{
 			float AttackTimeSeconds = AttackTime->GetSeconds();
@@ -479,18 +463,7 @@ namespace Metasound
 	 *  Creates an Attack/Decay envelope node.
 	 */
 	template<typename ValueType>
-	class METASOUNDSTANDARDNODES_API TADEnvelopeNode : public FNodeFacade
-	{
-	public:
-		/**
-		 * Constructor used by the Metasound Frontend.
-		 */
-		TADEnvelopeNode(const FNodeInitData& InInitData)
-			: FNodeFacade(InInitData.InstanceName, InInitData.InstanceID, TFacadeOperatorClass<TADEnvelopeNodeOperator<ValueType>>())
-		{}
-
-		virtual ~TADEnvelopeNode() = default;
-	};
+	using TADEnvelopeNode = TNodeFacade<TADEnvelopeNodeOperator<ValueType>>;
 
 	using FADEnvelopeNodeFloat = TADEnvelopeNode<float>;
 	METASOUND_REGISTER_NODE(FADEnvelopeNodeFloat)
diff --git a/Engine/Plugins/Runtime/Metasound/Source/MetasoundStandardNodes/Private/MetasoundADSREnvelopeNode.cpp b/Engine/Plugins/Runtime/Metasound/Source/MetasoundStandardNodes/Private/MetasoundADSREnvelopeNode.cpp
--- a/Engine/Plugins/Runtime/Metasound/Source/MetasoundStandardNodes/Private/MetasoundADSREnvelopeNode.cpp
+++ b/Engine/Plugins/Runtime/Metasound/Source/MetasoundStandardNodes/Private/MetasoundADSREnvelopeNode.cpp
@@ -651,22 +651,6 @@ namespace Metasound
 			InOutVertexData.BindReadVertex(METASOUND_GET_PARAM_NAME(OutputEnvelopeValue), OutputEnvelope);
 		}
 
-		virtual FDataReferenceCollection GetInputs() const override
-		{
-			// This should never be called. Bind(...) is called instead. This method
-			// exists as a stop-gap until the API can be deprecated and removed.
-			checkNoEntry();
-			return {};
-		}
-
-		virtual FDataReferenceCollection GetOutputs() const override
-		{
-			// This should never be called. Bind(...) is called instead. This method
-			// exists as a stop-gap until the API can be deprecated and removed.
-			checkNoEntry();
-			return {};
-		}
-
 		void UpdateParams()
 		{
 			float AttackTimeSeconds = AttackTime->GetSeconds();
@@ -798,18 +782,7 @@ namespace Metasound
 	 *  Creates an Attack/Decay envelope node.
 	 */
 	template<typename EnvelopeClass, typename ValueType>
-	class METASOUNDSTANDARDNODES_API TADSREnvelopeNode : public FNodeFacade
-	{
-	public:
-		/**
-		 * Constructor used by the Metasound Frontend.
-		 */
-		TADSREnvelopeNode(const FNodeInitData& InInitData)
-			: FNodeFacade(InInitData.InstanceName, InInitData.InstanceID, TFacadeOperatorClass<TADSREnvelopeNodeOperator<EnvelopeClass, ValueType>>())
-		{}
-
-		virtual ~TADSREnvelopeNode() = default;
-	};
+	using TADSREnvelopeNode = TNodeFacade<TADSREnvelopeNodeOperator<EnvelopeClass, ValueType>>;
 
 	using FADSREnvelopeNodeFloat = TADSREnvelopeNode<ADSREnvelopeNodePrivate::FFloatADSREnvelope, float>;
 	METASOUND_REGISTER_NODE(FADSREnvelopeNodeFloat)
diff --git a/Engine/Plugins/Runtime/Metasound/Source/MetasoundStandardNodes/Private/MetasoundAbsoluteValueNodes.cpp b/Engine/Plugins/Runtime/Metasound/Source/MetasoundStandardNodes/Private/MetasoundAbsoluteValueNodes.cpp
--- a/Engine/Plugins/Runtime/Metasound/Source/MetasoundStandardNodes/Private/MetasoundAbsoluteValueNodes.cpp
+++ b/Engine/Plugins/Runtime/Metasound/Source/MetasoundStandardNodes/Private/MetasoundAbsoluteValueNodes.cpp
@@ -195,22 +195,6 @@ namespace Metasound
 			InOutVertexData.BindReadVertex(METASOUND_GET_PARAM_NAME(OutputValue), OutputValue);
 		}
 
-		virtual FDataReferenceCollection GetInputs() const override
-		{
-			// This should never be called. Bind(...) is called instead. This method
-			// exists as a stop-gap until the API can be deprecated and removed.
-			checkNoEntry();
-			return {};
-		}
-
-		virtual FDataReferenceCollection GetOutputs() const override
-		{
-			// This should never be called. Bind(...) is called instead. This method
-			// exists as a stop-gap until the API can be deprecated and removed.
-			checkNoEntry();
-			return {};
-		}
-
 		void GetAbs()
 		{
 			using namespace MetasoundAbsNodePrivate;
@@ -239,18 +223,7 @@ namespace Metasound
 	 *  Returns the Aof the input
 	 */
 	template<typename ValueType>
-	class METASOUNDSTANDARDNODES_API TAbsNode : public FNodeFacade
-	{
-	public:
-		/**
-		 * Constructor used by the Metasound Frontend.
-		 */
-		TAbsNode(const FNodeInitData& InInitData)
-			: FNodeFacade(InInitData.InstanceName, InInitData.InstanceID, TFacadeOperatorClass<TAbsNodeOperator<ValueType>>())
-		{}
-
-		virtual ~TAbsNode() = default;
-	};
+	using TAbsNode = TNodeFacade<TAbsNodeOperator<ValueType>>;
 
 	using FAbsNodeInt32 = TAbsNode<int32>;
  	METASOUND_REGISTER_NODE(FAbsNodeInt32)
diff --git a/Engine/Plugins/Runtime/Metasound/Source/MetasoundStandardNodes/Private/MetasoundAdditiveSynthNode.cpp b/Engine/Plugins/Runtime/Metasound/Source/MetasoundStandardNodes/Private/MetasoundAdditiveSynthNode.cpp
--- a/Engine/Plugins/Runtime/Metasound/Source/MetasoundStandardNodes/Private/MetasoundAdditiveSynthNode.cpp
+++ b/Engine/Plugins/Runtime/Metasound/Source/MetasoundStandardNodes/Private/MetasoundAdditiveSynthNode.cpp
@@ -134,22 +134,6 @@ namespace Metasound
 			InOutVertexData.BindReadVertex(METASOUND_GET_PARAM_NAME(RightAudioOut), RightAudioOut);
 		}
 
-		virtual FDataReferenceCollection GetInputs() const override
-		{
-			// This should never be called. Bind(...) is called instead. This method
-			// exists as a stop-gap until the API can be deprecated and removed.
-			checkNoEntry();
-			return {};
-		}
-
-		virtual FDataReferenceCollection GetOutputs() const override
-		{
-			// This should never be called. Bind(...) is called instead. This method
-			// exists as a stop-gap until the API can be deprecated and removed.
-			checkNoEntry();
-			return {};
-		}
-
 		void Reset(const IOperator::FResetParams& InParams)
 		{
 			CreateSinusoids();
@@ -321,17 +305,7 @@ namespace Metasound
 		float SampleRate = 48000.0f;
 	};
 
-	class FAdditiveSynthNode : public FNodeFacade
-	{
-	public:
-		/**
-		 * Constructor used by the Metasound Frontend.
-		 */
-		FAdditiveSynthNode(const FNodeInitData& InitData)
-			: FNodeFacade(InitData.InstanceName, InitData.InstanceID, TFacadeOperatorClass<FAdditiveSynthOperator>())
-		{
-		}
-	};
+	using FAdditiveSynthNode = TNodeFacade<FAdditiveSynthOperator>;
 
 	METASOUND_REGISTER_NODE(FAdditiveSynthNode)
 }
diff --git a/Engine/Plugins/Runtime/Metasound/Source/MetasoundStandardNodes/Private/MetasoundAudioFormats.cpp b/Engine/Plugins/Runtime/Metasound/Source/MetasoundStandardNodes/Private/MetasoundAudioFormats.cpp
--- a/Engine/Plugins/Runtime/Metasound/Source/MetasoundStandardNodes/Private/MetasoundAudioFormats.cpp
+++ b/Engine/Plugins/Runtime/Metasound/Source/MetasoundStandardNodes/Private/MetasoundAudioFormats.cpp
@@ -125,22 +125,6 @@ namespace Metasound
 				InOutVertexData.BindReadVertex(OutputName, Stereo);
 			}
 
-			virtual FDataReferenceCollection GetInputs() const override
-			{
-				// This should never be called. Bind(...) is called instead. This method
-				// exists as a stop-gap until the API can be deprecated and removed.
-				checkNoEntry();
-				return {};
-			}
-
-			virtual FDataReferenceCollection GetOutputs() const override
-			{
-				// This should never be called. Bind(...) is called instead. This method
-				// exists as a stop-gap until the API can be deprecated and removed.
-				checkNoEntry();
-				return {};
-			}
-
 		private:
 
 			FVertexName OutputName;
@@ -272,22 +256,6 @@ namespace Metasound
 				InOutVertexData.BindReadVertex(METASOUND_GET_PARAM_NAME(RightChannelVertex), Right);
 			}
 
-			virtual FDataReferenceCollection GetInputs() const override
-			{
-				// This should never be called. Bind(...) is called instead. This method
-				// exists as a stop-gap until the API can be deprecated and removed.
-				checkNoEntry();
-				return {};
-			}
-
-			virtual FDataReferenceCollection GetOutputs() const override
-			{
-				// This should never be called. Bind(...) is called instead. This method
-				// exists as a stop-gap until the API can be deprecated and removed.
-				checkNoEntry();
-				return {};
-			}
-
 		private:
 
 			FVertexName InputName;
@@ -414,22 +382,6 @@ namespace Metasound
 				InOutVertexData.BindReadVertex(OutputName, Mono);
 			}
 
-			virtual FDataReferenceCollection GetInputs() const override
-			{
-				// This should never be called. Bind(...) is called instead. This method
-				// exists as a stop-gap until the API can be deprecated and removed.
-				checkNoEntry();
-				return {};
-			}
-
-			virtual FDataReferenceCollection GetOutputs() const override
-			{
-				// This should never be called. Bind(...) is called instead. This method
-				// exists as a stop-gap until the API can be deprecated and removed.
-				checkNoEntry();
-				return {};
-			}
-
 			virtual FExecuteFunction GetExecuteFunction() override
 			{
 				return nullptr;
@@ -560,22 +512,6 @@ namespace Metasound
 				InOutVertexData.BindReadVertex(InputName, Center);
 			}
 
-			virtual FDataReferenceCollection GetInputs() const override
-			{
-				// This should never be called. Bind(...) is called instead. This method
-				// exists as a stop-gap until the API can be deprecated and removed.
-				checkNoEntry();
-				return {};
-			}
-
-			virtual FDataReferenceCollection GetOutputs() const override
-			{
-				// This should never be called. Bind(...) is called instead. This method
-				// exists as a stop-gap until the API can be deprecated and removed.
-				checkNoEntry();
-				return {};
-			}
-
 			virtual FExecuteFunction GetExecuteFunction() override
 			{
 				return nullptr;
diff --git a/Engine/Plugins/Runtime/Metasound/Source/MetasoundStandardNodes/Private/MetasoundBPMToSecondsNode.cpp b/Engine/Plugins/Runtime/Metasound/Source/MetasoundStandardNodes/Private/MetasoundBPMToSecondsNode.cpp
--- a/Engine/Plugins/Runtime/Metasound/Source/MetasoundStandardNodes/Private/MetasoundBPMToSecondsNode.cpp
+++ b/Engine/Plugins/Runtime/Metasound/Source/MetasoundStandardNodes/Private/MetasoundBPMToSecondsNode.cpp
@@ -36,8 +36,6 @@ namespace Metasound
 
 		virtual void BindInputs(FInputVertexInterfaceData& InOutVertexData) override;
 		virtual void BindOutputs(FOutputVertexInterfaceData& InOutVertexData) override;
-		virtual FDataReferenceCollection GetInputs() const override;
-		virtual FDataReferenceCollection GetOutputs() const override;
 		void Reset(const IOperator::FResetParams& InParams);
 		void Execute();
 
@@ -89,22 +87,6 @@ namespace Metasound
 		InOutVertexData.BindReadVertex(METASOUND_GET_PARAM_NAME(OutputTimeSeconds), TimeSeconds);
 	}
 
-	FDataReferenceCollection FBPMToSecondsOperator::GetInputs() const
-	{
-		// This should never be called. Bind(...) is called instead. This method
-		// exists as a stop-gap until the API can be deprecated and removed.
-		checkNoEntry();
-		return {};
-	}
-
-	FDataReferenceCollection FBPMToSecondsOperator::GetOutputs() const
-	{
-		// This should never be called. Bind(...) is called instead. This method
-		// exists as a stop-gap until the API can be deprecated and removed.
-		checkNoEntry();
-		return {};
-	}
-
 	void FBPMToSecondsOperator::UpdateTime()
 	{
 		float CurrBPM = FMath::Clamp(*BPM, MinBPM, MaxBPM);
@@ -194,8 +176,18 @@ namespace Metasound
 	}
 
 	FBPMToSecondsNode::FBPMToSecondsNode(const FNodeInitData& InitData)
-		: FNodeFacade(InitData.InstanceName, InitData.InstanceID, TFacadeOperatorClass<FBPMToSecondsOperator>())
+	: FBPMToSecondsNode(FNodeData(InitData.InstanceName, InitData.InstanceID, FBPMToSecondsOperator::GetNodeInfo().DefaultInterface), MakeShared<const FNodeClassMetadata>(FBPMToSecondsOperator::GetNodeInfo()))
+	{
+	}
+
+	FBPMToSecondsNode::FBPMToSecondsNode(FNodeData InNodeData, TSharedRef<const FNodeClassMetadata> InClassMetadata)
+	: FNodeFacade(MoveTemp(InNodeData), MoveTemp(InClassMetadata), TFacadeOperatorClass<FBPMToSecondsOperator>())
+	{
+	}
+
+	FNodeClassMetadata FBPMToSecondsNode::CreateNodeClassMetadata()
 	{
+		return FBPMToSecondsOperator::GetNodeInfo();
 	}
 
 	METASOUND_REGISTER_NODE(FBPMToSecondsNode)
diff --git a/Engine/Plugins/Runtime/Metasound/Source/MetasoundStandardNodes/Private/MetasoundBandsplitterNode.cpp b/Engine/Plugins/Runtime/Metasound/Source/MetasoundStandardNodes/Private/MetasoundBandsplitterNode.cpp
--- a/Engine/Plugins/Runtime/Metasound/Source/MetasoundStandardNodes/Private/MetasoundBandsplitterNode.cpp
+++ b/Engine/Plugins/Runtime/Metasound/Source/MetasoundStandardNodes/Private/MetasoundBandsplitterNode.cpp
@@ -151,22 +151,6 @@ namespace Metasound
 			}
 		}
 
-		virtual FDataReferenceCollection GetInputs() const override
-		{
-			// This should never be called. Bind(...) is called instead. This method
-			// exists as a stop-gap until the API can be deprecated and removed.
-			checkNoEntry();
-			return {};
-		}
-
-		virtual FDataReferenceCollection GetOutputs() const override
-		{
-			// This should never be called. Bind(...) is called instead. This method
-			// exists as a stop-gap until the API can be deprecated and removed.
-			checkNoEntry();
-			return {};
-		}
-
 		static const FVertexInterface& GetDefaultInterface()
 		{
 			using namespace BandSplitterNode;
@@ -567,14 +551,7 @@ namespace Metasound
 
 	// Node Class
 	template<uint32 NumBands, uint32 NumChannels>
-	class METASOUNDSTANDARDNODES_API TBandSplitterNode : public FNodeFacade
-	{
-	public:
-		TBandSplitterNode(const FNodeInitData& InitData)
-			: FNodeFacade(InitData.InstanceName, InitData.InstanceID, TFacadeOperatorClass<TBandSplitterOperator<NumBands, NumChannels>>())
-		{
-		}
-	};
+	using TBandSplitterNode = TNodeFacade<TBandSplitterOperator<NumBands, NumChannels>>;
 
 #define REGISTER_BANDSPLITTER_NODE(A, B) \
 	using FBandSplitterNode_##A ## _ ##B = TBandSplitterNode<A, B>; \
diff --git a/Engine/Plugins/Runtime/Metasound/Source/MetasoundStandardNodes/Private/MetasoundBasicFilters.cpp b/Engine/Plugins/Runtime/Metasound/Source/MetasoundStandardNodes/Private/MetasoundBasicFilters.cpp
--- a/Engine/Plugins/Runtime/Metasound/Source/MetasoundStandardNodes/Private/MetasoundBasicFilters.cpp
+++ b/Engine/Plugins/Runtime/Metasound/Source/MetasoundStandardNodes/Private/MetasoundBasicFilters.cpp
@@ -109,22 +109,6 @@ namespace Metasound
 			InOutVertexData.BindReadVertex(METASOUND_GET_PARAM_NAME(ParamAudioOutput), AudioOutput);
 		}
 
-		virtual FDataReferenceCollection GetInputs() const override
-		{
-			// This should never be called. Bind(...) is called instead. This method
-			// exists as a stop-gap until the API can be deprecated and removed.
-			checkNoEntry();
-			return {};
-		}
-
-		virtual FDataReferenceCollection GetOutputs() const override
-		{
-			// This should never be called. Bind(...) is called instead. This method
-			// exists as a stop-gap until the API can be deprecated and removed.
-			checkNoEntry();
-			return {};
-		}
-
 		void Reset(const IOperator::FResetParams& InParams);
 		void Execute();
 
@@ -294,22 +278,6 @@ namespace Metasound
 			InOutVertexData.BindReadVertex(METASOUND_GET_PARAM_NAME(ParamBandStopOutput), BandStopOutput);
 		}
 
-		virtual FDataReferenceCollection GetInputs() const override
-		{
-			// This should never be called. Bind(...) is called instead. This method
-			// exists as a stop-gap until the API can be deprecated and removed.
-			checkNoEntry();
-			return {};
-		}
-
-		virtual FDataReferenceCollection GetOutputs() const override
-		{
-			// This should never be called. Bind(...) is called instead. This method
-			// exists as a stop-gap until the API can be deprecated and removed.
-			checkNoEntry();
-			return {};
-		}
-
 		void Reset(const IOperator::FResetParams& InParams);
 		void Execute();
 
@@ -514,22 +482,6 @@ namespace Metasound
 			InOutVertexData.BindReadVertex(METASOUND_GET_PARAM_NAME(ParamAudioOutput), AudioOutput);
 		}
 
-		virtual FDataReferenceCollection GetInputs() const override
-		{
-			// This should never be called. Bind(...) is called instead. This method
-			// exists as a stop-gap until the API can be deprecated and removed.
-			checkNoEntry();
-			return {};
-		}
-
-		virtual FDataReferenceCollection GetOutputs() const override
-		{
-			// This should never be called. Bind(...) is called instead. This method
-			// exists as a stop-gap until the API can be deprecated and removed.
-			checkNoEntry();
-			return {};
-		}
-
 		void Reset(const IOperator::FResetParams& InParams);
 		void Execute();
 
@@ -665,29 +617,6 @@ namespace Metasound
 			InOutVertexData.BindReadVertex(METASOUND_GET_PARAM_NAME(ParamAudioOutput), AudioOutput);
 		}
 
-		virtual void Bind(FVertexInterfaceData& InVertexData) const override
-		{
-			FInputVertexInterfaceData& Inputs = InVertexData.GetInputs();
-
-			FOutputVertexInterfaceData& Outputs = InVertexData.GetOutputs();
-		}
-
-		virtual FDataReferenceCollection GetInputs() const override
-		{
-			// This should never be called. Bind(...) is called instead. This method
-			// exists as a stop-gap until the API can be deprecated and removed.
-			checkNoEntry();
-			return {};
-		}
-
-		virtual FDataReferenceCollection GetOutputs() const override
-		{
-			// This should never be called. Bind(...) is called instead. This method
-			// exists as a stop-gap until the API can be deprecated and removed.
-			checkNoEntry();
-			return {};
-		}
-
 		void Reset(const IOperator::FResetParams& InParams);
 		void Execute();
 
@@ -826,29 +755,6 @@ namespace Metasound
 			InOutVertexData.BindReadVertex(METASOUND_GET_PARAM_NAME(ParamAudioOutput), AudioOutput);
 		}
 
-		virtual void Bind(FVertexInterfaceData& InVertexData) const override
-		{
-			FInputVertexInterfaceData& Inputs = InVertexData.GetInputs();
-			FOutputVertexInterfaceData& Outputs = InVertexData.GetOutputs();
-
-		}
-
-		virtual FDataReferenceCollection GetInputs() const override
-		{
-			// This should never be called. Bind(...) is called instead. This method
-			// exists as a stop-gap until the API can be deprecated and removed.
-			checkNoEntry();
-			return {};
-		}
-
-		virtual FDataReferenceCollection GetOutputs() const override
-		{
-			// This should never be called. Bind(...) is called instead. This method
-			// exists as a stop-gap until the API can be deprecated and removed.
-			checkNoEntry();
-			return {};
-		}
-
 		void Reset(const IOperator::FResetParams& InParams);
 		void Execute();
 
@@ -1037,80 +943,11 @@ namespace Metasound
 
 
 #pragma region Node Declarations
-	class METASOUNDSTANDARDNODES_API FLadderFilterNode : public FNodeFacade
-	{
-	public:
-		// public node api needs to define two conversion constructors:
-		// 1.) from FString
-		FLadderFilterNode(const FVertexName& InInstanceName, const FGuid& InInstanceID)
-			: FNodeFacade(InInstanceName, InInstanceID, TFacadeOperatorClass<FLadderFilterOperator>())
-		{ }
-
-		// 2.) From an NodeInitData struct
-		FLadderFilterNode(const FNodeInitData& InInitData)
-			: FLadderFilterNode(InInitData.InstanceName, InInitData.InstanceID)
-		{ }
-	};
-
-	class METASOUNDSTANDARDNODES_API FStateVariableFilterNode : public FNodeFacade
-	{
-	public:
-		// public node api needs to define two conversion constructors:
-		// 1.) from FString
-		FStateVariableFilterNode(const FVertexName& InInstanceName, const FGuid& InInstanceID)
-			: FNodeFacade(InInstanceName, InInstanceID, TFacadeOperatorClass<FStateVariableFilterOperator>())
-		{ }
-
-		// 2.) From an NodeInitData struct
-		FStateVariableFilterNode(const FNodeInitData& InInitData)
-			: FStateVariableFilterNode(InInitData.InstanceName, InInitData.InstanceID)
-		{ }
-	};
-
-	class METASOUNDSTANDARDNODES_API FOnePoleLowPassFilterNode : public FNodeFacade
-	{
-	public:
-		// public node api needs to define two conversion constructors:
-		// 1.) from FString
-		FOnePoleLowPassFilterNode(const FVertexName& InInstanceName, const FGuid& InInstanceID)
-			: FNodeFacade(InInstanceName, InInstanceID, TFacadeOperatorClass<FOnePoleLowPassFilterOperator>())
-		{ }
-
-		// 2.) From an NodeInitData struct
-		FOnePoleLowPassFilterNode(const FNodeInitData& InInitData)
-			: FOnePoleLowPassFilterNode(InInitData.InstanceName, InInitData.InstanceID)
-		{ }
-	};
-
-	class METASOUNDSTANDARDNODES_API FOnePoleHighPassFilterNode : public FNodeFacade
-	{
-	public:
-		// public node api needs to define two conversion constructors:
-		// 1.) from FString
-		FOnePoleHighPassFilterNode(const FVertexName& InInstanceName, const FGuid& InInstanceID)
-			: FNodeFacade(InInstanceName, InInstanceID, TFacadeOperatorClass<FOnePoleHighPassFilterOperator>())
-		{ }
-
-		// 2.) From an NodeInitData struct
-		FOnePoleHighPassFilterNode(const FNodeInitData& InInitData)
-			: FOnePoleHighPassFilterNode(InInitData.InstanceName, InInitData.InstanceID)
-		{ }
-	};
-
-	class METASOUNDSTANDARDNODES_API FBiquadFilterNode : public FNodeFacade
-	{
-	public:
-		// public node api needs to define two conversion constructors:
-		// 1.) from FString
-		FBiquadFilterNode(const FVertexName& InInstanceName, const FGuid& InInstanceID)
-			: FNodeFacade(InInstanceName, InInstanceID, TFacadeOperatorClass<FBiquadFilterOperator>())
-		{ }
-
-		// 2.) From an NodeInitData struct
-		FBiquadFilterNode(const FNodeInitData& InInitData)
-			: FBiquadFilterNode(InInitData.InstanceName, InInitData.InstanceID)
-		{ }
-	};
+	using FLadderFilterNode = TNodeFacade<FLadderFilterOperator>;
+	using FStateVariableFilterNode = TNodeFacade<FStateVariableFilterOperator>;
+	using FOnePoleLowPassFilterNode = TNodeFacade<FOnePoleLowPassFilterOperator>;
+	using FOnePoleHighPassFilterNode = TNodeFacade<FOnePoleHighPassFilterOperator>;
+	using FBiquadFilterNode = TNodeFacade<FBiquadFilterOperator>;
 #pragma endregion
 
 
diff --git a/Engine/Plugins/Runtime/Metasound/Source/MetasoundStandardNodes/Private/MetasoundBitcrusherNode.cpp b/Engine/Plugins/Runtime/Metasound/Source/MetasoundStandardNodes/Private/MetasoundBitcrusherNode.cpp
--- a/Engine/Plugins/Runtime/Metasound/Source/MetasoundStandardNodes/Private/MetasoundBitcrusherNode.cpp
+++ b/Engine/Plugins/Runtime/Metasound/Source/MetasoundStandardNodes/Private/MetasoundBitcrusherNode.cpp
@@ -109,22 +109,6 @@ namespace Metasound
 			InOutVertexData.BindReadVertex(METASOUND_GET_PARAM_NAME(OutputAudio), AudioOutput);
 		}
 
-		virtual FDataReferenceCollection GetInputs() const override
-		{
-			// This should never be called. Bind(...) is called instead. This method
-			// exists as a stop-gap until the API can be deprecated and removed.
-			checkNoEntry();
-			return {};
-		}
-
-		virtual FDataReferenceCollection GetOutputs() const override
-		{
-			// This should never be called. Bind(...) is called instead. This method
-			// exists as a stop-gap until the API can be deprecated and removed.
-			checkNoEntry();
-			return {};
-		}
-
 		static TUniquePtr<IOperator> CreateOperator(const FBuildOperatorParams& InParams, FBuildResults& OutResults)
 		{
 			using namespace BitcrusherVertexNames;
@@ -202,14 +186,7 @@ namespace Metasound
 	};
 
 	// Node Class
-	class FBitcrusherNode : public FNodeFacade
-	{
-	public:
-		FBitcrusherNode(const FNodeInitData& InitData)
-			: FNodeFacade(InitData.InstanceName, InitData.InstanceID, TFacadeOperatorClass<FBitcrusherOperator>())
-		{
-		}
-	};
+	using FBitcrusherNode = TNodeFacade<FBitcrusherOperator>;
 
 	// Register node
 	METASOUND_REGISTER_NODE(FBitcrusherNode)
diff --git a/Engine/Plugins/Runtime/Metasound/Source/MetasoundStandardNodes/Private/MetasoundClampNode.cpp b/Engine/Plugins/Runtime/Metasound/Source/MetasoundStandardNodes/Private/MetasoundClampNode.cpp
--- a/Engine/Plugins/Runtime/Metasound/Source/MetasoundStandardNodes/Private/MetasoundClampNode.cpp
+++ b/Engine/Plugins/Runtime/Metasound/Source/MetasoundStandardNodes/Private/MetasoundClampNode.cpp
@@ -232,22 +232,6 @@ namespace Metasound
 			InOutVertexData.BindReadVertex(METASOUND_GET_PARAM_NAME(OutputValue), OutputValue);
 		}
 
-		virtual FDataReferenceCollection GetInputs() const override
-		{
-			// This should never be called. Bind(...) is called instead. This method
-			// exists as a stop-gap until the API can be deprecated and removed.
-			checkNoEntry();
-			return {};
-		}
-
-		virtual FDataReferenceCollection GetOutputs() const override
-		{
-			// This should never be called. Bind(...) is called instead. This method
-			// exists as a stop-gap until the API can be deprecated and removed.
-			checkNoEntry();
-			return {};
-		}
-
 		void GetClamped()
 		{
 			using namespace MetasoundClampNodePrivate;
@@ -278,18 +262,7 @@ namespace Metasound
  *  Returns the clamped value of the input in the specified range.
  */
 	template<typename ValueType>
-	class METASOUNDSTANDARDNODES_API TClampNode : public FNodeFacade
-	{
-	public:
-		/**
-		 * Constructor used by the Metasound Frontend.
-		 */
-		TClampNode(const FNodeInitData& InInitData)
-			: FNodeFacade(InInitData.InstanceName, InInitData.InstanceID, TFacadeOperatorClass<TClampNodeOperator<ValueType>>())
-		{}
-
-		virtual ~TClampNode() = default;
-	};
+	using TClampNode = TNodeFacade<TClampNodeOperator<ValueType>>;
 
 	using FClampNodeInt32 = TClampNode<int32>;
  	METASOUND_REGISTER_NODE(FClampNodeInt32)
diff --git a/Engine/Plugins/Runtime/Metasound/Source/MetasoundStandardNodes/Private/MetasoundCompressorNode.cpp b/Engine/Plugins/Runtime/Metasound/Source/MetasoundStandardNodes/Private/MetasoundCompressorNode.cpp
--- a/Engine/Plugins/Runtime/Metasound/Source/MetasoundStandardNodes/Private/MetasoundCompressorNode.cpp
+++ b/Engine/Plugins/Runtime/Metasound/Source/MetasoundStandardNodes/Private/MetasoundCompressorNode.cpp
@@ -207,22 +207,6 @@ namespace Metasound
 			InOutVertexData.BindReadVertex(METASOUND_GET_PARAM_NAME(OutputEnvelope), EnvelopeOutput);
 		}
 
-		virtual FDataReferenceCollection GetInputs() const override
-		{
-			// This should never be called. Bind(...) is called instead. This method
-			// exists as a stop-gap until the API can be deprecated and removed.
-			checkNoEntry();
-			return {};
-		}
-
-		virtual FDataReferenceCollection GetOutputs() const override
-		{
-			// This should never be called. Bind(...) is called instead. This method
-			// exists as a stop-gap until the API can be deprecated and removed.
-			checkNoEntry();
-			return {};
-		}
-
 		static TUniquePtr<IOperator> CreateOperator(const FBuildOperatorParams& InParams, FBuildResults& OutResults)
 		{
 			using namespace CompressorVertexNames;
@@ -469,14 +453,7 @@ namespace Metasound
 	};
 
 	// Node Class
-	class FCompressorNode : public FNodeFacade
-	{
-	public:
-		FCompressorNode(const FNodeInitData& InitData)
-			: FNodeFacade(InitData.InstanceName, InitData.InstanceID, TFacadeOperatorClass<FCompressorOperator>())
-		{
-		}
-	};
+	using FCompressorNode = TNodeFacade<FCompressorOperator>;
 
 	// Register node
 	METASOUND_REGISTER_NODE(FCompressorNode)
diff --git a/Engine/Plugins/Runtime/Metasound/Source/MetasoundStandardNodes/Private/MetasoundConversionNodes.cpp b/Engine/Plugins/Runtime/Metasound/Source/MetasoundStandardNodes/Private/MetasoundConversionNodes.cpp
--- a/Engine/Plugins/Runtime/Metasound/Source/MetasoundStandardNodes/Private/MetasoundConversionNodes.cpp
+++ b/Engine/Plugins/Runtime/Metasound/Source/MetasoundStandardNodes/Private/MetasoundConversionNodes.cpp
@@ -254,22 +254,6 @@ namespace Metasound
 			InOutVertexData.BindReadVertex(METASOUND_GET_PARAM_NAME(OutputValue), OutputValue);
 		}
 
-		virtual FDataReferenceCollection GetInputs() const override
-		{
-			// This should never be called. Bind(...) is called instead. This method
-			// exists as a stop-gap until the API can be deprecated and removed.
-			checkNoEntry();
-			return {};
-		}
-
-		virtual FDataReferenceCollection GetOutputs() const override
-		{
-			// This should never be called. Bind(...) is called instead. This method
-			// exists as a stop-gap until the API can be deprecated and removed.
-			checkNoEntry();
-			return {};
-		}
-
 		void Reset(const IOperator::FResetParams& InParams)
 		{
 			using namespace MetasoundConversionNodePrivate;
@@ -293,18 +277,7 @@ namespace Metasound
 	 *  Generates a random float value when triggered.
 	 */
 	template<typename FromType, typename ToType>
-	class METASOUNDSTANDARDNODES_API TConversionNode : public FNodeFacade
-	{
-	public:
-		/**
-		 * Constructor used by the Metasound Frontend.
-		 */
-		TConversionNode(const FNodeInitData& InInitData)
-			: FNodeFacade(InInitData.InstanceName, InInitData.InstanceID, TFacadeOperatorClass<TConversionOperator<FromType, ToType>>())
-		{}
-
-		virtual ~TConversionNode() = default;
-	};
+	using TConversionNode = TNodeFacade<TConversionOperator<FromType, ToType>>;
 
 	using FConversionFloatToTime = TConversionNode<float, FTime>;
 	METASOUND_REGISTER_NODE(FConversionFloatToTime)
diff --git a/Engine/Plugins/Runtime/Metasound/Source/MetasoundStandardNodes/Private/MetasoundCrossfadeNode.cpp b/Engine/Plugins/Runtime/Metasound/Source/MetasoundStandardNodes/Private/MetasoundCrossfadeNode.cpp
--- a/Engine/Plugins/Runtime/Metasound/Source/MetasoundStandardNodes/Private/MetasoundCrossfadeNode.cpp
+++ b/Engine/Plugins/Runtime/Metasound/Source/MetasoundStandardNodes/Private/MetasoundCrossfadeNode.cpp
@@ -249,22 +249,6 @@ namespace Metasound
 			InOutVertexData.BindReadVertex(METASOUND_GET_PARAM_NAME(OutputTrigger), OutputValue);
 		}
 
-		virtual FDataReferenceCollection GetInputs() const override
-		{
-			// This should never be called. Bind(...) is called instead. This method
-			// exists as a stop-gap until the API can be deprecated and removed.
-			checkNoEntry();
-			return {};
-		}
-
-		virtual FDataReferenceCollection GetOutputs() const override
-		{
-			// This should never be called. Bind(...) is called instead. This method
-			// exists as a stop-gap until the API can be deprecated and removed.
-			checkNoEntry();
-			return {};
-		}
-
 		void PerformCrossfadeOutput()
 		{
 			// Clamp the cross fade value based on the number of inputs
@@ -306,18 +290,7 @@ namespace Metasound
 	};
 
 	template<typename ValueType, uint32 NumInputs>
-	class TCrossfadeNode : public FNodeFacade
-	{
-	public:
-		/**
-		 * Constructor used by the Metasound Frontend.
-		 */
-		TCrossfadeNode(const FNodeInitData& InInitData)
-			: FNodeFacade(InInitData.InstanceName, InInitData.InstanceID, TFacadeOperatorClass<TCrossfadeOperator<ValueType, NumInputs>>())
-		{}
-
-		virtual ~TCrossfadeNode() = default;
-	};
+	using TCrossfadeNode = TNodeFacade<TCrossfadeOperator<ValueType, NumInputs>>;
 
 	REGISTER_CROSSFADE_NODE(float, 2);
 	REGISTER_CROSSFADE_NODE(float, 3);
diff --git a/Engine/Plugins/Runtime/Metasound/Source/MetasoundStandardNodes/Private/MetasoundDebugLogNode.cpp b/Engine/Plugins/Runtime/Metasound/Source/MetasoundStandardNodes/Private/MetasoundDebugLogNode.cpp
--- a/Engine/Plugins/Runtime/Metasound/Source/MetasoundStandardNodes/Private/MetasoundDebugLogNode.cpp
+++ b/Engine/Plugins/Runtime/Metasound/Source/MetasoundStandardNodes/Private/MetasoundDebugLogNode.cpp
@@ -124,22 +124,6 @@ namespace Metasound
 			{
 			}
 
-			virtual FDataReferenceCollection GetInputs() const override
-			{
-				// This should never be called. Bind(...) is called instead. This method
-				// exists as a stop-gap until the API can be deprecated and removed.
-				checkNoEntry();
-				return {};
-			}
-
-			virtual FDataReferenceCollection GetOutputs() const override
-			{
-				// This should never be called. Bind(...) is called instead. This method
-				// exists as a stop-gap until the API can be deprecated and removed.
-				checkNoEntry();
-				return {};
-			}
-
 			void Reset(const IOperator::FResetParams& InParams)
 			{
 				GraphName = TEXT("");
@@ -185,18 +169,7 @@ namespace Metasound
 	 *  Records a value to the log when triggered
 	 */
 	template<typename PrintLogType>
-	class METASOUNDSTANDARDNODES_API TPrintLogNode : public FNodeFacade
-	{
-		public:
-			/**
-			 * Constructor used by the Metasound Frontend.
-			 */
-			TPrintLogNode(const FNodeInitData& InInitData)
-				: FNodeFacade(InInitData.InstanceName, InInitData.InstanceID, TFacadeOperatorClass<TPrintLogOperator<PrintLogType>>())
-			{}
-
-			virtual ~TPrintLogNode() = default;
-	};
+	using TPrintLogNode = TNodeFacade<TPrintLogOperator<PrintLogType>>;
 
 	using FPrintLogNodeInt32 = TPrintLogNode<int32>;
 	METASOUND_REGISTER_NODE(FPrintLogNodeInt32)
diff --git a/Engine/Plugins/Runtime/Metasound/Source/MetasoundStandardNodes/Private/MetasoundDebugTriggerOnNanNode.cpp b/Engine/Plugins/Runtime/Metasound/Source/MetasoundStandardNodes/Private/MetasoundDebugTriggerOnNanNode.cpp
--- a/Engine/Plugins/Runtime/Metasound/Source/MetasoundStandardNodes/Private/MetasoundDebugTriggerOnNanNode.cpp
+++ b/Engine/Plugins/Runtime/Metasound/Source/MetasoundStandardNodes/Private/MetasoundDebugTriggerOnNanNode.cpp
@@ -41,9 +41,6 @@ namespace Metasound
 
 		virtual void BindInputs(FInputVertexInterfaceData& InOutVertexData) override;
 		virtual void BindOutputs(FOutputVertexInterfaceData& InOutVertexData) override;
-		virtual FDataReferenceCollection GetInputs() const override;
-		virtual FDataReferenceCollection GetOutputs() const override;
-
 		void Execute();
 
 	private:
@@ -87,22 +84,6 @@ namespace Metasound
 		InOutVertexData.BindReadVertex(METASOUND_GET_PARAM_NAME(OutputTriggerOnNan), TriggerOutput);
 	}
 
-    FDataReferenceCollection FTriggerOnNanOperator::GetInputs() const
-    {
-    	// This should never be called. Bind(...) is called instead. This method
-    	// exists as a stop-gap until the API can be deprecated and removed.
-    	checkNoEntry();
-    	return {};
-    }
-
-    FDataReferenceCollection FTriggerOnNanOperator::GetOutputs() const
-    {
-    	// This should never be called. Bind(...) is called instead. This method
-    	// exists as a stop-gap until the API can be deprecated and removed.
-    	checkNoEntry();
-    	return {};
-    }
-
 	void FTriggerOnNanOperator::Execute()
 	{
 		const bool bShouldTriggerOnce = *bTriggerOnce;
@@ -176,14 +157,7 @@ namespace Metasound
 	}
 
 	// Node Class
-	class FTriggerOnNanNode : public FNodeFacade
-	{
-	public:
-		FTriggerOnNanNode(const FNodeInitData& InitData)
-			: FNodeFacade(InitData.InstanceName, InitData.InstanceID, TFacadeOperatorClass<FTriggerOnNanOperator>())
-		{
-		}
-	};
+	using FTriggerOnNanNode = TNodeFacade<FTriggerOnNanOperator>;
 
 	METASOUND_REGISTER_NODE(FTriggerOnNanNode)
 }
diff --git a/Engine/Plugins/Runtime/Metasound/Source/MetasoundStandardNodes/Private/MetasoundDecibelsToLinearGain.cpp b/Engine/Plugins/Runtime/Metasound/Source/MetasoundStandardNodes/Private/MetasoundDecibelsToLinearGain.cpp
--- a/Engine/Plugins/Runtime/Metasound/Source/MetasoundStandardNodes/Private/MetasoundDecibelsToLinearGain.cpp
+++ b/Engine/Plugins/Runtime/Metasound/Source/MetasoundStandardNodes/Private/MetasoundDecibelsToLinearGain.cpp
@@ -35,8 +35,6 @@ namespace Metasound
 
 		virtual void BindInputs(FInputVertexInterfaceData& InOutVertexData) override;
 		virtual void BindOutputs(FOutputVertexInterfaceData& InOutVertexData) override;
-		virtual FDataReferenceCollection GetInputs() const override;
-		virtual FDataReferenceCollection GetOutputs() const override;
 		void Reset(const IOperator::FResetParams& InParams);
 		void Execute();
 
@@ -67,22 +65,6 @@ namespace Metasound
 		InOutVertexData.BindReadVertex(METASOUND_GET_PARAM_NAME(OutputLinearGain), LinearGainOutput);
 	}
 
-	FDataReferenceCollection FDecibelsToLinearGainOperator::GetInputs() const
-	{
-		// This should never be called. Bind(...) is called instead. This method
-		// exists as a stop-gap until the API can be deprecated and removed.
-		checkNoEntry();
-		return {};
-	}
-
-	FDataReferenceCollection FDecibelsToLinearGainOperator::GetOutputs() const
-	{
-		// This should never be called. Bind(...) is called instead. This method
-		// exists as a stop-gap until the API can be deprecated and removed.
-		checkNoEntry();
-		return {};
-	}
-
 	void FDecibelsToLinearGainOperator::Reset(const IOperator::FResetParams& InParams)
 	{
 		Execute();
@@ -148,17 +130,7 @@ namespace Metasound
 		return MakeUnique<FDecibelsToLinearGainOperator>(InParams, InDecibelGain);
 	}
 
-	class METASOUNDSTANDARDNODES_API FDecibelsToLinearGainNode : public FNodeFacade
-	{
-	public:
-		/**
-		 * Constructor used by the Metasound Frontend.
-		 */
-		FDecibelsToLinearGainNode(const FNodeInitData& InitData)
-			: FNodeFacade(InitData.InstanceName, InitData.InstanceID, TFacadeOperatorClass <FDecibelsToLinearGainOperator>())
-		{
-		}
-	};
+	using FDecibelsToLinearGainNode = TNodeFacade<FDecibelsToLinearGainOperator>;
 
 	METASOUND_REGISTER_NODE(FDecibelsToLinearGainNode)
 }
diff --git a/Engine/Plugins/Runtime/Metasound/Source/MetasoundStandardNodes/Private/MetasoundDelayNode.cpp b/Engine/Plugins/Runtime/Metasound/Source/MetasoundStandardNodes/Private/MetasoundDelayNode.cpp
--- a/Engine/Plugins/Runtime/Metasound/Source/MetasoundStandardNodes/Private/MetasoundDelayNode.cpp
+++ b/Engine/Plugins/Runtime/Metasound/Source/MetasoundStandardNodes/Private/MetasoundDelayNode.cpp
@@ -208,22 +208,6 @@ namespace Metasound
 			InOutVertexData.BindReadVertex(METASOUND_GET_PARAM_NAME(OutParamAudio), AudioOutput);
 		}
 
-		virtual FDataReferenceCollection GetInputs() const override
-		{
-			// This should never be called. Bind(...) is called instead. This method
-			// exists as a stop-gap until the API can be deprecated and removed.
-			checkNoEntry();
-			return {};
-		}
-
-		virtual FDataReferenceCollection GetOutputs() const override
-		{
-			// This should never be called. Bind(...) is called instead. This method
-			// exists as a stop-gap until the API can be deprecated and removed.
-			checkNoEntry();
-			return {};
-		}
-
 		float GetClampedDelayTimeMsec(float InDelayTime) const
 		{
 			// Clamp the delay time to the max delay allowed
@@ -369,16 +353,7 @@ namespace Metasound
 	};
 
 	template<typename ValueType>
-	class METASOUNDSTANDARDNODES_API TDelayNode : public FNodeFacade
-	{
-	public:
-		TDelayNode(const FNodeInitData& InInitData)
-			: FNodeFacade(InInitData.InstanceName, InInitData.InstanceID, TFacadeOperatorClass<TDelayNodeOperator<ValueType>>())
-		{}
-
-		virtual ~TDelayNode() = default;
-	};
-
+	using TDelayNode = TNodeFacade<TDelayNodeOperator<ValueType>>;
 
 	using FDelayNode = TDelayNode<FTime>;
 	METASOUND_REGISTER_NODE(FDelayNode)
diff --git a/Engine/Plugins/Runtime/Metasound/Source/MetasoundStandardNodes/Private/MetasoundDiffuserNode.cpp b/Engine/Plugins/Runtime/Metasound/Source/MetasoundStandardNodes/Private/MetasoundDiffuserNode.cpp
--- a/Engine/Plugins/Runtime/Metasound/Source/MetasoundStandardNodes/Private/MetasoundDiffuserNode.cpp
+++ b/Engine/Plugins/Runtime/Metasound/Source/MetasoundStandardNodes/Private/MetasoundDiffuserNode.cpp
@@ -130,22 +130,6 @@ namespace Metasound
 			InOutVertexData.BindReadVertex(METASOUND_GET_PARAM_NAME(OutputAudio), AudioOutput);
 		}
 
-		virtual FDataReferenceCollection GetInputs() const override
-		{
-			// This should never be called. Bind(...) is called instead. This method
-			// exists as a stop-gap until the API can be deprecated and removed.
-			checkNoEntry();
-			return {};
-		}
-
-		virtual FDataReferenceCollection GetOutputs() const override
-		{
-			// This should never be called. Bind(...) is called instead. This method
-			// exists as a stop-gap until the API can be deprecated and removed.
-			checkNoEntry();
-			return {};
-		}
-
 		void Reset(const IOperator::FResetParams& InParams)
 		{
 			using namespace DiffuserNode;
@@ -221,17 +205,7 @@ namespace Metasound
 		}
 	};
 
-	class FDiffuserNode : public FNodeFacade
-	{
-	public:
-		/**
-		 * Constructor used by the Metasound Frontend.
-		 */
-		FDiffuserNode(const FNodeInitData& InitData)
-			: FNodeFacade(InitData.InstanceName, InitData.InstanceID, TFacadeOperatorClass<FDiffuserOperator>())
-		{
-		}
-	};
+	using FDiffuserNode = TNodeFacade<FDiffuserOperator>;
 
 	METASOUND_REGISTER_NODE(FDiffuserNode)
 }
diff --git a/Engine/Plugins/Runtime/Metasound/Source/MetasoundStandardNodes/Private/MetasoundDopplerPitchShiftNode.cpp b/Engine/Plugins/Runtime/Metasound/Source/MetasoundStandardNodes/Private/MetasoundDopplerPitchShiftNode.cpp
--- a/Engine/Plugins/Runtime/Metasound/Source/MetasoundStandardNodes/Private/MetasoundDopplerPitchShiftNode.cpp
+++ b/Engine/Plugins/Runtime/Metasound/Source/MetasoundStandardNodes/Private/MetasoundDopplerPitchShiftNode.cpp
@@ -34,8 +34,6 @@ namespace Metasound
 
 		virtual void BindInputs(FInputVertexInterfaceData& InOutVertexData) override;
 		virtual void BindOutputs(FOutputVertexInterfaceData& InOutVertexData) override;
-		virtual FDataReferenceCollection GetInputs() const override;
-		virtual FDataReferenceCollection GetOutputs() const override;
 		void Reset(const IOperator::FResetParams& InParams);
 		void Execute();
 
@@ -75,22 +73,6 @@ namespace Metasound
 		InOutVertexData.BindReadVertex(METASOUND_GET_PARAM_NAME(OutParamAudio), AudioOutput);
 	}
 
-	FDataReferenceCollection FDopplerPitchShiftOperator::GetInputs() const
-	{
-		// This should never be called. Bind(...) is called instead. This method
-		// exists as a stop-gap until the API can be deprecated and removed.
-		checkNoEntry();
-		return {};
-	}
-
-	FDataReferenceCollection FDopplerPitchShiftOperator::GetOutputs() const
-	{
-		// This should never be called. Bind(...) is called instead. This method
-		// exists as a stop-gap until the API can be deprecated and removed.
-		checkNoEntry();
-		return {};
-	}
-
 	void FDopplerPitchShiftOperator::Reset(const IOperator::FResetParams& InParams)
 	{
 		AudioOutput->Zero();
@@ -158,15 +140,7 @@ namespace Metasound
 		return MakeUnique<FDopplerPitchShiftOperator>(InParams, AudioIn, PitchShift, DelayLength);
 	}
 
-	class FDopperPitchShiftNode : public FNodeFacade
-	{
-	public:
-		FDopperPitchShiftNode(const FNodeInitData& InitData)
-			: FNodeFacade(InitData.InstanceName, InitData.InstanceID, TFacadeOperatorClass<FDopplerPitchShiftOperator>())
-		{
-		}
-	};
-
+	using FDopperPitchShiftNode = TNodeFacade<FDopplerPitchShiftOperator>;
 
 	METASOUND_REGISTER_NODE(FDopperPitchShiftNode)
 }
diff --git a/Engine/Plugins/Runtime/Metasound/Source/MetasoundStandardNodes/Private/MetasoundDynamicFilterNode.cpp b/Engine/Plugins/Runtime/Metasound/Source/MetasoundStandardNodes/Private/MetasoundDynamicFilterNode.cpp
--- a/Engine/Plugins/Runtime/Metasound/Source/MetasoundStandardNodes/Private/MetasoundDynamicFilterNode.cpp
+++ b/Engine/Plugins/Runtime/Metasound/Source/MetasoundStandardNodes/Private/MetasoundDynamicFilterNode.cpp
@@ -151,22 +151,6 @@ namespace Metasound
 			InOutVertexData.BindReadVertex(METASOUND_GET_PARAM_NAME(OutputAudio), AudioOutput);
 		}
 
-		virtual FDataReferenceCollection GetInputs() const override
-		{
-			// This should never be called. Bind(...) is called instead. This method
-			// exists as a stop-gap until the API can be deprecated and removed.
-			checkNoEntry();
-			return {};
-		}
-
-		virtual FDataReferenceCollection GetOutputs() const override
-		{
-			// This should never be called. Bind(...) is called instead. This method
-			// exists as a stop-gap until the API can be deprecated and removed.
-			checkNoEntry();
-			return {};
-		}
-
 		static const FVertexInterface& DeclareVertexInterface()
 		{
 			using namespace DynamicFilterNode;
@@ -309,14 +293,7 @@ namespace Metasound
 	};
 
 	// Node Class
-	class FDynamicFilterNode : public FNodeFacade
-	{
-	public:
-		FDynamicFilterNode(const FNodeInitData& InitData)
-			: FNodeFacade(InitData.InstanceName, InitData.InstanceID, TFacadeOperatorClass<FDynamicFilterOperator>())
-		{
-		}
-	};
+	using FDynamicFilterNode = TNodeFacade<FDynamicFilterOperator>;
 
 	// Register node
 	METASOUND_REGISTER_NODE(FDynamicFilterNode)
diff --git a/Engine/Plugins/Runtime/Metasound/Source/MetasoundStandardNodes/Private/MetasoundEnvelopeFollowerNode.cpp b/Engine/Plugins/Runtime/Metasound/Source/MetasoundStandardNodes/Private/MetasoundEnvelopeFollowerNode.cpp
--- a/Engine/Plugins/Runtime/Metasound/Source/MetasoundStandardNodes/Private/MetasoundEnvelopeFollowerNode.cpp
+++ b/Engine/Plugins/Runtime/Metasound/Source/MetasoundStandardNodes/Private/MetasoundEnvelopeFollowerNode.cpp
@@ -54,8 +54,6 @@ namespace Metasound
 
 		virtual void BindInputs(FInputVertexInterfaceData& InOutVertexData) override;
 		virtual void BindOutputs(FOutputVertexInterfaceData& InOutVertexData) override;
-		virtual FDataReferenceCollection GetInputs() const override;
-		virtual FDataReferenceCollection GetOutputs() const override;
 		void Reset(const IOperator::FResetParams& InParams);
 		void Execute();
 
@@ -124,22 +122,6 @@ namespace Metasound
 		InOutVertexData.BindReadVertex(METASOUND_GET_PARAM_NAME(OutputAudioEnvelope), EnvelopeAudioOutput);
 	}
 
-	FDataReferenceCollection FEnvelopeFollowerOperator::GetInputs() const
-	{
-		// This should never be called. Bind(...) is called instead. This method
-		// exists as a stop-gap until the API can be deprecated and removed.
-		checkNoEntry();
-		return {};
-	}
-
-	FDataReferenceCollection FEnvelopeFollowerOperator::GetOutputs() const
-	{
-		// This should never be called. Bind(...) is called instead. This method
-		// exists as a stop-gap until the API can be deprecated and removed.
-		checkNoEntry();
-		return {};
-	}
-
 	void FEnvelopeFollowerOperator::Reset(const IOperator::FResetParams& InParams)
 	{
 		OutputNeedsResetOnDisable = true;
@@ -281,8 +263,18 @@ namespace Metasound
 	}
 
 	FEnvelopeFollowerNode::FEnvelopeFollowerNode(const FNodeInitData& InitData)
-		: FNodeFacade(InitData.InstanceName, InitData.InstanceID, TFacadeOperatorClass<FEnvelopeFollowerOperator>())
+	: FEnvelopeFollowerNode(FNodeData(InitData.InstanceName, InitData.InstanceID, FEnvelopeFollowerOperator::GetNodeInfo().DefaultInterface), MakeShared<const FNodeClassMetadata>(FEnvelopeFollowerOperator::GetNodeInfo()))
+	{
+	}
+
+	FEnvelopeFollowerNode::FEnvelopeFollowerNode(FNodeData InNodeData, TSharedRef<const FNodeClassMetadata> InClassMetadata)
+	: FNodeFacade(MoveTemp(InNodeData), MoveTemp(InClassMetadata), TFacadeOperatorClass<FEnvelopeFollowerOperator>())
+	{
+	}
+
+	FNodeClassMetadata FEnvelopeFollowerNode::CreateNodeClassMetadata()
 	{
+		return FEnvelopeFollowerOperator::GetNodeInfo();
 	}
 
 	METASOUND_REGISTER_NODE(FEnvelopeFollowerNode)
diff --git a/Engine/Plugins/Runtime/Metasound/Source/MetasoundStandardNodes/Private/MetasoundFilterQToBandwidthNode.cpp b/Engine/Plugins/Runtime/Metasound/Source/MetasoundStandardNodes/Private/MetasoundFilterQToBandwidthNode.cpp
--- a/Engine/Plugins/Runtime/Metasound/Source/MetasoundStandardNodes/Private/MetasoundFilterQToBandwidthNode.cpp
+++ b/Engine/Plugins/Runtime/Metasound/Source/MetasoundStandardNodes/Private/MetasoundFilterQToBandwidthNode.cpp
@@ -100,22 +100,6 @@ namespace Metasound
 			InOutVertexData.BindReadVertex(METASOUND_GET_PARAM_NAME(OutputBandwidth), OutValue);
 		}
 
-		virtual FDataReferenceCollection GetInputs() const override
-		{
-			// This should never be called. Bind(...) is called instead. This method
-			// exists as a stop-gap until the API can be deprecated and removed.
-			checkNoEntry();
-			return {};
-		}
-
-		virtual FDataReferenceCollection GetOutputs() const override
-		{
-			// This should never be called. Bind(...) is called instead. This method
-			// exists as a stop-gap until the API can be deprecated and removed.
-			checkNoEntry();
-			return {};
-		}
-
 		void Reset(const IOperator::FResetParams& InParams)
 		{
 			*OutValue = Audio::GetBandwidthFromQ(FMath::Max(*InValue, KINDA_SMALL_NUMBER));
@@ -131,16 +115,7 @@ namespace Metasound
 		FFloatWriteRef OutValue;
 	};
 
-	class FConvertQToBandwidthNode : public FNodeFacade
-	{
-	public:
-		FConvertQToBandwidthNode(const FNodeInitData& InInitData)
-			: FNodeFacade(InInitData.InstanceName, InInitData.InstanceID, TFacadeOperatorClass<FConvertQToBandwidthNodeOperator>())
-		{
-		}
-
-		virtual ~FConvertQToBandwidthNode() = default;
-	};
+	using FConvertQToBandwidthNode = TNodeFacade<FConvertQToBandwidthNodeOperator>;
 
 	METASOUND_REGISTER_NODE(FConvertQToBandwidthNode)
 } 
diff --git a/Engine/Plugins/Runtime/Metasound/Source/MetasoundStandardNodes/Private/MetasoundFlangerNode.cpp b/Engine/Plugins/Runtime/Metasound/Source/MetasoundStandardNodes/Private/MetasoundFlangerNode.cpp
--- a/Engine/Plugins/Runtime/Metasound/Source/MetasoundStandardNodes/Private/MetasoundFlangerNode.cpp
+++ b/Engine/Plugins/Runtime/Metasound/Source/MetasoundStandardNodes/Private/MetasoundFlangerNode.cpp
@@ -128,22 +128,6 @@ namespace Metasound
 			InOutVertexData.BindReadVertex(METASOUND_GET_PARAM_NAME(OutputAudio), AudioOut);
 		}
 
-		virtual FDataReferenceCollection GetInputs() const override
-		{
-			// This should never be called. Bind(...) is called instead. This method
-			// exists as a stop-gap until the API can be deprecated and removed.
-			checkNoEntry();
-			return {};
-		}
-
-		virtual FDataReferenceCollection GetOutputs() const override
-		{
-			// This should never be called. Bind(...) is called instead. This method
-			// exists as a stop-gap until the API can be deprecated and removed.
-			checkNoEntry();
-			return {};
-		}
-
 		void Reset(const IOperator::FResetParams& InParams)
 		{
 			AudioOut->Zero();
@@ -185,17 +169,7 @@ namespace Metasound
 		Audio::FFlanger Flanger;
 	};
 
-	class METASOUNDSTANDARDNODES_API FFlangerNode : public FNodeFacade
-	{
-	public:
-		/**
-		 * Constructor used by the Metasound Frontend.
-		 */
-		FFlangerNode(const FNodeInitData& InitData)
-			: FNodeFacade(InitData.InstanceName, InitData.InstanceID, TFacadeOperatorClass<FFlangerOperator>())
-		{
-		}
-	};
+	using FFlangerNode = TNodeFacade<FFlangerOperator>;
 
 	METASOUND_REGISTER_NODE(FFlangerNode)
 }
diff --git a/Engine/Plugins/Runtime/Metasound/Source/MetasoundStandardNodes/Private/MetasoundFreqMultiplierToSemitoneNode.cpp b/Engine/Plugins/Runtime/Metasound/Source/MetasoundStandardNodes/Private/MetasoundFreqMultiplierToSemitoneNode.cpp
--- a/Engine/Plugins/Runtime/Metasound/Source/MetasoundStandardNodes/Private/MetasoundFreqMultiplierToSemitoneNode.cpp
+++ b/Engine/Plugins/Runtime/Metasound/Source/MetasoundStandardNodes/Private/MetasoundFreqMultiplierToSemitoneNode.cpp
@@ -36,8 +36,6 @@ namespace Metasound
 
 		virtual void BindInputs(FInputVertexInterfaceData& InOutVertexData) override;
 		virtual void BindOutputs(FOutputVertexInterfaceData& InOutVertexData) override;
-		virtual FDataReferenceCollection GetInputs() const override;
-		virtual FDataReferenceCollection GetOutputs() const override;
 		void Reset(const IOperator::FResetParams& InParams);
 		void Execute();
 
@@ -72,22 +70,6 @@ namespace Metasound
 		InOutVertexData.BindReadVertex(METASOUND_GET_PARAM_NAME(OutputSemitone), SemitoneOutput);
 	}
 
-	FDataReferenceCollection FFrequencyMultiplierToSemitoneOperator::GetInputs() const
-	{
-		// This should never be called. Bind(...) is called instead. This method
-		// exists as a stop-gap until the API can be deprecated and removed.
-		checkNoEntry();
-		return {};
-	}
-
-	FDataReferenceCollection FFrequencyMultiplierToSemitoneOperator::GetOutputs() const
-	{
-		// This should never be called. Bind(...) is called instead. This method
-		// exists as a stop-gap until the API can be deprecated and removed.
-		checkNoEntry();
-		return {};
-	}
-
 	void FFrequencyMultiplierToSemitoneOperator::Reset(const IOperator::FResetParams& InParams)
 	{
 		PrevFrequencyMultiplier = *FrequencyMultiplierInput;
@@ -159,17 +141,7 @@ namespace Metasound
 		return MakeUnique<FFrequencyMultiplierToSemitoneOperator>(InParams.OperatorSettings, InFrequencyMultiplier);
 	}
 
-	class METASOUNDSTANDARDNODES_API FFrequencyMultiplierToSemitoneNode : public FNodeFacade
-	{
-	public:
-		/**
-		 * Constructor used by the Metasound Frontend.
-		 */
-		FFrequencyMultiplierToSemitoneNode(const FNodeInitData& InitData)
-			: FNodeFacade(InitData.InstanceName, InitData.InstanceID, TFacadeOperatorClass <FFrequencyMultiplierToSemitoneOperator>())
-		{
-		}
-	};
+	using FFrequencyMultiplierToSemitoneNode = TNodeFacade<FFrequencyMultiplierToSemitoneOperator>;
 
 	METASOUND_REGISTER_NODE(FFrequencyMultiplierToSemitoneNode)
 }
diff --git a/Engine/Plugins/Runtime/Metasound/Source/MetasoundStandardNodes/Private/MetasoundFreqToMIDINode.cpp b/Engine/Plugins/Runtime/Metasound/Source/MetasoundStandardNodes/Private/MetasoundFreqToMIDINode.cpp
--- a/Engine/Plugins/Runtime/Metasound/Source/MetasoundStandardNodes/Private/MetasoundFreqToMIDINode.cpp
+++ b/Engine/Plugins/Runtime/Metasound/Source/MetasoundStandardNodes/Private/MetasoundFreqToMIDINode.cpp
@@ -32,8 +32,6 @@ namespace Metasound
 
 		virtual void BindInputs(FInputVertexInterfaceData& InOutVertexData) override;
 		virtual void BindOutputs(FOutputVertexInterfaceData& InOutVertexData) override;
-		virtual FDataReferenceCollection GetInputs() const override;
-		virtual FDataReferenceCollection GetOutputs() const override;
 		void Reset(const IOperator::FResetParams& InParams);
 		void Execute();
 
@@ -67,22 +65,6 @@ namespace Metasound
 		InOutVertexData.BindReadVertex(METASOUND_GET_PARAM_NAME(OutputMidi), MidiOutput);
 	}
 
-	FDataReferenceCollection FFreqToMidiOperator::GetInputs() const
-	{
-		// This should never be called. Bind(...) is called instead. This method
-		// exists as a stop-gap until the API can be deprecated and removed.
-		checkNoEntry();
-		return {};
-	}
-
-	FDataReferenceCollection FFreqToMidiOperator::GetOutputs() const
-	{
-		// This should never be called. Bind(...) is called instead. This method
-		// exists as a stop-gap until the API can be deprecated and removed.
-		checkNoEntry();
-		return {};
-	}
-
 	void FFreqToMidiOperator::Reset(const IOperator::FResetParams& InParams)
 	{
 		PrevFreq = *FreqInput;
@@ -156,17 +138,7 @@ namespace Metasound
 		return MakeUnique<FFreqToMidiOperator>(InParams.OperatorSettings, InFreq);
 	}
 
-	class METASOUNDSTANDARDNODES_API FFreqToMidiNode : public FNodeFacade
-	{
-	public:
-		/**
-		 * Constructor used by the Metasound Frontend.
-		 */
-		FFreqToMidiNode(const FNodeInitData& InitData)
-			: FNodeFacade(InitData.InstanceName, InitData.InstanceID, TFacadeOperatorClass <FFreqToMidiOperator>())
-		{
-		}
-	};
+	using FFreqToMidiNode = TNodeFacade<FFreqToMidiOperator>;
 
 	METASOUND_REGISTER_NODE(FFreqToMidiNode)
 }
diff --git a/Engine/Plugins/Runtime/Metasound/Source/MetasoundStandardNodes/Private/MetasoundGrainDelayNode.cpp b/Engine/Plugins/Runtime/Metasound/Source/MetasoundStandardNodes/Private/MetasoundGrainDelayNode.cpp
--- a/Engine/Plugins/Runtime/Metasound/Source/MetasoundStandardNodes/Private/MetasoundGrainDelayNode.cpp
+++ b/Engine/Plugins/Runtime/Metasound/Source/MetasoundStandardNodes/Private/MetasoundGrainDelayNode.cpp
@@ -72,8 +72,6 @@ namespace Metasound
 		static TUniquePtr<IOperator> CreateOperator(const FBuildOperatorParams& InParams, FBuildResults& OutResults);
 		virtual void BindInputs(FInputVertexInterfaceData& InOutVertexData) override;
 		virtual void BindOutputs(FOutputVertexInterfaceData& InOutVertexData) override;
-		virtual FDataReferenceCollection GetInputs() const override;
-		virtual FDataReferenceCollection GetOutputs() const override;
 		void Reset(const IOperator::FResetParams& InParams);
 		void Execute();
 
@@ -257,22 +255,6 @@ namespace Metasound
 		InOutVertexData.BindReadVertex(METASOUND_GET_PARAM_NAME(OutAudio), AudioOutput);
 	}
 
-	FDataReferenceCollection FGrainDelayOperator::GetInputs() const
-	{
-		// This should never be called. Bind(...) is called instead. This method
-		// exists as a stop-gap until the API can be deprecated and removed.
-		checkNoEntry();
-		return {};
-	}
-
-	FDataReferenceCollection FGrainDelayOperator::GetOutputs() const
-	{
-		// This should never be called. Bind(...) is called instead. This method
-		// exists as a stop-gap until the API can be deprecated and removed.
-		checkNoEntry();
-		return {};
-	}
-
 	void FGrainDelayOperator::Reset(const IOperator::FResetParams& InParams)
 	{
 		AudioOutput->Zero();
@@ -370,14 +352,7 @@ namespace Metasound
 		}
 	}
 	
-	class FGrainDelayNode final : public FNodeFacade
-	{
-	public:
-		explicit FGrainDelayNode(const FNodeInitData& InitData)
-			: FNodeFacade(InitData.InstanceName, InitData.InstanceID, TFacadeOperatorClass<FGrainDelayOperator>())
-		{
-		}
-	};
+	using FGrainDelayNode = TNodeFacade<FGrainDelayOperator>;
 
 	METASOUND_REGISTER_NODE(FGrainDelayNode)
 }
diff --git a/Engine/Plugins/Runtime/Metasound/Source/MetasoundStandardNodes/Private/MetasoundITDPannerNode.cpp b/Engine/Plugins/Runtime/Metasound/Source/MetasoundStandardNodes/Private/MetasoundITDPannerNode.cpp
--- a/Engine/Plugins/Runtime/Metasound/Source/MetasoundStandardNodes/Private/MetasoundITDPannerNode.cpp
+++ b/Engine/Plugins/Runtime/Metasound/Source/MetasoundStandardNodes/Private/MetasoundITDPannerNode.cpp
@@ -48,8 +48,6 @@ namespace Metasound
 
 		virtual void BindInputs(FInputVertexInterfaceData& InOutVertexData) override;
 		virtual void BindOutputs(FOutputVertexInterfaceData& InOutVertexData) override;
-		virtual FDataReferenceCollection GetInputs() const override;
-		virtual FDataReferenceCollection GetOutputs() const override;
 		void Reset(const IOperator::FResetParams& InParams);
 		void Execute();
 
@@ -128,22 +126,6 @@ namespace Metasound
 		InOutVertexData.BindReadVertex(METASOUND_GET_PARAM_NAME(OutputAudioRight), AudioRightOutput);
 	}
 
-	FDataReferenceCollection FITDPannerOperator::GetInputs() const
-	{
-		// This should never be called. Bind(...) is called instead. This method
-		// exists as a stop-gap until the API can be deprecated and removed.
-		checkNoEntry();
-		return {};
-	}
-
-	FDataReferenceCollection FITDPannerOperator::GetOutputs() const
-	{
-		// This should never be called. Bind(...) is called instead. This method
-		// exists as a stop-gap until the API can be deprecated and removed.
-		checkNoEntry();
-		return {};
-	}
-
 	void FITDPannerOperator::UpdateParams(bool bIsInit)
 	{
 		// ****************
@@ -314,17 +296,7 @@ namespace Metasound
 		return MakeUnique<FITDPannerOperator>(InParams.OperatorSettings, AudioIn, PanningAngle, DistanceFactor, HeadWidth);
 	}
 
-	class FITDPannerNode : public FNodeFacade
-	{
-	public:
-		/**
-		 * Constructor used by the Metasound Frontend.
-		 */
-		FITDPannerNode(const FNodeInitData& InitData)
-			: FNodeFacade(InitData.InstanceName, InitData.InstanceID, TFacadeOperatorClass<FITDPannerOperator>())
-		{
-		}
-	};
+	using FITDPannerNode = TNodeFacade<FITDPannerOperator>;
 
 	METASOUND_REGISTER_NODE(FITDPannerNode)
 }
diff --git a/Engine/Plugins/Runtime/Metasound/Source/MetasoundStandardNodes/Private/MetasoundInterpToNode.cpp b/Engine/Plugins/Runtime/Metasound/Source/MetasoundStandardNodes/Private/MetasoundInterpToNode.cpp
--- a/Engine/Plugins/Runtime/Metasound/Source/MetasoundStandardNodes/Private/MetasoundInterpToNode.cpp
+++ b/Engine/Plugins/Runtime/Metasound/Source/MetasoundStandardNodes/Private/MetasoundInterpToNode.cpp
@@ -24,18 +24,6 @@ namespace Metasound
 		METASOUND_PARAM(OutParamValue, "Value", "The current value.")
 	}
 
-	/** FInterpToNode
-	*
-	*  Interpolates to a target value over a given time.
-	*/
-	class METASOUNDSTANDARDNODES_API FInterpToNode : public FNodeFacade
-	{
-	public:
-		/**
-		 * Constructor used by the Metasound Frontend.
-		 */
-		FInterpToNode(const FNodeInitData& InitData);
-	};
 
 	class FInterpToOperator : public TExecutableOperator<FInterpToOperator>
 	{
@@ -49,8 +37,6 @@ namespace Metasound
 
 		virtual void BindInputs(FInputVertexInterfaceData& InOutVertexData) override;
 		virtual void BindOutputs(FOutputVertexInterfaceData& InOutVertexData) override;
-		virtual FDataReferenceCollection GetInputs() const override;
-		virtual FDataReferenceCollection GetOutputs() const override;
 		void Reset(const IOperator::FResetParams& InParams);
 		void Execute();
 
@@ -97,22 +83,6 @@ namespace Metasound
 		InOutVertexData.BindReadVertex(METASOUND_GET_PARAM_NAME(OutParamValue), ValueOutput);
 	}
 
-	FDataReferenceCollection FInterpToOperator::GetInputs() const
-	{
-		// This should never be called. Bind(...) is called instead. This method
-		// exists as a stop-gap until the API can be deprecated and removed.
-		checkNoEntry();
-		return {};
-	}
-
-	FDataReferenceCollection FInterpToOperator::GetOutputs() const
-	{
-		// This should never be called. Bind(...) is called instead. This method
-		// exists as a stop-gap until the API can be deprecated and removed.
-		checkNoEntry();
-		return {};
-	}
-
 	void FInterpToOperator::Reset(const IOperator::FResetParams& InParams)
 	{
 		// Set the fade to start at the value specified in the current value
@@ -189,12 +159,6 @@ namespace Metasound
 		return Info;
 	}
 
-
-	FInterpToNode::FInterpToNode(const FNodeInitData& InitData)
-		: FNodeFacade(InitData.InstanceName, InitData.InstanceID, TFacadeOperatorClass<FInterpToOperator>())
-	{
-	}
-
 	TUniquePtr<IOperator> FInterpToOperator::CreateOperator(const FBuildOperatorParams& InParams, FBuildResults& OutResults)
 	{
 		using namespace InterpToVertexNames; 
@@ -207,6 +171,7 @@ namespace Metasound
 		return MakeUnique<FInterpToOperator>(InParams, TargetValue, InterpTime);
 	}
 
+	using FInterpToNode = TNodeFacade<FInterpToOperator>;
 
 	METASOUND_REGISTER_NODE(FInterpToNode)
 }
diff --git a/Engine/Plugins/Runtime/Metasound/Source/MetasoundStandardNodes/Private/MetasoundLimiterNode.cpp b/Engine/Plugins/Runtime/Metasound/Source/MetasoundStandardNodes/Private/MetasoundLimiterNode.cpp
--- a/Engine/Plugins/Runtime/Metasound/Source/MetasoundStandardNodes/Private/MetasoundLimiterNode.cpp
+++ b/Engine/Plugins/Runtime/Metasound/Source/MetasoundStandardNodes/Private/MetasoundLimiterNode.cpp
@@ -138,22 +138,6 @@ namespace Metasound
 			InOutVertexData.BindReadVertex(METASOUND_GET_PARAM_NAME(OutputAudio), AudioOutput);
 		}
 
-		virtual FDataReferenceCollection GetInputs() const override
-		{
-			// This should never be called. Bind(...) is called instead. This method
-			// exists as a stop-gap until the API can be deprecated and removed.
-			checkNoEntry();
-			return {};
-		}
-
-		virtual FDataReferenceCollection GetOutputs() const override
-		{
-			// This should never be called. Bind(...) is called instead. This method
-			// exists as a stop-gap until the API can be deprecated and removed.
-			checkNoEntry();
-			return {};
-		}
-
 		static TUniquePtr<IOperator> CreateOperator(const FBuildOperatorParams& InParams, FBuildResults& OutResults)
 		{
 			using namespace LimiterVertexNames;
@@ -262,14 +246,7 @@ namespace Metasound
 	};
 
 	// Node Class
-	class FLimiterNode : public FNodeFacade
-	{
-	public:
-		FLimiterNode(const FNodeInitData& InitData)
-			: FNodeFacade(InitData.InstanceName, InitData.InstanceID, TFacadeOperatorClass<FLimiterOperator>())
-		{
-		}
-	};
+	using FLimiterNode = TNodeFacade<FLimiterOperator>;
 
 	// Register node
 	METASOUND_REGISTER_NODE(FLimiterNode)
diff --git a/Engine/Plugins/Runtime/Metasound/Source/MetasoundStandardNodes/Private/MetasoundLinearGainToDecibels.cpp b/Engine/Plugins/Runtime/Metasound/Source/MetasoundStandardNodes/Private/MetasoundLinearGainToDecibels.cpp
--- a/Engine/Plugins/Runtime/Metasound/Source/MetasoundStandardNodes/Private/MetasoundLinearGainToDecibels.cpp
+++ b/Engine/Plugins/Runtime/Metasound/Source/MetasoundStandardNodes/Private/MetasoundLinearGainToDecibels.cpp
@@ -33,8 +33,6 @@ namespace Metasound
 
 		virtual void BindInputs(FInputVertexInterfaceData& InOutVertexData) override;
 		virtual void BindOutputs(FOutputVertexInterfaceData& InOutVertexData) override;
-		virtual FDataReferenceCollection GetInputs() const override;
-		virtual FDataReferenceCollection GetOutputs() const override;
 		void Reset(const IOperator::FResetParams& InParams);
 		void Execute();
 
@@ -67,22 +65,6 @@ namespace Metasound
 		InOutVertexData.BindReadVertex(METASOUND_GET_PARAM_NAME(OutputDecibels), DecibelOutput);
 	}
 
-	FDataReferenceCollection FLinearGainToDecibelsOperator::GetInputs() const
-	{
-		// This should never be called. Bind(...) is called instead. This method
-		// exists as a stop-gap until the API can be deprecated and removed.
-		checkNoEntry();
-		return {};
-	}
-
-	FDataReferenceCollection FLinearGainToDecibelsOperator::GetOutputs() const
-	{
-		// This should never be called. Bind(...) is called instead. This method
-		// exists as a stop-gap until the API can be deprecated and removed.
-		checkNoEntry();
-		return {};
-	}
-
 	void FLinearGainToDecibelsOperator::Reset(const IOperator::FResetParams& InParams)
 	{
 		Execute();
@@ -153,17 +135,7 @@ namespace Metasound
 		return MakeUnique<FLinearGainToDecibelsOperator>(InParams, InLinearGain);
 	}
 
-	class METASOUNDSTANDARDNODES_API FLinearGainToDecibelsNode : public FNodeFacade
-	{
-	public:
-		/**
-		 * Constructor used by the Metasound Frontend.
-		 */
-		FLinearGainToDecibelsNode(const FNodeInitData& InitData)
-			: FNodeFacade(InitData.InstanceName, InitData.InstanceID, TFacadeOperatorClass <FLinearGainToDecibelsOperator>())
-		{
-		}
-	};
+	using FLinearGainToDecibelsNode = TNodeFacade<FLinearGainToDecibelsOperator>;
 
 	METASOUND_REGISTER_NODE(FLinearGainToDecibelsNode)
 }
diff --git a/Engine/Plugins/Runtime/Metasound/Source/MetasoundStandardNodes/Private/MetasoundLinearToLogFrequencyNode.cpp b/Engine/Plugins/Runtime/Metasound/Source/MetasoundStandardNodes/Private/MetasoundLinearToLogFrequencyNode.cpp
--- a/Engine/Plugins/Runtime/Metasound/Source/MetasoundStandardNodes/Private/MetasoundLinearToLogFrequencyNode.cpp
+++ b/Engine/Plugins/Runtime/Metasound/Source/MetasoundStandardNodes/Private/MetasoundLinearToLogFrequencyNode.cpp
@@ -125,22 +125,6 @@ namespace Metasound
 			InOutVertexData.BindReadVertex(METASOUND_GET_PARAM_NAME(OutputFreq), FreqOutput);
 		}
 
-		virtual FDataReferenceCollection GetInputs() const override
-		{
-			// This should never be called. Bind(...) is called instead. This method
-			// exists as a stop-gap until the API can be deprecated and removed.
-			checkNoEntry();
-			return {};
-		}
-
-		virtual FDataReferenceCollection GetOutputs() const override
-		{
-			// This should never be called. Bind(...) is called instead. This method
-			// exists as a stop-gap until the API can be deprecated and removed.
-			checkNoEntry();
-			return {};
-		}
-
 		void Reset(const IOperator::FResetParams& InParams)
 		{
 			PrevInputValue = *Value;
@@ -178,19 +162,7 @@ namespace Metasound
 		float PrevInputValue;
 	};
 
-
-	class METASOUNDSTANDARDNODES_API FLinearToLogFrequencyNode : public FNodeFacade
-	{
-	public:
-		/**
-		 * Constructor used by the Metasound Frontend.
-		 */
-		FLinearToLogFrequencyNode(const FNodeInitData& InitData)
-			: FNodeFacade(InitData.InstanceName, InitData.InstanceID, TFacadeOperatorClass<FLinearToLogFrequencyOperator>())
-		{
-
-		}
-	};
+	using FLinearToLogFrequencyNode = TNodeFacade<FLinearToLogFrequencyOperator>;
 
 	METASOUND_REGISTER_NODE(FLinearToLogFrequencyNode)
 }
diff --git a/Engine/Plugins/Runtime/Metasound/Source/MetasoundStandardNodes/Private/MetasoundLowFrequencyNoise.cpp b/Engine/Plugins/Runtime/Metasound/Source/MetasoundStandardNodes/Private/MetasoundLowFrequencyNoise.cpp
--- a/Engine/Plugins/Runtime/Metasound/Source/MetasoundStandardNodes/Private/MetasoundLowFrequencyNoise.cpp
+++ b/Engine/Plugins/Runtime/Metasound/Source/MetasoundStandardNodes/Private/MetasoundLowFrequencyNoise.cpp
@@ -55,8 +55,6 @@ namespace Metasound
 		METASOUND_PARAM(ScaledOutputPin, "Out", "Output value scaled (Min Value...Max Value)")
 	}
 
-	class FRandomLowFrequencyNoiseNode;
-
 	class FLowFrequencyNoiseOperator : public TExecutableOperator<FLowFrequencyNoiseOperator>
 	{
 	public:
@@ -235,22 +233,6 @@ namespace Metasound
 			InOutVertexData.BindReadVertex(METASOUND_GET_PARAM_NAME(NormalizedOutputPin), NormalizedOutput);
 		}
 
-		virtual FDataReferenceCollection GetInputs() const override
-		{
-			// This should never be called. Bind(...) is called instead. This method
-			// exists as a stop-gap until the API can be deprecated and removed.
-			checkNoEntry();
-			return {};
-		}
-
-		virtual FDataReferenceCollection GetOutputs() const override
-		{
-			// This should never be called. Bind(...) is called instead. This method
-			// exists as a stop-gap until the API can be deprecated and removed.
-			checkNoEntry();
-			return {};
-		}
-
 		void SetSeed(int32 InSeedValue)
 		{
 			if (InSeedValue != -1)
@@ -419,15 +401,7 @@ namespace Metasound
 		FFloatWriteRef NormalizedOutput;	// Unit value (0..1)
 	};
 
-	class FRandomLowFrequencyNoiseNode : public FNodeFacade
-	{
-	public:
-		FRandomLowFrequencyNoiseNode(const FNodeInitData& InInitData)
-			: FNodeFacade(InInitData.InstanceName, InInitData.InstanceID, TFacadeOperatorClass<FLowFrequencyNoiseOperator>())
-		{}
-
-		virtual ~FRandomLowFrequencyNoiseNode() = default;
-	};
+	using FRandomLowFrequencyNoiseNode = TNodeFacade<FLowFrequencyNoiseOperator>;
 
 	METASOUND_REGISTER_NODE(FRandomLowFrequencyNoiseNode);
 }
diff --git a/Engine/Plugins/Runtime/Metasound/Source/MetasoundStandardNodes/Private/MetasoundMIDIToFreqNode.cpp b/Engine/Plugins/Runtime/Metasound/Source/MetasoundStandardNodes/Private/MetasoundMIDIToFreqNode.cpp
--- a/Engine/Plugins/Runtime/Metasound/Source/MetasoundStandardNodes/Private/MetasoundMIDIToFreqNode.cpp
+++ b/Engine/Plugins/Runtime/Metasound/Source/MetasoundStandardNodes/Private/MetasoundMIDIToFreqNode.cpp
@@ -70,8 +70,6 @@ namespace Metasound
 
 		virtual void BindInputs(FInputVertexInterfaceData& InOutVertexData) override;
 		virtual void BindOutputs(FOutputVertexInterfaceData& InOutVertexData) override;
-		virtual FDataReferenceCollection GetInputs() const override;
-		virtual FDataReferenceCollection GetOutputs() const override;
 		void Reset(const IOperator::FResetParams& InParams);
 		void Execute();
 
@@ -111,24 +109,6 @@ namespace Metasound
 		InOutVertexData.BindReadVertex(METASOUND_GET_PARAM_NAME(OutputFreq), FreqOutput);
 	}
 
-	template<typename ValueType>
-	FDataReferenceCollection TMidiToFreqOperator<ValueType>::GetInputs() const
-	{
-		// This should never be called. Bind(...) is called instead. This method
-		// exists as a stop-gap until the API can be deprecated and removed.
-		checkNoEntry();
-		return {};
-	}
-
-	template<typename ValueType>
-	FDataReferenceCollection TMidiToFreqOperator<ValueType>::GetOutputs() const
-	{
-		// This should never be called. Bind(...) is called instead. This method
-		// exists as a stop-gap until the API can be deprecated and removed.
-		checkNoEntry();
-		return {};
-	}
-
 	template<typename ValueType>
 	void TMidiToFreqOperator<ValueType>::Reset(const IOperator::FResetParams& InParams)
 	{
@@ -210,18 +190,7 @@ namespace Metasound
 	}
 
 	template<typename ValueType>
-	class METASOUNDSTANDARDNODES_API TMidiToFreqNode : public FNodeFacade
-	{
-	public:
-		/**
-		 * Constructor used by the Metasound Frontend.
-		 */
-		TMidiToFreqNode(const FNodeInitData& InitData)
-			: FNodeFacade(InitData.InstanceName, InitData.InstanceID, TFacadeOperatorClass <TMidiToFreqOperator<ValueType>>())
-		{
-
-		}
-	};
+	using TMidiToFreqNode = TNodeFacade<TMidiToFreqOperator<ValueType>>;
 
 	using FMidiToFreqNodeInt32 = TMidiToFreqNode<int32>;
 	METASOUND_REGISTER_NODE(FMidiToFreqNodeInt32)
diff --git a/Engine/Plugins/Runtime/Metasound/Source/MetasoundStandardNodes/Private/MetasoundMapRangeNode.cpp b/Engine/Plugins/Runtime/Metasound/Source/MetasoundStandardNodes/Private/MetasoundMapRangeNode.cpp
--- a/Engine/Plugins/Runtime/Metasound/Source/MetasoundStandardNodes/Private/MetasoundMapRangeNode.cpp
+++ b/Engine/Plugins/Runtime/Metasound/Source/MetasoundStandardNodes/Private/MetasoundMapRangeNode.cpp
@@ -340,22 +340,6 @@ namespace Metasound
 			MappingObject.BindOutputs(InOutVertexData);
 		}
 
-		virtual FDataReferenceCollection GetInputs() const override
-		{
-			// This should never be called. Bind(...) is called instead. This method
-			// exists as a stop-gap until the API can be deprecated and removed.
-			checkNoEntry();
-			return {};
-		}
-
-		virtual FDataReferenceCollection GetOutputs() const override
-		{
-			// This should never be called. Bind(...) is called instead. This method
-			// exists as a stop-gap until the API can be deprecated and removed.
-			checkNoEntry();
-			return {};
-		}
-
 		void Reset(const IOperator::FResetParams& InParams)
 		{
 			MappingObject.DoMapping();
@@ -371,18 +355,7 @@ namespace Metasound
 	};
 
 	template<typename FMappingClass>
-	class TMapRangeNode : public FNodeFacade
-	{
-	public:
-		/**
-		 * Constructor used by the Metasound Frontend.
-		 */
-		TMapRangeNode(const FNodeInitData& InInitData)
-			: FNodeFacade(InInitData.InstanceName, InInitData.InstanceID, TFacadeOperatorClass<TMapRangeOperator<FMappingClass>>())
-		{}
-
-		virtual ~TMapRangeNode() = default;
-	};
+	using TMapRangeNode = TNodeFacade<TMapRangeOperator<FMappingClass>>;
 
 	using FMapRangeNodeInt32 = TMapRangeNode<MetasoundMapRangeNodePrivate::FIntRange>;
  	METASOUND_REGISTER_NODE(FMapRangeNodeInt32)
diff --git a/Engine/Plugins/Runtime/Metasound/Source/MetasoundStandardNodes/Private/MetasoundMathNodes.cpp b/Engine/Plugins/Runtime/Metasound/Source/MetasoundStandardNodes/Private/MetasoundMathNodes.cpp
--- a/Engine/Plugins/Runtime/Metasound/Source/MetasoundStandardNodes/Private/MetasoundMathNodes.cpp
+++ b/Engine/Plugins/Runtime/Metasound/Source/MetasoundStandardNodes/Private/MetasoundMathNodes.cpp
@@ -26,28 +26,28 @@
 
 #define DEFINE_METASOUND_MATHOP(OpName, DataTypeName, DataClass, Description, Keywords) \
 	class F##OpName##DataTypeName##Node \
-		: public ::Metasound::TMathOpNode<F##OpName##DataTypeName##Node, ::Metasound::TMathOp##OpName<DataClass, DataClass>, DataClass, DataClass> \
+	: public TNodeFacade<TMathOperator<F##OpName##DataTypeName##Node, ::Metasound::TMathOp##OpName<DataClass, DataClass>, DataClass, DataClass>>\
 	{ \
 	public: \
+		using TNodeFacade::TNodeFacade; \
 		static FNodeClassName GetClassName() { return { ::Metasound::StandardNodes::Namespace, TEXT(#OpName), TEXT(#DataTypeName)}; } \
 		static FText GetDisplayName() { return MathOpNames::Get##OpName##DisplayName<DataClass>(); } \
 		static FText GetDescription() { return Description; } \
 		static FName GetImageName() { return "MetasoundEditor.Graph.Node.Math." #OpName; } \
 		static TArray<FText> GetKeywords() { return Keywords; } \
-F##OpName##DataTypeName##Node(const FNodeInitData& InInitData) : TMathOpNode(InInitData) { } \
 	};
 
 #define DEFINE_METASOUND_OPERAND_TYPED_MATHOP(OpName, DataTypeName, DataClass, OperandTypeName, OperandDataClass, Description, Keywords) \
 	class F##OpName##DataTypeName##OperandTypeName##Node \
-		: public ::Metasound::TMathOpNode<F##OpName##DataTypeName##OperandTypeName##Node, ::Metasound::TMathOp##OpName<DataClass, OperandDataClass>, DataClass, OperandDataClass> \
+	: public TNodeFacade<TMathOperator<F##OpName##DataTypeName##OperandTypeName##Node, ::Metasound::TMathOp##OpName<DataClass, OperandDataClass>, DataClass, OperandDataClass>> \
 	{ \
 	public: \
+		using TNodeFacade::TNodeFacade; \
 		static FNodeClassName GetClassName() { return {::Metasound::StandardNodes::Namespace, TEXT(#OpName), TEXT(#DataTypeName " by " #OperandTypeName)}; } \
 		static FText GetDisplayName() { return MathOpNames::Get##OpName##DisplayName<DataClass, OperandDataClass>(); } \
 		static FText GetDescription() { return Description; } \
 		static FName GetImageName() { return "MetasoundEditor.Graph.Node.Math." #OpName; } \
 		static TArray<FText> GetKeywords() { return Keywords; } \
-		F##OpName##DataTypeName##OperandTypeName##Node(const FNodeInitData& InInitData) : TMathOpNode(InInitData) { } \
 	};
 
 namespace Metasound
@@ -195,112 +195,101 @@ namespace Metasound
 	}
 
 	template <typename TDerivedClass, typename TMathOpClass, typename TDataClass, typename TOperandDataClass>
-	class TMathOpNode : public FNodeFacade
+	class TMathOperator : public TExecutableOperator<TMathOperator<TDerivedClass, TMathOpClass, TDataClass, TOperandDataClass>>
 	{
 	private:
-		class TMathOperator : public TExecutableOperator<TMathOperator>
-		{
-		private:
-			using TDataClassReadRef = TDataReadReference<TDataClass>;
-			using TOperandDataClassReadRef = TDataReadReference<TOperandDataClass>;
-			using TDataClassWriteRef = TDataWriteReference<TDataClass>;
-
-			TMathOpClass InstanceData;
+		using TDataClassReadRef = TDataReadReference<TDataClass>;
+		using TOperandDataClassReadRef = TDataReadReference<TOperandDataClass>;
+		using TDataClassWriteRef = TDataWriteReference<TDataClass>;
 
-			TDataClassReadRef PrimaryOperandRef;
-			TArray<TOperandDataClassReadRef> AdditionalOperandRefs;
-			TDataClassWriteRef ValueRef;
+		TMathOpClass InstanceData;
 
-		public:
-			static const FNodeClassMetadata& GetNodeInfo()
-			{
-				auto InitNodeInfo = []() -> FNodeClassMetadata
-				{
-					FNodeDisplayStyle DisplayStyle;
-					DisplayStyle.ImageName = TDerivedClass::GetImageName();
-					DisplayStyle.bShowName = false;
-					DisplayStyle.bShowInputNames = false;
-					DisplayStyle.bShowOutputNames = false;
-
-					FNodeClassMetadata Info;
-					Info.ClassName = TDerivedClass::GetClassName();
-					Info.MajorVersion = 1;
-					Info.MinorVersion = 1;
-					Info.DisplayName = TDerivedClass::GetDisplayName();
-					Info.Description = TDerivedClass::GetDescription();
-					Info.Author = PluginAuthor;
-					Info.PromptIfMissing = PluginNodeMissingPrompt;
-					Info.DefaultInterface = TMathOpClass::GetVertexInterface();
-					Info.DisplayStyle = DisplayStyle;
-					Info.CategoryHierarchy = { NodeCategories::Math };
-					Info.Keywords = TDerivedClass::GetKeywords();
-
-					return Info;
-				};
-
-				static const FNodeClassMetadata Info = InitNodeInfo();
-				return Info;
-			}
+		TDataClassReadRef PrimaryOperandRef;
+		TArray<TOperandDataClassReadRef> AdditionalOperandRefs;
+		TDataClassWriteRef ValueRef;
 
-			void Reset(const IOperator::FResetParams& InParams)
-			{
-				TMathOpClass::Calculate(InstanceData, PrimaryOperandRef, AdditionalOperandRefs, ValueRef);
-			}
+	public:
+		static const FNodeClassMetadata& GetNodeInfo()
+		{
+			auto InitNodeInfo = []() -> FNodeClassMetadata
+			{
+				FNodeDisplayStyle DisplayStyle;
+				DisplayStyle.ImageName = TDerivedClass::GetImageName();
+				DisplayStyle.bShowName = false;
+				DisplayStyle.bShowInputNames = false;
+				DisplayStyle.bShowOutputNames = false;
+
+				FNodeClassMetadata Info;
+				Info.ClassName = TDerivedClass::GetClassName();
+				Info.MajorVersion = 1;
+				Info.MinorVersion = 1;
+				Info.DisplayName = TDerivedClass::GetDisplayName();
+				Info.Description = TDerivedClass::GetDescription();
+				Info.Author = PluginAuthor;
+				Info.PromptIfMissing = PluginNodeMissingPrompt;
+				Info.DefaultInterface = TMathOpClass::GetVertexInterface();
+				Info.DisplayStyle = DisplayStyle;
+				Info.CategoryHierarchy = { NodeCategories::Math };
+				Info.Keywords = TDerivedClass::GetKeywords();
 
-			void Execute()
-			{
-				TMathOpClass::Calculate(InstanceData, PrimaryOperandRef, AdditionalOperandRefs, ValueRef);
-			}
+				return Info;
+			};
 
-			static TUniquePtr<IOperator> CreateOperator(const FBuildOperatorParams& InParams, FBuildResults& OutResults)
-			{
-				using namespace MathOpNames;
+			static const FNodeClassMetadata Info = InitNodeInfo();
+			return Info;
+		}
 
-				const TMathOpNode& MathOpNode = static_cast<const TMathOpNode&>(InParams.Node);
+		void Reset(const IOperator::FResetParams& InParams)
+		{
+			TMathOpClass::Calculate(InstanceData, PrimaryOperandRef, AdditionalOperandRefs, ValueRef);
+		}
 
-				const FInputVertexInterfaceData& InputData = InParams.InputData;
-				const FInputVertexInterface& InputInterface = MathOpNode.GetVertexInterface().GetInputInterface();
+		void Execute()
+		{
+			TMathOpClass::Calculate(InstanceData, PrimaryOperandRef, AdditionalOperandRefs, ValueRef);
+		}
 
-				TDataClassReadRef PrimaryOperand = InputData.GetOrCreateDefaultDataReadReference<TDataClass>(METASOUND_GET_PARAM_NAME(PrimaryOperand), InParams.OperatorSettings);
+		static TUniquePtr<IOperator> CreateOperator(const FBuildOperatorParams& InParams, FBuildResults& OutResults)
+		{
+			using namespace MathOpNames;
 
-				// TODO: Support dynamic number of inputs
-				TOperandDataClassReadRef Op1 = InputData.GetOrCreateDefaultDataReadReference<TOperandDataClass>(METASOUND_GET_PARAM_NAME(AdditionalOperands), InParams.OperatorSettings);
-				TArray<TOperandDataClassReadRef> AdditionalOperandRefs = { Op1 };
+			const FInputVertexInterfaceData& InputData = InParams.InputData;
 
-				return MakeUnique<TMathOperator>(InParams, PrimaryOperand, AdditionalOperandRefs);
-			}
+			TDataClassReadRef PrimaryOperand = InputData.GetOrCreateDefaultDataReadReference<TDataClass>(METASOUND_GET_PARAM_NAME(PrimaryOperand), InParams.OperatorSettings);
 
-			TMathOperator(const FBuildOperatorParams& InParams, const TDataClassReadRef& InPrimaryOperand, const TArray<TOperandDataClassReadRef>& InAdditionalOperands)
-				: PrimaryOperandRef(InPrimaryOperand)
-				, AdditionalOperandRefs(InAdditionalOperands)
-				, ValueRef(TDataWriteReferenceFactory<TDataClass>::CreateAny(InParams.OperatorSettings))
-			{
-				Reset(InParams);
-			}
+			// TODO: Support dynamic number of inputs
+			TOperandDataClassReadRef Op1 = InputData.GetOrCreateDefaultDataReadReference<TOperandDataClass>(METASOUND_GET_PARAM_NAME(AdditionalOperands), InParams.OperatorSettings);
+			TArray<TOperandDataClassReadRef> AdditionalOperandRefs = { Op1 };
 
-			virtual void BindInputs(FInputVertexInterfaceData& InOutVertexData) override
-			{
-				using namespace MathOpNames;
+			return MakeUnique<TMathOperator>(InParams, PrimaryOperand, AdditionalOperandRefs);
+		}
 
-				InOutVertexData.BindReadVertex(METASOUND_GET_PARAM_NAME(PrimaryOperand), PrimaryOperandRef);
-				InOutVertexData.BindReadVertex(METASOUND_GET_PARAM_NAME(AdditionalOperands), AdditionalOperandRefs[0]);
-			}
+		TMathOperator(const FBuildOperatorParams& InParams, const TDataClassReadRef& InPrimaryOperand, const TArray<TOperandDataClassReadRef>& InAdditionalOperands)
+			: PrimaryOperandRef(InPrimaryOperand)
+			, AdditionalOperandRefs(InAdditionalOperands)
+			, ValueRef(TDataWriteReferenceFactory<TDataClass>::CreateAny(InParams.OperatorSettings))
+		{
+			Reset(InParams);
+		}
 
-			virtual void BindOutputs(FOutputVertexInterfaceData& InOutVertexData) override
-			{
-				using namespace MathOpNames;
+		virtual void BindInputs(FInputVertexInterfaceData& InOutVertexData) override
+		{
+			using namespace MathOpNames;
 
-				InOutVertexData.BindReadVertex(METASOUND_GET_PARAM_NAME(OutMath), ValueRef);
-			}
-		};
+			InOutVertexData.BindReadVertex(METASOUND_GET_PARAM_NAME(PrimaryOperand), PrimaryOperandRef);
+			InOutVertexData.BindReadVertex(METASOUND_GET_PARAM_NAME(AdditionalOperands), AdditionalOperandRefs[0]);
+		}
 
-	public:
-		TMathOpNode(const FNodeInitData& InInitData)
-			: FNodeFacade(InInitData.InstanceName, InInitData.InstanceID, TFacadeOperatorClass<TMathOperator>())
+		virtual void BindOutputs(FOutputVertexInterfaceData& InOutVertexData) override
 		{
+			using namespace MathOpNames;
+
+			InOutVertexData.BindReadVertex(METASOUND_GET_PARAM_NAME(OutMath), ValueRef);
 		}
 	};
 
+
+
 	// Standard Math Operations
 	template <typename TDataClass, typename TOperandDataClass = TDataClass>
 	class TMathOpAdd
@@ -1314,6 +1303,9 @@ namespace Metasound
 		}
 	};
 
+	template <typename TDerivedClass, typename TMathOpClass, typename TDataClass, typename TOperandDataClass>
+	using TMathOpNode = TNodeFacade<TMathOperator<TDerivedClass, TMathOpClass, TDataClass, TOperandDataClass>>;
+
 	// Definitions
 	DEFINE_METASOUND_MATHOP(Add, Float, float, METASOUND_LOCTEXT("MathAddFloatNode_Description", "Adds floats."), MathOpNames::AddKeywords)
 	DEFINE_METASOUND_MATHOP(Add, Int32, int32, METASOUND_LOCTEXT("MathAddInt32Node_Description", "Adds int32s."), MathOpNames::AddKeywords)
diff --git a/Engine/Plugins/Runtime/Metasound/Source/MetasoundStandardNodes/Private/MetasoundMaxNode.cpp b/Engine/Plugins/Runtime/Metasound/Source/MetasoundStandardNodes/Private/MetasoundMaxNode.cpp
--- a/Engine/Plugins/Runtime/Metasound/Source/MetasoundStandardNodes/Private/MetasoundMaxNode.cpp
+++ b/Engine/Plugins/Runtime/Metasound/Source/MetasoundStandardNodes/Private/MetasoundMaxNode.cpp
@@ -208,22 +208,6 @@ namespace Metasound
 			InOutVertexData.BindReadVertex(METASOUND_GET_PARAM_NAME(OutputValue), OutputValue);
 		}
 
-		virtual FDataReferenceCollection GetInputs() const override
-		{
-			// This should never be called. Bind(...) is called instead. This method
-			// exists as a stop-gap until the API can be deprecated and removed.
-			checkNoEntry();
-			return {};
-		}
-
-		virtual FDataReferenceCollection GetOutputs() const override
-		{
-			// This should never be called. Bind(...) is called instead. This method
-			// exists as a stop-gap until the API can be deprecated and removed.
-			checkNoEntry();
-			return {};
-		}
-
 		void GetMax()
 		{
 			using namespace MetasoundMaxNodePrivate;
@@ -253,18 +237,7 @@ namespace Metasound
 	 *  Returns the min of both inputs
 	 */
 	template<typename ValueType>
-	class METASOUNDSTANDARDNODES_API TMaxNode : public FNodeFacade
-	{
-	public:
-		/**
-		 * Constructor used by the Metasound Frontend.
-		 */
-		TMaxNode(const FNodeInitData& InInitData)
-			: FNodeFacade(InInitData.InstanceName, InInitData.InstanceID, TFacadeOperatorClass<TMaxNodeOperator<ValueType>>())
-		{}
-
-		virtual ~TMaxNode() = default;
-	};
+	using TMaxNode = TNodeFacade<TMaxNodeOperator<ValueType>>;
 
 	using FMaxNodeInt32 = TMaxNode<int32>;
  	METASOUND_REGISTER_NODE(FMaxNodeInt32)
diff --git a/Engine/Plugins/Runtime/Metasound/Source/MetasoundStandardNodes/Private/MetasoundMidSideNodes.cpp b/Engine/Plugins/Runtime/Metasound/Source/MetasoundStandardNodes/Private/MetasoundMidSideNodes.cpp
--- a/Engine/Plugins/Runtime/Metasound/Source/MetasoundStandardNodes/Private/MetasoundMidSideNodes.cpp
+++ b/Engine/Plugins/Runtime/Metasound/Source/MetasoundStandardNodes/Private/MetasoundMidSideNodes.cpp
@@ -122,22 +122,6 @@ namespace Metasound
 			InOutVertexData.BindReadVertex(METASOUND_GET_PARAM_NAME(OutputAudioSide), AudioOutputSide);
 		}
 
-		virtual FDataReferenceCollection GetInputs() const override
-		{
-			// This should never be called. Bind(...) is called instead. This method
-			// exists as a stop-gap until the API can be deprecated and removed.
-			checkNoEntry();
-			return {};
-		}
-
-		virtual FDataReferenceCollection GetOutputs() const override
-		{
-			// This should never be called. Bind(...) is called instead. This method
-			// exists as a stop-gap until the API can be deprecated and removed.
-			checkNoEntry();
-			return {};
-		}
-		
 		static TUniquePtr<IOperator> CreateOperator(const FBuildOperatorParams& InParams, FBuildResults& OutResults)
 		{
 			using namespace MidSideEncodeVertexNames;
@@ -251,14 +235,7 @@ namespace Metasound
 	};
 
 	// Node Class
-	class FMidSideEncodeNode : public FNodeFacade
-	{
-	public:
-		FMidSideEncodeNode(const FNodeInitData& InitData)
-			: FNodeFacade(InitData.InstanceName, InitData.InstanceID, TFacadeOperatorClass<FMidSideEncodeOperator>())
-		{
-		}
-	};
+	using FMidSideEncodeNode = TNodeFacade<FMidSideEncodeOperator>;
 
 	// Register node
 	METASOUND_REGISTER_NODE(FMidSideEncodeNode)
@@ -376,22 +353,6 @@ namespace Metasound
 			InOutVertexData.BindReadVertex(METASOUND_GET_PARAM_NAME(OutputAudioRight), AudioOutputRight);
 		}
 
-		virtual FDataReferenceCollection GetInputs() const override
-		{
-			// This should never be called. Bind(...) is called instead. This method
-			// exists as a stop-gap until the API can be deprecated and removed.
-			checkNoEntry();
-			return {};
-		}
-
-		virtual FDataReferenceCollection GetOutputs() const override
-		{
-			// This should never be called. Bind(...) is called instead. This method
-			// exists as a stop-gap until the API can be deprecated and removed.
-			checkNoEntry();
-			return {};
-		}
-
 		static TUniquePtr<IOperator> CreateOperator(const FBuildOperatorParams& InParams, FBuildResults& OutResults)
 		{
 			using namespace MidSideDecodeVertexNames;
@@ -507,14 +468,7 @@ namespace Metasound
 	};
 
 	// Node Class
-	class FMidSideDecodeNode : public FNodeFacade
-	{
-	public:
-		FMidSideDecodeNode(const FNodeInitData& InitData)
-			: FNodeFacade(InitData.InstanceName, InitData.InstanceID, TFacadeOperatorClass<FMidSideDecodeOperator>())
-		{
-		}
-	};
+	using FMidSideDecodeNode = TNodeFacade<FMidSideDecodeOperator>;
 
 	// Register node
 	METASOUND_REGISTER_NODE(FMidSideDecodeNode)
diff --git a/Engine/Plugins/Runtime/Metasound/Source/MetasoundStandardNodes/Private/MetasoundMidiNoteQuantizerNode.cpp b/Engine/Plugins/Runtime/Metasound/Source/MetasoundStandardNodes/Private/MetasoundMidiNoteQuantizerNode.cpp
--- a/Engine/Plugins/Runtime/Metasound/Source/MetasoundStandardNodes/Private/MetasoundMidiNoteQuantizerNode.cpp
+++ b/Engine/Plugins/Runtime/Metasound/Source/MetasoundStandardNodes/Private/MetasoundMidiNoteQuantizerNode.cpp
@@ -64,8 +64,6 @@ namespace Metasound
 
 		virtual void BindInputs(FInputVertexInterfaceData& InOutVertexData) override;
 		virtual void BindOutputs(FOutputVertexInterfaceData& InOutVertexData) override;
-		virtual FDataReferenceCollection GetInputs() const override;
-		virtual FDataReferenceCollection GetOutputs() const override;
 		void Reset(const IOperator::FResetParams& InParams);
 		void Execute();
 
@@ -191,22 +189,6 @@ namespace Metasound
 		InOutVertexData.BindReadVertex(METASOUND_GET_PARAM_NAME(ParamNoteOutput), MidiNoteOut);
 	}
 
-	FDataReferenceCollection FMidiNoteQuantizerOperator::GetInputs() const
-	{
-		// This should never be called. Bind(...) is called instead. This method
-		// exists as a stop-gap until the API can be deprecated and removed.
-		checkNoEntry();
-		return {};
-	}
-
-	FDataReferenceCollection FMidiNoteQuantizerOperator::GetOutputs() const
-	{
-		// This should never be called. Bind(...) is called instead. This method
-		// exists as a stop-gap until the API can be deprecated and removed.
-		checkNoEntry();
-		return {};
-	}
-
 	void FMidiNoteQuantizerOperator::Reset(const IOperator::FResetParams& InParams)
 	{
 		PreviousNoteIn = FMath::Clamp(*MidiNoteIn, MinNote, MaxNote);
@@ -262,21 +244,7 @@ namespace Metasound
 
 
 #pragma region Node Declaration
-	class METASOUNDSTANDARDNODES_API FMidiNoteQuantizerNode: public FNodeFacade
-	{
-	public:
-		// public node api needs to define two conversion constructors:
-		// (1: from FString)
-		FMidiNoteQuantizerNode(const FVertexName& InInstanceName, const FGuid & InInstanceID)
-			: FNodeFacade(InInstanceName, InInstanceID, TFacadeOperatorClass <FMidiNoteQuantizerOperator>())
-		{ }
-
-		// (2: From an NodeInitData struct)
-		FMidiNoteQuantizerNode(const FNodeInitData & InInitData)
-			: FMidiNoteQuantizerNode(InInitData.InstanceName, InInitData.InstanceID)
-		{ }
-
-	};
+	using FMidiNoteQuantizerNode = TNodeFacade<FMidiNoteQuantizerOperator>;
 #pragma endregion // Node Declaration
 
 
diff --git a/Engine/Plugins/Runtime/Metasound/Source/MetasoundStandardNodes/Private/MetasoundMidiScaleToArrayNode.cpp b/Engine/Plugins/Runtime/Metasound/Source/MetasoundStandardNodes/Private/MetasoundMidiScaleToArrayNode.cpp
--- a/Engine/Plugins/Runtime/Metasound/Source/MetasoundStandardNodes/Private/MetasoundMidiScaleToArrayNode.cpp
+++ b/Engine/Plugins/Runtime/Metasound/Source/MetasoundStandardNodes/Private/MetasoundMidiScaleToArrayNode.cpp
@@ -100,8 +100,6 @@ namespace Metasound
 		static TUniquePtr<IOperator> CreateOperator(const FBuildOperatorParams& InParams, FBuildResults& OutResults);
 		virtual void BindInputs(FInputVertexInterfaceData& InOutVertexData) override;
 		virtual void BindOutputs(FOutputVertexInterfaceData& InOutVertexData) override;
-		virtual FDataReferenceCollection GetInputs() const override;
-		virtual FDataReferenceCollection GetOutputs() const override;
 		void Reset(const IOperator::FResetParams& InParams);
 		void Execute();
 
@@ -207,22 +205,6 @@ namespace Metasound
 		InOutVertexData.BindReadVertex(METASOUND_GET_PARAM_NAME(ParamNoteArrayOutput), ScaleArrayOutput);
 	}
 
-	FDataReferenceCollection FMusicalScaleToNoteArrayOperator::GetInputs() const
-	{
-		// This should never be called. Bind(...) is called instead. This method
-		// exists as a stop-gap until the API can be deprecated and removed.
-		checkNoEntry();
-		return {};
-	}
-
-	FDataReferenceCollection FMusicalScaleToNoteArrayOperator::GetOutputs() const
-	{
-		// This should never be called. Bind(...) is called instead. This method
-		// exists as a stop-gap until the API can be deprecated and removed.
-		checkNoEntry();
-		return {};
-	}
-
 	void FMusicalScaleToNoteArrayOperator::Reset(const IOperator::FResetParams& InParams)
 	{
 		PreviousScale = *Scale;
@@ -248,21 +230,7 @@ namespace Metasound
 
 
 #pragma region Node Declaration
-	class METASOUNDSTANDARDNODES_API FMusicalScaleToNoteArrayNode: public FNodeFacade
-	{
-	public:
-		// public node api needs to define two conversion constructors:
-		// (1: from FString)
-		FMusicalScaleToNoteArrayNode(const FVertexName& InInstanceName, const FGuid & InInstanceID)
-			: FNodeFacade(InInstanceName, InInstanceID, TFacadeOperatorClass <FMusicalScaleToNoteArrayOperator>())
-		{ }
-
-		// (2: From an NodeInitData struct)
-		FMusicalScaleToNoteArrayNode(const FNodeInitData & InInitData)
-			: FMusicalScaleToNoteArrayNode(InInitData.InstanceName, InInitData.InstanceID)
-		{ }
-
-	};
+	using FMusicalScaleToNoteArrayNode = TNodeFacade<FMusicalScaleToNoteArrayOperator>;
 #pragma endregion // Node Declaration
 
 
diff --git a/Engine/Plugins/Runtime/Metasound/Source/MetasoundStandardNodes/Private/MetasoundMinNode.cpp b/Engine/Plugins/Runtime/Metasound/Source/MetasoundStandardNodes/Private/MetasoundMinNode.cpp
--- a/Engine/Plugins/Runtime/Metasound/Source/MetasoundStandardNodes/Private/MetasoundMinNode.cpp
+++ b/Engine/Plugins/Runtime/Metasound/Source/MetasoundStandardNodes/Private/MetasoundMinNode.cpp
@@ -204,22 +204,6 @@ namespace Metasound
 			InOutVertexData.BindReadVertex(METASOUND_GET_PARAM_NAME(OutputValue), OutputValue);
 		}
 
-		virtual FDataReferenceCollection GetInputs() const override
-		{
-			// This should never be called. Bind(...) is called instead. This method
-			// exists as a stop-gap until the API can be deprecated and removed.
-			checkNoEntry();
-			return {};
-		}
-
-		virtual FDataReferenceCollection GetOutputs() const override
-		{
-			// This should never be called. Bind(...) is called instead. This method
-			// exists as a stop-gap until the API can be deprecated and removed.
-			checkNoEntry();
-			return {};
-		}
-
 		void GetMin()
 		{
 			using namespace MetasoundMinNodePrivate;
@@ -249,18 +233,7 @@ namespace Metasound
 	 *  Returns the min of both inputs
 	 */
 	template<typename ValueType>
-	class METASOUNDSTANDARDNODES_API TMinNode : public FNodeFacade
-	{
-	public:
-		/**
-		 * Constructor used by the Metasound Frontend.
-		 */
-		TMinNode(const FNodeInitData& InInitData)
-			: FNodeFacade(InInitData.InstanceName, InInitData.InstanceID, TFacadeOperatorClass<TMinNodeOperator<ValueType>>())
-		{}
-
-		virtual ~TMinNode() = default;
-	};
+	using TMinNode = TNodeFacade<TMinNodeOperator<ValueType>>;
 
 	using FMinNodeInt32 = TMinNode<int32>;
 	METASOUND_REGISTER_NODE(FMinNodeInt32)
diff --git a/Engine/Plugins/Runtime/Metasound/Source/MetasoundStandardNodes/Private/MetasoundMixerNode.cpp b/Engine/Plugins/Runtime/Metasound/Source/MetasoundStandardNodes/Private/MetasoundMixerNode.cpp
--- a/Engine/Plugins/Runtime/Metasound/Source/MetasoundStandardNodes/Private/MetasoundMixerNode.cpp
+++ b/Engine/Plugins/Runtime/Metasound/Source/MetasoundStandardNodes/Private/MetasoundMixerNode.cpp
@@ -186,22 +186,6 @@ namespace Metasound
 			}
 		}
 
-		virtual FDataReferenceCollection GetInputs() const override
-		{
-			// This should never be called. Bind(...) is called instead. This method
-			// exists as a stop-gap until the API can be deprecated and removed.
-			checkNoEntry();
-			return {};
-		}
-
-		virtual FDataReferenceCollection GetOutputs() const override
-		{
-			// This should never be called. Bind(...) is called instead. This method
-			// exists as a stop-gap until the API can be deprecated and removed.
-			checkNoEntry();
-			return {};
-		}
-
 		void Reset(const IOperator::FResetParams& InParams)
 		{
 			for (uint32 i = 0; i < NumInputs; ++i)
@@ -429,19 +413,7 @@ namespace Metasound
 
 
 	template<uint32 NumInputs, uint32 NumChannels>
-	class METASOUNDSTANDARDNODES_API TAudioMixerNode : public FNodeFacade
-	{
-	public:
-		/**
-		 * Constructor used by the Metasound Frontend.
-		 */
-		TAudioMixerNode(const FNodeInitData& InInitData)
-			: FNodeFacade(InInitData.InstanceName, InInitData.InstanceID, TFacadeOperatorClass<TAudioMixerNodeOperator<NumInputs, NumChannels>>())
-		{}
-
-		virtual ~TAudioMixerNode() = default;
-	};
-
+	using TAudioMixerNode = TNodeFacade<TAudioMixerNodeOperator<NumInputs, NumChannels>>;
 
 	#define REGISTER_AUDIOMIXER_NODE(A, B) \
 		using FAudioMixerNode_##A ## _ ##B = TAudioMixerNode<A, B>; \
diff --git a/Engine/Plugins/Runtime/Metasound/Source/MetasoundStandardNodes/Private/MetasoundNoiseGenerator.cpp b/Engine/Plugins/Runtime/Metasound/Source/MetasoundStandardNodes/Private/MetasoundNoiseGenerator.cpp
--- a/Engine/Plugins/Runtime/Metasound/Source/MetasoundStandardNodes/Private/MetasoundNoiseGenerator.cpp
+++ b/Engine/Plugins/Runtime/Metasound/Source/MetasoundStandardNodes/Private/MetasoundNoiseGenerator.cpp
@@ -106,8 +106,6 @@ namespace Metasound
 		}
 	};
 
-	constexpr int32 FNoiseOperator::DefaultSeed;
-
 	struct FNoiseOperator_White final : public FNoiseOperator
 	{
 		Audio::FWhiteNoise Generator;
@@ -229,15 +227,24 @@ namespace Metasound
 		return Info;
 	}
 
-	FNoiseNode::FNoiseNode(const FVertexName& InName, const FGuid& InInstanceID, int32 InDefaultSeed)
-		: FNodeFacade(InName, InInstanceID, TFacadeOperatorClass<FNoiseOperator>())
-		, DefaultSeed(InDefaultSeed)
+	FNoiseNode::FNoiseNode(const FVertexName& InName, const FGuid& InInstanceID, int32)
+	: FNoiseNode(FNodeInitData{InName, InInstanceID})
 	{}
 
 	FNoiseNode::FNoiseNode(const FNodeInitData& InInitData)
-		: FNoiseNode(InInitData.InstanceName, InInitData.InstanceID, FNoiseOperator::DefaultSeed)
+	: FNoiseNode(FNodeData(InInitData.InstanceName, InInitData.InstanceID, FNoiseOperator::GetNodeInfo().DefaultInterface), MakeShared<const FNodeClassMetadata>(FNoiseOperator::GetNodeInfo()))
 	{}
 
+	FNoiseNode::FNoiseNode(FNodeData InNodeData, TSharedRef<const FNodeClassMetadata> InClassMetadata)
+	: FNodeFacade(MoveTemp(InNodeData), MoveTemp(InClassMetadata), TFacadeOperatorClass<FNoiseOperator>())
+	{
+	}
+
+	FNodeClassMetadata FNoiseNode::CreateNodeClassMetadata()
+	{
+		return FNoiseOperator::GetNodeInfo();
+	}
+
 	TUniquePtr<IOperator> FNoiseOperator::CreateOperator(const FBuildOperatorParams& InParams, FBuildResults& OutResults)
 	{
 		using namespace NoiseGeneratorVertexNames; 
diff --git a/Engine/Plugins/Runtime/Metasound/Source/MetasoundStandardNodes/Private/MetasoundOscillatorNodes.cpp b/Engine/Plugins/Runtime/Metasound/Source/MetasoundStandardNodes/Private/MetasoundOscillatorNodes.cpp
--- a/Engine/Plugins/Runtime/Metasound/Source/MetasoundStandardNodes/Private/MetasoundOscillatorNodes.cpp
+++ b/Engine/Plugins/Runtime/Metasound/Source/MetasoundStandardNodes/Private/MetasoundOscillatorNodes.cpp
@@ -255,13 +255,18 @@ namespace Metasound
 	};
 
 	FOscilatorNodeBase::FOscilatorNodeBase(const FVertexName& InInstanceName, const FGuid& InInstanceID, const FNodeClassMetadata& InInfo, const TSharedRef<IOperatorFactory, ESPMode::ThreadSafe>& InFactory, float InDefaultFrequency, float InDefaultGlideFactor, bool bInDefaultEnablement)
-		: FNode(InInstanceName, InInstanceID, InInfo)
+		: FBasicNode(FNodeData(InInstanceName, InInstanceID, InInfo.DefaultInterface), MakeShared<FNodeClassMetadata>(InInfo))
 		, Factory(InFactory)
-		, VertexInterface(GetMetadata().DefaultInterface)
 		, DefaultFrequency(InDefaultFrequency)
 		, DefaultGlideFactor(InDefaultGlideFactor)
 		, bDefaultEnablement(bInDefaultEnablement)
 	{}
+
+	FOscilatorNodeBase::FOscilatorNodeBase(FNodeData InNodeData, TSharedRef<const FNodeClassMetadata> InClassMetadata, TSharedRef<IOperatorFactory> InFactory)
+	: FBasicNode(MoveTemp(InNodeData), MoveTemp(InClassMetadata))
+	, Factory(MoveTemp(InFactory))
+	{
+	}
 	
 #pragma endregion Common
 
@@ -315,6 +320,7 @@ namespace Metasound
 			static const FNodeClassMetadata Info = InitNodeInfo();
 			return Info;
 		}
+
 		static const FVertexInterface& GetVertexInterface()
 		{
 			using namespace SineOscilatorVertexNames; 
@@ -388,9 +394,19 @@ namespace Metasound
 	{}
 
 	FSineOscilatorNode::FSineOscilatorNode(const FNodeInitData& InInitData)
-		: FSineOscilatorNode(InInitData.InstanceName, InInitData.InstanceID, 440.0f, 0.0f, true)
+	: FSineOscilatorNode(FNodeData(InInitData.InstanceName, InInitData.InstanceID, FFactory::GetNodeInfo().DefaultInterface), MakeShared<const FNodeClassMetadata>(FFactory::GetNodeInfo()))
 	{}
 
+	FSineOscilatorNode::FSineOscilatorNode(FNodeData InNodeData, TSharedRef<const FNodeClassMetadata> InMetadata)
+	: FOscilatorNodeBase(MoveTemp(InNodeData), MoveTemp(InMetadata), MakeShared<FFactory>())
+	{
+	}
+
+	FNodeClassMetadata FSineOscilatorNode::CreateNodeClassMetadata()
+	{
+		return FFactory::GetNodeInfo();
+	}
+
 	METASOUND_REGISTER_NODE(FSineOscilatorNode);
 
 #pragma endregion Sine
@@ -513,9 +529,19 @@ namespace Metasound
 	{}
 
 	FSawOscilatorNode::FSawOscilatorNode(const FNodeInitData& InInitData)
-		: FSawOscilatorNode(InInitData.InstanceName, InInitData.InstanceID, 440.0f, 0.0f, true)
+	: FSawOscilatorNode(FNodeData(InInitData.InstanceName, InInitData.InstanceID, FFactory::GetNodeInfo().DefaultInterface), MakeShared<FNodeClassMetadata>(FFactory::GetNodeInfo()))
 	{}
 
+	FSawOscilatorNode::FSawOscilatorNode(FNodeData InNodeData, TSharedRef<const FNodeClassMetadata> InMetadata)
+	: FOscilatorNodeBase(MoveTemp(InNodeData), MoveTemp(InMetadata), MakeShared<FFactory>())
+	{
+	}
+
+	FNodeClassMetadata FSawOscilatorNode::CreateNodeClassMetadata()
+	{
+		return FFactory::GetNodeInfo();
+	}
+
 	METASOUND_REGISTER_NODE(FSawOscilatorNode);
 
 #pragma endregion Saw
@@ -733,9 +759,19 @@ namespace Metasound
 	{}
 
 	FSquareOscilatorNode::FSquareOscilatorNode(const FNodeInitData& InInitData)
-		: FSquareOscilatorNode(InInitData.InstanceName, InInitData.InstanceID, 440.0f, 0.0f, true)
+	: FSquareOscilatorNode(FNodeData(InInitData.InstanceName, InInitData.InstanceID, FFactory::GetNodeInfo().DefaultInterface), MakeShared<FNodeClassMetadata>(FFactory::GetNodeInfo()))
 	{}
 
+	FSquareOscilatorNode::FSquareOscilatorNode(FNodeData InNodeData, TSharedRef<const FNodeClassMetadata> InMetadata)
+	: FOscilatorNodeBase(MoveTemp(InNodeData), MoveTemp(InMetadata), MakeShared<FFactory>())
+	{
+	}
+
+	FNodeClassMetadata FSquareOscilatorNode::CreateNodeClassMetadata()
+	{
+		return FFactory::GetNodeInfo();
+	}
+
 	METASOUND_REGISTER_NODE(FSquareOscilatorNode)
 #pragma endregion Square
 
@@ -857,9 +893,19 @@ namespace Metasound
 	}
 
 	FTriangleOscilatorNode::FTriangleOscilatorNode(const FNodeInitData& InInitData)
-		: FTriangleOscilatorNode(InInitData.InstanceName, InInitData.InstanceID, 440.0f, 0.0f, true)
+	: FTriangleOscilatorNode(FNodeData(InInitData.InstanceName, InInitData.InstanceID, FFactory::GetNodeInfo().DefaultInterface), MakeShared<FNodeClassMetadata>(FFactory::GetNodeInfo()))
 	{}
 
+	FTriangleOscilatorNode::FTriangleOscilatorNode(FNodeData InNodeData, TSharedRef<const FNodeClassMetadata> InMetadata)
+	: FOscilatorNodeBase(MoveTemp(InNodeData), MoveTemp(InMetadata), MakeShared<FFactory>())
+	{
+	}
+
+	FNodeClassMetadata FTriangleOscilatorNode::CreateNodeClassMetadata()
+	{
+		return FFactory::GetNodeInfo();
+	}
+
 	METASOUND_REGISTER_NODE(FTriangleOscilatorNode);
 #pragma endregion Triangle
 
@@ -1085,6 +1131,16 @@ namespace Metasound
 		: FNodeFacade(InInitData.InstanceName, InInitData.InstanceID, TFacadeOperatorClass<FLfoOperator>())
 	{}
 
+	FLfoNode::FLfoNode(FNodeData InNodeData, TSharedRef<const FNodeClassMetadata> InMetadata)
+	: FNodeFacade(MoveTemp(InNodeData), MoveTemp(InMetadata), TFacadeOperatorClass<FLfoOperator>())
+	{
+	}
+
+	FNodeClassMetadata FLfoNode::CreateNodeClassMetadata()
+	{
+		return FLfoOperator::GetNodeInfo();
+	}
+
 	METASOUND_REGISTER_NODE(FLfoNode);
 
 #pragma endregion LFO
diff --git a/Engine/Plugins/Runtime/Metasound/Source/MetasoundStandardNodes/Private/MetasoundPerlinNoiseNode.cpp b/Engine/Plugins/Runtime/Metasound/Source/MetasoundStandardNodes/Private/MetasoundPerlinNoiseNode.cpp
--- a/Engine/Plugins/Runtime/Metasound/Source/MetasoundStandardNodes/Private/MetasoundPerlinNoiseNode.cpp
+++ b/Engine/Plugins/Runtime/Metasound/Source/MetasoundStandardNodes/Private/MetasoundPerlinNoiseNode.cpp
@@ -172,22 +172,6 @@ namespace Metasound
 			InOutVertexData.BindReadVertex(METASOUND_GET_PARAM_NAME(OutputNormalizedPin), OutputNormalized);
 		}
 
-		virtual FDataReferenceCollection GetInputs() const override
-		{
-			// This should never be called. Bind(...) is called instead. This method
-			// exists as a stop-gap until the API can be deprecated and removed.
-			checkNoEntry();
-			return {};
-		}
-
-		virtual FDataReferenceCollection GetOutputs() const override
-		{
-			// This should never be called. Bind(...) is called instead. This method
-			// exists as a stop-gap until the API can be deprecated and removed.
-			checkNoEntry();
-			return {};
-		}
-
 		void AccumulateTime()
 		{
 			NumFramesElapsed += NumFramesPerBlock;
diff --git a/Engine/Plugins/Runtime/Metasound/Source/MetasoundStandardNodes/Private/MetasoundPlateReverbNode.cpp b/Engine/Plugins/Runtime/Metasound/Source/MetasoundStandardNodes/Private/MetasoundPlateReverbNode.cpp
--- a/Engine/Plugins/Runtime/Metasound/Source/MetasoundStandardNodes/Private/MetasoundPlateReverbNode.cpp
+++ b/Engine/Plugins/Runtime/Metasound/Source/MetasoundStandardNodes/Private/MetasoundPlateReverbNode.cpp
@@ -293,14 +293,7 @@ namespace Metasound
 		float LastWetLevel{ -1 };
 	};
 
-	class FPlateReverbNode final : public FNodeFacade
-	{
-	public:
-		explicit FPlateReverbNode(const FNodeInitData& InitData)
-			: FNodeFacade(InitData.InstanceName, InitData.InstanceID, TFacadeOperatorClass<FPlateReverbOperator>())
-		{
-		}
-	};
+	using FPlateReverbNode = TNodeFacade<FPlateReverbOperator>;
 
 	METASOUND_REGISTER_NODE(FPlateReverbNode)
 }
diff --git a/Engine/Plugins/Runtime/Metasound/Source/MetasoundStandardNodes/Private/MetasoundRingModNode.cpp b/Engine/Plugins/Runtime/Metasound/Source/MetasoundStandardNodes/Private/MetasoundRingModNode.cpp
--- a/Engine/Plugins/Runtime/Metasound/Source/MetasoundStandardNodes/Private/MetasoundRingModNode.cpp
+++ b/Engine/Plugins/Runtime/Metasound/Source/MetasoundStandardNodes/Private/MetasoundRingModNode.cpp
@@ -100,22 +100,6 @@ namespace Metasound
 			InOutVertexData.BindReadVertex(METASOUND_GET_PARAM_NAME(OutputAudio), AudioOutput);
 		}
 
-		virtual FDataReferenceCollection GetInputs() const override
-		{
-			// This should never be called. Bind(...) is called instead. This method
-			// exists as a stop-gap until the API can be deprecated and removed.
-			checkNoEntry();
-			return {};
-		}
-
-		virtual FDataReferenceCollection GetOutputs() const override
-		{
-			// This should never be called. Bind(...) is called instead. This method
-			// exists as a stop-gap until the API can be deprecated and removed.
-			checkNoEntry();
-			return {};
-		}
-		
 		static TUniquePtr<IOperator> CreateOperator(const FBuildOperatorParams& InParams, FBuildResults& OutResults)
 		{
 			using namespace RingModVertexNames;
@@ -151,14 +135,7 @@ namespace Metasound
 	};
 
 	// Node Class
-	class FRingModNode : public FNodeFacade
-	{
-	public:
-		FRingModNode(const FNodeInitData& InitData)
-			: FNodeFacade(InitData.InstanceName, InitData.InstanceID, TFacadeOperatorClass<FRingModOperator>())
-		{
-		}
-	};
+	using FRingModNode = TNodeFacade<FRingModOperator>;
 
 	// Register node
 	METASOUND_REGISTER_NODE(FRingModNode)
diff --git a/Engine/Plugins/Runtime/Metasound/Source/MetasoundStandardNodes/Private/MetasoundSampleAndHoldNode.cpp b/Engine/Plugins/Runtime/Metasound/Source/MetasoundStandardNodes/Private/MetasoundSampleAndHoldNode.cpp
--- a/Engine/Plugins/Runtime/Metasound/Source/MetasoundStandardNodes/Private/MetasoundSampleAndHoldNode.cpp
+++ b/Engine/Plugins/Runtime/Metasound/Source/MetasoundStandardNodes/Private/MetasoundSampleAndHoldNode.cpp
@@ -111,22 +111,6 @@ namespace Metasound
 			InOutVertexData.BindReadVertex(METASOUND_GET_PARAM_NAME(OutputAudio), AudioOutput);
 		}
 
-		virtual FDataReferenceCollection GetInputs() const override
-		{
-			// This should never be called. Bind(...) is called instead. This method
-			// exists as a stop-gap until the API can be deprecated and removed.
-			checkNoEntry();
-			return {};
-		}
-
-		virtual FDataReferenceCollection GetOutputs() const override
-		{
-			// This should never be called. Bind(...) is called instead. This method
-			// exists as a stop-gap until the API can be deprecated and removed.
-			checkNoEntry();
-			return {};
-		}
-		
 		void Execute()
 		{
 			TriggerSampleAndHold->ExecuteBlock(
@@ -176,15 +160,7 @@ namespace Metasound
 		float HoldValue = 0.0f;
 	};
 
-	class FSampleAndHoldNode : public FNodeFacade
-	{
-	public:
-		FSampleAndHoldNode(const FNodeInitData& InitData)
-			: FNodeFacade(InitData.InstanceName, InitData.InstanceID, TFacadeOperatorClass<FSampleAndHoldOperator>())
-		{
-		}
-	};
-
+	using FSampleAndHoldNode = TNodeFacade<FSampleAndHoldOperator>;
 
 	METASOUND_REGISTER_NODE(FSampleAndHoldNode)
 }
diff --git a/Engine/Plugins/Runtime/Metasound/Source/MetasoundStandardNodes/Private/MetasoundSemitoneToFreqMultiplierNode.cpp b/Engine/Plugins/Runtime/Metasound/Source/MetasoundStandardNodes/Private/MetasoundSemitoneToFreqMultiplierNode.cpp
--- a/Engine/Plugins/Runtime/Metasound/Source/MetasoundStandardNodes/Private/MetasoundSemitoneToFreqMultiplierNode.cpp
+++ b/Engine/Plugins/Runtime/Metasound/Source/MetasoundStandardNodes/Private/MetasoundSemitoneToFreqMultiplierNode.cpp
@@ -36,8 +36,6 @@ namespace Metasound
 
 		virtual void BindInputs(FInputVertexInterfaceData& InOutVertexData) override;
 		virtual void BindOutputs(FOutputVertexInterfaceData& InOutVertexData) override;
-		virtual FDataReferenceCollection GetInputs() const override;
-		virtual FDataReferenceCollection GetOutputs() const override;
 		void Execute();
 		void Reset(const IOperator::FResetParams& InParams);
 
@@ -74,22 +72,6 @@ namespace Metasound
 		InOutVertexData.BindReadVertex(METASOUND_GET_PARAM_NAME(OutputFrequencyMultiplier), FrequencyMultiplierOutput);
 	}
 
-	FDataReferenceCollection FSemitoneToFrequencyMultiplierOperator::GetInputs() const
-	{
-		// This should never be called. Bind(...) is called instead. This method
-		// exists as a stop-gap until the API can be deprecated and removed.
-		checkNoEntry();
-		return {};
-	}
-
-	FDataReferenceCollection FSemitoneToFrequencyMultiplierOperator::GetOutputs() const
-	{
-		// This should never be called. Bind(...) is called instead. This method
-		// exists as a stop-gap until the API can be deprecated and removed.
-		checkNoEntry();
-		return {};
-	}
-
 	void FSemitoneToFrequencyMultiplierOperator::Execute()
 	{
 		const float CurrSemitone = *SemitoneInput;
@@ -163,17 +145,7 @@ namespace Metasound
 		return MakeUnique<FSemitoneToFrequencyMultiplierOperator>(InParams.OperatorSettings, InSemitone);
 	}
 
-	class METASOUNDSTANDARDNODES_API FSemitoneToFrequencyMultiplierNode : public FNodeFacade
-	{
-	public:
-		/**
-		 * Constructor used by the Metasound Frontend.
-		 */
-		FSemitoneToFrequencyMultiplierNode(const FNodeInitData& InitData)
-			: FNodeFacade(InitData.InstanceName, InitData.InstanceID, TFacadeOperatorClass <FSemitoneToFrequencyMultiplierOperator>())
-		{
-		}
-	};
+	using FSemitoneToFrequencyMultiplierNode = TNodeFacade<FSemitoneToFrequencyMultiplierOperator>;
 
 	METASOUND_REGISTER_NODE(FSemitoneToFrequencyMultiplierNode)
 }
diff --git a/Engine/Plugins/Runtime/Metasound/Source/MetasoundStandardNodes/Private/MetasoundStereoDelayNode.cpp b/Engine/Plugins/Runtime/Metasound/Source/MetasoundStandardNodes/Private/MetasoundStereoDelayNode.cpp
--- a/Engine/Plugins/Runtime/Metasound/Source/MetasoundStandardNodes/Private/MetasoundStereoDelayNode.cpp
+++ b/Engine/Plugins/Runtime/Metasound/Source/MetasoundStandardNodes/Private/MetasoundStereoDelayNode.cpp
@@ -79,8 +79,6 @@ namespace Metasound
 
 		virtual void BindInputs(FInputVertexInterfaceData& InOutVertexData) override;
 		virtual void BindOutputs(FOutputVertexInterfaceData& InOutVertexData) override;
-		virtual FDataReferenceCollection GetInputs() const override;
-		virtual FDataReferenceCollection GetOutputs() const override;
 		void Execute();
 		void Reset(const IOperator::FResetParams& InParams);
 
@@ -191,22 +189,6 @@ namespace Metasound
 		InOutVertexData.BindReadVertex(METASOUND_GET_PARAM_NAME(OutAudioRight), RightAudioOutput);
 	}
 
-	FDataReferenceCollection FStereoDelayOperator::GetInputs() const
-	{
-		// This should never be called. Bind(...) is called instead. This method
-		// exists as a stop-gap until the API can be deprecated and removed.
-		checkNoEntry();
-		return {};
-	}
-
-	FDataReferenceCollection FStereoDelayOperator::GetOutputs() const
-	{
-		// This should never be called. Bind(...) is called instead. This method
-		// exists as a stop-gap until the API can be deprecated and removed.
-		checkNoEntry();
-		return {};
-	}
-
 	float FStereoDelayOperator::GetInputDelayTimeMsecClamped() const
 	{
 		// Clamp the delay time to the max delay allowed
@@ -419,18 +401,7 @@ namespace Metasound
 		return MakeUnique<FStereoDelayOperator>(InParams.OperatorSettings, LeftAudioIn, RightAudioIn, StereoDelayMode, DelayTime, DelayRatio, DryLevel, WetLevel, Feedback, MaxDelayTime.GetSeconds(), TriggerReset);
 	}
 
-
-	class FStereoDelayNode : public FNodeFacade
-	{
-	public:
-		/**
-		 * Constructor used by the Metasound Frontend.
-		 */
-		FStereoDelayNode(const FNodeInitData& InitData)
-			: FNodeFacade(InitData.InstanceName, InitData.InstanceID, TFacadeOperatorClass<FStereoDelayOperator>())
-		{
-		}
-	};
+	using FStereoDelayNode = TNodeFacade<FStereoDelayOperator>;
 
 	METASOUND_REGISTER_NODE(FStereoDelayNode)
 }
diff --git a/Engine/Plugins/Runtime/Metasound/Source/MetasoundStandardNodes/Private/MetasoundStereoPannerNode.cpp b/Engine/Plugins/Runtime/Metasound/Source/MetasoundStandardNodes/Private/MetasoundStereoPannerNode.cpp
--- a/Engine/Plugins/Runtime/Metasound/Source/MetasoundStandardNodes/Private/MetasoundStereoPannerNode.cpp
+++ b/Engine/Plugins/Runtime/Metasound/Source/MetasoundStandardNodes/Private/MetasoundStereoPannerNode.cpp
@@ -60,8 +60,6 @@ namespace Metasound
 
 		virtual void BindInputs(FInputVertexInterfaceData& InOutVertexData) override;
 		virtual void BindOutputs(FOutputVertexInterfaceData& InOutVertexData) override;
-		virtual FDataReferenceCollection GetInputs() const override;
-		virtual FDataReferenceCollection GetOutputs() const override;
 		void Execute();
 		void Reset(const IOperator::FResetParams& InParams);
 
@@ -120,22 +118,6 @@ namespace Metasound
 		InOutVertexData.BindReadVertex(METASOUND_GET_PARAM_NAME(OutputAudioRight), AudioRightOutput);
 	}
 
-	FDataReferenceCollection FStereoPannerOperator::GetInputs() const
-	{
-		// This should never be called. Bind(...) is called instead. This method
-		// exists as a stop-gap until the API can be deprecated and removed.
-		checkNoEntry();
-		return {};
-	}
-
-	FDataReferenceCollection FStereoPannerOperator::GetOutputs() const
-	{
-		// This should never be called. Bind(...) is called instead. This method
-		// exists as a stop-gap until the API can be deprecated and removed.
-		checkNoEntry();
-		return {};
-	}
-
 	void FStereoPannerOperator::ComputePanGains(float InPanningAmmount, float& OutLeftGain, float& OutRightGain) const
 	{
 		// Convert [-1.0, 1.0] to [0.0, 1.0]
@@ -257,18 +239,7 @@ namespace Metasound
 		return MakeUnique<FStereoPannerOperator>(InParams.OperatorSettings, AudioIn, PanningAmount, PanningLaw);
 	}
 
-	class FStereoPannerNode : public FNodeFacade
-	{
-	public:
-		/**
-		 * Constructor used by the Metasound Frontend.
-		 */
-		FStereoPannerNode(const FNodeInitData& InitData)
-			: FNodeFacade(InitData.InstanceName, InitData.InstanceID, TFacadeOperatorClass<FStereoPannerOperator>())
-		{
-		}
-	};
-
+	using FStereoPannerNode = TNodeFacade<FStereoPannerOperator>;
 
 	METASOUND_REGISTER_NODE(FStereoPannerNode)
 }
diff --git a/Engine/Plugins/Runtime/Metasound/Source/MetasoundStandardNodes/Private/MetasoundSuperOscillatorNode.cpp b/Engine/Plugins/Runtime/Metasound/Source/MetasoundStandardNodes/Private/MetasoundSuperOscillatorNode.cpp
--- a/Engine/Plugins/Runtime/Metasound/Source/MetasoundStandardNodes/Private/MetasoundSuperOscillatorNode.cpp
+++ b/Engine/Plugins/Runtime/Metasound/Source/MetasoundStandardNodes/Private/MetasoundSuperOscillatorNode.cpp
@@ -235,22 +235,6 @@ namespace Metasound
 			InOutVertexData.SetValue(METASOUND_GET_PARAM_NAME(WaveTypePin), OscType);
 		}
 		
-		virtual FDataReferenceCollection GetInputs() const override
-		{
-			checkNoEntry();
-
-			FDataReferenceCollection InputDataReferences;
-			return InputDataReferences;
-		}
-
-		virtual FDataReferenceCollection GetOutputs() const override
-		{
-			checkNoEntry();
-
-			FDataReferenceCollection OutputDataReferences;
-			return OutputDataReferences;
-		}
-
 		virtual void FillOutputs(const float Gain = 1.f, const float LinearPan = 0.f) {}
 		virtual void ZeroOutputs() {}
 		virtual float GetStereoWidth() const { return 0.f; }
@@ -693,21 +677,9 @@ namespace Metasound
 		Audio::FDynamicsProcessor LimiterRight;
 	};
 	
-	class FSuperOscillatorNodeMono : public FNodeFacade
-	{
-	public:
-		FSuperOscillatorNodeMono(const FNodeInitData& InitData)
-		: FNodeFacade(InitData.InstanceName, InitData.InstanceID, TFacadeOperatorClass<FSuperOscillatorOperatorMono>())
-		{}
-	};
+	using FSuperOscillatorNodeMono = TNodeFacade<FSuperOscillatorOperatorMono>;
 
-	class FSuperOscillatorNodeStereo : public FNodeFacade
-	{
-	public:
-		FSuperOscillatorNodeStereo(const FNodeInitData& InitData)
-		: FNodeFacade(InitData.InstanceName, InitData.InstanceID, TFacadeOperatorClass<FSuperOscillatorOperatorStereo>())
-		{}
-	};
+	using FSuperOscillatorNodeStereo = TNodeFacade<FSuperOscillatorOperatorStereo>;
 
 	METASOUND_REGISTER_NODE(FSuperOscillatorNodeMono);
 	METASOUND_REGISTER_NODE(FSuperOscillatorNodeStereo);
diff --git a/Engine/Plugins/Runtime/Metasound/Source/MetasoundStandardNodes/Private/MetasoundSwitchAndRampNode.cpp b/Engine/Plugins/Runtime/Metasound/Source/MetasoundStandardNodes/Private/MetasoundSwitchAndRampNode.cpp
--- a/Engine/Plugins/Runtime/Metasound/Source/MetasoundStandardNodes/Private/MetasoundSwitchAndRampNode.cpp
+++ b/Engine/Plugins/Runtime/Metasound/Source/MetasoundStandardNodes/Private/MetasoundSwitchAndRampNode.cpp
@@ -117,22 +117,6 @@ namespace Metasound
 			InOutVertexData.BindReadVertex(METASOUND_GET_PARAM_NAME(OutputAudio), AudioOut);
 		}
 
-		virtual FDataReferenceCollection GetInputs() const override
-		{
-			// This should never be called. Bind(...) is called instead. This method
-			// exists as a stop-gap until the API can be deprecated and removed.
-			checkNoEntry();
-			return {};
-		}
-
-		virtual FDataReferenceCollection GetOutputs() const override
-		{
-			// This should never be called. Bind(...) is called instead. This method
-			// exists as a stop-gap until the API can be deprecated and removed.
-			checkNoEntry();
-			return {};
-		}
-
 		void Execute()
 		{
 			// Lambda is used for both pre and post trigger frames
@@ -217,17 +201,7 @@ namespace Metasound
 		int32 RampSampleIndex = 0;
 	};
 
-	class METASOUNDSTANDARDNODES_API FSwitchAndRampNode : public FNodeFacade
-	{
-	public:
-		/**
-		 * Constructor used by the Metasound Frontend.
-		 */
-		FSwitchAndRampNode(const FNodeInitData& InitData)
-			: FNodeFacade(InitData.InstanceName, InitData.InstanceID, TFacadeOperatorClass<FSwitchAndRampOperator>())
-		{
-		}
-	};
+	using FSwitchAndRampNode = TNodeFacade<FSwitchAndRampOperator>;
 
 	METASOUND_REGISTER_NODE(FSwitchAndRampNode)
 }
diff --git a/Engine/Plugins/Runtime/Metasound/Source/MetasoundStandardNodes/Private/MetasoundTriggerAnyNode.cpp b/Engine/Plugins/Runtime/Metasound/Source/MetasoundStandardNodes/Private/MetasoundTriggerAnyNode.cpp
--- a/Engine/Plugins/Runtime/Metasound/Source/MetasoundStandardNodes/Private/MetasoundTriggerAnyNode.cpp
+++ b/Engine/Plugins/Runtime/Metasound/Source/MetasoundStandardNodes/Private/MetasoundTriggerAnyNode.cpp
@@ -201,18 +201,7 @@ namespace Metasound
 	*  Will output a trigger whenever any of its input triggers are set.
 	*/
 	template<uint32 NumInputs>
-	class METASOUNDSTANDARDNODES_API TTriggerAnyNode : public FNodeFacade
-	{
-	public:
-		/**
-		 * Constructor used by the Metasound Frontend.
-		 */
-		TTriggerAnyNode(const FNodeInitData& InInitData)
-			: FNodeFacade(InInitData.InstanceName, InInitData.InstanceID, TFacadeOperatorClass<TTriggerAnyOperator<NumInputs>>())
-		{}
-
-		virtual ~TTriggerAnyNode() = default;
-	};
+	using TTriggerAnyNode = TNodeFacade<TTriggerAnyOperator<NumInputs>>;
 
 	REGISTER_TRIGGER_ANY_NODE(2)
 	REGISTER_TRIGGER_ANY_NODE(3)
diff --git a/Engine/Plugins/Runtime/Metasound/Source/MetasoundStandardNodes/Private/MetasoundTriggerCoinNode.cpp b/Engine/Plugins/Runtime/Metasound/Source/MetasoundStandardNodes/Private/MetasoundTriggerCoinNode.cpp
--- a/Engine/Plugins/Runtime/Metasound/Source/MetasoundStandardNodes/Private/MetasoundTriggerCoinNode.cpp
+++ b/Engine/Plugins/Runtime/Metasound/Source/MetasoundStandardNodes/Private/MetasoundTriggerCoinNode.cpp
@@ -45,8 +45,6 @@ namespace Metasound
 
 		virtual void BindInputs(FInputVertexInterfaceData& InOutVertexData) override;
 		virtual void BindOutputs(FOutputVertexInterfaceData& InOutVertexData) override;
-		virtual FDataReferenceCollection GetInputs() const override;
-		virtual FDataReferenceCollection GetOutputs() const override;
 
 		void Execute();
 		void Reset(const IOperator::FResetParams& InParams);
@@ -111,22 +109,6 @@ namespace Metasound
 		InOutVertexData.BindReadVertex(METASOUND_GET_PARAM_NAME(OutputFalseTrigger), TriggerFalseOutput);
 	}
 
-	FDataReferenceCollection FTriggerCoinOperator::GetInputs() const
-	{
-		// This should never be called. Bind(...) is called instead. This method
-		// exists as a stop-gap until the API can be deprecated and removed.
-		checkNoEntry();
-		return {};
-	}
-
-	FDataReferenceCollection FTriggerCoinOperator::GetOutputs() const
-	{
-		// This should never be called. Bind(...) is called instead. This method
-		// exists as a stop-gap until the API can be deprecated and removed.
-		checkNoEntry();
-		return {};
-	}
-
 	void FTriggerCoinOperator::Execute()
 	{
 		TriggerTrueOutput->AdvanceBlock();
@@ -261,14 +243,7 @@ namespace Metasound
 	}
 
 	// Node Class
-	class FTriggerCoinNode : public FNodeFacade
-	{
-	public:
-		FTriggerCoinNode(const FNodeInitData& InitData)
-			: FNodeFacade(InitData.InstanceName, InitData.InstanceID, TFacadeOperatorClass<FTriggerCoinOperator>())
-		{
-		}
-	};
+	using FTriggerCoinNode = TNodeFacade<FTriggerCoinOperator>;
 
 	METASOUND_REGISTER_NODE(FTriggerCoinNode)
 }
diff --git a/Engine/Plugins/Runtime/Metasound/Source/MetasoundStandardNodes/Private/MetasoundTriggerControlNode.cpp b/Engine/Plugins/Runtime/Metasound/Source/MetasoundStandardNodes/Private/MetasoundTriggerControlNode.cpp
--- a/Engine/Plugins/Runtime/Metasound/Source/MetasoundStandardNodes/Private/MetasoundTriggerControlNode.cpp
+++ b/Engine/Plugins/Runtime/Metasound/Source/MetasoundStandardNodes/Private/MetasoundTriggerControlNode.cpp
@@ -49,10 +49,6 @@ namespace Metasound
 
 		virtual void BindOutputs(FOutputVertexInterfaceData& InOutVertexData) override;
 
-		virtual FDataReferenceCollection GetInputs() const override;
-
-		virtual FDataReferenceCollection GetOutputs() const override;
-
 		void Execute();
 		void Reset(const IOperator::FResetParams& InParams);
 
@@ -107,22 +103,6 @@ namespace Metasound
 		InOutVertexData.BindReadVertex(METASOUND_GET_PARAM_NAME(OutputExit), TriggerExitOutput);
 	}
 
-	FDataReferenceCollection FTriggerControlOperator::GetInputs() const
-	{
-		// This should never be called. Bind(...) is called instead. This method
-		// exists as a stop-gap until the API can be deprecated and removed.
-		checkNoEntry();
-		return {};
-	}
-
-	FDataReferenceCollection FTriggerControlOperator::GetOutputs() const
-	{
-		// This should never be called. Bind(...) is called instead. This method
-		// exists as a stop-gap until the API can be deprecated and removed.
-		checkNoEntry();
-		return {};
-	}
-
 	void FTriggerControlOperator::Execute()
 	{
 		TriggerExitOutput->AdvanceBlock();
@@ -239,14 +219,7 @@ namespace Metasound
 	}
 
 	// Node Class
-	class FTriggerControlNode : public FNodeFacade
-	{
-	public:
-		FTriggerControlNode(const FNodeInitData& InitData)
-			: FNodeFacade(InitData.InstanceName, InitData.InstanceID, TFacadeOperatorClass<FTriggerControlOperator>())
-		{
-		}
-	};
+	using FTriggerControlNode = TNodeFacade<FTriggerControlOperator>;
 
 	METASOUND_REGISTER_NODE(FTriggerControlNode)
 }
diff --git a/Engine/Plugins/Runtime/Metasound/Source/MetasoundStandardNodes/Private/MetasoundTriggerCounterNode.cpp b/Engine/Plugins/Runtime/Metasound/Source/MetasoundStandardNodes/Private/MetasoundTriggerCounterNode.cpp
--- a/Engine/Plugins/Runtime/Metasound/Source/MetasoundStandardNodes/Private/MetasoundTriggerCounterNode.cpp
+++ b/Engine/Plugins/Runtime/Metasound/Source/MetasoundStandardNodes/Private/MetasoundTriggerCounterNode.cpp
@@ -48,8 +48,6 @@ namespace Metasound
 
 			virtual void BindInputs(FInputVertexInterfaceData& InOutVertexData) override;
 			virtual void BindOutputs(FOutputVertexInterfaceData& InOutVertexData) override;
-			virtual FDataReferenceCollection GetInputs() const override;
-			virtual FDataReferenceCollection GetOutputs() const override;
 
 			void Execute();
 			void Reset(const IOperator::FResetParams& InParams);
@@ -115,22 +113,6 @@ namespace Metasound
 		InOutVertexData.BindReadVertex(METASOUND_GET_PARAM_NAME(OutputValue), OutValue);
 	}
 
-	FDataReferenceCollection FTriggerCounterOperator::GetInputs() const
-	{
-		// This should never be called. Bind(...) is called instead. This method
-		// exists as a stop-gap until the API can be deprecated and removed.
-		checkNoEntry();
-		return {};
-	}
-
-	FDataReferenceCollection FTriggerCounterOperator::GetOutputs() const
-	{
-		// This should never be called. Bind(...) is called instead. This method
-		// exists as a stop-gap until the API can be deprecated and removed.
-		checkNoEntry();
-		return {};
-	}
-
 	void FTriggerCounterOperator::Execute()
 	{	
 		TriggerOnReset->AdvanceBlock();
@@ -258,16 +240,7 @@ namespace Metasound
 		return Info;
 	}
 
-	class FTriggerCounterNode : public FNodeFacade
-	{
-	public:
-		FTriggerCounterNode(const FNodeInitData& InInitData)
-			: FNodeFacade(InInitData.InstanceName, InInitData.InstanceID, TFacadeOperatorClass<FTriggerCounterOperator>())
-		{
-		}
-
-		virtual ~FTriggerCounterNode() = default;
-	};
+	using FTriggerCounterNode = TNodeFacade<FTriggerCounterOperator>;
 
 	METASOUND_REGISTER_NODE(FTriggerCounterNode)
 }
diff --git a/Engine/Plugins/Runtime/Metasound/Source/MetasoundStandardNodes/Private/MetasoundTriggerDelayNode.cpp b/Engine/Plugins/Runtime/Metasound/Source/MetasoundStandardNodes/Private/MetasoundTriggerDelayNode.cpp
--- a/Engine/Plugins/Runtime/Metasound/Source/MetasoundStandardNodes/Private/MetasoundTriggerDelayNode.cpp
+++ b/Engine/Plugins/Runtime/Metasound/Source/MetasoundStandardNodes/Private/MetasoundTriggerDelayNode.cpp
@@ -38,8 +38,6 @@ namespace Metasound
 
 			virtual void BindInputs(FInputVertexInterfaceData& InOutVertexData) override;
 			virtual void BindOutputs(FOutputVertexInterfaceData& InOutVertexData) override;
-			virtual FDataReferenceCollection GetInputs() const override;
-			virtual FDataReferenceCollection GetOutputs() const override;
 
 			void Execute();
 			void Reset(const IOperator::FResetParams& InParams);
@@ -76,22 +74,6 @@ namespace Metasound
 		InOutVertexData.BindReadVertex(METASOUND_GET_PARAM_NAME(OutputOnTrigger), TriggerOut);
 	}
 
-	FDataReferenceCollection FTriggerDelayOperator::GetInputs() const
-	{
-		// This should never be called. Bind(...) is called instead. This method
-		// exists as a stop-gap until the API can be deprecated and removed.
-		checkNoEntry();
-		return {};
-	}
-
-	FDataReferenceCollection FTriggerDelayOperator::GetOutputs() const
-	{
-		// This should never be called. Bind(...) is called instead. This method
-		// exists as a stop-gap until the API can be deprecated and removed.
-		checkNoEntry();
-		return {};
-	}
-
 	void FTriggerDelayOperator::Execute()
 	{
 		// Advance internal counter to get rid of old triggers.
@@ -179,14 +161,7 @@ namespace Metasound
 		return Info;
 	}
 
-	class METASOUNDSTANDARDNODES_API FTriggerDelayNode : public FNodeFacade
-	{
-	public:
-		FTriggerDelayNode(const FNodeInitData& InInitData)
-			: FNodeFacade(InInitData.InstanceName, InInitData.InstanceID, TFacadeOperatorClass<FTriggerDelayOperator>())
-		{
-		}
-	};
+	using FTriggerDelayNode = TNodeFacade<FTriggerDelayOperator>;
 
 	METASOUND_REGISTER_NODE(FTriggerDelayNode)
 }
diff --git a/Engine/Plugins/Runtime/Metasound/Source/MetasoundStandardNodes/Private/MetasoundTriggerOnThresholdNode.cpp b/Engine/Plugins/Runtime/Metasound/Source/MetasoundStandardNodes/Private/MetasoundTriggerOnThresholdNode.cpp
--- a/Engine/Plugins/Runtime/Metasound/Source/MetasoundStandardNodes/Private/MetasoundTriggerOnThresholdNode.cpp
+++ b/Engine/Plugins/Runtime/Metasound/Source/MetasoundStandardNodes/Private/MetasoundTriggerOnThresholdNode.cpp
@@ -317,14 +317,7 @@ namespace Metasound
 	constexpr int32 TTriggerOnThresholdOperator<int32, int32>::DefaultThreshold = 1;
 	
 	template <typename ValueType, typename ThresholdType>
-	class TTriggerOnThresholdNode : public FNodeFacade
-	{
-	public:
-		TTriggerOnThresholdNode(const FNodeInitData& InitData)
-			: FNodeFacade(InitData.InstanceName, InitData.InstanceID, TFacadeOperatorClass<TTriggerOnThresholdOperator<ValueType, ThresholdType>>())
-		{
-		}
-	};
+	using TTriggerOnThresholdNode = TNodeFacade<TTriggerOnThresholdOperator<ValueType, ThresholdType>>;
 
 	using FTriggerOnThresholdAudioNode = TTriggerOnThresholdNode<FAudioBuffer, float>;
 	METASOUND_REGISTER_NODE(FTriggerOnThresholdAudioNode);
diff --git a/Engine/Plugins/Runtime/Metasound/Source/MetasoundStandardNodes/Private/MetasoundTriggerOnValueChangeNode.cpp b/Engine/Plugins/Runtime/Metasound/Source/MetasoundStandardNodes/Private/MetasoundTriggerOnValueChangeNode.cpp
--- a/Engine/Plugins/Runtime/Metasound/Source/MetasoundStandardNodes/Private/MetasoundTriggerOnValueChangeNode.cpp
+++ b/Engine/Plugins/Runtime/Metasound/Source/MetasoundStandardNodes/Private/MetasoundTriggerOnValueChangeNode.cpp
@@ -56,10 +56,6 @@ namespace Metasound
 
 		virtual void BindOutputs(FOutputVertexInterfaceData& InOutVertexData) override;
 
-		virtual FDataReferenceCollection GetInputs() const override;
-
-		virtual FDataReferenceCollection GetOutputs() const override;
-
 		void Execute();
 		void Reset(const IOperator::FResetParams& InParams);
 
@@ -96,24 +92,6 @@ namespace Metasound
 		InOutVertexData.BindReadVertex(METASOUND_GET_PARAM_NAME(OutputOnChange), TriggerOnChangeOutput);
 	}
 
-	template <typename ValueType>
-	FDataReferenceCollection TTriggerOnValueChangeOperator<ValueType>::GetInputs() const
-	{
-		// This should never be called. Bind(...) is called instead. This method
-		// exists as a stop-gap until the API can be deprecated and removed.
-		checkNoEntry();
-		return {};
-	}
-
-	template <typename ValueType>
-	FDataReferenceCollection TTriggerOnValueChangeOperator<ValueType>::GetOutputs() const
-	{
-		// This should never be called. Bind(...) is called instead. This method
-		// exists as a stop-gap until the API can be deprecated and removed.
-		checkNoEntry();
-		return {};
-	}
-
 	template <typename ValueType>
 	void TTriggerOnValueChangeOperator<ValueType>::Execute()
 	{
@@ -196,14 +174,7 @@ namespace Metasound
 
 	// Node Class
 	template <typename ValueType>
-	class TTriggerOnValueChangeNode : public FNodeFacade
-	{
-	public:
-		TTriggerOnValueChangeNode(const FNodeInitData& InitData)
-			: FNodeFacade(InitData.InstanceName, InitData.InstanceID, TFacadeOperatorClass<TTriggerOnValueChangeOperator<ValueType>>())
-		{
-		}
-	};
+	using TTriggerOnValueChangeNode = TNodeFacade<TTriggerOnValueChangeOperator<ValueType>>;
 
 	using FTriggerOnInt32Change = TTriggerOnValueChangeNode<int32>;
 	METASOUND_REGISTER_NODE(FTriggerOnInt32Change)
diff --git a/Engine/Plugins/Runtime/Metasound/Source/MetasoundStandardNodes/Private/MetasoundTriggerOnceNode.cpp b/Engine/Plugins/Runtime/Metasound/Source/MetasoundStandardNodes/Private/MetasoundTriggerOnceNode.cpp
--- a/Engine/Plugins/Runtime/Metasound/Source/MetasoundStandardNodes/Private/MetasoundTriggerOnceNode.cpp
+++ b/Engine/Plugins/Runtime/Metasound/Source/MetasoundStandardNodes/Private/MetasoundTriggerOnceNode.cpp
@@ -43,9 +43,6 @@ namespace Metasound
 
 		virtual void BindInputs(FInputVertexInterfaceData& InOutVertexData) override;
 		virtual void BindOutputs(FOutputVertexInterfaceData& InOutVertexData) override;
-		virtual FDataReferenceCollection GetInputs() const override;
-		virtual FDataReferenceCollection GetOutputs() const override;
-
 		void Execute();
 		void Reset(const IOperator::FResetParams& InParams);
 
@@ -91,22 +88,6 @@ namespace Metasound
 		InOutVertexData.BindReadVertex(METASOUND_GET_PARAM_NAME(OutputExit), TriggerExitOutput);
 	}
 
-	FDataReferenceCollection FTriggerOnceOperator::GetInputs() const
-	{
-		// This should never be called. Bind(...) is called instead. This method
-		// exists as a stop-gap until the API can be deprecated and removed.
-		checkNoEntry();
-		return {};
-	}
-
-	FDataReferenceCollection FTriggerOnceOperator::GetOutputs() const
-	{
-		// This should never be called. Bind(...) is called instead. This method
-		// exists as a stop-gap until the API can be deprecated and removed.
-		checkNoEntry();
-		return {};
-	}
-
 	void FTriggerOnceOperator::Execute()
 	{
 		TriggerExitOutput->AdvanceBlock();
@@ -199,14 +180,7 @@ namespace Metasound
 	}
 
 	// Node Class
-	class FTriggerOnceNode : public FNodeFacade
-	{
-	public:
-		FTriggerOnceNode(const FNodeInitData& InitData)
-			: FNodeFacade(InitData.InstanceName, InitData.InstanceID, TFacadeOperatorClass<FTriggerOnceOperator>())
-		{
-		}
-	};
+	using FTriggerOnceNode = TNodeFacade<FTriggerOnceOperator>;
 
 	METASOUND_REGISTER_NODE(FTriggerOnceNode)
 }
diff --git a/Engine/Plugins/Runtime/Metasound/Source/MetasoundStandardNodes/Private/MetasoundTriggerPipeNode.cpp b/Engine/Plugins/Runtime/Metasound/Source/MetasoundStandardNodes/Private/MetasoundTriggerPipeNode.cpp
--- a/Engine/Plugins/Runtime/Metasound/Source/MetasoundStandardNodes/Private/MetasoundTriggerPipeNode.cpp
+++ b/Engine/Plugins/Runtime/Metasound/Source/MetasoundStandardNodes/Private/MetasoundTriggerPipeNode.cpp
@@ -40,8 +40,6 @@ namespace Metasound
 
 		virtual void BindInputs(FInputVertexInterfaceData& InOutVertexData) override;
 		virtual void BindOutputs(FOutputVertexInterfaceData& InOutVertexData) override;
-		virtual FDataReferenceCollection GetInputs() const override;
-		virtual FDataReferenceCollection GetOutputs() const override;
 		void Execute();
 		void Reset(const IOperator::FResetParams& InParams);
 
@@ -82,22 +80,6 @@ namespace Metasound
 		InOutVertexData.BindReadVertex(METASOUND_GET_PARAM_NAME(OutputOutTrigger), TriggerOut);
 	}
 
-	FDataReferenceCollection FTriggerPipeOperator::GetInputs() const
-	{
-		// This should never be called. Bind(...) is called instead. This method
-		// exists as a stop-gap until the API can be deprecated and removed.
-		checkNoEntry();
-		return {};
-	}
-
-	FDataReferenceCollection FTriggerPipeOperator::GetOutputs() const
-	{
-		// This should never be called. Bind(...) is called instead. This method
-		// exists as a stop-gap until the API can be deprecated and removed.
-		checkNoEntry();
-		return {};
-	}
-
 	void FTriggerPipeOperator::Execute()
 	{
 		// Advance internal counter to get rid of old triggers.
@@ -205,16 +187,7 @@ namespace Metasound
 		return Info;
 	}
 
-	class METASOUNDSTANDARDNODES_API FTriggerPipeNode : public FNodeFacade
-	{
-	public:
-		FTriggerPipeNode(const FNodeInitData& InInitData)
-			: FNodeFacade(InInitData.InstanceName, InInitData.InstanceID, TFacadeOperatorClass<FTriggerPipeOperator>())
-		{
-		}
-
-		virtual ~FTriggerPipeNode() = default;
-	};
+	using FTriggerPipeNode = TNodeFacade<FTriggerPipeOperator>;
 
 	METASOUND_REGISTER_NODE(FTriggerPipeNode)
 }
diff --git a/Engine/Plugins/Runtime/Metasound/Source/MetasoundStandardNodes/Private/MetasoundTriggerRepeatNode.cpp b/Engine/Plugins/Runtime/Metasound/Source/MetasoundStandardNodes/Private/MetasoundTriggerRepeatNode.cpp
--- a/Engine/Plugins/Runtime/Metasound/Source/MetasoundStandardNodes/Private/MetasoundTriggerRepeatNode.cpp
+++ b/Engine/Plugins/Runtime/Metasound/Source/MetasoundStandardNodes/Private/MetasoundTriggerRepeatNode.cpp
@@ -47,10 +47,6 @@ namespace Metasound
 
 			virtual void BindOutputs(FOutputVertexInterfaceData& InOutVertexData) override;
 
-			virtual FDataReferenceCollection GetInputs() const override;
-
-			virtual FDataReferenceCollection GetOutputs() const override;
-
 			void Execute();
 			void Reset(const IOperator::FResetParams& InParams);
 
@@ -102,22 +98,6 @@ namespace Metasound
 		InOutVertexData.BindReadVertex(METASOUND_GET_PARAM_NAME(RepeatOutputOnTrigger), TriggerOut);
 	}
 
-	FDataReferenceCollection FTriggerRepeatOperator::GetInputs() const
-	{
-		// This should never be called. Bind(...) is called instead. This method
-		// exists as a stop-gap until the API can be deprecated and removed.
-		checkNoEntry();
-		return {};
-	}
-
-	FDataReferenceCollection FTriggerRepeatOperator::GetOutputs() const
-	{
-		// This should never be called. Bind(...) is called instead. This method
-		// exists as a stop-gap until the API can be deprecated and removed.
-		checkNoEntry();
-		return {};
-	}
-
 	void FTriggerRepeatOperator::Execute()
 	{
 		TriggerEnable->ExecuteBlock([](int32, int32) { },
@@ -240,6 +220,16 @@ namespace Metasound
 	:	FNodeFacade(InInitData.InstanceName, InInitData.InstanceID, TFacadeOperatorClass<FTriggerRepeatOperator>())
 	{
 	}
+
+	FTriggerRepeatNode::FTriggerRepeatNode(FNodeData InNodeData, TSharedRef<const FNodeClassMetadata> InMetadata)
+	:	FNodeFacade(MoveTemp(InNodeData), MoveTemp(InMetadata), TFacadeOperatorClass<FTriggerRepeatOperator>())
+	{
+	}
+
+	FNodeClassMetadata FTriggerRepeatNode::CreateNodeClassMetadata()
+	{
+		return FTriggerRepeatOperator::GetNodeInfo();
+	}
 }
 
 #undef LOCTEXT_NAMESPACE //MetasoundPeriodicTriggerNode
diff --git a/Engine/Plugins/Runtime/Metasound/Source/MetasoundStandardNodes/Private/MetasoundTriggerRouteNode.cpp b/Engine/Plugins/Runtime/Metasound/Source/MetasoundStandardNodes/Private/MetasoundTriggerRouteNode.cpp
--- a/Engine/Plugins/Runtime/Metasound/Source/MetasoundStandardNodes/Private/MetasoundTriggerRouteNode.cpp
+++ b/Engine/Plugins/Runtime/Metasound/Source/MetasoundStandardNodes/Private/MetasoundTriggerRouteNode.cpp
@@ -194,22 +194,6 @@ namespace Metasound
 			InOutVertexData.BindReadVertex(METASOUND_GET_PARAM_NAME(OutputValue), OutputValue);
 		}
 
-		virtual FDataReferenceCollection GetInputs() const override
-		{
-			// This should never be called. Bind(...) is called instead. This method
-			// exists as a stop-gap until the API can be deprecated and removed.
-			checkNoEntry();
-			return {};
-		}
-
-		virtual FDataReferenceCollection GetOutputs() const override
-		{
-			// This should never be called. Bind(...) is called instead. This method
-			// exists as a stop-gap until the API can be deprecated and removed.
-			checkNoEntry();
-			return {};
-		}
-
 		void Execute()
 		{
 			OutputTrigger->AdvanceBlock();
@@ -257,18 +241,7 @@ namespace Metasound
 	 *  Routes values from multiple input pins to a single output pin based on trigger inputs.
 	 */
 	template<typename ValueType, uint32 NumInputs>
-	class METASOUNDSTANDARDNODES_API TTriggerRouteNode : public FNodeFacade
-	{
-	public:
-		/**
-		 * Constructor used by the Metasound Frontend.
-		 */
-		TTriggerRouteNode(const FNodeInitData& InInitData)
-			: FNodeFacade(InInitData.InstanceName, InInitData.InstanceID, TFacadeOperatorClass<TTriggerRouteOperator<ValueType, NumInputs>>())
-		{}
-
-		virtual ~TTriggerRouteNode() = default;
-	};
+	using TTriggerRouteNode = TNodeFacade<TTriggerRouteOperator<ValueType, NumInputs>>;
 
 	REGISTER_TRIGGER_ROUTE_NODE(int32, 2)
 	REGISTER_TRIGGER_ROUTE_NODE(int32, 3)
diff --git a/Engine/Plugins/Runtime/Metasound/Source/MetasoundStandardNodes/Private/MetasoundTriggerSelectNode.cpp b/Engine/Plugins/Runtime/Metasound/Source/MetasoundStandardNodes/Private/MetasoundTriggerSelectNode.cpp
--- a/Engine/Plugins/Runtime/Metasound/Source/MetasoundStandardNodes/Private/MetasoundTriggerSelectNode.cpp
+++ b/Engine/Plugins/Runtime/Metasound/Source/MetasoundStandardNodes/Private/MetasoundTriggerSelectNode.cpp
@@ -46,9 +46,6 @@ namespace Metasound
 
 		virtual void BindInputs(FInputVertexInterfaceData& InOutVertexData) override;
 		virtual void BindOutputs(FOutputVertexInterfaceData& InOutVertexData) override;
-		virtual FDataReferenceCollection GetInputs() const override;
-		virtual FDataReferenceCollection GetOutputs() const override;
-
 		void Execute();
 		void Reset(const IOperator::FResetParams& InParams);
 
@@ -90,25 +87,6 @@ namespace Metasound
 		}
 	}
 
-	template<uint32 NumOutputs>
-	FDataReferenceCollection TTriggerSelectOperator<NumOutputs>::GetInputs() const
-	{
-		// This should never be called. Bind(...) is called instead. This method
-		// exists as a stop-gap until the API can be deprecated and removed.
-		checkNoEntry();
-		return {};
-	}
-
-	template<uint32 NumOutputs>
-	FDataReferenceCollection TTriggerSelectOperator<NumOutputs>::GetOutputs() const
-	{
-
-		// This should never be called. Bind(...) is called instead. This method
-		// exists as a stop-gap until the API can be deprecated and removed.
-		checkNoEntry();
-		return {};
-	}
-
 	template<uint32 NumOutputs>
 	void TTriggerSelectOperator<NumOutputs>::Execute()
 	{
@@ -209,14 +187,7 @@ namespace Metasound
 
 	// Node Class
 	template <uint32 NumOutputs>
-	class TTriggerSelectNode : public FNodeFacade
-	{
-	public:
-		TTriggerSelectNode(const FNodeInitData& InitData)
-			: FNodeFacade(InitData.InstanceName, InitData.InstanceID, TFacadeOperatorClass<TTriggerSelectOperator<NumOutputs>>())
-		{
-		}
-	};
+	using TTriggerSelectNode = TNodeFacade<TTriggerSelectOperator<NumOutputs>>;
 
 	REGISTER_TRIGGER_SELECT_NODE(2)
 	REGISTER_TRIGGER_SELECT_NODE(3)
diff --git a/Engine/Plugins/Runtime/Metasound/Source/MetasoundStandardNodes/Private/MetasoundTriggerSequenceNode.cpp b/Engine/Plugins/Runtime/Metasound/Source/MetasoundStandardNodes/Private/MetasoundTriggerSequenceNode.cpp
--- a/Engine/Plugins/Runtime/Metasound/Source/MetasoundStandardNodes/Private/MetasoundTriggerSequenceNode.cpp
+++ b/Engine/Plugins/Runtime/Metasound/Source/MetasoundStandardNodes/Private/MetasoundTriggerSequenceNode.cpp
@@ -49,9 +49,6 @@ namespace Metasound
 
 		virtual void BindInputs(FInputVertexInterfaceData& InOutVertexData) override;
 		virtual void BindOutputs(FOutputVertexInterfaceData& InOutVertexData) override;
-		virtual FDataReferenceCollection GetInputs() const override;
-		virtual FDataReferenceCollection GetOutputs() const override;
-
 		void Execute();
 		void Reset(const IOperator::FResetParams& InParams);
 
@@ -100,27 +97,6 @@ namespace Metasound
 		}
 	}
 
-
-	template<uint32 NumOutputs>
-	FDataReferenceCollection TTriggerSequenceOperator<NumOutputs>::GetInputs() const
-	{
-		// This should never be called. Bind(...) is called instead. This method
-		// exists as a stop-gap until the API can be deprecated and removed.
-		checkNoEntry();
-		return {};
-	}
-
-	template<uint32 NumOutputs>
-	FDataReferenceCollection TTriggerSequenceOperator<NumOutputs>::GetOutputs() const
-	{
-		// This should never be called. Bind(...) is called instead. This method
-		// exists as a stop-gap until the API can be deprecated and removed.
-		checkNoEntry();
-		return {};
-	}
-
-
-
 	template<uint32 NumOutputs>
 	void TTriggerSequenceOperator<NumOutputs>::Execute()
 	{
@@ -242,14 +218,7 @@ namespace Metasound
 
 	// Node Class
 	template <uint32 NumOutputs>
-	class TTriggerSequenceNode : public FNodeFacade
-	{
-	public:
-		TTriggerSequenceNode(const FNodeInitData& InitData)
-			: FNodeFacade(InitData.InstanceName, InitData.InstanceID, TFacadeOperatorClass<TTriggerSequenceOperator<NumOutputs>>())
-		{
-		}
-	};
+	using TTriggerSequenceNode = TNodeFacade<TTriggerSequenceOperator<NumOutputs>>;
 
 	REGISTER_TRIGGER_SEQUENCE_NODE(2)
 	REGISTER_TRIGGER_SEQUENCE_NODE(3)
diff --git a/Engine/Plugins/Runtime/Metasound/Source/MetasoundStandardNodes/Private/MetasoundTriggerToggleNode.cpp b/Engine/Plugins/Runtime/Metasound/Source/MetasoundStandardNodes/Private/MetasoundTriggerToggleNode.cpp
--- a/Engine/Plugins/Runtime/Metasound/Source/MetasoundStandardNodes/Private/MetasoundTriggerToggleNode.cpp
+++ b/Engine/Plugins/Runtime/Metasound/Source/MetasoundStandardNodes/Private/MetasoundTriggerToggleNode.cpp
@@ -41,8 +41,6 @@ namespace Metasound
 
 			virtual void BindInputs(FInputVertexInterfaceData& InOutVertexData) override;
 			virtual void BindOutputs(FOutputVertexInterfaceData& InOutVertexData) override;
-			virtual FDataReferenceCollection GetInputs() const override;
-			virtual FDataReferenceCollection GetOutputs() const override;
 
 			void Execute();
 			void Reset(const IOperator::FResetParams& InParams);
@@ -80,22 +78,6 @@ namespace Metasound
 		InOutVertexData.BindReadVertex(METASOUND_GET_PARAM_NAME(OutputValue), ValueOutput);
 	}
 
-	FDataReferenceCollection FTriggerToggleOperator::GetInputs() const
-	{
-		// This should never be called. Bind(...) is called instead. This method
-		// exists as a stop-gap until the API can be deprecated and removed.
-		checkNoEntry();
-		return {};
-	}
-
-	FDataReferenceCollection FTriggerToggleOperator::GetOutputs() const
-	{
-		// This should never be called. Bind(...) is called instead. This method
-		// exists as a stop-gap until the API can be deprecated and removed.
-		checkNoEntry();
-		return {};
-	}
-
 	void FTriggerToggleOperator::Execute()
 	{
 		TriggerOutput->AdvanceBlock();
@@ -187,6 +169,17 @@ namespace Metasound
 	{
 	}
 
+	FTriggerToggleNode::FTriggerToggleNode(FNodeData InNodeData, TSharedRef<const FNodeClassMetadata> InMetadata)
+	:	FNodeFacade(MoveTemp(InNodeData), InMetadata, TFacadeOperatorClass<FTriggerToggleOperator>())
+	{
+	}
+
+	FNodeClassMetadata FTriggerToggleNode::CreateNodeClassMetadata()
+	{
+		return FTriggerToggleOperator::GetNodeInfo();
+	}
+	
+
 	METASOUND_REGISTER_NODE(FTriggerToggleNode)
 }
 
diff --git a/Engine/Plugins/Runtime/Metasound/Source/MetasoundStandardNodes/Private/MetasoundWaveShaperNode.cpp b/Engine/Plugins/Runtime/Metasound/Source/MetasoundStandardNodes/Private/MetasoundWaveShaperNode.cpp
--- a/Engine/Plugins/Runtime/Metasound/Source/MetasoundStandardNodes/Private/MetasoundWaveShaperNode.cpp
+++ b/Engine/Plugins/Runtime/Metasound/Source/MetasoundStandardNodes/Private/MetasoundWaveShaperNode.cpp
@@ -56,8 +56,6 @@ namespace Metasound
 
 		virtual void BindInputs(FInputVertexInterfaceData& InOutVertexData) override;
 		virtual void BindOutputs(FOutputVertexInterfaceData& InOutVertexData) override;
-		virtual FDataReferenceCollection GetInputs() const override;
-		virtual FDataReferenceCollection GetOutputs() const override;
 		void Execute();
 		void Reset(const IOperator::FResetParams& InParams);
 
@@ -114,22 +112,6 @@ namespace Metasound
 		InOutVertexData.BindReadVertex(METASOUND_GET_PARAM_NAME(OutAudio), AudioOutput);
 	}
 
-	FDataReferenceCollection FWaveShaperOperator::GetInputs() const
-	{
-		// This should never be called. Bind(...) is called instead. This method
-		// exists as a stop-gap until the API can be deprecated and removed.
-		checkNoEntry();
-		return {};
-	}
-
-	FDataReferenceCollection FWaveShaperOperator::GetOutputs() const
-	{
-		// This should never be called. Bind(...) is called instead. This method
-		// exists as a stop-gap until the API can be deprecated and removed.
-		checkNoEntry();
-		return {};
-	}
-
 	void FWaveShaperOperator::Execute()
 	{
 		const float* InputAudio = AudioInput->GetData();
@@ -205,17 +187,7 @@ namespace Metasound
 		return MakeUnique<FWaveShaperOperator>(InParams.OperatorSettings, AudioIn, InAmount, InBias, OutGain, InType);
 	}
 
-	class FWaveShaperNode : public FNodeFacade
-	{
-	public:
-		/**
-		 * Constructor used by the Metasound Frontend.
-		 */
-		FWaveShaperNode(const FNodeInitData& InitData)
-			: FNodeFacade(InitData.InstanceName, InitData.InstanceID, TFacadeOperatorClass<FWaveShaperOperator>())
-		{
-		}
-	};
+	using FWaveShaperNode = TNodeFacade<FWaveShaperOperator>;
 
 	METASOUND_REGISTER_NODE(FWaveShaperNode)
 }
diff --git a/Engine/Plugins/Runtime/Metasound/Source/MetasoundStandardNodes/Private/MetasoundWaveWriterNode.cpp b/Engine/Plugins/Runtime/Metasound/Source/MetasoundStandardNodes/Private/MetasoundWaveWriterNode.cpp
--- a/Engine/Plugins/Runtime/Metasound/Source/MetasoundStandardNodes/Private/MetasoundWaveWriterNode.cpp
+++ b/Engine/Plugins/Runtime/Metasound/Source/MetasoundStandardNodes/Private/MetasoundWaveWriterNode.cpp
@@ -115,22 +115,6 @@ namespace Metasound
 		{
 		}
 
-		virtual FDataReferenceCollection GetInputs() const override
-		{
-			// This should never be called. Bind(...) is called instead. This method
-			// exists as a stop-gap until the API can be deprecated and removed.
-			checkNoEntry();
-			return {};
-		}
-	
-		virtual FDataReferenceCollection GetOutputs() const override
-		{
-			// This should never be called. Bind(...) is called instead. This method
-			// exists as a stop-gap until the API can be deprecated and removed.
-			checkNoEntry();
-			return {};
-		}
-
 		static const FVertexInterface& DeclareVertexInterface()
 		{
 			auto CreateDefaultInterface = []()-> FVertexInterface
@@ -396,20 +380,8 @@ namespace Metasound
 	}
 
 	template<int32 NumInputChannels>
-	class METASOUNDSTANDARDNODES_API TWaveWriterNode : public FNodeFacade
-	{
-	public:
-		TWaveWriterNode(const FVertexName& InName, const FGuid& InInstanceID)
-			: FNodeFacade(InName, InInstanceID, TFacadeOperatorClass<TWaveWriterOperator<NumInputChannels>>())
-		{
-		}
-
-		TWaveWriterNode(const FNodeInitData& InInitData)
-			: TWaveWriterNode(InInitData.InstanceName, InInitData.InstanceID)
-		{
-		}
-	};
-
+	using TWaveWriterNode = TNodeFacade<TWaveWriterOperator<NumInputChannels>>;
+		
 	#define REGISTER_WAVEWRITER_NODE(A) \
 		using FWaveWriterNode_##A = TWaveWriterNode<A>; \
 		METASOUND_REGISTER_NODE(FWaveWriterNode_##A)
diff --git a/Engine/Plugins/Runtime/Metasound/Source/MetasoundStandardNodes/Public/MetasoundBPMToSecondsNode.h b/Engine/Plugins/Runtime/Metasound/Source/MetasoundStandardNodes/Public/MetasoundBPMToSecondsNode.h
--- a/Engine/Plugins/Runtime/Metasound/Source/MetasoundStandardNodes/Public/MetasoundBPMToSecondsNode.h
+++ b/Engine/Plugins/Runtime/Metasound/Source/MetasoundStandardNodes/Public/MetasoundBPMToSecondsNode.h
@@ -18,6 +18,10 @@ namespace Metasound
 		 * Constructor used by the Metasound Frontend.
 		 */
 		FBPMToSecondsNode(const FNodeInitData& InitData);
+
+		FBPMToSecondsNode(FNodeData InNodeData, TSharedRef<const FNodeClassMetadata> InClassMetadata);
+
+		static FNodeClassMetadata CreateNodeClassMetadata();
 	};
 } // namespace Metasound
 
diff --git a/Engine/Plugins/Runtime/Metasound/Source/MetasoundStandardNodes/Public/MetasoundEnvelopeFollowerNode.h b/Engine/Plugins/Runtime/Metasound/Source/MetasoundStandardNodes/Public/MetasoundEnvelopeFollowerNode.h
--- a/Engine/Plugins/Runtime/Metasound/Source/MetasoundStandardNodes/Public/MetasoundEnvelopeFollowerNode.h
+++ b/Engine/Plugins/Runtime/Metasound/Source/MetasoundStandardNodes/Public/MetasoundEnvelopeFollowerNode.h
@@ -18,6 +18,10 @@ namespace Metasound
 		 * Constructor used by the Metasound Frontend.
 		 */
 		FEnvelopeFollowerNode(const FNodeInitData& InitData);
+
+		FEnvelopeFollowerNode(FNodeData InNodeData, TSharedRef<const FNodeClassMetadata> InClassMetadata);
+
+		static FNodeClassMetadata CreateNodeClassMetadata();
 	};
 } // namespace Metasound
 
diff --git a/Engine/Plugins/Runtime/Metasound/Source/MetasoundStandardNodes/Public/MetasoundGain.h b/Engine/Plugins/Runtime/Metasound/Source/MetasoundStandardNodes/Public/MetasoundGain.h
--- a/Engine/Plugins/Runtime/Metasound/Source/MetasoundStandardNodes/Public/MetasoundGain.h
+++ b/Engine/Plugins/Runtime/Metasound/Source/MetasoundStandardNodes/Public/MetasoundGain.h
@@ -8,7 +8,7 @@
 namespace Metasound
 {
 	/** Representation of a gain */
-	enum class EGainRepresentation  : uint8
+	enum class UE_DEPRECATED(5.6, "EGainRepresentation is not used and will be removed") EGainRepresentation  : uint8
 	{
 		/** Gain in linear scale  */
 		Linear,
@@ -21,7 +21,7 @@ namespace Metasound
 	 * getters and setters as well as several convenience functions for 
 	 * converting gain to decibels.
 	 */
-	class FGain
+	class UE_DEPRECATED(5.6, "FGain is not used and will be removed") FGain
 	{
 		public:
 			static constexpr const float DefaultGain = 1.0f;
@@ -32,6 +32,7 @@ namespace Metasound
 			{
 			}
 
+PRAGMA_DISABLE_DEPRECATION_WARNINGS
 			/** Construct a Gain with a given value and representation.
 			 *
 			 * @param InValue - The value of the gain at the given 
@@ -62,6 +63,7 @@ namespace Metasound
 				: FGain(InValue, EGainRepresentation::Linear)
 			{}
 
+PRAGMA_ENABLE_DEPRECATION_WARNINGS
 			/** Implicit operator used for math operations. */
 			operator float() const
 			{
diff --git a/Engine/Plugins/Runtime/Metasound/Source/MetasoundStandardNodes/Public/MetasoundNoiseGenerator.h b/Engine/Plugins/Runtime/Metasound/Source/MetasoundStandardNodes/Public/MetasoundNoiseGenerator.h
--- a/Engine/Plugins/Runtime/Metasound/Source/MetasoundStandardNodes/Public/MetasoundNoiseGenerator.h
+++ b/Engine/Plugins/Runtime/Metasound/Source/MetasoundStandardNodes/Public/MetasoundNoiseGenerator.h
@@ -13,9 +13,12 @@ namespace Metasound
 	public:
 		FNoiseNode(const FVertexName& InName, const FGuid& InInstanceID, int32 InDefaultSeed);
 		FNoiseNode(const FNodeInitData& InInitData);
+		FNoiseNode(FNodeData InNodeData, TSharedRef<const FNodeClassMetadata> InClassMetadata);
 
-		int32 GetDefaultSeed() const { return DefaultSeed; }
+		UE_DEPRECATED(5.6, "The default seed is stored on the vertex interface of the node")
+		int32 GetDefaultSeed() const { return INDEX_NONE; }
+
+		static FNodeClassMetadata CreateNodeClassMetadata();
 	private:
-		int32 DefaultSeed;
 	};
 }
diff --git a/Engine/Plugins/Runtime/Metasound/Source/MetasoundStandardNodes/Public/MetasoundOscillatorNodes.h b/Engine/Plugins/Runtime/Metasound/Source/MetasoundStandardNodes/Public/MetasoundOscillatorNodes.h
--- a/Engine/Plugins/Runtime/Metasound/Source/MetasoundStandardNodes/Public/MetasoundOscillatorNodes.h
+++ b/Engine/Plugins/Runtime/Metasound/Source/MetasoundStandardNodes/Public/MetasoundOscillatorNodes.h
@@ -5,6 +5,7 @@
 #include "CoreMinimal.h"
 
 #include "MetasoundAudioBuffer.h"
+#include "MetasoundBasicNode.h"
 #include "MetasoundBuilderInterface.h"
 #include "MetasoundDataReferenceCollection.h"
 #include "MetasoundExecutableOperator.h"
@@ -16,26 +17,31 @@
 
 namespace Metasound
 {
-	class METASOUNDSTANDARDNODES_API FOscilatorNodeBase : public FNode
+	class METASOUNDSTANDARDNODES_API FOscilatorNodeBase : public FBasicNode
 	{
 	protected:
 		FOscilatorNodeBase(const FVertexName& InInstanceName, const FGuid& InInstanceID, const FNodeClassMetadata& InInfo, const TSharedRef<IOperatorFactory, ESPMode::ThreadSafe>& InFactory, float InDefaultFrequency, float InDefaultGlideTime, bool bInDefaultEnablement);
+		FOscilatorNodeBase(FNodeData InNodeData, TSharedRef<const FNodeClassMetadata> InClassMetadata, TSharedRef<IOperatorFactory> InFactory);
 
+		UE_DEPRECATED(5.6, "This function should not be used")
 		float GetDefaultPhaseOffset() const
 		{
 			return DefaultPhaseOffset; 
 		}
 
+		UE_DEPRECATED(5.6, "This function should not be used")
 		float GetDefaultFrequency() const
 		{
 			return DefaultFrequency;
 		}
 
+		UE_DEPRECATED(5.6, "This function should not be used")
 		float GetDefaultGlideFactor() const
 		{
 			return DefaultGlideFactor;
 		}
 
+		UE_DEPRECATED(5.6, "This function should not be used")
 		bool GetDefaultEnablement() const
 		{
 			return bDefaultEnablement;
@@ -46,29 +52,10 @@ namespace Metasound
 			return Factory;
 		}
 
-		const FVertexInterface& GetVertexInterface() const override
-		{
-			return VertexInterface;
-		}
-
-		bool SetVertexInterface(const FVertexInterface& InInterface) override
-		{
-			if (IsVertexInterfaceSupported(InInterface))
-			{
-				VertexInterface = InInterface;
-				return true;
-			}
-			return false;
-		}
-
-		bool IsVertexInterfaceSupported(const FVertexInterface& InInterface) const override
-		{
-			return VertexInterface == InInterface;
-		}
-
 		TSharedRef<IOperatorFactory, ESPMode::ThreadSafe> Factory;
 		FVertexInterface VertexInterface;
 
+	private:
 		float DefaultPhaseOffset = 0.f;
 		float DefaultFrequency = 440.f;
 		float DefaultGlideFactor = 0.0f;
@@ -79,38 +66,60 @@ namespace Metasound
 	{
 	public:
 		class FFactory;
+
+		UE_DEPRECATED(5.6, "Do not construct oscillators with this function")
 		FSineOscilatorNode(const FVertexName& InInstanceName, const FGuid& InInstanceID, float InDefaultFrequency, float InDefautlGlideTime, bool bInDefaultEnablement);
+
 		FSineOscilatorNode(const FNodeInitData& InInitData);
+		FSineOscilatorNode(FNodeData InNodeData, TSharedRef<const FNodeClassMetadata> InMetadata);
+
+		static FNodeClassMetadata CreateNodeClassMetadata();
 	};
 		
 	class METASOUNDSTANDARDNODES_API FSawOscilatorNode : public FOscilatorNodeBase
 	{
 	public:
 		class FFactory;
+		UE_DEPRECATED(5.6, "Do not construct oscillators with this function")
 		FSawOscilatorNode(const FVertexName& InInstanceName, const FGuid& InInstanceID, float InDefaultFrequency, float InDefautlGlideTime, bool bInDefaultEnablement);
+
 		FSawOscilatorNode(const FNodeInitData& InInitData);
+		FSawOscilatorNode(FNodeData InNodeData, TSharedRef<const FNodeClassMetadata> InMetadata);
+
+		static FNodeClassMetadata CreateNodeClassMetadata();
 	};
 
 	class METASOUNDSTANDARDNODES_API FTriangleOscilatorNode : public FOscilatorNodeBase
 	{
 	public:
 		class FFactory;
+		UE_DEPRECATED(5.6, "Do not construct oscillators with this function")
 		FTriangleOscilatorNode(const FVertexName& InInstanceName, const FGuid& InInstanceID, float InDefaultFrequency, float InDefautlGlideTime, bool bInDefaultEnablement);
 		FTriangleOscilatorNode(const FNodeInitData& InInitData);
+		FTriangleOscilatorNode(FNodeData InNodeData, TSharedRef<const FNodeClassMetadata> InMetadata);
+
+		static FNodeClassMetadata CreateNodeClassMetadata();
 	};
 
 	class METASOUNDSTANDARDNODES_API FSquareOscilatorNode : public FOscilatorNodeBase
 	{
 	public:
 		class FFactory;
+		UE_DEPRECATED(5.6, "Do not construct oscillators with this function")
 		FSquareOscilatorNode(const FVertexName& InInstanceName, const FGuid& InInstanceID, float InDefaultFrequency, float InDefautlGlideTime, bool bInDefaultEnablement);
 		FSquareOscilatorNode(const FNodeInitData& InInitData);
+		FSquareOscilatorNode(FNodeData InNodeData, TSharedRef<const FNodeClassMetadata> InMetadata);
+
+		static FNodeClassMetadata CreateNodeClassMetadata();
 	};
 
 	class METASOUNDSTANDARDNODES_API FLfoNode : public FNodeFacade
 	{
 	public:		
 		FLfoNode(const FNodeInitData& InInitData);
+		FLfoNode(FNodeData InNodeData, TSharedRef<const FNodeClassMetadata> InMetadata);
 		virtual ~FLfoNode() = default;
+
+		static FNodeClassMetadata CreateNodeClassMetadata();
 	};
 }
diff --git a/Engine/Plugins/Runtime/Metasound/Source/MetasoundStandardNodes/Public/MetasoundRandomNode.h b/Engine/Plugins/Runtime/Metasound/Source/MetasoundStandardNodes/Public/MetasoundRandomNode.h
--- a/Engine/Plugins/Runtime/Metasound/Source/MetasoundStandardNodes/Public/MetasoundRandomNode.h
+++ b/Engine/Plugins/Runtime/Metasound/Source/MetasoundStandardNodes/Public/MetasoundRandomNode.h
@@ -405,22 +405,6 @@ namespace Metasound
 			InOutVertexData.BindReadVertex(METASOUND_GET_PARAM_NAME(OutputValue), OutputValue);
 		}
 
-		virtual FDataReferenceCollection GetInputs() const override
-		{
-			// This should never be called. Bind(...) is called instead. This method
-			// exists as a stop-gap until the API can be deprecated and removed.
-			checkNoEntry();
-			return {};
-		}
-
-		virtual FDataReferenceCollection GetOutputs() const override
-		{
-			// This should never be called. Bind(...) is called instead. This method
-			// exists as a stop-gap until the API can be deprecated and removed.
-			checkNoEntry();
-			return {};
-		}
-
 		void Execute()
 		{
 			TriggerOutOnReset->AdvanceBlock();
@@ -517,18 +501,7 @@ namespace Metasound
 	 *  Generates a random float value when triggered.
 	 */
 	template<typename ValueType>
-	class METASOUNDSTANDARDNODES_API TRandomNode : public FNodeFacade
-	{
-	public:
-		/**
-		 * Constructor used by the Metasound Frontend.
-		 */
-		TRandomNode(const FNodeInitData& InInitData)
-			: FNodeFacade(InInitData.InstanceName, InInitData.InstanceID, TFacadeOperatorClass<TRandomNodeOperator<ValueType>>())
-		{}
-
-		virtual ~TRandomNode() = default;
-	};
+	using TRandomNode = TNodeFacade<TRandomNodeOperator<ValueType>>;
 
 } // namespace Metasound
 
diff --git a/Engine/Plugins/Runtime/Metasound/Source/MetasoundStandardNodes/Public/MetasoundTriggerAccumulatorNode.h b/Engine/Plugins/Runtime/Metasound/Source/MetasoundStandardNodes/Public/MetasoundTriggerAccumulatorNode.h
--- a/Engine/Plugins/Runtime/Metasound/Source/MetasoundStandardNodes/Public/MetasoundTriggerAccumulatorNode.h
+++ b/Engine/Plugins/Runtime/Metasound/Source/MetasoundStandardNodes/Public/MetasoundTriggerAccumulatorNode.h
@@ -120,22 +120,6 @@ namespace Metasound
 			InOutVertexData.BindReadVertex(METASOUND_GET_PARAM_NAME(AccumulateOutputOnTrigger), OutputTrigger);
 		}
 
-		virtual FDataReferenceCollection GetInputs() const override
-		{
-			// This should never be called. Bind(...) is called instead. This method
-			// exists as a stop-gap until the API can be deprecated and removed.
-			checkNoEntry();
-			return {};
-		}
-
-		virtual FDataReferenceCollection GetOutputs() const override
-		{
-			// This should never be called. Bind(...) is called instead. This method
-			// exists as a stop-gap until the API can be deprecated and removed.
-			checkNoEntry();
-			return {};
-		}
-
 		void TriggerOutputIfReady(int32 InStartFrame)
 		{
 			for (bool WasTriggered : bInputWasTriggered)
@@ -210,18 +194,7 @@ namespace Metasound
 	*  Routes values from multiple input pins to a single output pin based on trigger inputs.
 	*/
 	template<uint32 NumInputs>
-	class METASOUNDSTANDARDNODES_API TTriggerAccumulatorNode : public FNodeFacade
-	{
-	public:
-		/**
-		 * Constructor used by the Metasound Frontend.
-		 */
-		TTriggerAccumulatorNode(const FNodeInitData& InInitData)
-			: FNodeFacade(InInitData.InstanceName, InInitData.InstanceID, TFacadeOperatorClass<TTriggerAccumulatorOperator<NumInputs>>())
-		{}
-
-		virtual ~TTriggerAccumulatorNode() = default;
-	};
+	using TTriggerAccumulatorNode = TNodeFacade<TTriggerAccumulatorOperator<NumInputs>>;
 
 } // namespace Metasound
 
diff --git a/Engine/Plugins/Runtime/Metasound/Source/MetasoundStandardNodes/Public/MetasoundTriggerCompareNode.h b/Engine/Plugins/Runtime/Metasound/Source/MetasoundStandardNodes/Public/MetasoundTriggerCompareNode.h
--- a/Engine/Plugins/Runtime/Metasound/Source/MetasoundStandardNodes/Public/MetasoundTriggerCompareNode.h
+++ b/Engine/Plugins/Runtime/Metasound/Source/MetasoundStandardNodes/Public/MetasoundTriggerCompareNode.h
@@ -143,22 +143,6 @@ namespace Metasound
 			InOutVertexData.BindReadVertex(METASOUND_GET_PARAM_NAME(OutputOnFalse), TriggerOutOnFalse);
 		}
 
-		virtual FDataReferenceCollection GetInputs() const override
-		{
-			// This should never be called. Bind(...) is called instead. This method
-			// exists as a stop-gap until the API can be deprecated and removed.
-			checkNoEntry();
-			return {};
-		}
-
-		virtual FDataReferenceCollection GetOutputs() const override
-		{
-			// This should never be called. Bind(...) is called instead. This method
-			// exists as a stop-gap until the API can be deprecated and removed.
-			checkNoEntry();
-			return {};
-		}
-
 		void Execute()
 		{
 			TriggerOutOnTrue->AdvanceBlock();
@@ -230,18 +214,7 @@ namespace Metasound
 	 *  Compares two inputs against enumerated comparison types.
 	 */
 	template<typename ValueType>
-	class TTriggerCompareNode : public FNodeFacade
-	{
-	public:
-		/**
-		 * Constructor used by the Metasound Frontend.
-		 */
-		TTriggerCompareNode(const FNodeInitData& InInitData)
-			: FNodeFacade(InInitData.InstanceName, InInitData.InstanceID, TFacadeOperatorClass<TTriggerCompareNodeOperator<ValueType>>())
-		{}
-
-		virtual ~TTriggerCompareNode() = default;
-	};
+	using TTriggerCompareNode = TNodeFacade<TTriggerCompareNodeOperator<ValueType>>;
 }
 
 #undef LOCTEXT_NAMESPACE
diff --git a/Engine/Plugins/Runtime/Metasound/Source/MetasoundStandardNodes/Public/MetasoundTriggerRepeatNode.h b/Engine/Plugins/Runtime/Metasound/Source/MetasoundStandardNodes/Public/MetasoundTriggerRepeatNode.h
--- a/Engine/Plugins/Runtime/Metasound/Source/MetasoundStandardNodes/Public/MetasoundTriggerRepeatNode.h
+++ b/Engine/Plugins/Runtime/Metasound/Source/MetasoundStandardNodes/Public/MetasoundTriggerRepeatNode.h
@@ -20,10 +20,10 @@ namespace Metasound
 	{
 		public:
 			FTriggerRepeatNode(const FNodeInitData& InInitData);
+			FTriggerRepeatNode(FNodeData InNodeData, TSharedRef<const FNodeClassMetadata> InMetadata);
 
 			virtual ~FTriggerRepeatNode() = default;
 
-		private:
-			float DefaultPeriod = 1.0f;
+			static FNodeClassMetadata CreateNodeClassMetadata();
 	};
 }
diff --git a/Engine/Plugins/Runtime/Metasound/Source/MetasoundStandardNodes/Public/MetasoundTriggerToggleNode.h b/Engine/Plugins/Runtime/Metasound/Source/MetasoundStandardNodes/Public/MetasoundTriggerToggleNode.h
--- a/Engine/Plugins/Runtime/Metasound/Source/MetasoundStandardNodes/Public/MetasoundTriggerToggleNode.h
+++ b/Engine/Plugins/Runtime/Metasound/Source/MetasoundStandardNodes/Public/MetasoundTriggerToggleNode.h
@@ -21,6 +21,10 @@ namespace Metasound
 		public:
 			FTriggerToggleNode(const FNodeInitData& InInitData);
 
+			FTriggerToggleNode(FNodeData InNodeData, TSharedRef<const FNodeClassMetadata> InMetadata);
+
 			virtual ~FTriggerToggleNode() = default;
+
+			static FNodeClassMetadata CreateNodeClassMetadata();
 	};
 }
diff --git a/Engine/Plugins/Runtime/Metasound/Source/MetasoundStandardNodes/Public/MetasoundValueNode.h b/Engine/Plugins/Runtime/Metasound/Source/MetasoundStandardNodes/Public/MetasoundValueNode.h
--- a/Engine/Plugins/Runtime/Metasound/Source/MetasoundStandardNodes/Public/MetasoundValueNode.h
+++ b/Engine/Plugins/Runtime/Metasound/Source/MetasoundStandardNodes/Public/MetasoundValueNode.h
@@ -199,18 +199,8 @@ namespace Metasound
 	 *  Generates a random float value when triggered.
 	 */
 	template<typename ValueType>
-	class METASOUNDSTANDARDNODES_API TValueNode : public FNodeFacade
-	{
-	public:
-		/**
-		 * Constructor used by the Metasound Frontend.
-		 */
-		TValueNode(const FNodeInitData& InInitData)
-			: FNodeFacade(InInitData.InstanceName, InInitData.InstanceID, TFacadeOperatorClass<TValueOperator<ValueType>>())
-		{}
-
-		virtual ~TValueNode() = default;
-	};
+	using TValueNode = TNodeFacade<TValueOperator<ValueType>>;
+
 } // namespace Metasound
 
 #undef LOCTEXT_NAMESPACE
