Update various MetaSound nodes to use FNodeData for node configuration

diff --git a/Engine/Plugins/Runtime/AudioModulation/Source/AudioModulation/Private/MetaSoundNodes/GetModulatorValueNode.cpp b/Engine/Plugins/Runtime/AudioModulation/Source/AudioModulation/Private/MetaSoundNodes/GetModulatorValueNode.cpp
--- a/Engine/Plugins/Runtime/AudioModulation/Source/AudioModulation/Private/MetaSoundNodes/GetModulatorValueNode.cpp
+++ b/Engine/Plugins/Runtime/AudioModulation/Source/AudioModulation/Private/MetaSoundNodes/GetModulatorValueNode.cpp
@@ -107,22 +107,6 @@ namespace AudioModulation
 			InOutVertexData.BindReadVertex("Out", OutValue);
 		}
 
-		virtual Metasound::FDataReferenceCollection GetInputs() const override
-		{
-			// This should never be called. Bind(...) is called instead. This method
-			// exists as a stop-gap until the API can be deprecated and removed.
-			checkNoEntry();
-			return {};
-		}
-
-		virtual Metasound::FDataReferenceCollection GetOutputs() const override
-		{
-			// This should never be called. Bind(...) is called instead. This method
-			// exists as a stop-gap until the API can be deprecated and removed.
-			checkNoEntry();
-			return {};
-		}
-
 		void Execute()
 		{
 			const FSoundModulatorAsset& ModulatorAsset = *Modulator;
@@ -185,16 +169,7 @@ namespace AudioModulation
 		Metasound::TDataWriteReference<float> OutValue;
 	};
 
-	class FGetModulatorValueNode : public Metasound::FNodeFacade
-	{
-	public:
-		FGetModulatorValueNode(const Metasound::FNodeInitData& InInitData)
-			: Metasound::FNodeFacade(InInitData.InstanceName, InInitData.InstanceID, Metasound::TFacadeOperatorClass<FGetModulatorValueNodeOperator>())
-		{
-		}
-
-		virtual ~FGetModulatorValueNode() = default;
-	};
+	using FGetModulatorValueNode = Metasound::TNodeFacade<FGetModulatorValueNodeOperator>;
 
 	METASOUND_REGISTER_NODE(FGetModulatorValueNode)
 } // namespace AudioModulation
diff --git a/Engine/Plugins/Runtime/AudioModulation/Source/AudioModulation/Private/MetaSoundNodes/MixModulatorsNode.cpp b/Engine/Plugins/Runtime/AudioModulation/Source/AudioModulation/Private/MetaSoundNodes/MixModulatorsNode.cpp
--- a/Engine/Plugins/Runtime/AudioModulation/Source/AudioModulation/Private/MetaSoundNodes/MixModulatorsNode.cpp
+++ b/Engine/Plugins/Runtime/AudioModulation/Source/AudioModulation/Private/MetaSoundNodes/MixModulatorsNode.cpp
@@ -120,22 +120,6 @@ namespace AudioModulation
 			InOutVertexData.BindReadVertex("Out", OutValue);
 		}
 
-		virtual Metasound::FDataReferenceCollection GetInputs() const override
-		{
-			// This should never be called. Bind(...) is called instead. This method
-			// exists as a stop-gap until the API can be deprecated and removed.
-			checkNoEntry();
-			return {};
-		}
-
-		virtual Metasound::FDataReferenceCollection GetOutputs() const override
-		{
-			// This should never be called. Bind(...) is called instead. This method
-			// exists as a stop-gap until the API can be deprecated and removed.
-			checkNoEntry();
-			return {};
-		}
-
 		void Execute()
 		{
 			auto UpdateValue = [](const FSoundModulatorAssetReadRef& ReadRef, Audio::FDeviceId InDeviceId, Audio::FModulatorHandle& OutHandle)
@@ -220,17 +204,7 @@ namespace AudioModulation
 		Metasound::TDataWriteReference<float> OutValue;
 	};
 
-	class FMixModulatorsNode : public Metasound::FNodeFacade
-	{
-	public:
-		FMixModulatorsNode(const Metasound::FNodeInitData& InInitData)
-			: Metasound::FNodeFacade(InInitData.InstanceName, InInitData.InstanceID, Metasound::TFacadeOperatorClass<FMixModulatorsNodeOperator>())
-		{
-		}
-
-		virtual ~FMixModulatorsNode() = default;
-	};
-
+	using FMixModulatorsNode = Metasound::TNodeFacade<FMixModulatorsNodeOperator>;
 	METASOUND_REGISTER_NODE(FMixModulatorsNode)
 } // namespace AudioModulation
 
diff --git a/Engine/Plugins/Runtime/Harmonix/Source/HarmonixMetasound/Private/HarmonixMetasound/Nodes/DelayNode.cpp b/Engine/Plugins/Runtime/Harmonix/Source/HarmonixMetasound/Private/HarmonixMetasound/Nodes/DelayNode.cpp
--- a/Engine/Plugins/Runtime/Harmonix/Source/HarmonixMetasound/Private/HarmonixMetasound/Nodes/DelayNode.cpp
+++ b/Engine/Plugins/Runtime/Harmonix/Source/HarmonixMetasound/Private/HarmonixMetasound/Nodes/DelayNode.cpp
@@ -310,22 +310,6 @@ namespace HarmonixMetasound::DelayNode
 			InVertexData.BindReadVertex(Outputs::AudioRightName, Outputs.AudioRight);
 		}
 		
-		virtual FDataReferenceCollection GetInputs() const override
-		{
-			// This should never be called. Bind(...) is called instead. This method
-			// exists as a stop-gap until the API can be deprecated and removed.
-			checkNoEntry();
-			return {};
-		}
-		
-		virtual FDataReferenceCollection GetOutputs() const override
-		{
-			// This should never be called. Bind(...) is called instead. This method
-			// exists as a stop-gap until the API can be deprecated and removed.
-			checkNoEntry();
-			return {};
-		}
-
 		void Reset(const FResetParams& Params)
 		{
 			Delay.Unprepare();
@@ -441,14 +425,7 @@ namespace HarmonixMetasound::DelayNode
 		FOutputs Outputs;
 	};
 
-	class FDelayNode final : public FNodeFacade
-	{
-	public:
-		explicit FDelayNode(const FNodeInitData& InitData)
-		: FNodeFacade(InitData.InstanceName, InitData.InstanceID, TFacadeOperatorClass<FDelayOperator>())
-		{}
-	};
-
+	using FDelayNode = Metasound::TNodeFacade<FDelayOperator>;
 	METASOUND_REGISTER_NODE(FDelayNode);
 }
 
diff --git a/Engine/Plugins/Runtime/Harmonix/Source/HarmonixMetasound/Private/HarmonixMetasound/Nodes/DistortionNode.cpp b/Engine/Plugins/Runtime/Harmonix/Source/HarmonixMetasound/Private/HarmonixMetasound/Nodes/DistortionNode.cpp
--- a/Engine/Plugins/Runtime/Harmonix/Source/HarmonixMetasound/Private/HarmonixMetasound/Nodes/DistortionNode.cpp
+++ b/Engine/Plugins/Runtime/Harmonix/Source/HarmonixMetasound/Private/HarmonixMetasound/Nodes/DistortionNode.cpp
@@ -425,14 +425,7 @@ namespace HarmonixMetasound
 		}
 	};
 
-	class FDistortionNode final : public FNodeFacade
-	{
-	public:
-		explicit FDistortionNode(const FNodeInitData& InitData)
-			: FNodeFacade(InitData.InstanceName, InitData.InstanceID, TFacadeOperatorClass<FDistortionOperator>())
-		{}
-	};
-
+	using FDistortionNode = Metasound::TNodeFacade<FDistortionOperator>;
 	METASOUND_REGISTER_NODE(FDistortionNode);
 }
 
diff --git a/Engine/Plugins/Runtime/Harmonix/Source/HarmonixMetasound/Private/HarmonixMetasound/Nodes/DjFilterNode.cpp b/Engine/Plugins/Runtime/Harmonix/Source/HarmonixMetasound/Private/HarmonixMetasound/Nodes/DjFilterNode.cpp
--- a/Engine/Plugins/Runtime/Harmonix/Source/HarmonixMetasound/Private/HarmonixMetasound/Nodes/DjFilterNode.cpp
+++ b/Engine/Plugins/Runtime/Harmonix/Source/HarmonixMetasound/Private/HarmonixMetasound/Nodes/DjFilterNode.cpp
@@ -207,14 +207,7 @@ namespace HarmonixMetasound::Nodes::DjFilter
 		Harmonix::Dsp::Effects::FDjFilter Filter;
 	};
 
-	class FDjFilterNode final : public Metasound::FNodeFacade
-	{
-	public:
-		explicit FDjFilterNode(const Metasound::FNodeInitData& InitData)
-		: FNodeFacade(InitData.InstanceName, InitData.InstanceID, Metasound::TFacadeOperatorClass<FOp>())
-		{}
-	};
-	
+	using FDjFilterNode = Metasound::TNodeFacade<FOp>;
 	METASOUND_REGISTER_NODE(FDjFilterNode);
 }
 
diff --git a/Engine/Plugins/Runtime/Harmonix/Source/HarmonixMetasound/Private/HarmonixMetasound/Nodes/FFTAnalyzerNode.cpp b/Engine/Plugins/Runtime/Harmonix/Source/HarmonixMetasound/Private/HarmonixMetasound/Nodes/FFTAnalyzerNode.cpp
--- a/Engine/Plugins/Runtime/Harmonix/Source/HarmonixMetasound/Private/HarmonixMetasound/Nodes/FFTAnalyzerNode.cpp
+++ b/Engine/Plugins/Runtime/Harmonix/Source/HarmonixMetasound/Private/HarmonixMetasound/Nodes/FFTAnalyzerNode.cpp
@@ -265,14 +265,7 @@ namespace HarmonixMetasound
 		int32 SampleRate = 0;
 	};
 
-	class FFFTAnalyzerNode final : public FNodeFacade
-	{
-	public:
-		explicit FFFTAnalyzerNode(const FNodeInitData& InitData)
-			: FNodeFacade(InitData.InstanceName, InitData.InstanceID, TFacadeOperatorClass<FFFTAnalyzerOperator>())
-		{}
-	};
-
+	using FFFTAnalyzerNode = Metasound::TNodeFacade<FFFTAnalyzerOperator>;
 	METASOUND_REGISTER_NODE(FFFTAnalyzerNode);
 }
 
diff --git a/Engine/Plugins/Runtime/Harmonix/Source/HarmonixMetasound/Private/HarmonixMetasound/Nodes/FusionSamplerNode.cpp b/Engine/Plugins/Runtime/Harmonix/Source/HarmonixMetasound/Private/HarmonixMetasound/Nodes/FusionSamplerNode.cpp
--- a/Engine/Plugins/Runtime/Harmonix/Source/HarmonixMetasound/Private/HarmonixMetasound/Nodes/FusionSamplerNode.cpp
+++ b/Engine/Plugins/Runtime/Harmonix/Source/HarmonixMetasound/Private/HarmonixMetasound/Nodes/FusionSamplerNode.cpp
@@ -330,8 +330,6 @@ namespace HarmonixMetasound
 		}
 	};
 
-	template <int32 NUM_CHANNELS> class FFusionSamplerNode;
-
 	/**********************************************************************************************************
 	***********************************************************************************************************
 	* TEMPLATE FFusionSamplerOperator
@@ -378,7 +376,7 @@ namespace HarmonixMetasound
 			}
 			check (OutputChannelLayout != EAudioBufferChannelLayout::Raw);
 
-			const FFusionSamplerNode<NUM_CHANNELS>& PlayerNode = static_cast<const FFusionSamplerNode<NUM_CHANNELS>&>(InParams.Node);
+			const TNodeFacade<FFusionSamplerOperator<NUM_CHANNELS>>& PlayerNode = static_cast<const TNodeFacade<FFusionSamplerOperator<NUM_CHANNELS>>&>(InParams.Node);
 			FFusionSamplerOperatorBase::FConstructionArgs ConstructionArgs(InParams, OutputChannelLayout);
 			return MakeUnique<FFusionSamplerOperator<NUM_CHANNELS>>(ConstructionArgs);
 		}
@@ -520,21 +518,6 @@ namespace HarmonixMetasound
 		InVertexData.BindReadVertex(METASOUND_GET_PARAM_NAME(FusionSamplerNodePinNames::RenderSync), SyncLinkOutPin);
 	}
 
-	/**********************************************************************************************************
-	***********************************************************************************************************
-	* TEMPLATE FFusionSamplerNode
-	***********************************************************************************************************
-	*********************************************************************************************************/
-	template <int32 NUM_CHANNELS>
-	class FFusionSamplerNode : public FNodeFacade
-	{
-	public: 
-		FFusionSamplerNode(const FNodeInitData& InInitData)
-			: FNodeFacade(InInitData.InstanceName, InInitData.InstanceID, TFacadeOperatorClass<FFusionSamplerOperator<NUM_CHANNELS>>())
-		{}
-		virtual ~FFusionSamplerNode() = default;
-	};
-
 	/**********************************************************************************************************
 	***********************************************************************************************************
 	* TEMPLATE FFusionSamplerOperator & FFusionSamplerNode EXPLICIT INSTANTIATIONS
@@ -542,8 +525,8 @@ namespace HarmonixMetasound
 	*********************************************************************************************************/
 	template class FFusionSamplerOperator<1>;
 	template class FFusionSamplerOperator<2>;
-	using FFusionSamplerNodeMono   = FFusionSamplerNode<1>;
-	using FFusionSamplerNodeStereo = FFusionSamplerNode<2>;
+	using FFusionSamplerNodeMono   = TNodeFacade<FFusionSamplerOperator<1>>;
+	using FFusionSamplerNodeStereo = TNodeFacade<FFusionSamplerOperator<2>>;
 	METASOUND_REGISTER_NODE(FFusionSamplerNodeMono)
 	METASOUND_REGISTER_NODE(FFusionSamplerNodeStereo)
 
diff --git a/Engine/Plugins/Runtime/Harmonix/Source/HarmonixMetasound/Private/HarmonixMetasound/Nodes/MetronomeNode.cpp b/Engine/Plugins/Runtime/Harmonix/Source/HarmonixMetasound/Private/HarmonixMetasound/Nodes/MetronomeNode.cpp
--- a/Engine/Plugins/Runtime/Harmonix/Source/HarmonixMetasound/Private/HarmonixMetasound/Nodes/MetronomeNode.cpp
+++ b/Engine/Plugins/Runtime/Harmonix/Source/HarmonixMetasound/Private/HarmonixMetasound/Nodes/MetronomeNode.cpp
@@ -59,15 +59,7 @@ namespace HarmonixMetasound::Nodes::MetronomeNode
 			const int32 InSeekPrerollBars);
 	};
 
-	class FMetronomeNode : public FNodeFacade
-	{
-	public:
-		FMetronomeNode(const FNodeInitData& InInitData)
-			: FNodeFacade(InInitData.InstanceName, InInitData.InstanceID, TFacadeOperatorClass<FMetronomeOperator>())
-		{}
-		virtual ~FMetronomeNode() = default;
-	};
-
+	using FMetronomeNode = Metasound::TNodeFacade<FMetronomeOperator>;
 	METASOUND_REGISTER_NODE(FMetronomeNode)
 
 	const FNodeClassMetadata& FMetronomeOperator::GetNodeInfo()
diff --git a/Engine/Plugins/Runtime/Harmonix/Source/HarmonixMetasound/Private/HarmonixMetasound/Nodes/MidiCCTriggerNode.cpp b/Engine/Plugins/Runtime/Harmonix/Source/HarmonixMetasound/Private/HarmonixMetasound/Nodes/MidiCCTriggerNode.cpp
--- a/Engine/Plugins/Runtime/Harmonix/Source/HarmonixMetasound/Private/HarmonixMetasound/Nodes/MidiCCTriggerNode.cpp
+++ b/Engine/Plugins/Runtime/Harmonix/Source/HarmonixMetasound/Private/HarmonixMetasound/Nodes/MidiCCTriggerNode.cpp
@@ -200,14 +200,7 @@ namespace HarmonixMetasound::Nodes::MidiCCTriggerNode
 		FOutputs Outputs;
 	};
 
-	class FMidiCCTriggerNode_V1 final : public FNodeFacade
-	{
-	public:
-		explicit FMidiCCTriggerNode_V1(const FNodeInitData& InInitData)
-			: FNodeFacade(InInitData.InstanceName, InInitData.InstanceID, Metasound::TFacadeOperatorClass<FMidiCCTriggerOperator_V1>())
-		{}
-	};
-
+	using FMidiCCTriggerNode_V1 = Metasound::TNodeFacade<FMidiCCTriggerOperator_V1>;
 	METASOUND_REGISTER_NODE(FMidiCCTriggerNode_V1);
 
 	class FMidiCCTriggerOperator_V0 final : public Metasound::TExecutableOperator<FMidiCCTriggerOperator_V0>
@@ -382,14 +375,7 @@ namespace HarmonixMetasound::Nodes::MidiCCTriggerNode
 			FOutputs Outputs;
 	};
 
-	class FMidiCCTriggerNode_V0 final : public FNodeFacade
-	{
-	public:
-		explicit FMidiCCTriggerNode_V0(const FNodeInitData& InInitData)
-			: FNodeFacade(InInitData.InstanceName, InInitData.InstanceID, Metasound::TFacadeOperatorClass<FMidiCCTriggerOperator_V0>())
-		{}
-	};
-
+	using FMidiCCTriggerNode_V0 = Metasound::TNodeFacade<FMidiCCTriggerOperator_V0>;
 	METASOUND_REGISTER_NODE(FMidiCCTriggerNode_V0);
 }
 
diff --git a/Engine/Plugins/Runtime/Harmonix/Source/HarmonixMetasound/Private/HarmonixMetasound/Nodes/MidiChannelFilterNode.cpp b/Engine/Plugins/Runtime/Harmonix/Source/HarmonixMetasound/Private/HarmonixMetasound/Nodes/MidiChannelFilterNode.cpp
--- a/Engine/Plugins/Runtime/Harmonix/Source/HarmonixMetasound/Private/HarmonixMetasound/Nodes/MidiChannelFilterNode.cpp
+++ b/Engine/Plugins/Runtime/Harmonix/Source/HarmonixMetasound/Private/HarmonixMetasound/Nodes/MidiChannelFilterNode.cpp
@@ -168,15 +168,7 @@ namespace HarmonixMetasound::Nodes::MidiChannelFilter
 		int32 LastChannelValue = -1;
 	};
 
-	class FMidiChannelFilterNode final : public Metasound::FNodeFacade
-	{
-	public:
-		explicit FMidiChannelFilterNode(const Metasound::FNodeInitData& InInitData)
-			: FNodeFacade(InInitData.InstanceName, InInitData.InstanceID, Metasound::TFacadeOperatorClass<FMidiChannelFilterOperator>())
-		{}
-		virtual ~FMidiChannelFilterNode() override = default;
-	};
-
+	using FMidiChannelFilterNode = Metasound::TNodeFacade<FMidiChannelFilterOperator>;
 	METASOUND_REGISTER_NODE(FMidiChannelFilterNode)
 }
 
diff --git a/Engine/Plugins/Runtime/Harmonix/Source/HarmonixMetasound/Private/HarmonixMetasound/Nodes/MidiClockOffsetNode.cpp b/Engine/Plugins/Runtime/Harmonix/Source/HarmonixMetasound/Private/HarmonixMetasound/Nodes/MidiClockOffsetNode.cpp
--- a/Engine/Plugins/Runtime/Harmonix/Source/HarmonixMetasound/Private/HarmonixMetasound/Nodes/MidiClockOffsetNode.cpp
+++ b/Engine/Plugins/Runtime/Harmonix/Source/HarmonixMetasound/Private/HarmonixMetasound/Nodes/MidiClockOffsetNode.cpp
@@ -272,15 +272,7 @@ namespace HarmonixMetasound::Nodes::MidiClockOffset
 		float PrevOffsetBeats = 0.0f;
 	};
 
-	class FMidiClockOffsetNode : public FNodeFacade
-	{
-	public:
-		FMidiClockOffsetNode(const FNodeInitData& InInitData)
-			: FNodeFacade(InInitData.InstanceName, InInitData.InstanceID, TFacadeOperatorClass<FMidiClockOffsetOperator>())
-		{}
-	virtual ~FMidiClockOffsetNode() = default;
-	};
-
+	using FMidiClockOffsetNode = Metasound::TNodeFacade<FMidiClockOffsetOperator>;
 	METASOUND_REGISTER_NODE(FMidiClockOffsetNode)
 }
 
diff --git a/Engine/Plugins/Runtime/Harmonix/Source/HarmonixMetasound/Private/HarmonixMetasound/Nodes/MidiClockStateNode.cpp b/Engine/Plugins/Runtime/Harmonix/Source/HarmonixMetasound/Private/HarmonixMetasound/Nodes/MidiClockStateNode.cpp
--- a/Engine/Plugins/Runtime/Harmonix/Source/HarmonixMetasound/Private/HarmonixMetasound/Nodes/MidiClockStateNode.cpp
+++ b/Engine/Plugins/Runtime/Harmonix/Source/HarmonixMetasound/Private/HarmonixMetasound/Nodes/MidiClockStateNode.cpp
@@ -62,15 +62,7 @@ namespace HarmonixMetasound
 		static constexpr float kTinyTempoSpeed = 0.0001f;
 	};
 
-	class FMidiClockStateNode : public FNodeFacade
-	{
-	public:
-		FMidiClockStateNode(const FNodeInitData& InInitData)
-			: FNodeFacade(InInitData.InstanceName, InInitData.InstanceID, TFacadeOperatorClass<FMidiClockStateOperator>())
-		{}
-		virtual ~FMidiClockStateNode() = default;
-	};
-
+	using FMidiClockStateNode = Metasound::TNodeFacade<FMidiClockStateOperator>;
 	METASOUND_REGISTER_NODE(FMidiClockStateNode)
 		
 	const FNodeClassMetadata& FMidiClockStateOperator::GetNodeInfo()
diff --git a/Engine/Plugins/Runtime/Harmonix/Source/HarmonixMetasound/Private/HarmonixMetasound/Nodes/MidiClockSubdivisionTriggerNode.cpp b/Engine/Plugins/Runtime/Harmonix/Source/HarmonixMetasound/Private/HarmonixMetasound/Nodes/MidiClockSubdivisionTriggerNode.cpp
--- a/Engine/Plugins/Runtime/Harmonix/Source/HarmonixMetasound/Private/HarmonixMetasound/Nodes/MidiClockSubdivisionTriggerNode.cpp
+++ b/Engine/Plugins/Runtime/Harmonix/Source/HarmonixMetasound/Private/HarmonixMetasound/Nodes/MidiClockSubdivisionTriggerNode.cpp
@@ -76,15 +76,7 @@ namespace HarmonixMetasound::Nodes::MidiClockSubdivisionTriggerNode
 		Harmonix::Midi::Ops::FPulseGenerator PulseGenerator;
 	};
 
-	class FMidiClockSubdivisionTriggerNode : public FNodeFacade
-	{
-	public:
-		FMidiClockSubdivisionTriggerNode(const FNodeInitData& InInitData)
-			: FNodeFacade(InInitData.InstanceName, InInitData.InstanceID, TFacadeOperatorClass<FMidiClockSubdivisionTriggerOperator>())
-		{}
-		virtual ~FMidiClockSubdivisionTriggerNode() override = default;
-	};
-
+	using FMidiClockSubdivisionTriggerNode = Metasound::TNodeFacade<FMidiClockSubdivisionTriggerOperator>;
 	METASOUND_REGISTER_NODE(FMidiClockSubdivisionTriggerNode)
 		
 	const FNodeClassMetadata& FMidiClockSubdivisionTriggerOperator::GetNodeInfo()
diff --git a/Engine/Plugins/Runtime/Harmonix/Source/HarmonixMetasound/Private/HarmonixMetasound/Nodes/MidiClockTimestampTriggerNode.cpp b/Engine/Plugins/Runtime/Harmonix/Source/HarmonixMetasound/Private/HarmonixMetasound/Nodes/MidiClockTimestampTriggerNode.cpp
--- a/Engine/Plugins/Runtime/Harmonix/Source/HarmonixMetasound/Private/HarmonixMetasound/Nodes/MidiClockTimestampTriggerNode.cpp
+++ b/Engine/Plugins/Runtime/Harmonix/Source/HarmonixMetasound/Private/HarmonixMetasound/Nodes/MidiClockTimestampTriggerNode.cpp
@@ -69,15 +69,7 @@ namespace HarmonixMetasound
 		void CalculateTriggerTick();
 	};
 
-	class FMidiClockTimestampTriggerNode : public FNodeFacade
-	{
-	public:
-		FMidiClockTimestampTriggerNode(const FNodeInitData& InInitData)
-			: FNodeFacade(InInitData.InstanceName, InInitData.InstanceID, TFacadeOperatorClass<FMidiClockTimestampTriggerOperator>())
-		{}
-		virtual ~FMidiClockTimestampTriggerNode() = default;
-	};
-
+	using FMidiClockTimestampTriggerNode = Metasound::TNodeFacade<FMidiClockTimestampTriggerOperator>;
 	METASOUND_REGISTER_NODE(FMidiClockTimestampTriggerNode)
 		
 	const FNodeClassMetadata& FMidiClockTimestampTriggerOperator::GetNodeInfo()
diff --git a/Engine/Plugins/Runtime/Harmonix/Source/HarmonixMetasound/Private/HarmonixMetasound/Nodes/MidiNoteFilterNode.cpp b/Engine/Plugins/Runtime/Harmonix/Source/HarmonixMetasound/Private/HarmonixMetasound/Nodes/MidiNoteFilterNode.cpp
--- a/Engine/Plugins/Runtime/Harmonix/Source/HarmonixMetasound/Private/HarmonixMetasound/Nodes/MidiNoteFilterNode.cpp
+++ b/Engine/Plugins/Runtime/Harmonix/Source/HarmonixMetasound/Private/HarmonixMetasound/Nodes/MidiNoteFilterNode.cpp
@@ -184,15 +184,7 @@ namespace HarmonixMetasound::Nodes::MidiNoteFilter
 		Harmonix::Midi::Ops::FMidiNoteFilter Filter;
 	};
 
-	class FMidiNoteFilterNode final : public Metasound::FNodeFacade
-	{
-	public:
-		explicit FMidiNoteFilterNode(const Metasound::FNodeInitData& InInitData)
-			: FNodeFacade(InInitData.InstanceName, InInitData.InstanceID, Metasound::TFacadeOperatorClass<FMidiNoteFilterOperator>())
-		{}
-		virtual ~FMidiNoteFilterNode() override = default;
-	};
-
+	using FMidiNoteFilterNode = Metasound::TNodeFacade<FMidiNoteFilterOperator>;
 	METASOUND_REGISTER_NODE(FMidiNoteFilterNode)
 }
 
diff --git a/Engine/Plugins/Runtime/Harmonix/Source/HarmonixMetasound/Private/HarmonixMetasound/Nodes/MidiNoteTriggerNode.cpp b/Engine/Plugins/Runtime/Harmonix/Source/HarmonixMetasound/Private/HarmonixMetasound/Nodes/MidiNoteTriggerNode.cpp
--- a/Engine/Plugins/Runtime/Harmonix/Source/HarmonixMetasound/Private/HarmonixMetasound/Nodes/MidiNoteTriggerNode.cpp
+++ b/Engine/Plugins/Runtime/Harmonix/Source/HarmonixMetasound/Private/HarmonixMetasound/Nodes/MidiNoteTriggerNode.cpp
@@ -237,15 +237,7 @@ namespace HarmonixMetasound::Nodes::MidiNoteTriggerNode
 		Harmonix::Midi::Ops::FStuckNoteGuard StuckNoteGuard;
 	};
 
-	class FMidiNoteTriggerNode_V1 final : public FNodeFacade
-	{
-	public:
-		explicit FMidiNoteTriggerNode_V1(const FNodeInitData& InInitData)
-			: FNodeFacade(InInitData.InstanceName, InInitData.InstanceID, TFacadeOperatorClass<FMidiNoteTriggerOperator_V1>())
-		{}
-		virtual ~FMidiNoteTriggerNode_V1() override = default;
-	};
-
+	using FMidiNoteTriggerNode_V1 = Metasound::TNodeFacade<FMidiNoteTriggerOperator_V1>;
 	METASOUND_REGISTER_NODE(FMidiNoteTriggerNode_V1)
 }
 
diff --git a/Engine/Plugins/Runtime/Harmonix/Source/HarmonixMetasound/Private/HarmonixMetasound/Nodes/MidiPlayerNode.cpp b/Engine/Plugins/Runtime/Harmonix/Source/HarmonixMetasound/Private/HarmonixMetasound/Nodes/MidiPlayerNode.cpp
--- a/Engine/Plugins/Runtime/Harmonix/Source/HarmonixMetasound/Private/HarmonixMetasound/Nodes/MidiPlayerNode.cpp
+++ b/Engine/Plugins/Runtime/Harmonix/Source/HarmonixMetasound/Private/HarmonixMetasound/Nodes/MidiPlayerNode.cpp
@@ -171,15 +171,7 @@ namespace HarmonixMetasound::Nodes::MidiPlayerNode
 		virtual void InitTransportImpl() override;
 	};
 
-	class FMidiPlayerNode : public FNodeFacade
-	{
-	public:
-		FMidiPlayerNode(const FNodeInitData& InInitData)
-			: FNodeFacade(InInitData.InstanceName, InInitData.InstanceID, TFacadeOperatorClass<FMidiPlayerOperator>())
-		{}
-		virtual ~FMidiPlayerNode() = default;
-	};
-
+	using FMidiPlayerNode = Metasound::TNodeFacade<FMidiPlayerOperator>;
 	METASOUND_REGISTER_NODE(FMidiPlayerNode)
 
 	const FNodeClassMetadata& FMidiPlayerOperator::GetNodeInfo()
diff --git a/Engine/Plugins/Runtime/Harmonix/Source/HarmonixMetasound/Private/HarmonixMetasound/Nodes/MidiPulseGeneratorNode.cpp b/Engine/Plugins/Runtime/Harmonix/Source/HarmonixMetasound/Private/HarmonixMetasound/Nodes/MidiPulseGeneratorNode.cpp
--- a/Engine/Plugins/Runtime/Harmonix/Source/HarmonixMetasound/Private/HarmonixMetasound/Nodes/MidiPulseGeneratorNode.cpp
+++ b/Engine/Plugins/Runtime/Harmonix/Source/HarmonixMetasound/Private/HarmonixMetasound/Nodes/MidiPulseGeneratorNode.cpp
@@ -211,14 +211,7 @@ namespace HarmonixMetasound::Nodes::MidiPulseGeneratorNode
 		Harmonix::Midi::Ops::FMidiPulseGenerator PulseGenerator;
 	};
 
-	class FMidiPulseGeneratorNode final : public Metasound::FNodeFacade
-	{
-	public:
-		explicit FMidiPulseGeneratorNode(const Metasound::FNodeInitData& InInitData)
-			: FNodeFacade(InInitData.InstanceName, InInitData.InstanceID, Metasound::TFacadeOperatorClass<FMidiPulseGeneratorOperator>())
-		{}
-	};
-
+	using FMidiPulseGeneratorNode = Metasound::TNodeFacade<FMidiPulseGeneratorOperator>;
 	METASOUND_REGISTER_NODE(FMidiPulseGeneratorNode)
 }
 
diff --git a/Engine/Plugins/Runtime/Harmonix/Source/HarmonixMetasound/Private/HarmonixMetasound/Nodes/MidiStreamLoggerNode.cpp b/Engine/Plugins/Runtime/Harmonix/Source/HarmonixMetasound/Private/HarmonixMetasound/Nodes/MidiStreamLoggerNode.cpp
--- a/Engine/Plugins/Runtime/Harmonix/Source/HarmonixMetasound/Private/HarmonixMetasound/Nodes/MidiStreamLoggerNode.cpp
+++ b/Engine/Plugins/Runtime/Harmonix/Source/HarmonixMetasound/Private/HarmonixMetasound/Nodes/MidiStreamLoggerNode.cpp
@@ -33,8 +33,6 @@ namespace HarmonixMetasound
 
 		virtual void BindInputs(FInputVertexInterfaceData& InVertexData) override;
 		virtual void BindOutputs(FOutputVertexInterfaceData& InVertexData) override;
-		virtual FDataReferenceCollection GetInputs() const override;
-		virtual FDataReferenceCollection GetOutputs() const override;
 
 		void Reset(const FResetParams& ResetParams);
 
@@ -49,15 +47,7 @@ namespace HarmonixMetasound
 		void DumpMidiEvent(const FMidiStreamEvent& MidiEvent);
 	};
 
-	class FMidiStreamLoggerNode : public FNodeFacade
-	{
-	public:
-		FMidiStreamLoggerNode(const FNodeInitData& InInitData)
-			: FNodeFacade(InInitData.InstanceName, InInitData.InstanceID, TFacadeOperatorClass<FMidiStreamLoggerOperator>())
-		{}
-		virtual ~FMidiStreamLoggerNode() = default;
-	};
-
+	using FMidiStreamLoggerNode = Metasound::TNodeFacade<FMidiStreamLoggerOperator>;
 	METASOUND_REGISTER_NODE(FMidiStreamLoggerNode)
 		
 	const FNodeClassMetadata& FMidiStreamLoggerOperator::GetNodeInfo()
@@ -132,22 +122,6 @@ namespace HarmonixMetasound
 	{
 	}
 
-	FDataReferenceCollection FMidiStreamLoggerOperator::GetInputs() const
-	{
-		// This should never be called. Bind(...) is called instead. This method
-		// exists as a stop-gap until the API can be deprecated and removed.
-		checkNoEntry();
-		return {};
-	}
-
-	FDataReferenceCollection FMidiStreamLoggerOperator::GetOutputs() const
-	{
-		// This should never be called. Bind(...) is called instead. This method
-		// exists as a stop-gap until the API can be deprecated and removed.
-		checkNoEntry();
-		return {};
-	}
-
 	void FMidiStreamLoggerOperator::Reset(const FResetParams& ResetParams)
 	{
 	}
diff --git a/Engine/Plugins/Runtime/Harmonix/Source/HarmonixMetasound/Private/HarmonixMetasound/Nodes/MidiStreamMergeNode.cpp b/Engine/Plugins/Runtime/Harmonix/Source/HarmonixMetasound/Private/HarmonixMetasound/Nodes/MidiStreamMergeNode.cpp
--- a/Engine/Plugins/Runtime/Harmonix/Source/HarmonixMetasound/Private/HarmonixMetasound/Nodes/MidiStreamMergeNode.cpp
+++ b/Engine/Plugins/Runtime/Harmonix/Source/HarmonixMetasound/Private/HarmonixMetasound/Nodes/MidiStreamMergeNode.cpp
@@ -45,15 +45,7 @@ namespace HarmonixMetasound
 		FMidiStreamWriteRef MidiStreamOutPin;
 	};
 
-	class FMidiStreamMergeNode : public FNodeFacade
-	{
-	public:
-		FMidiStreamMergeNode(const FNodeInitData& InInitData)
-			: FNodeFacade(InInitData.InstanceName, InInitData.InstanceID, TFacadeOperatorClass<FMidiStreamMergeOperator>())
-		{}
-		virtual ~FMidiStreamMergeNode() = default;
-	};
-
+	using FMidiStreamMergeNode = Metasound::TNodeFacade<FMidiStreamMergeOperator>;
 	METASOUND_REGISTER_NODE(FMidiStreamMergeNode)
 		
 	const FNodeClassMetadata& FMidiStreamMergeOperator::GetNodeInfo()
diff --git a/Engine/Plugins/Runtime/Harmonix/Source/HarmonixMetasound/Private/HarmonixMetasound/Nodes/MidiStreamSelectNode.cpp b/Engine/Plugins/Runtime/Harmonix/Source/HarmonixMetasound/Private/HarmonixMetasound/Nodes/MidiStreamSelectNode.cpp
--- a/Engine/Plugins/Runtime/Harmonix/Source/HarmonixMetasound/Private/HarmonixMetasound/Nodes/MidiStreamSelectNode.cpp
+++ b/Engine/Plugins/Runtime/Harmonix/Source/HarmonixMetasound/Private/HarmonixMetasound/Nodes/MidiStreamSelectNode.cpp
@@ -60,15 +60,7 @@ namespace HarmonixMetasound
 		void MergeOrCreateNoteOffsForDeselectedStream(int32 SelectedStreamIndex);
 	};
 
-	class FMidiStreamSelectNode : public FNodeFacade
-	{
-	public:
-		FMidiStreamSelectNode(const FNodeInitData& InInitData)
-			: FNodeFacade(InInitData.InstanceName, InInitData.InstanceID, TFacadeOperatorClass<FMidiStreamSelectOperator>())
-		{}
-		virtual ~FMidiStreamSelectNode() = default;
-	};
-
+	using FMidiStreamSelectNode = Metasound::TNodeFacade<FMidiStreamSelectOperator>;
 	METASOUND_REGISTER_NODE(FMidiStreamSelectNode)
 
 	const FNodeClassMetadata& FMidiStreamSelectOperator::GetNodeInfo()
diff --git a/Engine/Plugins/Runtime/Harmonix/Source/HarmonixMetasound/Private/HarmonixMetasound/Nodes/MidiStreamTrackFilterNode.cpp b/Engine/Plugins/Runtime/Harmonix/Source/HarmonixMetasound/Private/HarmonixMetasound/Nodes/MidiStreamTrackFilterNode.cpp
--- a/Engine/Plugins/Runtime/Harmonix/Source/HarmonixMetasound/Private/HarmonixMetasound/Nodes/MidiStreamTrackFilterNode.cpp
+++ b/Engine/Plugins/Runtime/Harmonix/Source/HarmonixMetasound/Private/HarmonixMetasound/Nodes/MidiStreamTrackFilterNode.cpp
@@ -171,15 +171,7 @@ namespace HarmonixMetasound::Nodes::MidiTrackFilter
 		Harmonix::Midi::Ops::FMidiTrackFilter Filter;
 	};
 
-	class FMidiStreamTrackFilterNode_V1 final : public FNodeFacade
-	{
-	public:
-		explicit FMidiStreamTrackFilterNode_V1(const FNodeInitData& InInitData)
-			: FNodeFacade(InInitData.InstanceName, InInitData.InstanceID, TFacadeOperatorClass<FMidiStreamTrackFilterOperator_V1>())
-		{}
-		virtual ~FMidiStreamTrackFilterNode_V1() override = default;
-	};
-
+	using FMidiStreamTrackFilterNode_V1 = Metasound::TNodeFacade<FMidiStreamTrackFilterOperator_V1>;
 	METASOUND_REGISTER_NODE(FMidiStreamTrackFilterNode_V1)
 }
 
diff --git a/Engine/Plugins/Runtime/Harmonix/Source/HarmonixMetasound/Private/HarmonixMetasound/Nodes/MidiStreamTransposerNode.cpp b/Engine/Plugins/Runtime/Harmonix/Source/HarmonixMetasound/Private/HarmonixMetasound/Nodes/MidiStreamTransposerNode.cpp
--- a/Engine/Plugins/Runtime/Harmonix/Source/HarmonixMetasound/Private/HarmonixMetasound/Nodes/MidiStreamTransposerNode.cpp
+++ b/Engine/Plugins/Runtime/Harmonix/Source/HarmonixMetasound/Private/HarmonixMetasound/Nodes/MidiStreamTransposerNode.cpp
@@ -169,15 +169,7 @@ namespace HarmonixMetasound::Nodes::MidiNoteTranspose
 		FOutputs Outputs;
 	};
 
-	class FMidiNoteTransposeNode final : public FNodeFacade
-	{
-	public:
-		explicit FMidiNoteTransposeNode(const FNodeInitData& InInitData)
-			: FNodeFacade(InInitData.InstanceName, InInitData.InstanceID, TFacadeOperatorClass<FMidiNoteTransposeOperator>())
-		{}
-		virtual ~FMidiNoteTransposeNode() override = default;
-	};
-
+	using FMidiNoteTransposeNode = TNodeFacade<FMidiNoteTransposeOperator>;
 	METASOUND_REGISTER_NODE(FMidiNoteTransposeNode)
 }
 
diff --git a/Engine/Plugins/Runtime/Harmonix/Source/HarmonixMetasound/Private/HarmonixMetasound/Nodes/MidiStreamWriterNode.cpp b/Engine/Plugins/Runtime/Harmonix/Source/HarmonixMetasound/Private/HarmonixMetasound/Nodes/MidiStreamWriterNode.cpp
--- a/Engine/Plugins/Runtime/Harmonix/Source/HarmonixMetasound/Private/HarmonixMetasound/Nodes/MidiStreamWriterNode.cpp
+++ b/Engine/Plugins/Runtime/Harmonix/Source/HarmonixMetasound/Private/HarmonixMetasound/Nodes/MidiStreamWriterNode.cpp
@@ -1,4 +1,4 @@
-﻿// Copyright Epic Games, Inc. All Rights Reserved.
+// Copyright Epic Games, Inc. All Rights Reserved.
 
 #include "HarmonixMetasound/Nodes/MidiStreamWriterNode.h"
 
@@ -173,14 +173,7 @@ namespace HarmonixMetasound::Nodes::MidiStreamWriter
 		TUniquePtr<Harmonix::Midi::Ops::FMidiStreamWriter> MidiStreamWriter;
 	};
 
-	class FMidiStreamWriterNode : public FNodeFacade
-	{
-	public:
-		FMidiStreamWriterNode(const FNodeInitData& InInitData)
-			: FNodeFacade(InInitData.InstanceName, InInitData.InstanceID, TFacadeOperatorClass<FMidiStreamWriterOperator>())
-		{}
-	};
-	
+	using FMidiStreamWriterNode = Metasound::TNodeFacade<FMidiStreamWriterOperator>;
 	METASOUND_REGISTER_NODE(FMidiStreamWriterNode);
 }
 
diff --git a/Engine/Plugins/Runtime/Harmonix/Source/HarmonixMetasound/Private/HarmonixMetasound/Nodes/MidiTextTriggerNode.cpp b/Engine/Plugins/Runtime/Harmonix/Source/HarmonixMetasound/Private/HarmonixMetasound/Nodes/MidiTextTriggerNode.cpp
--- a/Engine/Plugins/Runtime/Harmonix/Source/HarmonixMetasound/Private/HarmonixMetasound/Nodes/MidiTextTriggerNode.cpp
+++ b/Engine/Plugins/Runtime/Harmonix/Source/HarmonixMetasound/Private/HarmonixMetasound/Nodes/MidiTextTriggerNode.cpp
@@ -150,15 +150,8 @@ namespace HarmonixMetasound
 		FTriggerWriteRef TriggerOutPin;
 	};
 
-	class FMidiTextTriggerNode_V1 final : public FNodeFacade
-	{
-	public:
-		explicit FMidiTextTriggerNode_V1(const FNodeInitData& InInitData)
-			: FNodeFacade(InInitData.InstanceName, InInitData.InstanceID, TFacadeOperatorClass<FMidiTextTriggerOperator_V1>())
-		{}
-		virtual ~FMidiTextTriggerNode_V1() override = default;
-	};
-
+	
+	using FMidiTextTriggerNode_V1 = Metasound::TNodeFacade<FMidiTextTriggerOperator_V1>;
 	METASOUND_REGISTER_NODE(FMidiTextTriggerNode_V1)
 }
 
diff --git a/Engine/Plugins/Runtime/Harmonix/Source/HarmonixMetasound/Private/HarmonixMetasound/Nodes/MorphingLfoNode.cpp b/Engine/Plugins/Runtime/Harmonix/Source/HarmonixMetasound/Private/HarmonixMetasound/Nodes/MorphingLfoNode.cpp
--- a/Engine/Plugins/Runtime/Harmonix/Source/HarmonixMetasound/Private/HarmonixMetasound/Nodes/MorphingLfoNode.cpp
+++ b/Engine/Plugins/Runtime/Harmonix/Source/HarmonixMetasound/Private/HarmonixMetasound/Nodes/MorphingLfoNode.cpp
@@ -312,19 +312,10 @@ namespace HarmonixMetasound::Nodes::MorphingLFO
 			LFO.Advance(Output, NumSamples);
 		}
 	}
-	
-	template<typename OutputDataType>
-	class TNode final : public FNodeFacade
-	{
-	public:
-		explicit TNode(const FNodeInitData& InitData)
-		: FNodeFacade(InitData.InstanceName, InitData.InstanceID, Metasound::TFacadeOperatorClass<TOp<OutputDataType>>())
-		{}
-	};
 
-	using FMorphingLFOFloatNode = TNode<float>;
+	using FMorphingLFOFloatNode = Metasound::TNodeFacade<TOp<float>>;
 	METASOUND_REGISTER_NODE(FMorphingLFOFloatNode);
-	using FMorphingLFOAudioNode = TNode<FAudioBuffer>;
+	using FMorphingLFOAudioNode = Metasound::TNodeFacade<TOp<FAudioBuffer>>;
 	METASOUND_REGISTER_NODE(FMorphingLFOAudioNode);
 }
 
diff --git a/Engine/Plugins/Runtime/Harmonix/Source/HarmonixMetasound/Private/HarmonixMetasound/Nodes/MultiThreadedFusionSynchronizerNode.cpp b/Engine/Plugins/Runtime/Harmonix/Source/HarmonixMetasound/Private/HarmonixMetasound/Nodes/MultiThreadedFusionSynchronizerNode.cpp
--- a/Engine/Plugins/Runtime/Harmonix/Source/HarmonixMetasound/Private/HarmonixMetasound/Nodes/MultiThreadedFusionSynchronizerNode.cpp
+++ b/Engine/Plugins/Runtime/Harmonix/Source/HarmonixMetasound/Private/HarmonixMetasound/Nodes/MultiThreadedFusionSynchronizerNode.cpp
@@ -301,22 +301,13 @@ namespace HarmonixMetasound
 	* TEMPLATE FMultiThreadedFusionSynchronizerNode
 	***********************************************************************************************************
 	*********************************************************************************************************/
-	template <int32 NUM_SAMPLERS, int32 NUM_CHANNELS>
-	class TMultiThreadedFusionSynchronizerNode : public FNodeFacade
-	{
-	public:
-		TMultiThreadedFusionSynchronizerNode(const FNodeInitData& InInitData)
-			: FNodeFacade(InInitData.InstanceName, InInitData.InstanceID, TFacadeOperatorClass<TMultiThreadedFusionSynchronizerOperator<NUM_SAMPLERS, NUM_CHANNELS>>())
-		{}
-		virtual ~TMultiThreadedFusionSynchronizerNode() = default;
-	};
 
-	using F4MonoInstanceSynchronizerNode = TMultiThreadedFusionSynchronizerNode<4,1>;
-	using F8MonoInstanceSynchronizerNode = TMultiThreadedFusionSynchronizerNode<8,1>;
+	using F4MonoInstanceSynchronizerNode = Metasound::TNodeFacade<TMultiThreadedFusionSynchronizerOperator<4,1>>;
+	using F8MonoInstanceSynchronizerNode = Metasound::TNodeFacade<TMultiThreadedFusionSynchronizerOperator<8,1>>;
 	METASOUND_REGISTER_NODE(F4MonoInstanceSynchronizerNode)
 	METASOUND_REGISTER_NODE(F8MonoInstanceSynchronizerNode)
-	using F4StereoInstanceSynchronizerNode = TMultiThreadedFusionSynchronizerNode<4,2>;
-	using F8StereoInstanceSynchronizerNode = TMultiThreadedFusionSynchronizerNode<8,2>;
+	using F4StereoInstanceSynchronizerNode = Metasound::TNodeFacade<TMultiThreadedFusionSynchronizerOperator<4,2>>;
+	using F8StereoInstanceSynchronizerNode = Metasound::TNodeFacade<TMultiThreadedFusionSynchronizerOperator<8,2>>;
 	METASOUND_REGISTER_NODE(F4StereoInstanceSynchronizerNode)
 	METASOUND_REGISTER_NODE(F8StereoInstanceSynchronizerNode)
 }
diff --git a/Engine/Plugins/Runtime/Harmonix/Source/HarmonixMetasound/Private/HarmonixMetasound/Nodes/MultibandAnalyzerNode.cpp b/Engine/Plugins/Runtime/Harmonix/Source/HarmonixMetasound/Private/HarmonixMetasound/Nodes/MultibandAnalyzerNode.cpp
--- a/Engine/Plugins/Runtime/Harmonix/Source/HarmonixMetasound/Private/HarmonixMetasound/Nodes/MultibandAnalyzerNode.cpp
+++ b/Engine/Plugins/Runtime/Harmonix/Source/HarmonixMetasound/Private/HarmonixMetasound/Nodes/MultibandAnalyzerNode.cpp
@@ -391,14 +391,7 @@ namespace HarmonixMetasound::Nodes::MultibandAnalyzer
 		bool BypassedOutputValid{ false };
 	};
 
-	class FMultibandAnalyzerNode final : public Metasound::FNodeFacade
-	{
-	public:
-		explicit FMultibandAnalyzerNode(const Metasound::FNodeInitData& InInitData)
-			: FNodeFacade(InInitData.InstanceName, InInitData.InstanceID, Metasound::TFacadeOperatorClass<FOp>())
-		{}
-	};
-
+	using FMultibandAnalyzerNode = Metasound::TNodeFacade<FOp>;
 	METASOUND_REGISTER_NODE(FMultibandAnalyzerNode);
 }
 
diff --git a/Engine/Plugins/Runtime/Harmonix/Source/HarmonixMetasound/Private/HarmonixMetasound/Nodes/MusicSeekTargetBuilder.cpp b/Engine/Plugins/Runtime/Harmonix/Source/HarmonixMetasound/Private/HarmonixMetasound/Nodes/MusicSeekTargetBuilder.cpp
--- a/Engine/Plugins/Runtime/Harmonix/Source/HarmonixMetasound/Private/HarmonixMetasound/Nodes/MusicSeekTargetBuilder.cpp
+++ b/Engine/Plugins/Runtime/Harmonix/Source/HarmonixMetasound/Private/HarmonixMetasound/Nodes/MusicSeekTargetBuilder.cpp
@@ -1,4 +1,4 @@
-﻿// Copyright Epic Games, Inc. All Rights Reserved.
+// Copyright Epic Games, Inc. All Rights Reserved.
 
 #include "MetasoundExecutableOperator.h"
 #include "MetasoundFacade.h"
@@ -131,15 +131,7 @@ namespace HarmonixMetasound::Nodes::BarBeatToSeekTarget
 		FMusicSeekTargetWriteRef SeekTargetOutPin;
 	};
 
-	class FBarBeatToSeekTargetNode : public FNodeFacade
-	{
-	public:
-		FBarBeatToSeekTargetNode(const FNodeInitData& InInitData)
-			: FNodeFacade(InInitData.InstanceName, InInitData.InstanceID, TFacadeOperatorClass<FBarBeatToSeekTargetOperator>())
-		{}
-		virtual ~FBarBeatToSeekTargetNode() = default;
-	};
-
+	using FBarBeatToSeekTargetNode = Metasound::TNodeFacade<FBarBeatToSeekTargetOperator>;
 	METASOUND_REGISTER_NODE(FBarBeatToSeekTargetNode);
 }
 
@@ -252,15 +244,7 @@ namespace HarmonixMetasound::Nodes::TimeMsToSeekTarget
 		FMusicSeekTargetWriteRef SeekTargetOutPin;
 	};
 
-	class FTimeMsToSeekTargetNode : public FNodeFacade
-	{
-	public:
-		FTimeMsToSeekTargetNode(const FNodeInitData& InInitData)
-			: FNodeFacade(InInitData.InstanceName, InInitData.InstanceID, TFacadeOperatorClass<FTimeMsToSeekTargetOperator>())
-		{}
-		virtual ~FTimeMsToSeekTargetNode() = default;
-	};
-
+	using FTimeMsToSeekTargetNode = Metasound::TNodeFacade<FTimeMsToSeekTargetOperator>;
 	METASOUND_REGISTER_NODE(FTimeMsToSeekTargetNode);
 };
 
@@ -373,15 +357,7 @@ namespace HarmonixMetasound::Nodes::MusicTimestampToSeekTarget
 		FMusicSeekTargetWriteRef SeekTargetOutPin;
 	};
 	
-	class FMusicTimestampToSeekTargetNode : public FNodeFacade
-	{
-	public:
-		FMusicTimestampToSeekTargetNode(const FNodeInitData& InInitData)
-			: FNodeFacade(InInitData.InstanceName, InInitData.InstanceID, TFacadeOperatorClass<FMusicTimestampToSeekTargetOperator>())
-		{}
-		virtual ~FMusicTimestampToSeekTargetNode() = default;
-	};
-	
+	using FMusicTimestampToSeekTargetNode = Metasound::TNodeFacade<FMusicTimestampToSeekTargetOperator>;
 	METASOUND_REGISTER_NODE(FMusicTimestampToSeekTargetNode);
 	
 };
diff --git a/Engine/Plugins/Runtime/Harmonix/Source/HarmonixMetasound/Private/HarmonixMetasound/Nodes/MusicTimestampBuilder.cpp b/Engine/Plugins/Runtime/Harmonix/Source/HarmonixMetasound/Private/HarmonixMetasound/Nodes/MusicTimestampBuilder.cpp
--- a/Engine/Plugins/Runtime/Harmonix/Source/HarmonixMetasound/Private/HarmonixMetasound/Nodes/MusicTimestampBuilder.cpp
+++ b/Engine/Plugins/Runtime/Harmonix/Source/HarmonixMetasound/Private/HarmonixMetasound/Nodes/MusicTimestampBuilder.cpp
@@ -50,15 +50,7 @@ namespace HarmonixMetasound
 		FMusicTimestampWriteRef MusicTimeStampOutPin;
 	};
 
-	class FMusicTimestampBuilderNode : public FNodeFacade
-	{
-	public:
-		FMusicTimestampBuilderNode(const FNodeInitData& InInitData)
-			: FNodeFacade(InInitData.InstanceName, InInitData.InstanceID, TFacadeOperatorClass<FMusicTimestampBuilderOperator>())
-		{}
-		virtual ~FMusicTimestampBuilderNode() = default;
-	};
-
+	using FMusicTimestampBuilderNode = Metasound::TNodeFacade<FMusicTimestampBuilderOperator>;
 	METASOUND_REGISTER_NODE(FMusicTimestampBuilderNode)
 
 	const FNodeClassMetadata& FMusicTimestampBuilderOperator::GetNodeInfo()
diff --git a/Engine/Plugins/Runtime/Harmonix/Source/HarmonixMetasound/Private/HarmonixMetasound/Nodes/PeakNode.cpp b/Engine/Plugins/Runtime/Harmonix/Source/HarmonixMetasound/Private/HarmonixMetasound/Nodes/PeakNode.cpp
--- a/Engine/Plugins/Runtime/Harmonix/Source/HarmonixMetasound/Private/HarmonixMetasound/Nodes/PeakNode.cpp
+++ b/Engine/Plugins/Runtime/Harmonix/Source/HarmonixMetasound/Private/HarmonixMetasound/Nodes/PeakNode.cpp
@@ -165,14 +165,7 @@ namespace HarmonixMetasound::Nodes::Peak
 		FOutputs Outputs;
 	};
 
-	class FPeakNode final : public Metasound::FNodeFacade
-	{
-	public:
-		explicit FPeakNode(const Metasound::FNodeInitData& InitData)
-		: FNodeFacade(InitData.InstanceName, InitData.InstanceID, Metasound::TFacadeOperatorClass<FOp>())
-		{}
-	};
-	
+	using FPeakNode = Metasound::TNodeFacade<FOp>;
 	METASOUND_REGISTER_NODE(FPeakNode);
 }
 
diff --git a/Engine/Plugins/Runtime/Harmonix/Source/HarmonixMetasound/Private/HarmonixMetasound/Nodes/SecsUntilMusicTimestampNode.cpp b/Engine/Plugins/Runtime/Harmonix/Source/HarmonixMetasound/Private/HarmonixMetasound/Nodes/SecsUntilMusicTimestampNode.cpp
--- a/Engine/Plugins/Runtime/Harmonix/Source/HarmonixMetasound/Private/HarmonixMetasound/Nodes/SecsUntilMusicTimestampNode.cpp
+++ b/Engine/Plugins/Runtime/Harmonix/Source/HarmonixMetasound/Private/HarmonixMetasound/Nodes/SecsUntilMusicTimestampNode.cpp
@@ -206,14 +206,7 @@ namespace HarmonixMetasound::Nodes::SecsUntilMusicTimestampNode
 		float MsOfTimestamp = Never;
 	};
 
-	class FSecsUntilMusicTimestampNode final : public FNodeFacade
-	{
-	public:
-		explicit FSecsUntilMusicTimestampNode(const FNodeInitData& InInitData)
-			: FNodeFacade(InInitData.InstanceName, InInitData.InstanceID, Metasound::TFacadeOperatorClass<FOp>())
-		{}
-	};
-
+	using FSecsUntilMusicTimestampNode = Metasound::TNodeFacade<FOp>;
 	METASOUND_REGISTER_NODE(FSecsUntilMusicTimestampNode);
 }
 
diff --git a/Engine/Plugins/Runtime/Harmonix/Source/HarmonixMetasound/Private/HarmonixMetasound/Nodes/StepSequencePlayerNode.cpp b/Engine/Plugins/Runtime/Harmonix/Source/HarmonixMetasound/Private/HarmonixMetasound/Nodes/StepSequencePlayerNode.cpp
--- a/Engine/Plugins/Runtime/Harmonix/Source/HarmonixMetasound/Private/HarmonixMetasound/Nodes/StepSequencePlayerNode.cpp
+++ b/Engine/Plugins/Runtime/Harmonix/Source/HarmonixMetasound/Private/HarmonixMetasound/Nodes/StepSequencePlayerNode.cpp
@@ -82,8 +82,6 @@ namespace HarmonixMetasound::Nodes::StepSequencePlayer
 
 		virtual void BindInputs(FInputVertexInterfaceData& InVertexData) override;
 		virtual void BindOutputs(FOutputVertexInterfaceData& InVertexData) override;
-		virtual FDataReferenceCollection GetInputs() const override;
-		virtual FDataReferenceCollection GetOutputs() const override;
 		void Reset(const FResetParams& ResetParams);
 		void Execute();
 
@@ -145,15 +143,7 @@ namespace HarmonixMetasound::Nodes::StepSequencePlayer
 
 	};
 
-	class FStepSequencePlayerNode : public FNodeFacade
-	{
-	public:
-		FStepSequencePlayerNode(const FNodeInitData& InInitData)
-			: FNodeFacade(InInitData.InstanceName, InInitData.InstanceID, TFacadeOperatorClass<FStepSequencePlayerOperator>())
-		{}
-		virtual ~FStepSequencePlayerNode() = default;
-	};
-
+	using FStepSequencePlayerNode = TNodeFacade<FStepSequencePlayerOperator>;
 	METASOUND_REGISTER_NODE(FStepSequencePlayerNode)
 
 	const FNodeClassMetadata& FStepSequencePlayerOperator::GetNodeInfo()
@@ -312,22 +302,6 @@ namespace HarmonixMetasound::Nodes::StepSequencePlayer
 		InVertexData.BindReadVertex(METASOUND_GET_PARAM_NAME(Outputs::MidiStream), MidiOutPin);
 	}
 
-	FDataReferenceCollection FStepSequencePlayerOperator::GetInputs() const
-	{
-		// This should never be called. Bind(...) is called instead. This method
-		// exists as a stop-gap until the API can be deprecated and removed.
-		checkNoEntry();
-		return {};
-	}
-
-	FDataReferenceCollection FStepSequencePlayerOperator::GetOutputs() const
-	{
-		// This should never be called. Bind(...) is called instead. This method
-		// exists as a stop-gap until the API can be deprecated and removed.
-		checkNoEntry();
-		return {};
-	}
-
 	void FStepSequencePlayerOperator::Reset(const FResetParams& ResetParams)
 	{
 		BlockSize = ResetParams.OperatorSettings.GetNumFramesPerBlock();
diff --git a/Engine/Plugins/Runtime/Harmonix/Source/HarmonixMetasound/Private/HarmonixMetasound/Nodes/TransportToTriggerNode.cpp b/Engine/Plugins/Runtime/Harmonix/Source/HarmonixMetasound/Private/HarmonixMetasound/Nodes/TransportToTriggerNode.cpp
--- a/Engine/Plugins/Runtime/Harmonix/Source/HarmonixMetasound/Private/HarmonixMetasound/Nodes/TransportToTriggerNode.cpp
+++ b/Engine/Plugins/Runtime/Harmonix/Source/HarmonixMetasound/Private/HarmonixMetasound/Nodes/TransportToTriggerNode.cpp
@@ -94,22 +94,6 @@ namespace HarmonixMetasound
 			InVertexData.BindReadVertex(METASOUND_GET_PARAM_NAME(Outputs::TransportKill), KillOutPin);
 		}
 
-		virtual FDataReferenceCollection GetInputs() const override
-		{
-			// This should never be called. Bind(...) is called instead. This method
-			// exists as a stop-gap until the API can be deprecated and removed.
-			checkNoEntry();
-			return {};
-		}
-
-		virtual FDataReferenceCollection GetOutputs() const override
-		{
-			// This should never be called. Bind(...) is called instead. This method
-			// exists as a stop-gap until the API can be deprecated and removed.
-			checkNoEntry();
-			return {};
-		}
-
 		static TUniquePtr<IOperator> CreateOperator(const FBuildOperatorParams& InParams, FBuildResults& OutResults)
 		{
 			using namespace CommonPinNames;
@@ -185,15 +169,7 @@ namespace HarmonixMetasound
 		FTriggerWriteRef KillOutPin;
 	};
 
-	class FTransportToTriggerNode : public FNodeFacade
-	{
-	public:
-		FTransportToTriggerNode(const FNodeInitData& InitData)
-			: FNodeFacade(InitData.InstanceName, InitData.InstanceID, TFacadeOperatorClass<FTransportToTriggerOperator>())
-		{
-		}
-	};
-
+	using FTransportToTriggerNode = Metasound::TNodeFacade<FTransportToTriggerOperator>;
 	METASOUND_REGISTER_NODE(FTransportToTriggerNode);
 }
 
diff --git a/Engine/Plugins/Runtime/Harmonix/Source/HarmonixMetasound/Private/HarmonixMetasound/Nodes/TransportWavePlayerControllerNode.cpp b/Engine/Plugins/Runtime/Harmonix/Source/HarmonixMetasound/Private/HarmonixMetasound/Nodes/TransportWavePlayerControllerNode.cpp
--- a/Engine/Plugins/Runtime/Harmonix/Source/HarmonixMetasound/Private/HarmonixMetasound/Nodes/TransportWavePlayerControllerNode.cpp
+++ b/Engine/Plugins/Runtime/Harmonix/Source/HarmonixMetasound/Private/HarmonixMetasound/Nodes/TransportWavePlayerControllerNode.cpp
@@ -103,22 +103,6 @@ namespace HarmonixMetasound
 			bNeedsTransportInit = true;
 		}
 		
-		virtual FDataReferenceCollection GetInputs() const override
-		{
-			// This should never be called. Bind(...) is called instead. This method
-			// exists as a stop-gap until the API can be deprecated and removed.
-			checkNoEntry();
-			return {};
-		}
-
-		virtual FDataReferenceCollection GetOutputs() const override
-		{
-			// This should never be called. Bind(...) is called instead. This method
-			// exists as a stop-gap until the API can be deprecated and removed.
-			checkNoEntry();
-			return {};
-		}
-
 		static TUniquePtr<IOperator> CreateOperator(const FBuildOperatorParams& InParams, FBuildResults& OutResults)
 		{
 			using namespace CommonPinNames;
@@ -294,15 +278,7 @@ namespace HarmonixMetasound
 		}
 	};
 
-	class FTransportWavePlayerControllerNode : public FNodeFacade
-	{
-	public:
-		FTransportWavePlayerControllerNode(const FNodeInitData& InitData)
-			: FNodeFacade(InitData.InstanceName, InitData.InstanceID, TFacadeOperatorClass<FTransportWavePlayerControllerOperator>())
-		{
-		}
-	};
-
+	using FTransportWavePlayerControllerNode = Metasound::TNodeFacade<FTransportWavePlayerControllerOperator>;
 	METASOUND_REGISTER_NODE(FTransportWavePlayerControllerNode);
 }
 
diff --git a/Engine/Plugins/Runtime/Harmonix/Source/HarmonixMetasound/Private/HarmonixMetasound/Nodes/TriggerToTransportNode.cpp b/Engine/Plugins/Runtime/Harmonix/Source/HarmonixMetasound/Private/HarmonixMetasound/Nodes/TriggerToTransportNode.cpp
--- a/Engine/Plugins/Runtime/Harmonix/Source/HarmonixMetasound/Private/HarmonixMetasound/Nodes/TriggerToTransportNode.cpp
+++ b/Engine/Plugins/Runtime/Harmonix/Source/HarmonixMetasound/Private/HarmonixMetasound/Nodes/TriggerToTransportNode.cpp
@@ -50,8 +50,6 @@ namespace HarmonixMetasound::Nodes::TriggerToTransportNode
 
 		virtual void BindInputs(FInputVertexInterfaceData& InVertexData) override;
 		virtual void BindOutputs(FOutputVertexInterfaceData& InVertexData) override;
-		virtual FDataReferenceCollection GetInputs() const override;
-		virtual FDataReferenceCollection GetOutputs() const override;
 
 		void Reset(const FResetParams& ResetParams);
 		
@@ -72,15 +70,7 @@ namespace HarmonixMetasound::Nodes::TriggerToTransportNode
 		FMusicTransportEventStreamWriteRef TransportOutPin;
 	};
 
-	class FTriggerToTransportNode : public FNodeFacade
-	{
-	public:
-		FTriggerToTransportNode(const FNodeInitData& InInitData)
-			: FNodeFacade(InInitData.InstanceName, InInitData.InstanceID, TFacadeOperatorClass<FTriggerToTransportOperator>())
-		{}
-		virtual ~FTriggerToTransportNode() override = default;
-	};
-
+	using FTriggerToTransportNode = Metasound::TNodeFacade<FTriggerToTransportOperator>;
 	METASOUND_REGISTER_NODE(FTriggerToTransportNode)
 
 	const FNodeClassMetadata& FTriggerToTransportOperator::GetNodeInfo()
@@ -204,22 +194,6 @@ namespace HarmonixMetasound::Nodes::TriggerToTransportNode
 		InVertexData.BindReadVertex(METASOUND_GET_PARAM_NAME(Outputs::Transport), TransportOutPin);
 	}
 
-	FDataReferenceCollection FTriggerToTransportOperator::GetInputs() const
-	{
-		// This should never be called. Bind(...) is called instead. This method
-		// exists as a stop-gap until the API can be deprecated and removed.
-		checkNoEntry();
-		return {};
-	}
-
-	FDataReferenceCollection FTriggerToTransportOperator::GetOutputs() const
-	{
-		// This should never be called. Bind(...) is called instead. This method
-		// exists as a stop-gap until the API can be deprecated and removed.
-		checkNoEntry();
-		return {};
-	}
-
 	void FTriggerToTransportOperator::Reset(const FResetParams& ResetParams)
 	{
 		TransportOutPin->Reset();
diff --git a/Engine/Plugins/Runtime/Metasound/Source/MetasoundEngine/Private/MetasoundAudioBusReaderNode.cpp b/Engine/Plugins/Runtime/Metasound/Source/MetasoundEngine/Private/MetasoundAudioBusReaderNode.cpp
--- a/Engine/Plugins/Runtime/Metasound/Source/MetasoundEngine/Private/MetasoundAudioBusReaderNode.cpp
+++ b/Engine/Plugins/Runtime/Metasound/Source/MetasoundEngine/Private/MetasoundAudioBusReaderNode.cpp
@@ -154,22 +154,6 @@ namespace Metasound
 			}
 		}
 
-		virtual FDataReferenceCollection GetInputs() const override
-		{
-			// This should never be called. Bind(...) is called instead. This method
-			// exists as a stop-gap until the API can be deprecated and removed.
-			checkNoEntry();
-			return {};
-		}
-
-		virtual FDataReferenceCollection GetOutputs() const override
-		{
-			// This should never be called. Bind(...) is called instead. This method
-			// exists as a stop-gap until the API can be deprecated and removed.
-			checkNoEntry();
-			return {};
-		}
-
 		void Execute()
 		{
 			const FAudioBusProxyPtr& BusProxy = AudioBusAsset->GetAudioBusProxy();
@@ -328,14 +312,7 @@ namespace Metasound
 	};
 
 	template<uint32 NumChannels>
-	class TAudioBusReaderNode : public FNodeFacade
-	{
-	public:
-		TAudioBusReaderNode(const FNodeInitData& InitData)
-			: FNodeFacade(InitData.InstanceName, InitData.InstanceID, TFacadeOperatorClass<TAudioBusReaderOperator<NumChannels>>())
-		{
-		}
-	};
+	using TAudioBusReaderNode = TNodeFacade<TAudioBusReaderOperator<NumChannels>>;
 
 #define REGISTER_AUDIO_BUS_READER_NODE(ChannelCount) \
 	using FAudioBusReaderNode_##ChannelCount = TAudioBusReaderNode<ChannelCount>; \
diff --git a/Engine/Plugins/Runtime/Metasound/Source/MetasoundEngine/Private/MetasoundAudioBusWriterNode.cpp b/Engine/Plugins/Runtime/Metasound/Source/MetasoundEngine/Private/MetasoundAudioBusWriterNode.cpp
--- a/Engine/Plugins/Runtime/Metasound/Source/MetasoundEngine/Private/MetasoundAudioBusWriterNode.cpp
+++ b/Engine/Plugins/Runtime/Metasound/Source/MetasoundEngine/Private/MetasoundAudioBusWriterNode.cpp
@@ -233,22 +233,6 @@ namespace Metasound
 		{
 		}
 
-		virtual FDataReferenceCollection GetInputs() const override
-		{
-			// This should never be called. Bind(...) is called instead. This method
-			// exists as a stop-gap until the API can be deprecated and removed.
-			checkNoEntry();
-			return {};
-		}
-
-		virtual FDataReferenceCollection GetOutputs() const override
-		{
-			// This should never be called. Bind(...) is called instead. This method
-			// exists as a stop-gap until the API can be deprecated and removed.
-			checkNoEntry();
-			return {};
-		}
-
 		void Execute()
 		{
 			const FAudioBusProxyPtr& BusProxy = AudioBusAsset->GetAudioBusProxy();
@@ -341,14 +325,7 @@ namespace Metasound
 	};
 
 	template<uint32 NumChannels>
-	class TAudioBusWriterNode : public FNodeFacade
-	{
-	public:
-		TAudioBusWriterNode(const FNodeInitData& InitData)
-			: FNodeFacade(InitData.InstanceName, InitData.InstanceID, TFacadeOperatorClass<TAudioBusWriterOperator<NumChannels>>())
-		{
-		}
-	};
+	using TAudioBusWriterNode = TNodeFacade<TAudioBusWriterOperator<NumChannels>>;
 
 #define REGISTER_AUDIO_BUS_WRITER_NODE(ChannelCount) \
 	using FAudioBusWriterNode_##ChannelCount = TAudioBusWriterNode<ChannelCount>; \
diff --git a/Engine/Plugins/Runtime/Metasound/Source/MetasoundEngine/Private/MetasoundWaveInfoNode.cpp b/Engine/Plugins/Runtime/Metasound/Source/MetasoundEngine/Private/MetasoundWaveInfoNode.cpp
--- a/Engine/Plugins/Runtime/Metasound/Source/MetasoundEngine/Private/MetasoundWaveInfoNode.cpp
+++ b/Engine/Plugins/Runtime/Metasound/Source/MetasoundEngine/Private/MetasoundWaveInfoNode.cpp
@@ -47,8 +47,6 @@ namespace Metasound
 		static TUniquePtr<IOperator> CreateOperator(const FBuildOperatorParams& InParams, FBuildResults& OutResults);
 		virtual void BindInputs(FInputVertexInterfaceData& InOutVertexData) override;
 		virtual void BindOutputs(FOutputVertexInterfaceData& InOutVertexData) override;
-		virtual FDataReferenceCollection GetInputs() const override;
-		virtual FDataReferenceCollection GetOutputs() const override;
 		void Execute();
 		void Reset(const IOperator::FResetParams& InParams);
 
@@ -143,22 +141,6 @@ namespace Metasound
 		InOutVertexData.BindReadVertex(METASOUND_GET_PARAM_NAME(OutParamAssetPath), PathOutput);
 	}
 
-	FDataReferenceCollection FWaveInfoNodeOperator::GetInputs() const
-	{
-		// This should never be called. Bind(...) is called instead. This method
-		// exists as a stop-gap until the API can be deprecated and removed.
-		checkNoEntry();
-		return {};
-	}
-
-	FDataReferenceCollection FWaveInfoNodeOperator::GetOutputs() const
-	{
-		// This should never be called. Bind(...) is called instead. This method
-		// exists as a stop-gap until the API can be deprecated and removed.
-		checkNoEntry();
-		return {};
-	}
-
 	void FWaveInfoNodeOperator::Execute()
 	{
 		FSoundWaveProxyPtr ProxyPtr = (*WaveAsset).GetSoundWaveProxy();
@@ -197,22 +179,7 @@ namespace Metasound
 		Execute();
 	}
 
-	class FWaveInfoNode : public FNodeFacade
-	{
-	public:
-		// public node api needs to define two conversion constructors:
-		// (1: from FString)
-		FWaveInfoNode(const Metasound::FVertexName& InInstanceName, const FGuid& InInstanceID)
-			: FNodeFacade(InInstanceName, InInstanceID, TFacadeOperatorClass < FWaveInfoNodeOperator >())
-		{ }
-
-		// (2: From an NodeInitData struct)
-		FWaveInfoNode(const FNodeInitData& InInitData)
-			: FWaveInfoNode(InInitData.InstanceName, InInitData.InstanceID)
-		{ }
-
-	};
-
+	using FWaveInfoNode = TNodeFacade<FWaveInfoNodeOperator>;
 	METASOUND_REGISTER_NODE(FWaveInfoNode);
 
 } // namespace Metasound
diff --git a/Engine/Plugins/Runtime/Metasound/Source/MetasoundEngine/Private/MetasoundWavePlayerNode.cpp b/Engine/Plugins/Runtime/Metasound/Source/MetasoundEngine/Private/MetasoundWavePlayerNode.cpp
--- a/Engine/Plugins/Runtime/Metasound/Source/MetasoundEngine/Private/MetasoundWavePlayerNode.cpp
+++ b/Engine/Plugins/Runtime/Metasound/Source/MetasoundEngine/Private/MetasoundWavePlayerNode.cpp
@@ -329,6 +329,97 @@
 		TEXT("Default: 512"),
 		ECVF_Default);
 
+	namespace WavePlayerOperatorPrivate
+	{
+		template<typename AudioChannelConfigurationInfoType>
+		FVertexInterface GetVertexInterface()
+		{
+			using namespace WavePlayerVertexNames;
+
+			//Adds uncommonly used pins to the Advanced View, to reduce the size of the node.
+			FDataVertexMetadata InputLoopStartMetaData = METASOUND_GET_PARAM_METADATA(InputLoopStart);
+			InputLoopStartMetaData.bIsAdvancedDisplay = true;
+
+			FDataVertexMetadata InputLoopDurationMetaData = METASOUND_GET_PARAM_METADATA(InputLoopDuration);
+			InputLoopDurationMetaData.bIsAdvancedDisplay = true;
+
+			FDataVertexMetadata OutputTriggerOnNearlyDoneMetaData = METASOUND_GET_PARAM_METADATA(OutputTriggerOnNearlyDone);
+			OutputTriggerOnNearlyDoneMetaData.bIsAdvancedDisplay = true;
+
+			FDataVertexMetadata OutputTriggerOnCuePointMetaData = METASOUND_GET_PARAM_METADATA(OutputTriggerOnCuePoint);
+			OutputTriggerOnCuePointMetaData.bIsAdvancedDisplay = true;
+
+			FDataVertexMetadata OutputCuePointIDMetaData = METASOUND_GET_PARAM_METADATA(OutputCuePointID);
+			OutputCuePointIDMetaData.bIsAdvancedDisplay = true;
+
+			FDataVertexMetadata OutputCuePointLabelMetaData = METASOUND_GET_PARAM_METADATA(OutputCuePointLabel);
+			OutputCuePointLabelMetaData.bIsAdvancedDisplay = true;
+
+			FDataVertexMetadata OutputPlaybackLocationMetaData = METASOUND_GET_PARAM_METADATA(OutputPlaybackLocation);
+			OutputPlaybackLocationMetaData.bIsAdvancedDisplay = true;
+
+			// Workaround to override display name of OutputLoopRatio
+			FDataVertexMetadata OutputLoopRatioMetadata
+			{
+				METASOUND_GET_PARAM_TT(OutputLoopRatio), // description 
+				METASOUND_LOCTEXT("OutputLoopRatioNotPercentDisplayName", "Loop Ratio"), // display name  
+				true // Is Advanced Display
+			};
+
+			FVertexInterface VertexInterface(
+				FInputVertexInterface(
+					TInputDataVertex<FTrigger>(METASOUND_GET_PARAM_NAME_AND_METADATA(InputTriggerPlay)),
+					TInputDataVertex<FTrigger>(METASOUND_GET_PARAM_NAME_AND_METADATA(InputTriggerStop)),
+					TInputDataVertex<FWaveAsset>(METASOUND_GET_PARAM_NAME_AND_METADATA(InputWaveAsset)),
+					TInputDataVertex<FTime>(METASOUND_GET_PARAM_NAME_AND_METADATA(InputStartTime), 0.0f),
+					TInputDataVertex<float>(METASOUND_GET_PARAM_NAME_AND_METADATA(InputPitchShift), 0.0f),
+					TInputDataVertex<bool>(METASOUND_GET_PARAM_NAME_AND_METADATA(InputLoop), false),
+					TInputDataVertex<FTime>(METASOUND_GET_PARAM_NAME(InputLoopStart), InputLoopStartMetaData, 0.0f),
+					TInputDataVertex<FTime>(METASOUND_GET_PARAM_NAME(InputLoopDuration), InputLoopDurationMetaData, -1.0f)
+				),
+				FOutputVertexInterface(
+					TOutputDataVertex<FTrigger>(METASOUND_GET_PARAM_NAME_AND_METADATA(OutputTriggerOnPlay)),
+					TOutputDataVertex<FTrigger>(METASOUND_GET_PARAM_NAME_AND_METADATA(OutputTriggerOnDone)),
+					TOutputDataVertex<FTrigger>(METASOUND_GET_PARAM_NAME(OutputTriggerOnNearlyDone), OutputTriggerOnNearlyDoneMetaData),
+					TOutputDataVertex<FTrigger>(METASOUND_GET_PARAM_NAME_AND_METADATA(OutputTriggerOnLooped)),
+					TOutputDataVertex<FTrigger>(METASOUND_GET_PARAM_NAME(OutputTriggerOnCuePoint), OutputTriggerOnCuePointMetaData),
+					TOutputDataVertex<int32>(METASOUND_GET_PARAM_NAME(OutputCuePointID), OutputCuePointIDMetaData),
+					TOutputDataVertex<FString>(METASOUND_GET_PARAM_NAME(OutputCuePointLabel), OutputCuePointLabelMetaData),
+					TOutputDataVertex<float>(METASOUND_GET_PARAM_NAME(OutputLoopRatio), OutputLoopRatioMetadata),
+					TOutputDataVertex<float>(METASOUND_GET_PARAM_NAME(OutputPlaybackLocation), OutputPlaybackLocationMetaData),
+					TOutputDataVertex<FTime>(METASOUND_GET_PARAM_NAME_AND_METADATA(OutputPlaybackTime))
+				)
+			);
+
+			// Add audio outputs dependent upon source info.
+			for (const FOutputDataVertex& OutputDataVertex : AudioChannelConfigurationInfoType::GetAudioOutputs())
+			{
+				VertexInterface.GetOutputInterface().Add(OutputDataVertex);
+			}
+
+			return VertexInterface;
+		}
+
+		template<typename AudioChannelConfigurationInfoType>
+		FNodeClassMetadata GetNodeInfo()
+		{
+			FNodeClassMetadata Info;
+			Info.ClassName = { Metasound::EngineNodes::Namespace, TEXT("Wave Player"), AudioChannelConfigurationInfoType::GetVariantName() };
+			Info.MajorVersion = 1;
+			Info.MinorVersion = 0;
+			Info.DisplayName = AudioChannelConfigurationInfoType::GetNodeDisplayName();
+			Info.Description = METASOUND_LOCTEXT("Metasound_WavePlayerNodeDescription", "Plays a wave asset. The wave's channel configurations will be up or down mixed to match the wave players audio channel format.");
+			Info.Author = PluginAuthor;
+			Info.PromptIfMissing = PluginNodeMissingPrompt;
+			Info.DefaultInterface = GetVertexInterface<AudioChannelConfigurationInfoType>();
+			Info.Keywords = { METASOUND_LOCTEXT("WavePlayerSoundKeyword", "Sound"),
+				                METASOUND_LOCTEXT("WavePlayerCueKeyword", "Cue")
+			};
+
+			return Info;
+		}
+	}
+
 	/** MetaSound operator for the wave player node. */
 	class FWavePlayerOperator : public TExecutableOperator<FWavePlayerOperator>
 	{	
@@ -410,23 +501,6 @@
 			}
 		}
 
-		virtual FDataReferenceCollection GetInputs() const override
-		{
-			// This should never be called. Bind(...) is called instead. This method
-			// exists as a stop-gap until the API can be deprecated and removed.
-			checkNoEntry();
-			return {};
-		}
-
-		virtual FDataReferenceCollection GetOutputs() const override
-		{
-			// This should never be called. Bind(...) is called instead. This method
-			// exists as a stop-gap until the API can be deprecated and removed.
-			checkNoEntry();
-			return {};
-		}
-
-
 		void Execute()
 		{
 			using namespace WavePlayerNodePrivate;
@@ -1127,142 +1201,37 @@
 		TArray<FOutputDataVertex> OutputAudioVertices;
 	};
 
-	
 	template<typename AudioChannelConfigurationInfoType>
-	class TWavePlayerNode : public FNode
+	class TWavePlayerNode : public FBasicNode
 	{
 
 	public:
-		static FVertexInterface DeclareVertexInterface()
-		{
-			using namespace WavePlayerVertexNames;
-
-			//Adds uncommonly used pins to the Advanced View, to reduce the size of the node.
-			FDataVertexMetadata InputLoopStartMetaData = METASOUND_GET_PARAM_METADATA(InputLoopStart);
-			InputLoopStartMetaData.bIsAdvancedDisplay = true;
-
-			FDataVertexMetadata InputLoopDurationMetaData = METASOUND_GET_PARAM_METADATA(InputLoopDuration);
-			InputLoopDurationMetaData.bIsAdvancedDisplay = true;
-
-			FDataVertexMetadata OutputTriggerOnNearlyDoneMetaData = METASOUND_GET_PARAM_METADATA(OutputTriggerOnNearlyDone);
-			OutputTriggerOnNearlyDoneMetaData.bIsAdvancedDisplay = true;		
-
-			FDataVertexMetadata OutputTriggerOnCuePointMetaData = METASOUND_GET_PARAM_METADATA(OutputTriggerOnCuePoint);
-			OutputTriggerOnCuePointMetaData.bIsAdvancedDisplay = true;
-
-			FDataVertexMetadata OutputCuePointIDMetaData = METASOUND_GET_PARAM_METADATA(OutputCuePointID);
-			OutputCuePointIDMetaData.bIsAdvancedDisplay = true;
-
-			FDataVertexMetadata OutputCuePointLabelMetaData = METASOUND_GET_PARAM_METADATA(OutputCuePointLabel);
-			OutputCuePointLabelMetaData.bIsAdvancedDisplay = true;
-
-			FDataVertexMetadata OutputPlaybackLocationMetaData = METASOUND_GET_PARAM_METADATA(OutputPlaybackLocation);
-			OutputPlaybackLocationMetaData.bIsAdvancedDisplay = true;
-
-			// Workaround to override display name of OutputLoopRatio
-			static const FDataVertexMetadata OutputLoopRatioMetadata
-			{ 
-				METASOUND_GET_PARAM_TT(OutputLoopRatio), // description 
-				METASOUND_LOCTEXT("OutputLoopRatioNotPercentDisplayName", "Loop Ratio"), // display name  
-				true // Is Advanced Display
-			};
-
-			FVertexInterface VertexInterface(
-				FInputVertexInterface(
-					TInputDataVertex<FTrigger>(METASOUND_GET_PARAM_NAME_AND_METADATA(InputTriggerPlay)),
-					TInputDataVertex<FTrigger>(METASOUND_GET_PARAM_NAME_AND_METADATA(InputTriggerStop)),
-					TInputDataVertex<FWaveAsset>(METASOUND_GET_PARAM_NAME_AND_METADATA(InputWaveAsset)),
-					TInputDataVertex<FTime>(METASOUND_GET_PARAM_NAME_AND_METADATA(InputStartTime), 0.0f),
-					TInputDataVertex<float>(METASOUND_GET_PARAM_NAME_AND_METADATA(InputPitchShift), 0.0f),
-					TInputDataVertex<bool>(METASOUND_GET_PARAM_NAME_AND_METADATA(InputLoop), false),
-					TInputDataVertex<FTime>(METASOUND_GET_PARAM_NAME(InputLoopStart), InputLoopStartMetaData, 0.0f),
-					TInputDataVertex<FTime>(METASOUND_GET_PARAM_NAME(InputLoopDuration), InputLoopDurationMetaData, -1.0f)
-					),
-				FOutputVertexInterface(
-					TOutputDataVertex<FTrigger>(METASOUND_GET_PARAM_NAME_AND_METADATA(OutputTriggerOnPlay)),
-					TOutputDataVertex<FTrigger>(METASOUND_GET_PARAM_NAME_AND_METADATA(OutputTriggerOnDone)),
-					TOutputDataVertex<FTrigger>(METASOUND_GET_PARAM_NAME(OutputTriggerOnNearlyDone), OutputTriggerOnNearlyDoneMetaData),
-					TOutputDataVertex<FTrigger>(METASOUND_GET_PARAM_NAME_AND_METADATA(OutputTriggerOnLooped)),
-					TOutputDataVertex<FTrigger>(METASOUND_GET_PARAM_NAME(OutputTriggerOnCuePoint), OutputTriggerOnCuePointMetaData),
-					TOutputDataVertex<int32>(METASOUND_GET_PARAM_NAME(OutputCuePointID), OutputCuePointIDMetaData),
-					TOutputDataVertex<FString>(METASOUND_GET_PARAM_NAME(OutputCuePointLabel), OutputCuePointLabelMetaData),
-					TOutputDataVertex<float>(METASOUND_GET_PARAM_NAME(OutputLoopRatio), OutputLoopRatioMetadata),
-					TOutputDataVertex<float>(METASOUND_GET_PARAM_NAME(OutputPlaybackLocation), OutputPlaybackLocationMetaData),
-					TOutputDataVertex<FTime>(METASOUND_GET_PARAM_NAME_AND_METADATA(OutputPlaybackTime))
-				)
-			);
-
-			// Add audio outputs dependent upon source info.
-			for (const FOutputDataVertex& OutputDataVertex : AudioChannelConfigurationInfoType::GetAudioOutputs())
-			{
-				VertexInterface.GetOutputInterface().Add(OutputDataVertex);
-			}
-
-			return VertexInterface;
-		}
-
-		static const FNodeClassMetadata& GetNodeInfo()
-		{
-			auto InitNodeInfo = []() -> FNodeClassMetadata
-			{
-				FNodeClassMetadata Info;
-				Info.ClassName = { Metasound::EngineNodes::Namespace, TEXT("Wave Player"), AudioChannelConfigurationInfoType::GetVariantName() };
-				Info.MajorVersion = 1;
-				Info.MinorVersion = 0;
-				Info.DisplayName = AudioChannelConfigurationInfoType::GetNodeDisplayName();
-				Info.Description = METASOUND_LOCTEXT("Metasound_WavePlayerNodeDescription", "Plays a wave asset. The wave's channel configurations will be up or down mixed to match the wave players audio channel format.");
-				Info.Author = PluginAuthor;
-				Info.PromptIfMissing = PluginNodeMissingPrompt;
-				Info.DefaultInterface = DeclareVertexInterface();
-				Info.Keywords = { METASOUND_LOCTEXT("WavePlayerSoundKeyword", "Sound"),
-				                  METASOUND_LOCTEXT("WavePlayerCueKeyword", "Cue")
-				};
-
-				return Info;
-			};
-
-			static const FNodeClassMetadata Info = InitNodeInfo();
-
-			return Info;
-		}
-
-		TWavePlayerNode(const FVertexName& InName, const FGuid& InInstanceID)
-			:	FNode(InName, InInstanceID, GetNodeInfo())
-			,	Factory(MakeOperatorFactoryRef<FWavePlayerOperatorFactory>(AudioChannelConfigurationInfoType::GetAudioOutputs()))
-			,	Interface(DeclareVertexInterface())
-		{
-		}
-
 		TWavePlayerNode(const FNodeInitData& InInitData)
-			: TWavePlayerNode(InInitData.InstanceName, InInitData.InstanceID)
+			: TWavePlayerNode(FNodeData(InInitData.InstanceName, InInitData.InstanceID, WavePlayerOperatorPrivate::GetVertexInterface<AudioChannelConfigurationInfoType>()), MakeShared<const FNodeClassMetadata>(CreateNodeClassMetadata()))
 		{
 		}
 
+		TWavePlayerNode(FNodeData InNodeData, TSharedRef<const FNodeClassMetadata> InClassMetadata)
+			: FBasicNode(MoveTemp(InNodeData), MoveTemp(InClassMetadata))
+			, Factory(MakeOperatorFactoryRef<FWavePlayerOperatorFactory>(AudioChannelConfigurationInfoType::GetAudioOutputs()))
+		{
+		}
+
+
 		virtual ~TWavePlayerNode() = default;
 
+		static FNodeClassMetadata CreateNodeClassMetadata()
+		{
+			return WavePlayerOperatorPrivate::GetNodeInfo<AudioChannelConfigurationInfoType>();
+		}
+
 		virtual FOperatorFactorySharedRef GetDefaultOperatorFactory() const override
 		{
 			return Factory;
 		}
 
-		virtual const FVertexInterface& GetVertexInterface() const override
-		{
-			return Interface;
-		}
-
-		virtual bool SetVertexInterface(const FVertexInterface& InInterface) override
-		{
-			return InInterface == Interface;
-		}
-
-		virtual bool IsVertexInterfaceSupported(const FVertexInterface& InInterface) const override
-		{
-			return InInterface == Interface;
-		}
-
 	private:
 		FOperatorFactorySharedRef Factory;
-		FVertexInterface Interface;
 	};
 
 	struct FMonoAudioChannelConfigurationInfo
diff --git a/Engine/Plugins/Runtime/Metasound/Source/MetasoundEngine/Private/MetasoundWaveTableBankEvaluateNode.cpp b/Engine/Plugins/Runtime/Metasound/Source/MetasoundEngine/Private/MetasoundWaveTableBankEvaluateNode.cpp
--- a/Engine/Plugins/Runtime/Metasound/Source/MetasoundEngine/Private/MetasoundWaveTableBankEvaluateNode.cpp
+++ b/Engine/Plugins/Runtime/Metasound/Source/MetasoundEngine/Private/MetasoundWaveTableBankEvaluateNode.cpp
@@ -127,26 +127,6 @@ namespace Metasound
 			using namespace WaveTableBankEvaluateNode;
 			InOutVertexData.BindReadVertex(METASOUND_GET_PARAM_NAME(WaveTableBankEval_OutParam), OutWriteRef);
 		}
-		
-		virtual FDataReferenceCollection GetInputs() const override
-		{
-			// This should never be called. Bind(...) is called instead. This method
-			// exists as a stop-gap until the API can be deprecated and removed. 
-			checkNoEntry();
-
-			FDataReferenceCollection InputDataReferences;
-			return InputDataReferences;
-		}
-
-		virtual FDataReferenceCollection GetOutputs() const override
-		{
-			// This should never be called. Bind(...) is called instead. This method
-			// exists as a stop-gap until the API can be deprecated and removed. 
-			checkNoEntry();
-
-			FDataReferenceCollection OutputDataReferences;
-			return OutputDataReferences;
-		}
 
 		void Execute()
 		{
@@ -268,16 +248,7 @@ namespace Metasound
 		} CachedState;
 	};
 
-	class FMetasoundWaveTableBankEvaluateNode : public FNodeFacade
-	{
-	public:
-		FMetasoundWaveTableBankEvaluateNode(const FNodeInitData& InInitData)
-			: FNodeFacade(InInitData.InstanceName, InInitData.InstanceID, TFacadeOperatorClass<FMetasoundWaveTableBankEvaluateNodeOperator>())
-		{
-		}
-
-		virtual ~FMetasoundWaveTableBankEvaluateNode() = default;
-	};
+	using FMetasoundWaveTableBankEvaluateNode = TNodeFacade<FMetasoundWaveTableBankEvaluateNodeOperator>;
 
 	METASOUND_REGISTER_NODE(FMetasoundWaveTableBankEvaluateNode)
 } // namespace Metasound
diff --git a/Engine/Plugins/Runtime/Metasound/Source/MetasoundEngine/Private/MetasoundWaveTableEnvelopeNode.cpp b/Engine/Plugins/Runtime/Metasound/Source/MetasoundEngine/Private/MetasoundWaveTableEnvelopeNode.cpp
--- a/Engine/Plugins/Runtime/Metasound/Source/MetasoundEngine/Private/MetasoundWaveTableEnvelopeNode.cpp
+++ b/Engine/Plugins/Runtime/Metasound/Source/MetasoundEngine/Private/MetasoundWaveTableEnvelopeNode.cpp
@@ -156,22 +156,6 @@ namespace Metasound
 			InOutVertexData.BindReadVertex("Out", OutWriteRef);
 		}
 
-		virtual FDataReferenceCollection GetInputs() const override
-		{
-			// This should never be called. Bind(...) is called instead. This method
-			// exists as a stop-gap until the API can be deprecated and removed.
-			checkNoEntry();
-			return {};
-		}
-
-		virtual FDataReferenceCollection GetOutputs() const override
-		{
-			// This should never be called. Bind(...) is called instead. This method
-			// exists as a stop-gap until the API can be deprecated and removed.
-			checkNoEntry();
-			return {};
-		}
-
 		float GetStopValue(WaveTable::FWaveTableSampler::ESingleSampleMode InSampleMode) const
 		{
 			using namespace WaveTable;
@@ -355,17 +339,7 @@ namespace Metasound
 		bool bPaused = false;
 	};
 
-	class FMetasoundWaveTableEnvelopeNode : public FNodeFacade
-	{
-	public:
-		FMetasoundWaveTableEnvelopeNode(const FNodeInitData& InInitData)
-			: FNodeFacade(InInitData.InstanceName, InInitData.InstanceID, TFacadeOperatorClass<FMetasoundWaveTableEnvelopeNodeOperator>())
-		{
-		}
-
-		virtual ~FMetasoundWaveTableEnvelopeNode() = default;
-	};
-
+	using FMetasoundWaveTableEnvelopeNode = TNodeFacade<FMetasoundWaveTableEnvelopeNodeOperator>;
 	METASOUND_REGISTER_NODE(FMetasoundWaveTableEnvelopeNode)
 } // namespace Metasound
 
diff --git a/Engine/Plugins/Runtime/Metasound/Source/MetasoundEngine/Private/MetasoundWaveTableEvaluateNode.cpp b/Engine/Plugins/Runtime/Metasound/Source/MetasoundEngine/Private/MetasoundWaveTableEvaluateNode.cpp
--- a/Engine/Plugins/Runtime/Metasound/Source/MetasoundEngine/Private/MetasoundWaveTableEvaluateNode.cpp
+++ b/Engine/Plugins/Runtime/Metasound/Source/MetasoundEngine/Private/MetasoundWaveTableEvaluateNode.cpp
@@ -118,26 +118,6 @@ namespace Metasound
 			using namespace WaveTableEvaluateNode;
 			InOutVertexData.BindReadVertex(METASOUND_GET_PARAM_NAME(OutParam), OutWriteRef);
 		}
-		
-		virtual FDataReferenceCollection GetInputs() const override
-		{
-			// This should never be called. Bind(...) is called instead. This method
-			// exists as a stop-gap until the API can be deprecated and removed. 
-			checkNoEntry();
-
-			FDataReferenceCollection InputDataReferences;
-			return InputDataReferences;
-		}
-
-		virtual FDataReferenceCollection GetOutputs() const override
-		{
-			// This should never be called. Bind(...) is called instead. This method
-			// exists as a stop-gap until the API can be deprecated and removed. 
-			checkNoEntry();
-
-			FDataReferenceCollection OutputDataReferences;
-			return OutputDataReferences;
-		}
 
 		void Execute()
 		{
@@ -178,17 +158,7 @@ namespace Metasound
 		FFloatWriteRef OutWriteRef;
 	};
 
-	class FMetasoundWaveTableEvaluateNode : public FNodeFacade
-	{
-	public:
-		FMetasoundWaveTableEvaluateNode(const FNodeInitData& InInitData)
-			: FNodeFacade(InInitData.InstanceName, InInitData.InstanceID, TFacadeOperatorClass<FMetasoundWaveTableEvaluateNodeOperator>())
-		{
-		}
-
-		virtual ~FMetasoundWaveTableEvaluateNode() = default;
-	};
-
+	using FMetasoundWaveTableEvaluateNode = TNodeFacade<FMetasoundWaveTableEvaluateNodeOperator>;
 	METASOUND_REGISTER_NODE(FMetasoundWaveTableEvaluateNode)
 } // namespace Metasound
 
diff --git a/Engine/Plugins/Runtime/Metasound/Source/MetasoundEngine/Private/MetasoundWaveTableGetNode.cpp b/Engine/Plugins/Runtime/Metasound/Source/MetasoundEngine/Private/MetasoundWaveTableGetNode.cpp
--- a/Engine/Plugins/Runtime/Metasound/Source/MetasoundEngine/Private/MetasoundWaveTableGetNode.cpp
+++ b/Engine/Plugins/Runtime/Metasound/Source/MetasoundEngine/Private/MetasoundWaveTableGetNode.cpp
@@ -107,22 +107,6 @@ namespace Metasound
 			InOutVertexData.BindReadVertex("Out", OutTable);
 		}
 
-		virtual FDataReferenceCollection GetInputs() const override
-		{
-			// This should never be called. Bind(...) is called instead. This method
-			// exists as a stop-gap until the API can be deprecated and removed.
-			checkNoEntry();
-			return {};
-		}
-
-		virtual FDataReferenceCollection GetOutputs() const override
-		{
-			// This should never be called. Bind(...) is called instead. This method
-			// exists as a stop-gap until the API can be deprecated and removed.
-			checkNoEntry();
-			return {};
-		}
-
 	private:
 		void ComputeGainIndexData(const TArray<FWaveTableData>& WaveTables, EWaveTableSamplingMode InSampleMode, ::WaveTable::FWaveTable& OutputWaveTable)
 		{
@@ -286,17 +270,7 @@ namespace Metasound
 		FWaveTableWriteRef OutTable;
 	};
 
-	class FMetasoundWaveTableGetNode : public FNodeFacade
-	{
-	public:
-		FMetasoundWaveTableGetNode(const FNodeInitData& InInitData)
-			: FNodeFacade(InInitData.InstanceName, InInitData.InstanceID, TFacadeOperatorClass<FMetasoundWaveTableGetNodeOperator>())
-		{
-		}
-
-		virtual ~FMetasoundWaveTableGetNode() = default;
-	};
-
+	using FMetasoundWaveTableGetNode = TNodeFacade<FMetasoundWaveTableGetNodeOperator>;
 	METASOUND_REGISTER_NODE(FMetasoundWaveTableGetNode)
 } // namespace Metasound
 
diff --git a/Engine/Plugins/Runtime/Metasound/Source/MetasoundEngine/Private/MetasoundWaveTableOscillatorNode.cpp b/Engine/Plugins/Runtime/Metasound/Source/MetasoundEngine/Private/MetasoundWaveTableOscillatorNode.cpp
--- a/Engine/Plugins/Runtime/Metasound/Source/MetasoundEngine/Private/MetasoundWaveTableOscillatorNode.cpp
+++ b/Engine/Plugins/Runtime/Metasound/Source/MetasoundEngine/Private/MetasoundWaveTableOscillatorNode.cpp
@@ -239,17 +239,7 @@ namespace Metasound
 		TDataWriteReference<FAudioBuffer> OutBufferWriteRef;
 	};
 
-	class FMetasoundWaveTableOscillatorNode : public FNodeFacade
-	{
-	public:
-		FMetasoundWaveTableOscillatorNode(const FNodeInitData& InInitData)
-			: FNodeFacade(InInitData.InstanceName, InInitData.InstanceID, TFacadeOperatorClass<FMetasoundWaveTableOscillatorNodeOperator>())
-		{
-		}
-
-		virtual ~FMetasoundWaveTableOscillatorNode() = default;
-	};
-
+	using FMetasoundWaveTableOscillatorNode = TNodeFacade<FMetasoundWaveTableOscillatorNodeOperator>;
 	METASOUND_REGISTER_NODE(FMetasoundWaveTableOscillatorNode)
 } // namespace Metasound
 
diff --git a/Engine/Plugins/Runtime/Metasound/Source/MetasoundEngine/Private/MetasoundWaveTablePlayerNode.cpp b/Engine/Plugins/Runtime/Metasound/Source/MetasoundEngine/Private/MetasoundWaveTablePlayerNode.cpp
--- a/Engine/Plugins/Runtime/Metasound/Source/MetasoundEngine/Private/MetasoundWaveTablePlayerNode.cpp
+++ b/Engine/Plugins/Runtime/Metasound/Source/MetasoundEngine/Private/MetasoundWaveTablePlayerNode.cpp
@@ -182,22 +182,6 @@ namespace Metasound
 			InOutVertexData.BindReadVertex("OnFinished", OutOnFinishedRef);
 		}
 
-		virtual FDataReferenceCollection GetInputs() const override
-		{
-			// This should never be called. Bind(...) is called instead. This method
-			// exists as a stop-gap until the API can be deprecated and removed.
-			checkNoEntry();
-			return {};
-		}
-
-		virtual FDataReferenceCollection GetOutputs() const override
-		{
-			// This should never be called. Bind(...) is called instead. This method
-			// exists as a stop-gap until the API can be deprecated and removed.
-			checkNoEntry();
-			return {};
-		}
-
 		void Execute()
 		{
 			using namespace WaveTable;
@@ -347,17 +331,7 @@ namespace Metasound
 		TDataWriteReference<FTrigger> OutOnFinishedRef;
 	};
 
-	class FMetaSoundWaveTablePlayerNode : public FNodeFacade
-	{
-	public:
-		FMetaSoundWaveTablePlayerNode(const FNodeInitData& InInitData)
-			: FNodeFacade(InInitData.InstanceName, InInitData.InstanceID, TFacadeOperatorClass<FMetaSoundWaveTablePlayerNodeOperator>())
-		{
-		}
-
-		virtual ~FMetaSoundWaveTablePlayerNode() = default;
-	};
-
+	using FMetaSoundWaveTablePlayerNode = TNodeFacade<FMetaSoundWaveTablePlayerNodeOperator>;
 	METASOUND_REGISTER_NODE(FMetaSoundWaveTablePlayerNode)
 } // namespace Metasound
 
