Fix crash when playing a MetaSound preset while autosave occurs

diff --git a/Engine/Plugins/Runtime/Metasound/Source/MetasoundEditor/Private/MetasoundEditorGraphConnectionDrawingPolicy.cpp b/Engine/Plugins/Runtime/Metasound/Source/MetasoundEditor/Private/MetasoundEditorGraphConnectionDrawingPolicy.cpp
--- a/Engine/Plugins/Runtime/Metasound/Source/MetasoundEditor/Private/MetasoundEditorGraphConnectionDrawingPolicy.cpp
+++ b/Engine/Plugins/Runtime/Metasound/Source/MetasoundEditor/Private/MetasoundEditorGraphConnectionDrawingPolicy.cpp
@@ -219,8 +219,10 @@ namespace Metasound
 				if (const UEdGraphPin* ReroutePin = FGraphBuilder::FindReroutedOutputPin(OutputPin))
 				{
 					OutputVertex = FGraphBuilder::GetPinVertex(*Builder, ReroutePin, &OutputNode);
-					check(OutputNode);
-					EdgeStyle = Builder->FindConstEdgeStyle(OutputNode->GetID(), OutputVertex->Name);
+					if (OutputVertex && OutputNode)
+					{
+						EdgeStyle = Builder->FindConstEdgeStyle(OutputNode->GetID(), OutputVertex->Name);
+					}
 				}
 			}
 		}
