First step in reworking frontend graph builder.
#rnx

diff --git a/Engine/Plugins/Runtime/Metasound/Source/MetasoundFrontend/Private/MetasoundFrontendGraphBuilder.cpp b/Engine/Plugins/Runtime/Metasound/Source/MetasoundFrontend/Private/MetasoundFrontendGraphBuilder.cpp
--- a/Engine/Plugins/Runtime/Metasound/Source/MetasoundFrontend/Private/MetasoundFrontendGraphBuilder.cpp
+++ b/Engine/Plugins/Runtime/Metasound/Source/MetasoundFrontend/Private/MetasoundFrontendGraphBuilder.cpp
@@ -24,6 +24,44 @@ namespace Metasound::Frontend
 {
 	namespace GraphBuilderPrivate
 	{
+		struct FDefaultLiteralData
+		{
+			FGuid DestinationNodeID;
+			FGuid DestinationVertexID;
+			FVertexName DestinationVertexKey;
+			FName TypeName;
+			FLiteralNodeConstructorParams InitParams;
+		};
+
+		// Map of Input VertexID to variable data required to construct and connect default variable
+		using FNodeIDVertexID = TTuple<FGuid, FGuid>;
+		using FDependencyByIDMap = TMap<FGuid, const FMetasoundFrontendClass*>;
+		using FSharedNodeByIDMap = TMap<FGuid, TSharedPtr<const INode>>;
+		using FDefaultInputByIDMap = TMap<FNodeIDVertexID, FDefaultLiteralData>;
+
+		// Context used throughout entire graph build process
+		// (for both a root and nested subgraphs)
+		struct FBuildContext
+		{
+			FString DebugAssetName;
+			FDependencyByIDMap FrontendClasses;
+			FSharedNodeByIDMap Graphs;
+			const Frontend::IDataTypeRegistry& DataTypeRegistry;
+			const Frontend::FProxyDataCache& ProxyDataCache;
+			FGuid GraphId;
+		};
+
+		// Transient context used for building a specific graph
+		struct FBuildGraphContext
+		{
+			TUniquePtr<FFrontendGraph> Graph;
+			const FMetasoundFrontendGraphClass& GraphClass;
+			const FMetasoundFrontendGraph& PagedGraph;
+			FBuildContext& BuildContext;
+			FDefaultInputByIDMap DefaultInputs;
+		};
+
+
 		FNodeInitData CreateNodeInitData(const FMetasoundFrontendNode& InNode)
 		{
 			FNodeInitData InitData;
@@ -47,550 +85,678 @@ namespace Metasound::Frontend
 
 			return Frontend::DefaultPageID;
 		}
-	} // namespace GraphBuilderPrivate
-
 
-	TUniquePtr<INode> FGraphBuilder::CreateVariableNode(const FBuildContext& InContext, const FMetasoundFrontendNode& InNode, const FMetasoundFrontendGraph& InGraph)
-	{
-		using namespace Metasound::Frontend;
+		const FMetasoundFrontendLiteral* FindLiteralForInputVertex(const FVertexName& InVertexName, const FMetasoundFrontendNode& InNode, const FMetasoundFrontendClassInput& InNodeClassInput)
+		{
+			using namespace Frontend;
 
-		const FMetasoundFrontendVariable* FrontendVariable = FindVariableForVariableNode(InNode, InGraph);
+			// Default value priority is:
+			// 1. A value set directly on the node
+			// 2. A default value of the owning graph
+			// 3. A default value on the input node class.
 
-		if (nullptr != FrontendVariable)
-		{
-			IDataTypeRegistry& DataTypeRegistry = IDataTypeRegistry::Get();
-			const bool IsLiteralParsableByDataType = DataTypeRegistry.IsLiteralTypeSupported(FrontendVariable->TypeName, FrontendVariable->Literal.GetType());
+			const FMetasoundFrontendLiteral* Literal = nullptr;
 
-			if (IsLiteralParsableByDataType)
+			// Check for default value directly on node.
+			if (InNode.InputLiterals.Num())
 			{
-				FLiteral Literal = FrontendVariable->Literal.ToLiteral(FrontendVariable->TypeName, &InContext.DataTypeRegistry, &InContext.ProxyDataCache);
-
-				FVariableNodeConstructorParams InitParams =
+				const FMetasoundFrontendVertex* InputVertex = InNode.Interface.Inputs.FindByPredicate([&InVertexName](const FMetasoundFrontendVertex& Vertex) { return Vertex.Name == InVertexName; });
+				if (ensure(InputVertex))
 				{
-					InNode.Name,
-					InNode.GetID(),
-					MoveTemp(Literal)
-				};
-
-				return DataTypeRegistry.CreateVariableNode(FrontendVariable->TypeName, MoveTemp(InitParams));
+					const FMetasoundFrontendVertexLiteral* VertexLiteral = InNode.InputLiterals.FindByPredicate([&InputVertex](const FMetasoundFrontendVertexLiteral& InVertexLiteral) { return InVertexLiteral.VertexID == InputVertex->VertexID; });
+					if (VertexLiteral)
+					{
+						Literal = &VertexLiteral->Value;
+					}
+				}
 			}
-			else
+
+			// Check for default value on node class
+			if (nullptr == Literal)
 			{
-				UE_LOG(LogMetaSound, Error, TEXT("Cannot create variable node [NodeID:%s]. [Variable:%s] cannot be constructed with the provided literal type."), *InNode.GetID().ToString(), *FrontendVariable->Name.ToString());
+				const FGuid PageID = GraphBuilderPrivate::ResolveTargetPageID(InNodeClassInput);
+				const FMetasoundFrontendLiteral* DefaultLiteral = InNodeClassInput.FindConstDefault(PageID);
+				if (ensure(DefaultLiteral))
+				{
+					if (DefaultLiteral->IsValid())
+					{
+						Literal = DefaultLiteral;
+					}
+				}
 			}
+
+			return Literal;
 		}
-		else
+
+		const FMetasoundFrontendVariable* FindVariableForVariableNode(const FMetasoundFrontendNode& InVariableNode, const FMetasoundFrontendGraph& InGraph)
 		{
-			UE_LOG(LogMetaSound, Error, TEXT("Cannot create variable node [NodeID:%s]. No variable found for variable node."), *InNode.GetID().ToString());
+			const FGuid& DesiredID = InVariableNode.GetID();
+			return InGraph.Variables.FindByPredicate([&](const FMetasoundFrontendVariable& InVar) { return InVar.VariableNodeID == DesiredID; });
 		}
 
-		return TUniquePtr<INode>(nullptr);
-	}
-
-	TUniquePtr<INode> FGraphBuilder::CreateInputNode(const FBuildContext& InContext, const FMetasoundFrontendNode& InNode, const FMetasoundFrontendClass& InClass, const FMetasoundFrontendClassInput& InOwningGraphClassInput)
-	{
-		using namespace Metasound::Frontend;
+		TUniquePtr<INode> CreateVariableNode(const FBuildContext& InContext, const FMetasoundFrontendNode& InNode, const FMetasoundFrontendGraph& InGraph)
+		{
+			using namespace Metasound::Frontend;
 
-		const FMetasoundFrontendLiteral* FrontendLiteral = FindInputLiteralForInputNode(InNode, InClass, InOwningGraphClassInput);
+			const FMetasoundFrontendVariable* FrontendVariable = FindVariableForVariableNode(InNode, InGraph);
 
-		if (nullptr != FrontendLiteral)
-		{
-			if (ensure(InNode.Interface.Inputs.Num() == 1))
+			if (nullptr != FrontendVariable)
 			{
 				IDataTypeRegistry& DataTypeRegistry = IDataTypeRegistry::Get();
-				const FMetasoundFrontendVertex& InputVertex = InNode.Interface.Inputs[0];
-
-				const bool IsLiteralParsableByDataType = DataTypeRegistry.IsLiteralTypeSupported(InputVertex.TypeName, FrontendLiteral->GetType());
+				const bool IsLiteralParsableByDataType = DataTypeRegistry.IsLiteralTypeSupported(FrontendVariable->TypeName, FrontendVariable->Literal.GetType());
 
 				if (IsLiteralParsableByDataType)
 				{
-					FLiteral Literal = FrontendLiteral->ToLiteral(InputVertex.TypeName, &InContext.DataTypeRegistry, &InContext.ProxyDataCache);
+					FLiteral Literal = FrontendVariable->Literal.ToLiteral(FrontendVariable->TypeName, &InContext.DataTypeRegistry, &InContext.ProxyDataCache);
 
-					FInputNodeConstructorParams InitParams =
+					FVariableNodeConstructorParams InitParams =
 					{
 						InNode.Name,
 						InNode.GetID(),
-						InputVertex.Name,
 						MoveTemp(Literal)
 					};
 
-					ensureAlwaysMsgf(InOwningGraphClassInput.AccessType != EMetasoundFrontendVertexAccessType::Unset, TEXT("Graph Class Input cannot be set to access type of 'Unset'"));
-					if (InOwningGraphClassInput.AccessType == EMetasoundFrontendVertexAccessType::Reference)
-					{
-						return DataTypeRegistry.CreateInputNode(InputVertex.TypeName, MoveTemp(InitParams));
-					}
-					else // InOwningGraphClassInput.AccessType == EMetasoundFrontendVertexAccessType::Value
-					{
-						return DataTypeRegistry.CreateConstructorInputNode(InputVertex.TypeName, MoveTemp(InitParams));
-					}
+					return DataTypeRegistry.CreateVariableNode(FrontendVariable->TypeName, MoveTemp(InitParams));
 				}
 				else
 				{
-					UE_LOG(LogMetaSound, Error, TEXT("Cannot create input node [NodeID:%s]. [Vertex:%s] cannot be constructed with the provided literal type."), *InNode.GetID().ToString(), *InputVertex.Name.ToString());
+					UE_LOG(LogMetaSound, Error, TEXT("Cannot create variable node [NodeID:%s]. [Variable:%s] cannot be constructed with the provided literal type."), *InNode.GetID().ToString(), *FrontendVariable->Name.ToString());
 				}
 			}
-		}
-		else
-		{
-			UE_LOG(LogMetaSound, Error, TEXT("Cannot create input node [NodeID:%s]. No default literal set for input node."), *InNode.GetID().ToString());
-		}
+			else
+			{
+				UE_LOG(LogMetaSound, Error, TEXT("Cannot create variable node [NodeID:%s]. No variable found for variable node."), *InNode.GetID().ToString());
+			}
 
-		return TUniquePtr<INode>(nullptr);
-	}
+			return TUniquePtr<INode>(nullptr);
+		}
 
-	TUniquePtr<INode> FGraphBuilder::CreateOutputNode(const FBuildContext& InContext, const FMetasoundFrontendNode& InNode, const FMetasoundFrontendClass& InClass, FBuildGraphContext& InGraphContext, const TSet<FNodeIDVertexID>& InEdgeDestinations)
-	{
-		using namespace Metasound::Frontend;
+		const FMetasoundFrontendLiteral* FindInputLiteralForInputNode(const FMetasoundFrontendNode& InInputNode, const FMetasoundFrontendClass& InInputNodeClass, const FMetasoundFrontendClassInput& InOwningGraphClassInput)
+		{
+			using namespace Frontend;
 
-		check(InClass.Metadata.GetType() == EMetasoundFrontendClassType::Output);
-		check(InNode.ClassID == InClass.ID);
+			// Default value priority is:
+			// 1. A value set directly on the node
+			// 2. A default value of the owning graph
+			// 3. A default value on the input node class.
 
-		if (ensure(InNode.Interface.Outputs.Num() == 1))
-		{
-			const FMetasoundFrontendVertex& OutputVertex = InNode.Interface.Outputs[0];
+			const FMetasoundFrontendLiteral* Literal = nullptr;
 
-			FOutputNodeConstructorParams InitParams =
+			// Check for default value directly on node.
+			if (ensure(InInputNode.Interface.Inputs.Num() == 1))
 			{
-				InNode.Name,
-				InNode.GetID(),
-				OutputVertex.Name
-			};
+				const FMetasoundFrontendVertex& InputVertex = InInputNode.Interface.Inputs[0];
 
-			{
-				const FNodeInitData InitData = GraphBuilderPrivate::CreateNodeInitData(InNode);
-				TArray<FDefaultLiteralData> DefaultLiteralData = GetInputDefaultLiteralData(InContext, InNode, InitData, InEdgeDestinations);
-				for (FDefaultLiteralData& Data : DefaultLiteralData)
+				// Find input literal matching VertexID
+				const FMetasoundFrontendVertexLiteral* VertexLiteral = InInputNode.InputLiterals.FindByPredicate(
+					[&](const FMetasoundFrontendVertexLiteral& InVertexLiteral)
+					{
+						return InVertexLiteral.VertexID == InputVertex.VertexID;
+					}
+				);
+
+				if (nullptr != VertexLiteral)
 				{
-					InGraphContext.DefaultInputs.Emplace(FNodeIDVertexID { InNode.GetID(), Data.DestinationVertexID }, MoveTemp(Data));
+					Literal = &VertexLiteral->Value;
 				}
 			}
 
-			ensure(InClass.Interface.Outputs.Num() == 1);
-			if (InClass.Interface.Outputs[0].AccessType == EMetasoundFrontendVertexAccessType::Reference)
+			// Check for default value on owning graph.
+			if (nullptr == Literal)
 			{
-				return IDataTypeRegistry::Get().CreateOutputNode(OutputVertex.TypeName, MoveTemp(InitParams));
+				// Find Class Default that is not invalid
+				const FGuid PageID = GraphBuilderPrivate::ResolveTargetPageID(InOwningGraphClassInput);
+				const FMetasoundFrontendLiteral& DefaultLiteral = InOwningGraphClassInput.FindConstDefaultChecked(PageID);
+				if (DefaultLiteral.IsValid())
+				{
+					Literal = &DefaultLiteral;
+				}
 			}
-			else // InClass.Interface.Outputs[0].AccessType == EMetasoundFrontendVertexAccessType::Value
+
+			// Check for default value on input node class
+			if (nullptr == Literal && ensure(InInputNodeClass.Interface.Inputs.Num() == 1))
 			{
-				return IDataTypeRegistry::Get().CreateConstructorOutputNode(OutputVertex.TypeName, MoveTemp(InitParams));
+				const FMetasoundFrontendClassInput& InputNodeClassInput = InInputNodeClass.Interface.Inputs.Last();
+				const FGuid PageID = GraphBuilderPrivate::ResolveTargetPageID(InputNodeClassInput);
+				const FMetasoundFrontendLiteral& DefaultLiteral = InputNodeClassInput.FindConstDefaultChecked(PageID);
+				if (DefaultLiteral.IsValid())
+				{
+					Literal = &DefaultLiteral;
+				}
 			}
-		}
-		return TUniquePtr<INode>(nullptr);
-	}
 
-	TUniquePtr<INode> FGraphBuilder::CreateExternalNode(const FBuildContext& InContext, const FMetasoundFrontendNode& InNode, const FMetasoundFrontendClass& InClass, FBuildGraphContext& InGraphContext, const TSet<FNodeIDVertexID>& InEdgeDestinations)
-	{
-		using namespace Frontend;
+			return Literal;
+		}
 
-		check(InNode.ClassID == InClass.ID);
 
-		const FNodeInitData InitData = GraphBuilderPrivate::CreateNodeInitData(InNode);
+		TUniquePtr<INode> CreateInputNode(const FBuildContext& InContext, const FMetasoundFrontendNode& InNode, const FMetasoundFrontendClass& InClass, const FMetasoundFrontendClassInput& InOwningGraphClassInput)
 		{
-			TArray<FDefaultLiteralData> DefaultLiteralData = GetInputDefaultLiteralData(InContext, InNode, InitData, InEdgeDestinations);
-			for (FDefaultLiteralData& Data : DefaultLiteralData)
+			using namespace Metasound::Frontend;
+
+			const FMetasoundFrontendLiteral* FrontendLiteral = FindInputLiteralForInputNode(InNode, InClass, InOwningGraphClassInput);
+
+			if (nullptr != FrontendLiteral)
 			{
-				InGraphContext.DefaultInputs.Emplace(FNodeIDVertexID{ InNode.GetID(), Data.DestinationVertexID }, MoveTemp(Data));
-			}
-		}
+				if (ensure(InNode.Interface.Inputs.Num() == 1))
+				{
+					IDataTypeRegistry& DataTypeRegistry = IDataTypeRegistry::Get();
+					const FMetasoundFrontendVertex& InputVertex = InNode.Interface.Inputs[0];
 
-		// TODO: handle check to see if node interface conforms to class interface here. 
-		// TODO: check to see if external object supports class interface.
+					const bool IsLiteralParsableByDataType = DataTypeRegistry.IsLiteralTypeSupported(InputVertex.TypeName, FrontendLiteral->GetType());
 
-		const FNodeRegistryKey Key = FNodeRegistryKey(InClass.Metadata);
-		return FMetasoundFrontendRegistryContainer::Get()->CreateNode(Key, InitData);
-	}
+					if (IsLiteralParsableByDataType)
+					{
+						FLiteral Literal = FrontendLiteral->ToLiteral(InputVertex.TypeName, &InContext.DataTypeRegistry, &InContext.ProxyDataCache);
 
-	const FMetasoundFrontendClassInput* FGraphBuilder::FindClassInputForInputNode(const FMetasoundFrontendGraphClass& InOwningGraph, const FMetasoundFrontendNode& InInputNode)
-	{
-		// Input nodes should have exactly one input.
-		if (ensure(InInputNode.Interface.Inputs.Num() == 1))
-		{
-			const FName& TypeName = InInputNode.Interface.Inputs[0].TypeName;
+						FInputNodeConstructorParams InitParams =
+						{
+							InNode.Name,
+							InNode.GetID(),
+							InputVertex.Name,
+							MoveTemp(Literal)
+						};
 
-			auto IsMatchingInput = [&](const FMetasoundFrontendClassInput& GraphInput)
+						ensureAlwaysMsgf(InOwningGraphClassInput.AccessType != EMetasoundFrontendVertexAccessType::Unset, TEXT("Graph Class Input cannot be set to access type of 'Unset'"));
+						if (InOwningGraphClassInput.AccessType == EMetasoundFrontendVertexAccessType::Reference)
+						{
+							return DataTypeRegistry.CreateInputNode(InputVertex.TypeName, MoveTemp(InitParams));
+						}
+						else // InOwningGraphClassInput.AccessType == EMetasoundFrontendVertexAccessType::Value
+						{
+							return DataTypeRegistry.CreateConstructorInputNode(InputVertex.TypeName, MoveTemp(InitParams));
+						}
+					}
+					else
+					{
+						UE_LOG(LogMetaSound, Error, TEXT("Cannot create input node [NodeID:%s]. [Vertex:%s] cannot be constructed with the provided literal type."), *InNode.GetID().ToString(), *InputVertex.Name.ToString());
+					}
+				}
+			}
+			else
 			{
-				return (InInputNode.GetID() == GraphInput.NodeID);
-			};
+				UE_LOG(LogMetaSound, Error, TEXT("Cannot create input node [NodeID:%s]. No default literal set for input node."), *InNode.GetID().ToString());
+			}
 
-			return InOwningGraph.Interface.Inputs.FindByPredicate(IsMatchingInput);
+			return TUniquePtr<INode>(nullptr);
 		}
-		return nullptr;
-	}
 
-	const FMetasoundFrontendClassOutput* FGraphBuilder::FindClassOutputForOutputNode(const FMetasoundFrontendGraphClass& InOwningGraph, const FMetasoundFrontendNode& InOutputNode)
-	{
-		// Output nodes should have exactly one output
-		if (ensure(InOutputNode.Interface.Outputs.Num() == 1))
+		TArray<FDefaultLiteralData> GetInputDefaultLiteralData(const FBuildContext& InContext, const FMetasoundFrontendNode& InNode, const FNodeInitData& InInitData, const TSet<FNodeIDVertexID>& InEdgeDestinations)
 		{
-			const FName& TypeName = InOutputNode.Interface.Outputs[0].TypeName;
+			TArray<FDefaultLiteralData> DefaultLiteralData;
 
-			auto IsMatchingOutput = [&](const FMetasoundFrontendClassOutput& GraphOutput)
+			TArray<FMetasoundFrontendVertex> InputVertices = InNode.Interface.Inputs;
+			for (const FMetasoundFrontendVertexLiteral& Literal : InNode.InputLiterals)
 			{
-				return (InOutputNode.GetID() == GraphOutput.NodeID);
-			};
+				FLiteralNodeConstructorParams InitParams;
+				FName TypeName;
 
-			return InOwningGraph.Interface.Outputs.FindByPredicate(IsMatchingOutput);
-		}
-		return nullptr;
-	}
+				FGuid VertexID = Literal.VertexID;
+				FVertexName DestinationVertexName;
+				bool bRequiresDefault = false;
+				auto RemoveAndBuildParams = [&](const FMetasoundFrontendVertex& Vertex)
+				{
+					if (Vertex.VertexID == VertexID)
+					{
+						// Only build params and forward along to connect to dynamically generated
+						// literal node if edge is not explicitly connected to vertex as destination.
+						bRequiresDefault = !InEdgeDestinations.Contains({InNode.GetID(), Vertex.VertexID});
+						if (bRequiresDefault)
+						{
+							InitParams.Literal = Literal.Value.ToLiteral(Vertex.TypeName, &InContext.DataTypeRegistry, &InContext.ProxyDataCache);
+							InitParams.InstanceID = Frontend::CreateLocallyUniqueId();
+							InitParams.NodeName = "Literal";
+							TypeName = Vertex.TypeName;
 
-	const FMetasoundFrontendVariable* FGraphBuilder::FindVariableForVariableNode(const FMetasoundFrontendNode& InVariableNode, const FMetasoundFrontendGraph& InGraph)
-	{
-		const FGuid& DesiredID = InVariableNode.GetID();
-		return InGraph.Variables.FindByPredicate([&](const FMetasoundFrontendVariable& InVar) { return InVar.VariableNodeID == DesiredID; });
-	}
+							DestinationVertexName = Vertex.Name;
+						}
 
-	const FMetasoundFrontendLiteral* FGraphBuilder::FindInputLiteralForInputNode(const FMetasoundFrontendNode& InInputNode, const FMetasoundFrontendClass& InInputNodeClass, const FMetasoundFrontendClassInput& InOwningGraphClassInput)
-	{
-		using namespace Frontend;
+						return true;
+					}
+
+					return false;
+				};
 
-		// Default value priority is:
-		// 1. A value set directly on the node
-		// 2. A default value of the owning graph
-		// 3. A default value on the input node class.
+				const bool bRemoved = InputVertices.RemoveAllSwap(RemoveAndBuildParams, EAllowShrinking::No) > 0;
+				if (ensure(bRemoved))
+				{
+					if (bRequiresDefault)
+					{
+						DefaultLiteralData.Emplace(FDefaultLiteralData
+							{
+								InNode.GetID(),
+								VertexID,
+								DestinationVertexName,
+								TypeName,
+								MoveTemp(InitParams)
+							});
+					}
+				}
+			}
+
+			return DefaultLiteralData;
+		}
 
-		const FMetasoundFrontendLiteral* Literal = nullptr;
 
-		// Check for default value directly on node.
-		if (ensure(InInputNode.Interface.Inputs.Num() == 1))
+		TUniquePtr<INode> CreateOutputNode(const FBuildContext& InContext, const FMetasoundFrontendNode& InNode, const FMetasoundFrontendClass& InClass, FBuildGraphContext& InGraphContext, const TSet<FNodeIDVertexID>& InEdgeDestinations)
 		{
-			const FMetasoundFrontendVertex& InputVertex = InInputNode.Interface.Inputs[0];
+			using namespace Metasound::Frontend;
+
+			check(InClass.Metadata.GetType() == EMetasoundFrontendClassType::Output);
+			check(InNode.ClassID == InClass.ID);
+
+			if (ensure(InNode.Interface.Outputs.Num() == 1))
+			{
+				const FMetasoundFrontendVertex& OutputVertex = InNode.Interface.Outputs[0];
+
+				FOutputNodeConstructorParams InitParams =
+				{
+					InNode.Name,
+					InNode.GetID(),
+					OutputVertex.Name
+				};
 
-			// Find input literal matching VertexID
-			const FMetasoundFrontendVertexLiteral* VertexLiteral = InInputNode.InputLiterals.FindByPredicate(
-				[&](const FMetasoundFrontendVertexLiteral& InVertexLiteral)
 				{
-					return InVertexLiteral.VertexID == InputVertex.VertexID;
+					const FNodeInitData InitData = GraphBuilderPrivate::CreateNodeInitData(InNode);
+					TArray<FDefaultLiteralData> DefaultLiteralData = GetInputDefaultLiteralData(InContext, InNode, InitData, InEdgeDestinations);
+					for (FDefaultLiteralData& Data : DefaultLiteralData)
+					{
+						InGraphContext.DefaultInputs.Emplace(FNodeIDVertexID { InNode.GetID(), Data.DestinationVertexID }, MoveTemp(Data));
+					}
 				}
-			);
 
-			if (nullptr != VertexLiteral)
-			{
-				Literal = &VertexLiteral->Value;
+				ensure(InClass.Interface.Outputs.Num() == 1);
+				if (InClass.Interface.Outputs[0].AccessType == EMetasoundFrontendVertexAccessType::Reference)
+				{
+					return IDataTypeRegistry::Get().CreateOutputNode(OutputVertex.TypeName, MoveTemp(InitParams));
+				}
+				else // InClass.Interface.Outputs[0].AccessType == EMetasoundFrontendVertexAccessType::Value
+				{
+					return IDataTypeRegistry::Get().CreateConstructorOutputNode(OutputVertex.TypeName, MoveTemp(InitParams));
+				}
 			}
+			return TUniquePtr<INode>(nullptr);
 		}
 
-		// Check for default value on owning graph.
-		if (nullptr == Literal)
+		TUniquePtr<INode> CreateExternalNode(const FBuildContext& InContext, const FMetasoundFrontendNode& InNode, const FMetasoundFrontendClass& InClass, FBuildGraphContext& InGraphContext, const TSet<FNodeIDVertexID>& InEdgeDestinations)
 		{
-			// Find Class Default that is not invalid
-			const FGuid PageID = GraphBuilderPrivate::ResolveTargetPageID(InOwningGraphClassInput);
-			const FMetasoundFrontendLiteral& DefaultLiteral = InOwningGraphClassInput.FindConstDefaultChecked(PageID);
-			if (DefaultLiteral.IsValid())
+			using namespace Frontend;
+
+			check(InNode.ClassID == InClass.ID);
+
+			const FNodeInitData InitData = GraphBuilderPrivate::CreateNodeInitData(InNode);
 			{
-				Literal = &DefaultLiteral;
+				TArray<FDefaultLiteralData> DefaultLiteralData = GetInputDefaultLiteralData(InContext, InNode, InitData, InEdgeDestinations);
+				for (FDefaultLiteralData& Data : DefaultLiteralData)
+				{
+					InGraphContext.DefaultInputs.Emplace(FNodeIDVertexID{ InNode.GetID(), Data.DestinationVertexID }, MoveTemp(Data));
+				}
 			}
+
+			// TODO: handle check to see if node interface conforms to class interface here. 
+			// TODO: check to see if external object supports class interface.
+
+			const FNodeRegistryKey Key = FNodeRegistryKey(InClass.Metadata);
+			return FMetasoundFrontendRegistryContainer::Get()->CreateNode(Key, InitData);
 		}
 
-		// Check for default value on input node class
-		if (nullptr == Literal && ensure(InInputNodeClass.Interface.Inputs.Num() == 1))
+		const FMetasoundFrontendClassInput* FindClassInputForInputNode(const FMetasoundFrontendGraphClass& InOwningGraph, const FMetasoundFrontendNode& InInputNode)
 		{
-			const FMetasoundFrontendClassInput& InputNodeClassInput = InInputNodeClass.Interface.Inputs.Last();
-			const FGuid PageID = GraphBuilderPrivate::ResolveTargetPageID(InputNodeClassInput);
-			const FMetasoundFrontendLiteral& DefaultLiteral = InputNodeClassInput.FindConstDefaultChecked(PageID);
-			if (DefaultLiteral.IsValid())
+			// Input nodes should have exactly one input.
+			if (ensure(InInputNode.Interface.Inputs.Num() == 1))
 			{
-				Literal = &DefaultLiteral;
+				const FName& TypeName = InInputNode.Interface.Inputs[0].TypeName;
+
+				auto IsMatchingInput = [&](const FMetasoundFrontendClassInput& GraphInput)
+				{
+					return (InInputNode.GetID() == GraphInput.NodeID);
+				};
+
+				return InOwningGraph.Interface.Inputs.FindByPredicate(IsMatchingInput);
 			}
+			return nullptr;
 		}
 
-		return Literal;
-	}
-
-	bool FGraphBuilder::AddNodesToGraph(FBuildGraphContext& InGraphContext)
-	{
-		TSet<FNodeIDVertexID> GraphEdgeDestinations;
-		const TArray<FMetasoundFrontendEdge>& GraphEdges = InGraphContext.PagedGraph.Edges;
-		Algo::Transform(GraphEdges, GraphEdgeDestinations, [](const FMetasoundFrontendEdge& Edge) 
+		const FMetasoundFrontendClassOutput* FindClassOutputForOutputNode(const FMetasoundFrontendGraphClass& InOwningGraph, const FMetasoundFrontendNode& InOutputNode)
 		{
-			return FNodeIDVertexID{Edge.ToNodeID, Edge.ToVertexID};
-		});
+			// Output nodes should have exactly one output
+			if (ensure(InOutputNode.Interface.Outputs.Num() == 1))
+			{
+				const FName& TypeName = InOutputNode.Interface.Outputs[0].TypeName;
 
-		for (const FMetasoundFrontendNode& Node : InGraphContext.PagedGraph.Nodes)
+				auto IsMatchingOutput = [&](const FMetasoundFrontendClassOutput& GraphOutput)
+				{
+					return (InOutputNode.GetID() == GraphOutput.NodeID);
+				};
+
+				return InOwningGraph.Interface.Outputs.FindByPredicate(IsMatchingOutput);
+			}
+			return nullptr;
+		}
+
+		bool AddNodesToGraph(FBuildGraphContext& InGraphContext)
 		{
-			const FMetasoundFrontendClass* NodeClass = InGraphContext.BuildContext.FrontendClasses.FindRef(Node.ClassID);
+			TSet<FNodeIDVertexID> GraphEdgeDestinations;
+			const TArray<FMetasoundFrontendEdge>& GraphEdges = InGraphContext.PagedGraph.Edges;
+			Algo::Transform(GraphEdges, GraphEdgeDestinations, [](const FMetasoundFrontendEdge& Edge) 
+			{
+				return FNodeIDVertexID{Edge.ToNodeID, Edge.ToVertexID};
+			});
 
-			if (ensure(nullptr != NodeClass))
+			for (const FMetasoundFrontendNode& Node : InGraphContext.PagedGraph.Nodes)
 			{
-				switch (NodeClass->Metadata.GetType())
+				const FMetasoundFrontendClass* NodeClass = InGraphContext.BuildContext.FrontendClasses.FindRef(Node.ClassID);
+
+				if (ensure(nullptr != NodeClass))
 				{
-					case EMetasoundFrontendClassType::Input:
+					switch (NodeClass->Metadata.GetType())
 					{
-						const FMetasoundFrontendClassInput* ClassInput = FindClassInputForInputNode(InGraphContext.GraphClass, Node);
-
-						if (nullptr != ClassInput)
+						case EMetasoundFrontendClassType::Input:
 						{
-							TSharedPtr<const INode> InputNode(CreateInputNode(InGraphContext.BuildContext, Node, *NodeClass, *ClassInput).Release());
-							InGraphContext.Graph->AddInputNode(Node.GetID(), ClassInput->Name, InputNode);
+							const FMetasoundFrontendClassInput* ClassInput = FindClassInputForInputNode(InGraphContext.GraphClass, Node);
+
+							if (nullptr != ClassInput)
+							{
+								TSharedPtr<const INode> InputNode(CreateInputNode(InGraphContext.BuildContext, Node, *NodeClass, *ClassInput).Release());
+								InGraphContext.Graph->AddInputNode(Node.GetID(), ClassInput->Name, InputNode);
+							}
+							else
+							{
+								const FString GraphClassIDString = InGraphContext.GraphClass.ID.ToString();
+								UE_LOG(LogMetaSound, Error, TEXT("MetaSound '%s': Failed to match input node [NodeID:%s, NodeName:%s] to owning graph [ClassID:%s] input."), *InGraphContext.BuildContext.DebugAssetName, *Node.GetID().ToString(), *Node.Name.ToString(), *GraphClassIDString);
+								return false;
+							}
 						}
-						else
+						break;
+
+						case EMetasoundFrontendClassType::Output:
 						{
-							const FString GraphClassIDString = InGraphContext.GraphClass.ID.ToString();
-							UE_LOG(LogMetaSound, Error, TEXT("MetaSound '%s': Failed to match input node [NodeID:%s, NodeName:%s] to owning graph [ClassID:%s] input."), *InGraphContext.BuildContext.DebugAssetName, *Node.GetID().ToString(), *Node.Name.ToString(), *GraphClassIDString);
-							return false;
+							const FMetasoundFrontendClassOutput* ClassOutput = FindClassOutputForOutputNode(InGraphContext.GraphClass, Node);
+							if (nullptr != ClassOutput)
+							{
+								TSharedPtr<const INode> OutputNode(CreateOutputNode(InGraphContext.BuildContext, Node, *NodeClass, InGraphContext, GraphEdgeDestinations).Release());
+								InGraphContext.Graph->AddOutputNode(Node.GetID(), ClassOutput->Name, OutputNode);
+							}
+							else
+							{
+								const FString GraphClassIDString = InGraphContext.GraphClass.ID.ToString();
+								UE_LOG(LogMetaSound, Error, TEXT("MetaSound '%s': Failed to match output node [NodeID:%s, NodeName:%s] to owning graph [ClassID:%s] output."), *InGraphContext.BuildContext.DebugAssetName, *Node.GetID().ToString(), *Node.Name.ToString(), *GraphClassIDString);
+								return false;
+							}
 						}
-					}
-					break;
+						break;
 
-					case EMetasoundFrontendClassType::Output:
-					{
-						const FMetasoundFrontendClassOutput* ClassOutput = FindClassOutputForOutputNode(InGraphContext.GraphClass, Node);
-						if (nullptr != ClassOutput)
+						case EMetasoundFrontendClassType::Graph:
 						{
-							TSharedPtr<const INode> OutputNode(CreateOutputNode(InGraphContext.BuildContext, Node, *NodeClass, InGraphContext, GraphEdgeDestinations).Release());
-							InGraphContext.Graph->AddOutputNode(Node.GetID(), ClassOutput->Name, OutputNode);
+							const TSharedPtr<const INode> SubgraphPtr = InGraphContext.BuildContext.Graphs.FindRef(Node.ClassID);
+
+							if (SubgraphPtr.IsValid())
+							{
+								InGraphContext.Graph->AddNode(Node.GetID(), SubgraphPtr);
+							}
+							else
+							{
+								UE_LOG(LogMetaSound, Error, TEXT("MetaSound '%s': Failed to find subgraph for node [NodeID:%s, NodeName:%s, ClassID:%s]"), *InGraphContext.BuildContext.DebugAssetName, *Node.GetID().ToString(), *Node.Name.ToString(), *Node.ClassID.ToString());
+								return false;
+							}
 						}
-						else
+						break;
+
+						case EMetasoundFrontendClassType::Literal:
 						{
-							const FString GraphClassIDString = InGraphContext.GraphClass.ID.ToString();
-							UE_LOG(LogMetaSound, Error, TEXT("MetaSound '%s': Failed to match output node [NodeID:%s, NodeName:%s] to owning graph [ClassID:%s] output."), *InGraphContext.BuildContext.DebugAssetName, *Node.GetID().ToString(), *Node.Name.ToString(), *GraphClassIDString);
+							checkNoEntry(); // Unsupported.
 							return false;
 						}
-					}
-					break;
 
-					case EMetasoundFrontendClassType::Graph:
-					{
-						const TSharedPtr<const INode> SubgraphPtr = InGraphContext.BuildContext.Graphs.FindRef(Node.ClassID);
-
-						if (SubgraphPtr.IsValid())
+						case EMetasoundFrontendClassType::Variable:
 						{
-							InGraphContext.Graph->AddNode(Node.GetID(), SubgraphPtr);
+							TSharedPtr<const INode> VariableNode(CreateVariableNode(InGraphContext.BuildContext, Node, InGraphContext.PagedGraph).Release());
+							InGraphContext.Graph->AddNode(Node.GetID(), VariableNode);
 						}
-						else
+						break;
+
+						// Templates, variable accessors and mutators are
+						// constructed with the same parameters as external nodes.
+						case EMetasoundFrontendClassType::Template:
+						case EMetasoundFrontendClassType::VariableAccessor:
+						case EMetasoundFrontendClassType::VariableDeferredAccessor:
+						case EMetasoundFrontendClassType::VariableMutator:
+						case EMetasoundFrontendClassType::External:
+						default:
 						{
-							UE_LOG(LogMetaSound, Error, TEXT("MetaSound '%s': Failed to find subgraph for node [NodeID:%s, NodeName:%s, ClassID:%s]"), *InGraphContext.BuildContext.DebugAssetName, *Node.GetID().ToString(), *Node.Name.ToString(), *Node.ClassID.ToString());
-							return false;
+							TSharedPtr<const INode> ExternalNode(CreateExternalNode(InGraphContext.BuildContext, Node, *NodeClass, InGraphContext, GraphEdgeDestinations).Release());
+							InGraphContext.Graph->AddNode(Node.GetID(), ExternalNode);
 						}
+						break;
 					}
-					break;
-
-					case EMetasoundFrontendClassType::Literal:
-					{
-						checkNoEntry(); // Unsupported.
-						return false;
-					}
-
-					case EMetasoundFrontendClassType::Variable:
-					{
-						TSharedPtr<const INode> VariableNode(CreateVariableNode(InGraphContext.BuildContext, Node, InGraphContext.PagedGraph).Release());
-						InGraphContext.Graph->AddNode(Node.GetID(), VariableNode);
-					}
-					break;
-
-					// Templates, variable accessors and mutators are
-					// constructed with the same parameters as external nodes.
-					case EMetasoundFrontendClassType::Template:
-					case EMetasoundFrontendClassType::VariableAccessor:
-					case EMetasoundFrontendClassType::VariableDeferredAccessor:
-					case EMetasoundFrontendClassType::VariableMutator:
-					case EMetasoundFrontendClassType::External:
-					default:
-					{
-						TSharedPtr<const INode> ExternalNode(CreateExternalNode(InGraphContext.BuildContext, Node, *NodeClass, InGraphContext, GraphEdgeDestinations).Release());
-						InGraphContext.Graph->AddNode(Node.GetID(), ExternalNode);
-					}
-					break;
 				}
 			}
-		}
-
-		return true;
-	}
 
-	bool FGraphBuilder::AddEdgesToGraph(FBuildGraphContext& InGraphContext)
-	{
-		// Pair of frontend node and core node. The frontend node can be one of
-		// several types.
-		struct FCoreNodeAndFrontendVertex
-		{
-			const INode* Node = nullptr;
-			const FMetasoundFrontendVertex* Vertex = nullptr;
-		};
-
-		TMap<FNodeIDVertexID, FCoreNodeAndFrontendVertex> NodeSourcesByID;
-		TMap<FNodeIDVertexID, FCoreNodeAndFrontendVertex> NodeDestinationsByID;
+			return true;
+		}
 
-		// Add nodes to NodeID/VertexID map
-		for (const FMetasoundFrontendNode& Node : InGraphContext.PagedGraph.Nodes)
+		bool AddEdgesToGraph(FBuildGraphContext& InGraphContext)
 		{
-			const INode* CoreNode = InGraphContext.Graph->FindNode(Node.GetID());
-			if (nullptr == CoreNode)
+			// Pair of frontend node and core node. The frontend node can be one of
+			// several types.
+			struct FCoreNodeAndFrontendVertex
 			{
-				UE_LOG(LogMetaSound, Warning, TEXT("Metasound '%s': Could not find referenced node [Name:%s, NodeID:%s]"), *InGraphContext.BuildContext.DebugAssetName, *Node.Name.ToString(), *Node.GetID().ToString());
-				return false;
-			}
+				const INode* Node = nullptr;
+				const FMetasoundFrontendVertex* Vertex = nullptr;
+			};
 
-			for (const FMetasoundFrontendVertex& Vertex : Node.Interface.Inputs)
-			{
-				NodeDestinationsByID.Add(FNodeIDVertexID(Node.GetID(), Vertex.VertexID), FCoreNodeAndFrontendVertex({CoreNode, &Vertex}));
-			}
+			TMap<FNodeIDVertexID, FCoreNodeAndFrontendVertex> NodeSourcesByID;
+			TMap<FNodeIDVertexID, FCoreNodeAndFrontendVertex> NodeDestinationsByID;
 
-			for (const FMetasoundFrontendVertex& Vertex : Node.Interface.Outputs)
+			// Add nodes to NodeID/VertexID map
+			for (const FMetasoundFrontendNode& Node : InGraphContext.PagedGraph.Nodes)
 			{
-				NodeSourcesByID.Add(FNodeIDVertexID(Node.GetID(), Vertex.VertexID), FCoreNodeAndFrontendVertex({CoreNode, &Vertex}));
-			}
-		};
+				const INode* CoreNode = InGraphContext.Graph->FindNode(Node.GetID());
+				if (nullptr == CoreNode)
+				{
+					UE_LOG(LogMetaSound, Warning, TEXT("Metasound '%s': Could not find referenced node [Name:%s, NodeID:%s]"), *InGraphContext.BuildContext.DebugAssetName, *Node.Name.ToString(), *Node.GetID().ToString());
+					return false;
+				}
 
-		for (const FMetasoundFrontendEdge& Edge : InGraphContext.PagedGraph.Edges)
-		{
-			const FNodeIDVertexID DestinationKey(Edge.ToNodeID, Edge.ToVertexID);
-			const FCoreNodeAndFrontendVertex* DestinationNodeAndVertex = NodeDestinationsByID.Find(DestinationKey);
+				for (const FMetasoundFrontendVertex& Vertex : Node.Interface.Inputs)
+				{
+					NodeDestinationsByID.Add(FNodeIDVertexID(Node.GetID(), Vertex.VertexID), FCoreNodeAndFrontendVertex({CoreNode, &Vertex}));
+				}
 
-			if (nullptr == DestinationNodeAndVertex)
-			{
-				UE_LOG(LogMetaSound, Error, TEXT("MetaSound '%s': Failed to add edge. Could not find destination [NodeID:%s, VertexID:%s]"), *InGraphContext.BuildContext.DebugAssetName, *Edge.ToNodeID.ToString(), *Edge.ToVertexID.ToString());
-				return false;
-			}
+				for (const FMetasoundFrontendVertex& Vertex : Node.Interface.Outputs)
+				{
+					NodeSourcesByID.Add(FNodeIDVertexID(Node.GetID(), Vertex.VertexID), FCoreNodeAndFrontendVertex({CoreNode, &Vertex}));
+				}
+			};
 
-			if (nullptr == DestinationNodeAndVertex->Node)
+			for (const FMetasoundFrontendEdge& Edge : InGraphContext.PagedGraph.Edges)
 			{
-				UE_LOG(LogMetaSound, Warning, TEXT("MetaSound '%s': 'Failed to add edge. Null destination node [NodeID:%s]"), *InGraphContext.BuildContext.DebugAssetName, *Edge.ToNodeID.ToString());
-				return false;
-			}
+				const FNodeIDVertexID DestinationKey(Edge.ToNodeID, Edge.ToVertexID);
+				const FCoreNodeAndFrontendVertex* DestinationNodeAndVertex = NodeDestinationsByID.Find(DestinationKey);
 
-			const FNodeIDVertexID SourceKey(Edge.FromNodeID, Edge.FromVertexID);
-			const FCoreNodeAndFrontendVertex* SourceNodeAndVertex = NodeSourcesByID.Find(SourceKey);
+				if (nullptr == DestinationNodeAndVertex)
+				{
+					UE_LOG(LogMetaSound, Error, TEXT("MetaSound '%s': Failed to add edge. Could not find destination [NodeID:%s, VertexID:%s]"), *InGraphContext.BuildContext.DebugAssetName, *Edge.ToNodeID.ToString(), *Edge.ToVertexID.ToString());
+					return false;
+				}
 
-			if (nullptr == SourceNodeAndVertex)
-			{
-				UE_LOG(LogMetaSound, Error, TEXT("MetaSound '%s': Failed to add edge. Could not find source [NodeID:%s, VertexID:%s]"), *InGraphContext.BuildContext.DebugAssetName, *Edge.FromNodeID.ToString(), *Edge.FromVertexID.ToString());
-				return false;
-			}
+				if (nullptr == DestinationNodeAndVertex->Node)
+				{
+					UE_LOG(LogMetaSound, Warning, TEXT("MetaSound '%s': 'Failed to add edge. Null destination node [NodeID:%s]"), *InGraphContext.BuildContext.DebugAssetName, *Edge.ToNodeID.ToString());
+					return false;
+				}
 
-			if (nullptr == SourceNodeAndVertex->Node)
-			{
-				UE_LOG(LogMetaSound, Warning, TEXT("MetaSound '%s': Skipping edge. Null source node [NodeID:%s]"), *InGraphContext.BuildContext.DebugAssetName, *Edge.FromNodeID.ToString());
-				return false;
-			}
+				const FNodeIDVertexID SourceKey(Edge.FromNodeID, Edge.FromVertexID);
+				const FCoreNodeAndFrontendVertex* SourceNodeAndVertex = NodeSourcesByID.Find(SourceKey);
 
-			const INode* FromNode = SourceNodeAndVertex->Node;
-			const FVertexName FromVertexKey = SourceNodeAndVertex->Vertex->Name;
+				if (nullptr == SourceNodeAndVertex)
+				{
+					UE_LOG(LogMetaSound, Error, TEXT("MetaSound '%s': Failed to add edge. Could not find source [NodeID:%s, VertexID:%s]"), *InGraphContext.BuildContext.DebugAssetName, *Edge.FromNodeID.ToString(), *Edge.FromVertexID.ToString());
+					return false;
+				}
 
-			const INode* ToNode = DestinationNodeAndVertex->Node;
-			const FVertexName ToVertexKey = DestinationNodeAndVertex->Vertex->Name;
+				if (nullptr == SourceNodeAndVertex->Node)
+				{
+					UE_LOG(LogMetaSound, Warning, TEXT("MetaSound '%s': Skipping edge. Null source node [NodeID:%s]"), *InGraphContext.BuildContext.DebugAssetName, *Edge.FromNodeID.ToString());
+					return false;
+				}
 
-			bool bSuccess = InGraphContext.Graph->AddDataEdge(*FromNode, FromVertexKey,  *ToNode, ToVertexKey);
+				const INode* FromNode = SourceNodeAndVertex->Node;
+				const FVertexName FromVertexKey = SourceNodeAndVertex->Vertex->Name;
 
-			// If succeeded, remove input as viable vertex to construct default variable, as it has been superceded by a connection.
-			if (bSuccess)
-			{
-				FNodeIDVertexID DestinationPair { ToNode->GetInstanceID(), DestinationNodeAndVertex->Vertex->VertexID };
-				InGraphContext.DefaultInputs.Remove(DestinationPair);
-			}
-			else
-			{
-				UE_LOG(LogMetaSound, Error, TEXT("MetaSound '%s': Failed to connect edge from [NodeID:%s, VertexID:%s] to [NodeID:%s, VertexID:%s]"), *InGraphContext.BuildContext.DebugAssetName, *Edge.FromNodeID.ToString(), *Edge.FromVertexID.ToString(), *Edge.ToNodeID.ToString(), *Edge.ToVertexID.ToString());
-				return false;
-			}
-		}
+				const INode* ToNode = DestinationNodeAndVertex->Node;
+				const FVertexName ToVertexKey = DestinationNodeAndVertex->Vertex->Name;
 
-		return true;
-	}
+				bool bSuccess = InGraphContext.Graph->AddDataEdge(*FromNode, FromVertexKey,  *ToNode, ToVertexKey);
 
-	bool FGraphBuilder::AddDefaultInputLiterals(FBuildGraphContext& InGraphContext)
-	{
-		using namespace Metasound::Frontend;
-		using namespace LiteralNodeNames; 
+				// If succeeded, remove input as viable vertex to construct default variable, as it has been superceded by a connection.
+				if (bSuccess)
+				{
+					FNodeIDVertexID DestinationPair { ToNode->GetInstanceID(), DestinationNodeAndVertex->Vertex->VertexID };
+					InGraphContext.DefaultInputs.Remove(DestinationPair);
+				}
+				else
+				{
+					UE_LOG(LogMetaSound, Error, TEXT("MetaSound '%s': Failed to connect edge from [NodeID:%s, VertexID:%s] to [NodeID:%s, VertexID:%s]"), *InGraphContext.BuildContext.DebugAssetName, *Edge.FromNodeID.ToString(), *Edge.FromVertexID.ToString(), *Edge.ToNodeID.ToString(), *Edge.ToVertexID.ToString());
+					return false;
+				}
+			}
 
-		using FIterator = FDefaultInputByIDMap::TIterator;
+			return true;
+		}
 
-		for (FDefaultInputByIDMap::ElementType& Pair : InGraphContext.DefaultInputs)
+		bool AddDefaultInputLiterals(FBuildGraphContext& InGraphContext)
 		{
-			FDefaultLiteralData& LiteralData = Pair.Value;
-			const FGuid LiteralNodeID = Frontend::CreateLocallyUniqueId();
+			using namespace Metasound::Frontend;
+			using namespace LiteralNodeNames; 
 
-			// 1. Construct and add the default variable to the graph
-			{
-				TUniquePtr<const INode> DefaultLiteral = IDataTypeRegistry::Get().CreateLiteralNode(LiteralData.TypeName, MoveTemp(LiteralData.InitParams));
-				InGraphContext.Graph->AddNode(LiteralNodeID, TSharedPtr<const INode>(DefaultLiteral.Release()));
-			}
+			using FIterator = FDefaultInputByIDMap::TIterator;
 
-			// 2. Connect the default variable to the expected input
-			const INode* FromNode = InGraphContext.Graph->FindNode(LiteralNodeID);
-			if (nullptr == FromNode)
+			for (FDefaultInputByIDMap::ElementType& Pair : InGraphContext.DefaultInputs)
 			{
-				UE_LOG(LogMetaSound, Error, TEXT("MetaSound '%s': Failed to find node in graph [NodeID:%s]"), *InGraphContext.BuildContext.DebugAssetName, *LiteralNodeID.ToString());
-				return false;
-			}
-			const FVertexName& FromVertexName = METASOUND_GET_PARAM_NAME(OutputValue);
+				FDefaultLiteralData& LiteralData = Pair.Value;
+				const FGuid LiteralNodeID = Frontend::CreateLocallyUniqueId();
 
-			const INode* ToNode = InGraphContext.Graph->FindNode(LiteralData.DestinationNodeID);
-			if (nullptr == ToNode)
-			{
-				UE_LOG(LogMetaSound, Error, TEXT("MetaSound '%s': Failed to node in graph [NodeID:%s]"), *InGraphContext.BuildContext.DebugAssetName, *LiteralData.DestinationNodeID.ToString());
-				return false;
-			}
-			const FVertexName& ToVertexName = LiteralData.DestinationVertexKey;
+				// 1. Construct and add the default variable to the graph
+				{
+					TUniquePtr<const INode> DefaultLiteral = IDataTypeRegistry::Get().CreateLiteralNode(LiteralData.TypeName, MoveTemp(LiteralData.InitParams));
+					InGraphContext.Graph->AddNode(LiteralNodeID, TSharedPtr<const INode>(DefaultLiteral.Release()));
+				}
 
-			bool bSuccess = InGraphContext.Graph->AddDataEdge(*FromNode, FromVertexName, *ToNode, ToVertexName);
-			if (!bSuccess)
-			{
-				UE_LOG(LogMetaSound, Error, TEXT("MetaSound '%s': Failed to connect default variable edge: from '%s' to '%s'"), *InGraphContext.BuildContext.DebugAssetName, *FromVertexName.ToString(), *ToVertexName.ToString());
-				return false;
+				// 2. Connect the default variable to the expected input
+				const INode* FromNode = InGraphContext.Graph->FindNode(LiteralNodeID);
+				if (nullptr == FromNode)
+				{
+					UE_LOG(LogMetaSound, Error, TEXT("MetaSound '%s': Failed to find node in graph [NodeID:%s]"), *InGraphContext.BuildContext.DebugAssetName, *LiteralNodeID.ToString());
+					return false;
+				}
+				const FVertexName& FromVertexName = METASOUND_GET_PARAM_NAME(OutputValue);
+
+				const INode* ToNode = InGraphContext.Graph->FindNode(LiteralData.DestinationNodeID);
+				if (nullptr == ToNode)
+				{
+					UE_LOG(LogMetaSound, Error, TEXT("MetaSound '%s': Failed to node in graph [NodeID:%s]"), *InGraphContext.BuildContext.DebugAssetName, *LiteralData.DestinationNodeID.ToString());
+					return false;
+				}
+				const FVertexName& ToVertexName = LiteralData.DestinationVertexKey;
+
+				bool bSuccess = InGraphContext.Graph->AddDataEdge(*FromNode, FromVertexName, *ToNode, ToVertexName);
+				if (!bSuccess)
+				{
+					UE_LOG(LogMetaSound, Error, TEXT("MetaSound '%s': Failed to connect default variable edge: from '%s' to '%s'"), *InGraphContext.BuildContext.DebugAssetName, *FromVertexName.ToString(), *ToVertexName.ToString());
+					return false;
+				}
 			}
-		}
 
-		// Clear default inputs because literals have been moved out of default input map.
-		InGraphContext.DefaultInputs.Reset();
-		return true;
-	}
+			// Clear default inputs because literals have been moved out of default input map.
+			InGraphContext.DefaultInputs.Reset();
+			return true;
+		}
 
-	TArray<FGraphBuilder::FDefaultLiteralData> FGraphBuilder::GetInputDefaultLiteralData(const FBuildContext& InContext, const FMetasoundFrontendNode& InNode, const FNodeInitData& InInitData, const TSet<FNodeIDVertexID>& InEdgeDestinations)
-	{
-		TArray<FDefaultLiteralData> DefaultLiteralData;
 
-		TArray<FMetasoundFrontendVertex> InputVertices = InNode.Interface.Inputs;
-		for (const FMetasoundFrontendVertexLiteral& Literal : InNode.InputLiterals)
+		bool SortSubgraphDependencies(TArray<const FMetasoundFrontendGraphClass*>& Subgraphs)
 		{
-			FLiteralNodeConstructorParams InitParams;
-			FName TypeName;
-
-			FGuid VertexID = Literal.VertexID;
-			FVertexName DestinationVertexName;
-			bool bRequiresDefault = false;
-			auto RemoveAndBuildParams = [&](const FMetasoundFrontendVertex& Vertex)
+			// Helper for caching and querying subgraph dependencies
+			struct FSubgraphDependencyLookup
 			{
-				if (Vertex.VertexID == VertexID)
+				FSubgraphDependencyLookup(TArrayView<const FMetasoundFrontendGraphClass*> InGraphs)
 				{
-					// Only build params and forward along to connect to dynamically generated
-					// literal node if edge is not explicitly connected to vertex as destination.
-					bRequiresDefault = !InEdgeDestinations.Contains({InNode.GetID(), Vertex.VertexID});
-					if (bRequiresDefault)
+					// Map ClassID to graph pointer. 
+					TMap<FGuid, const FMetasoundFrontendGraphClass*> ClassIDAndGraph;
+					for (const FMetasoundFrontendGraphClass* Graph: InGraphs)
 					{
-						InitParams.Literal = Literal.Value.ToLiteral(Vertex.TypeName, &InContext.DataTypeRegistry, &InContext.ProxyDataCache);
-						InitParams.InstanceID = Frontend::CreateLocallyUniqueId();
-						InitParams.NodeName = "Literal";
-						TypeName = Vertex.TypeName;
+						ClassIDAndGraph.Add(Graph->ID, Graph);
+					}
 
-						DestinationVertexName = Vertex.Name;
+					// Cache subgraph dependencies.
+					for (const FMetasoundFrontendGraphClass* GraphClass : InGraphs)
+					{
+						for (const FMetasoundFrontendNode& Node : GraphClass->GetConstDefaultGraph().Nodes)
+						{
+							if (ClassIDAndGraph.Contains(Node.ClassID))
+							{
+								DependencyMap.Add(GraphClass, ClassIDAndGraph[Node.ClassID]);
+							}
+						}
 					}
+				}
 
-					return true;
+				TArray<const FMetasoundFrontendGraphClass*> operator()(const FMetasoundFrontendGraphClass* InParent) const
+				{
+					TArray<const FMetasoundFrontendGraphClass*> Dependencies;
+					DependencyMap.MultiFind(InParent, Dependencies);
+					return Dependencies;
 				}
 
-				return false;
+			private:
+
+				TMultiMap<const FMetasoundFrontendGraphClass*, const FMetasoundFrontendGraphClass*> DependencyMap;
 			};
 
-			const bool bRemoved = InputVertices.RemoveAllSwap(RemoveAndBuildParams, EAllowShrinking::No) > 0;
-			if (ensure(bRemoved))
+			bool bSuccess = Algo::TopologicalSort(Subgraphs, FSubgraphDependencyLookup(Subgraphs));
+			if (!bSuccess)
 			{
-				if (bRequiresDefault)
-				{
-					DefaultLiteralData.Emplace(FDefaultLiteralData
-						{
-							InNode.GetID(),
-							VertexID,
-							DestinationVertexName,
-							TypeName,
-							MoveTemp(InitParams)
-						});
-				}
+				UE_LOG(LogMetaSound, Error, TEXT("Failed to topologically sort subgraphs. Possible recursive subgraph dependency"));
 			}
+
+			return bSuccess;
 		}
 
-		return DefaultLiteralData;
-	}
+		TUniquePtr<FFrontendGraph> CreateGraphInternal(FBuildContext& InContext, const FMetasoundFrontendGraphClass& InGraphClass)
+		{
+			const FString GraphName = InContext.DebugAssetName;
+
+			const FGuid PageID = GraphBuilderPrivate::ResolveTargetPageID(InGraphClass);
+			const FMetasoundFrontendGraph& PageGraph = InGraphClass.FindConstGraphChecked(PageID);
+			FBuildGraphContext BuildGraphContext
+			{
+				MakeUnique<FFrontendGraph>(GraphName, InContext.GraphId),
+				InGraphClass,
+				PageGraph,
+				InContext
+			};
+
+			bool bSuccess = AddNodesToGraph(BuildGraphContext);
+
+			if (bSuccess)
+			{
+				bSuccess = AddEdgesToGraph(BuildGraphContext);
+			}
+
+			if (bSuccess)
+			{
+				bSuccess = AddDefaultInputLiterals(BuildGraphContext);
+			}
+
+			if (bSuccess)
+			{
+				return MoveTemp(BuildGraphContext.Graph);
+			}
+			else
+			{
+				return TUniquePtr<FFrontendGraph>(nullptr);
+			}
+		}
+	} // namespace GraphBuilderPrivate
 
 	/** Check that all dependencies are C++ class dependencies. */
 	bool FGraphBuilder::IsFlat(const FMetasoundFrontendDocument& InDocument)
@@ -637,94 +803,11 @@ PRAGMA_ENABLE_DEPRECATION_WARNINGS
 		return bIsEveryDependencyMet;
 	}
 
-	bool FGraphBuilder::SortSubgraphDependencies(TArray<const FMetasoundFrontendGraphClass*>& Subgraphs)
-	{
-		// Helper for caching and querying subgraph dependencies
-		struct FSubgraphDependencyLookup
-		{
-			FSubgraphDependencyLookup(TArrayView<const FMetasoundFrontendGraphClass*> InGraphs)
-			{
-				// Map ClassID to graph pointer. 
-				TMap<FGuid, const FMetasoundFrontendGraphClass*> ClassIDAndGraph;
-				for (const FMetasoundFrontendGraphClass* Graph: InGraphs)
-				{
-					ClassIDAndGraph.Add(Graph->ID, Graph);
-				}
-
-				// Cache subgraph dependencies.
-				for (const FMetasoundFrontendGraphClass* GraphClass : InGraphs)
-				{
-					for (const FMetasoundFrontendNode& Node : GraphClass->GetConstDefaultGraph().Nodes)
-					{
-						if (ClassIDAndGraph.Contains(Node.ClassID))
-						{
-							DependencyMap.Add(GraphClass, ClassIDAndGraph[Node.ClassID]);
-						}
-					}
-				}
-			}
-
-			TArray<const FMetasoundFrontendGraphClass*> operator()(const FMetasoundFrontendGraphClass* InParent) const
-			{
-				TArray<const FMetasoundFrontendGraphClass*> Dependencies;
-				DependencyMap.MultiFind(InParent, Dependencies);
-				return Dependencies;
-			}
-
-		private:
-
-			TMultiMap<const FMetasoundFrontendGraphClass*, const FMetasoundFrontendGraphClass*> DependencyMap;
-		};
-
-		bool bSuccess = Algo::TopologicalSort(Subgraphs, FSubgraphDependencyLookup(Subgraphs));
-		if (!bSuccess)
-		{
-			UE_LOG(LogMetaSound, Error, TEXT("Failed to topologically sort subgraphs. Possible recursive subgraph dependency"));
-		}
-
-		return bSuccess;
-	}
-
-	TUniquePtr<FFrontendGraph> FGraphBuilder::CreateGraph(FBuildContext& InContext, const FMetasoundFrontendGraphClass& InGraphClass)
-	{
-		using namespace Frontend;
-
-		const FString GraphName = InContext.DebugAssetName;
-
-		const FGuid PageID = GraphBuilderPrivate::ResolveTargetPageID(InGraphClass);
-		const FMetasoundFrontendGraph& PageGraph = InGraphClass.FindConstGraphChecked(PageID);
-		FBuildGraphContext BuildGraphContext
-		{
-			MakeUnique<FFrontendGraph>(GraphName, InContext.GraphId),
-			InGraphClass,
-			PageGraph,
-			InContext
-		};
-
-		bool bSuccess = AddNodesToGraph(BuildGraphContext);
-
-		if (bSuccess)
-		{
-			bSuccess = AddEdgesToGraph(BuildGraphContext);
-		}
-
-		if (bSuccess)
-		{
-			bSuccess = AddDefaultInputLiterals(BuildGraphContext);
-		}
-
-		if (bSuccess)
-		{
-			return MoveTemp(BuildGraphContext.Graph);
-		}
-		else
-		{
-			return TUniquePtr<FFrontendGraph>(nullptr);
-		}
-	}
 
 	TUniquePtr<FFrontendGraph> FGraphBuilder::CreateGraph(const FMetasoundFrontendGraphClass& InGraph, const TArray<FMetasoundFrontendGraphClass>& InSubgraphs, const TArray<FMetasoundFrontendClass>& InDependencies, const Frontend::FProxyDataCache& InProxyDataCache, const FString& InDebugAssetName, const FGuid InGraphId)
 	{
+		using namespace GraphBuilderPrivate;
+
 		FBuildContext Context
 		{
 			InDebugAssetName, 					// DebugAssetName
@@ -759,7 +842,7 @@ PRAGMA_ENABLE_DEPRECATION_WARNINGS
 		// Create each subgraph.
 		for (const FMetasoundFrontendGraphClass* FrontendSubgraphPtr : FrontendSubgraphPtrs)
 		{
-			TSharedPtr<const INode> Subgraph(CreateGraph(Context, *FrontendSubgraphPtr).Release());
+			TSharedPtr<const INode> Subgraph(CreateGraphInternal(Context, *FrontendSubgraphPtr).Release());
 			if (!Subgraph.IsValid())
 			{
 				UE_LOG(LogMetaSound, Warning, TEXT("Failed to create subgraph [SubgraphName: %s] in asset '%s'"), *FrontendSubgraphPtr->Metadata.GetClassName().ToString(), *InDebugAssetName);
@@ -772,7 +855,7 @@ PRAGMA_ENABLE_DEPRECATION_WARNINGS
 		}
 
 		// Create parent graph.
-		return CreateGraph(Context, InGraph);
+		return CreateGraphInternal(Context, InGraph);
 	}
 
 	TUniquePtr<FFrontendGraph> FGraphBuilder::CreateGraph(const FMetasoundFrontendDocument& InDocument, const FString& InDebugAssetName)
diff --git a/Engine/Plugins/Runtime/Metasound/Source/MetasoundFrontend/Public/MetasoundFrontendGraphBuilder.h b/Engine/Plugins/Runtime/Metasound/Source/MetasoundFrontend/Public/MetasoundFrontendGraphBuilder.h
--- a/Engine/Plugins/Runtime/Metasound/Source/MetasoundFrontend/Public/MetasoundFrontendGraphBuilder.h
+++ b/Engine/Plugins/Runtime/Metasound/Source/MetasoundFrontend/Public/MetasoundFrontendGraphBuilder.h
@@ -62,68 +62,5 @@ namespace Metasound::Frontend
 			const FString& InDebugAssetName,
 			const FGuid InGraphId = Frontend::CreateLocallyUniqueId());
 
-	private:
-		struct FDefaultLiteralData
-		{
-			FGuid DestinationNodeID;
-			FGuid DestinationVertexID;
-			FVertexName DestinationVertexKey;
-			FName TypeName;
-			FLiteralNodeConstructorParams InitParams;
-		};
-
-		// Map of Input VertexID to variable data required to construct and connect default variable
-		using FNodeIDVertexID = TTuple<FGuid, FGuid>;
-		using FDependencyByIDMap = TMap<FGuid, const FMetasoundFrontendClass*>;
-		using FSharedNodeByIDMap = TMap<FGuid, TSharedPtr<const INode>>;
-		using FDefaultInputByIDMap = TMap<FNodeIDVertexID, FDefaultLiteralData>;
-
-		// Context used throughout entire graph build process
-		// (for both a root and nested subgraphs)
-		struct FBuildContext
-		{
-			FString DebugAssetName;
-			FDependencyByIDMap FrontendClasses;
-			FSharedNodeByIDMap Graphs;
-			const Frontend::IDataTypeRegistry& DataTypeRegistry;
-			const Frontend::FProxyDataCache& ProxyDataCache;
-			FGuid GraphId;
-		};
-
-		// Transient context used for building a specific graph
-		struct FBuildGraphContext
-		{
-			TUniquePtr<FFrontendGraph> Graph;
-			const FMetasoundFrontendGraphClass& GraphClass;
-			const FMetasoundFrontendGraph& PagedGraph;
-			FBuildContext& BuildContext;
-			FDefaultInputByIDMap DefaultInputs;
-		};
-
-
-		static TArray<FDefaultLiteralData> GetInputDefaultLiteralData(const FBuildContext& InContext, const FMetasoundFrontendNode& InNode, const FNodeInitData& InInitData, const TSet<FNodeIDVertexID>& InEdgeDestinations);
-
-		static bool SortSubgraphDependencies(TArray<const FMetasoundFrontendGraphClass*>& Subgraphs);
-
-		static TUniquePtr<FFrontendGraph> CreateGraph(FBuildContext& InContext, const FMetasoundFrontendGraphClass& InSubgraph);
-
-		static const FMetasoundFrontendClassInput* FindClassInputForInputNode(const FMetasoundFrontendGraphClass& InOwningGraph, const FMetasoundFrontendNode& InInputNode);
-		static const FMetasoundFrontendClassOutput* FindClassOutputForOutputNode(const FMetasoundFrontendGraphClass& InOwningGraph, const FMetasoundFrontendNode& InOutputNode);
-		static const FMetasoundFrontendLiteral* FindInputLiteralForInputNode(const FMetasoundFrontendNode& InInputNode, const FMetasoundFrontendClass& InInputNodeClass, const FMetasoundFrontendClassInput& InOwningGraphClassInput);
-		static const FMetasoundFrontendVariable* FindVariableForVariableNode(const FMetasoundFrontendNode& InVariableNode, const FMetasoundFrontendGraph& InGraph);
-
-		static TUniquePtr<INode> CreateInputNode(const FBuildContext& InContext, const FMetasoundFrontendNode& InNode, const FMetasoundFrontendClass& InClass, const FMetasoundFrontendClassInput& InOwningGraphClassInput);
-		static TUniquePtr<INode> CreateOutputNode(const FBuildContext& InContext, const FMetasoundFrontendNode& InNode, const FMetasoundFrontendClass& InClass, FBuildGraphContext& InGraphContext, const TSet<FNodeIDVertexID>& InEdgeDestinations);
-		static TUniquePtr<INode> CreateVariableNode(const FBuildContext& InContext, const FMetasoundFrontendNode& InNode, const FMetasoundFrontendGraph& InGraph);
-		static TUniquePtr<INode> CreateExternalNode(const FBuildContext& InContext, const FMetasoundFrontendNode& InNode, const FMetasoundFrontendClass& InClass, FBuildGraphContext& InGraphContext, const TSet<FNodeIDVertexID>& InEdgeDestinations);
-
-		// Returns false on error
-		static bool AddNodesToGraph(FBuildGraphContext& InGraphContext);
-
-		// Returns false on error
-		static bool AddEdgesToGraph(FBuildGraphContext& InGraphContext);
-
-		// Returns false on error
-		static bool AddDefaultInputLiterals(FBuildGraphContext& InGraphContext);
 	};
 }
