UMetaSoundSource: Generators are stored in an Array with their AudioComponentID and InstanceID instead of in a map

[RN] minor framework/audio

diff --git a/Engine/Plugins/Runtime/Metasound/Source/MetasoundEngine/Private/MetasoundGeneratorHandle.cpp b/Engine/Plugins/Runtime/Metasound/Source/MetasoundEngine/Private/MetasoundGeneratorHandle.cpp
--- a/Engine/Plugins/Runtime/Metasound/Source/MetasoundEngine/Private/MetasoundGeneratorHandle.cpp
+++ b/Engine/Plugins/Runtime/Metasound/Source/MetasoundEngine/Private/MetasoundGeneratorHandle.cpp
@@ -67,12 +67,12 @@ namespace Metasound
 			}
 			
 			// Listen for the source creating a new generator
-			Handle->GeneratorCreatedDelegateHandle = Source->OnGeneratorInstanceCreated.AddSP(
+			Handle->GeneratorCreatedDelegateHandle = Source->OnGeneratorInstanceInfoCreated.AddSP(
 				Handle,
 				&FMetasoundGeneratorHandle::HandleGeneratorCreated);
 
 			// Listen for the generator being destroyed
-			Handle->GeneratorDestroyedDelegateHandle = Source->OnGeneratorInstanceDestroyed.AddSP(
+			Handle->GeneratorDestroyedDelegateHandle = Source->OnGeneratorInstanceInfoDestroyed.AddSP(
 				Handle,
 				&FMetasoundGeneratorHandle::HandleGeneratorDestroyed);
 
@@ -598,19 +598,15 @@ namespace Metasound
 		}		
 	}
 
-	void FMetasoundGeneratorHandle::HandleGeneratorCreated(
-		const uint64 InAudioComponentId,
-		TSharedPtr<FMetasoundGenerator> InGenerator)
+	void FMetasoundGeneratorHandle::HandleGeneratorCreated(const FGeneratorInstanceInfo& GeneratorInfo)
 	{
 		METASOUND_LLM_SCOPE;
 		METASOUND_TRACE_CPUPROFILER_EVENT_SCOPE(FMetasoundGeneratorHandle::HandleGeneratorCreated);
-		
-		check(InAudioComponentId != INDEX_NONE);
 					
-		if (InAudioComponentId == GetAudioComponentId())
+		if (GeneratorInfo.AudioComponentID == GetAudioComponentId())
 		{
 			// Set the generator on the game thread. We grab a weak pointer in case this gets destroyed while we wait.
-			ExecuteOnGameThread(UE_SOURCE_LOCATION, [WeakThis = AsWeak(), WeakGenerator = InGenerator.ToWeakPtr()]()
+			ExecuteOnGameThread(UE_SOURCE_LOCATION, [WeakThis = AsWeak(), WeakGenerator = GeneratorInfo.Generator]()
 			{
 				if (const TSharedPtr<FMetasoundGeneratorHandle> PinnedThis = WeakThis.Pin())
 				{
@@ -620,16 +616,12 @@ namespace Metasound
 		}
 	}
 
-	void FMetasoundGeneratorHandle::HandleGeneratorDestroyed(
-		const uint64 InAudioComponentId,
-		TSharedPtr<FMetasoundGenerator>)
+	void FMetasoundGeneratorHandle::HandleGeneratorDestroyed(const FGeneratorInstanceInfo& GeneratorInfo)
 	{
 		METASOUND_LLM_SCOPE;
 		METASOUND_TRACE_CPUPROFILER_EVENT_SCOPE(FMetasoundGeneratorHandle::HandleGeneratorDestroyed);
-		
-		check(InAudioComponentId != INDEX_NONE);
 
-		if (InAudioComponentId == GetAudioComponentId())
+		if (GeneratorInfo.AudioComponentID == GetAudioComponentId())
 		{
 			// Unset the generator on the game thread. We grab a weak pointer in case this gets destroyed while we wait.
 			ExecuteOnGameThread(UE_SOURCE_LOCATION, [WeakThis = AsWeak()]()
diff --git a/Engine/Plugins/Runtime/Metasound/Source/MetasoundEngine/Private/MetasoundSource.cpp b/Engine/Plugins/Runtime/Metasound/Source/MetasoundEngine/Private/MetasoundSource.cpp
--- a/Engine/Plugins/Runtime/Metasound/Source/MetasoundEngine/Private/MetasoundSource.cpp
+++ b/Engine/Plugins/Runtime/Metasound/Source/MetasoundEngine/Private/MetasoundSource.cpp
@@ -1100,7 +1100,8 @@ ISoundGeneratorPtr UMetaSoundSource::CreateSoundGenerator(const FSoundGeneratorI
 
 	if (Generator.IsValid())
 	{
-		TrackGenerator(InParams.AudioComponentId, Generator);
+		FGeneratorInstanceInfo Info = FGeneratorInstanceInfo(InParams.AudioComponentId, InParams.InstanceID, Generator);
+		TrackGenerator(MoveTemp(Info));
 	}
 
 	return ISoundGeneratorPtr(Generator);
@@ -1598,11 +1599,14 @@ const TArray<Metasound::FVertexName>& UMetaSoundSource::GetOutputAudioChannelOrd
 	}
 }
 
-void UMetaSoundSource::TrackGenerator(uint64 Id, TSharedPtr<Metasound::FMetasoundGenerator> Generator)
+void UMetaSoundSource::TrackGenerator(Metasound::FGeneratorInstanceInfo&& GeneratorInfo)
 {
 	FScopeLock GeneratorMapLock(&GeneratorMapCriticalSection);
-	Generators.Add(Id, Generator);
-	OnGeneratorInstanceCreated.Broadcast(Id, Generator);
+	Generators.Add(MoveTemp(GeneratorInfo));
+	PRAGMA_DISABLE_DEPRECATION_WARNINGS
+	OnGeneratorInstanceCreated.Broadcast(Generators.Last().AudioComponentID, Generators.Last().Generator.Pin());
+	PRAGMA_ENABLE_DEPRECATION_WARNINGS
+	OnGeneratorInstanceInfoCreated.Broadcast(Generators.Last());
 }
 
 void UMetaSoundSource::ForgetGenerator(ISoundGeneratorPtr Generator)
@@ -1610,29 +1614,50 @@ void UMetaSoundSource::ForgetGenerator(ISoundGeneratorPtr Generator)
 	using namespace Metasound;
 	FMetasoundGenerator* AsMetasoundGenerator = static_cast<FMetasoundGenerator*>(Generator.Get());
 	FScopeLock GeneratorMapLock(&GeneratorMapCriticalSection);
-	for (auto It = Generators.begin(); It != Generators.end(); ++It)
+	const uint32 Index = Generators.IndexOfByPredicate([AsMetasoundGenerator](const FGeneratorInstanceInfo& GeneratorInfo)
 	{
-		if ((*It).Value.HasSameObject(AsMetasoundGenerator))
-		{
-			OnGeneratorInstanceDestroyed.Broadcast((*It).Key, StaticCastSharedPtr<Metasound::FMetasoundGenerator>(Generator));
-			Generators.Remove((*It).Key);
-			return;
-		}
+		return GeneratorInfo.Generator.HasSameObject(AsMetasoundGenerator);
+	});
+	if (Index != INDEX_NONE)
+	{
+		PRAGMA_DISABLE_DEPRECATION_WARNINGS
+		OnGeneratorInstanceDestroyed.Broadcast(Generators[Index].AudioComponentID, Generators[Index].Generator.Pin());
+		PRAGMA_ENABLE_DEPRECATION_WARNINGS
+		OnGeneratorInstanceInfoDestroyed.Broadcast(Generators[Index]);
+		Generators.RemoveAtSwap(Index);
 	}
 }
 
 TWeakPtr<Metasound::FMetasoundGenerator> UMetaSoundSource::GetGeneratorForAudioComponent(uint64 ComponentId) const
 {
 	using namespace Metasound;
+	check(ComponentId != INDEX_NONE)
 	FScopeLock GeneratorMapLock(&GeneratorMapCriticalSection);
-	const TWeakPtr<FMetasoundGenerator>* Result = Generators.Find(ComponentId);
-	if (!Result)
+	const FGeneratorInstanceInfo* GeneratorInfo = Generators.FindByPredicate([ComponentId](const FGeneratorInstanceInfo& Info)
+	{
+		return Info.AudioComponentID == ComponentId;
+	});
+	if (GeneratorInfo != nullptr)
 	{
-		return TWeakPtr<FMetasoundGenerator>(nullptr);
+		return GeneratorInfo->Generator;
 	}
-	return *Result;
+	return TWeakPtr<FMetasoundGenerator>(nullptr);
 }
 
+TWeakPtr<Metasound::FMetasoundGenerator> UMetaSoundSource::GetGeneratorForInstanceID(uint64 InstanceId) const
+{
+	using namespace Metasound;
+	FScopeLock GeneratorMapLock(&GeneratorMapCriticalSection);
+	const FGeneratorInstanceInfo* GeneratorInfo = Generators.FindByPredicate([InstanceId](const FGeneratorInstanceInfo& Info)
+	{
+		return Info.InstanceID == InstanceId;
+	});
+	if (GeneratorInfo != nullptr)
+	{
+		return GeneratorInfo->Generator;
+	}
+	return TWeakPtr<FMetasoundGenerator>(nullptr);
+}
 
 bool UMetaSoundSource::IsDynamic() const
 {
diff --git a/Engine/Plugins/Runtime/Metasound/Source/MetasoundEngine/Public/MetasoundGeneratorHandle.h b/Engine/Plugins/Runtime/Metasound/Source/MetasoundEngine/Public/MetasoundGeneratorHandle.h
--- a/Engine/Plugins/Runtime/Metasound/Source/MetasoundEngine/Public/MetasoundGeneratorHandle.h
+++ b/Engine/Plugins/Runtime/Metasound/Source/MetasoundEngine/Public/MetasoundGeneratorHandle.h
@@ -20,6 +20,7 @@ namespace Metasound
 {
 	class FMetasoundGenerator;
 	struct FVertexInterfaceChange;
+	struct FGeneratorInstanceInfo;
 }
 
 DECLARE_DYNAMIC_DELEGATE_TwoParams(FOnMetasoundOutputValueChanged, FName, OutputName, const FMetaSoundOutput&, Output);
@@ -250,10 +251,10 @@ namespace Metasound
 		/**
 		 * Delegate management
 		 */
-		void HandleGeneratorCreated(uint64 InAudioComponentId, TSharedPtr<FMetasoundGenerator> Generator);
+		void HandleGeneratorCreated(const FGeneratorInstanceInfo& GeneratorInfo);
 		FDelegateHandle GeneratorCreatedDelegateHandle;
 
-		void HandleGeneratorDestroyed(uint64 InAudioComponentId, TSharedPtr<FMetasoundGenerator> Generator);
+		void HandleGeneratorDestroyed(const FGeneratorInstanceInfo& GeneratorInfo);
 		FDelegateHandle GeneratorDestroyedDelegateHandle;
 
 		void HandleGeneratorGraphSet();
diff --git a/Engine/Plugins/Runtime/Metasound/Source/MetasoundEngine/Public/MetasoundSource.h b/Engine/Plugins/Runtime/Metasound/Source/MetasoundEngine/Public/MetasoundSource.h
--- a/Engine/Plugins/Runtime/Metasound/Source/MetasoundEngine/Public/MetasoundSource.h
+++ b/Engine/Plugins/Runtime/Metasound/Source/MetasoundEngine/Public/MetasoundSource.h
@@ -62,11 +62,20 @@ namespace Metasound
 		using FCookedQualitySettings = FMetaSoundQualitySettings;
 	} // namespace SourcePrivate
 
+	struct FGeneratorInstanceInfo
+	{
+		uint64 AudioComponentID;
+		uint64 InstanceID;
+		TWeakPtr<FMetasoundGenerator> Generator;
+	};
 } // namespace Metasound
 
 
+UE_DEPRECATED(5.6, "Use FGeneratorInstanceInfoDelegate instead.")
 DECLARE_TS_MULTICAST_DELEGATE_TwoParams(FOnGeneratorInstanceCreated, uint64, TSharedPtr<Metasound::FMetasoundGenerator>);
+UE_DEPRECATED(5.6, "Use FGeneratorInstanceInfoDelegate instead.")
 DECLARE_TS_MULTICAST_DELEGATE_TwoParams(FOnGeneratorInstanceDestroyed, uint64, TSharedPtr<Metasound::FMetasoundGenerator>);
+DECLARE_TS_MULTICAST_DELEGATE_OneParam(FGeneratorInstanceInfoDelegate, const Metasound::FGeneratorInstanceInfo&);
 
 /**
  * This Metasound type can be played as an audio source.
@@ -272,11 +281,32 @@ public:
 	virtual bool IsOneShot() const override;
 	virtual bool EnableSubmixSendsOnPreview() const override { return true; }
 
+	/**
+	 * Attempt to get a generator for a given AudioComponentID.
+	 * You should only pass a valid AudioComponentID.
+	 * If the ID is set to INDEX_NONE, e.g. from using FAudioDevice::PlaySoundAtLocation, use GetGeneratorForInstanceID instead.
+	 * @param ComponentId The ID of the AudioComponent you want the generator from.
+	 * @return A generator if it exists. A null weak pointer otherwise.
+	 */
 	TWeakPtr<Metasound::FMetasoundGenerator> GetGeneratorForAudioComponent(uint64 ComponentId) const;
+
+	/**
+	 * Attempt to get a generator for a given InstanceID.
+	 * Use Audio::GetTransmitterID to construct an appropriate InstanceID.
+	 * @param InstanceId The hash of the ComponentID, WaveInstanceHash and the ActiveSound Play Order
+	 * @return A generator if it exists. A null weak pointer otherwise.
+	 */
+	TWeakPtr<Metasound::FMetasoundGenerator> GetGeneratorForInstanceID(uint64 InstanceId) const;
 	bool IsDynamic() const;
 
+	PRAGMA_DISABLE_DEPRECATION_WARNINGS
+	UE_DEPRECATED(5.6, "Use OnGeneratorInstanceInfoCreated instead.")
 	FOnGeneratorInstanceCreated OnGeneratorInstanceCreated;
+	UE_DEPRECATED(5.6, "Use OnGeneratorInstanceInfoDestroyed instead.")
 	FOnGeneratorInstanceDestroyed OnGeneratorInstanceDestroyed;
+	PRAGMA_ENABLE_DEPRECATION_WARNINGS
+	FGeneratorInstanceInfoDelegate OnGeneratorInstanceInfoCreated;
+	FGeneratorInstanceInfoDelegate OnGeneratorInstanceInfoDestroyed;
 	Metasound::FOperatorSettings GetOperatorSettings(Metasound::FSampleRate InDeviceSampleRate) const;
 
 	virtual const FMetasoundFrontendDocument& GetConstDocument() const override;
@@ -333,8 +363,8 @@ private:
 	Metasound::FMetasoundEnvironment CreateEnvironment(const Audio::FParameterTransmitterInitParams& InParams) const;
 
 	mutable FCriticalSection GeneratorMapCriticalSection;
-	TSortedMap<uint64, TWeakPtr<Metasound::FMetasoundGenerator>> Generators;
-	void TrackGenerator(uint64 Id, TSharedPtr<Metasound::FMetasoundGenerator> Generator);
+	TArray<Metasound::FGeneratorInstanceInfo> Generators;
+	void TrackGenerator(Metasound::FGeneratorInstanceInfo&& GeneratorInfo);
 	void ForgetGenerator(ISoundGeneratorPtr Generator);
 
 	static FRuntimeInput CreateRuntimeInput(const Metasound::Frontend::IDataTypeRegistry& Registry, const FMetasoundFrontendClassInput& Input, bool bCreateUObjectProxies);
