UMetaSoundSource: Fix Mac compile error with FGeneratorInstanceInfo constructor

#rnx

diff --git a/Engine/Plugins/Runtime/Metasound/Source/MetasoundEngine/Private/MetasoundSource.cpp b/Engine/Plugins/Runtime/Metasound/Source/MetasoundEngine/Private/MetasoundSource.cpp
--- a/Engine/Plugins/Runtime/Metasound/Source/MetasoundEngine/Private/MetasoundSource.cpp
+++ b/Engine/Plugins/Runtime/Metasound/Source/MetasoundEngine/Private/MetasoundSource.cpp
@@ -328,6 +328,11 @@ namespace Metasound
 		}
 
 	} // namespace SourcePrivate
+
+	FGeneratorInstanceInfo::FGeneratorInstanceInfo(uint64 InAudioComponentID, uint64 InInstanceID, TWeakPtr<FMetasoundGenerator> InGenerator) :
+	AudioComponentID(InAudioComponentID), InstanceID(InInstanceID), Generator(InGenerator)
+	{
+	}
 } // namespace Metasound
 
 FAutoConsoleVariableRef CVarMetaSoundEnableExperimentalRUntimePresetGraphInflation(
diff --git a/Engine/Plugins/Runtime/Metasound/Source/MetasoundEngine/Public/MetasoundSource.h b/Engine/Plugins/Runtime/Metasound/Source/MetasoundEngine/Public/MetasoundSource.h
--- a/Engine/Plugins/Runtime/Metasound/Source/MetasoundEngine/Public/MetasoundSource.h
+++ b/Engine/Plugins/Runtime/Metasound/Source/MetasoundEngine/Public/MetasoundSource.h
@@ -64,6 +64,8 @@ namespace Metasound
 
 	struct FGeneratorInstanceInfo
 	{
+		FGeneratorInstanceInfo() = default;
+		FGeneratorInstanceInfo(uint64 InAudioComponentID, uint64 InInstanceID, TWeakPtr<FMetasoundGenerator> InGenerator);
 		uint64 AudioComponentID;
 		uint64 InstanceID;
 		TWeakPtr<FMetasoundGenerator> Generator;
