Rework metasound data type registry to use new methods for node creation.
#rnx

diff --git a/Engine/Plugins/Runtime/Metasound/Source/MetasoundEngineTest/Private/EngineTestMetaSoundAutomatedNodeTest.cpp b/Engine/Plugins/Runtime/Metasound/Source/MetasoundEngineTest/Private/EngineTestMetaSoundAutomatedNodeTest.cpp
--- a/Engine/Plugins/Runtime/Metasound/Source/MetasoundEngineTest/Private/EngineTestMetaSoundAutomatedNodeTest.cpp
+++ b/Engine/Plugins/Runtime/Metasound/Source/MetasoundEngineTest/Private/EngineTestMetaSoundAutomatedNodeTest.cpp
@@ -819,10 +819,16 @@ namespace Metasound::EngineTest{
 
 					if (const Frontend::IDataTypeRegistryEntry* DataTypeEntry = DataTypeRegistry.FindDataTypeRegistryEntry(UnderlyingDataTypeName))
 					{
+						TSharedPtr<const FNodeClassMetadata> VariableNodeMetadata = DataTypeEntry->GetVariableClassMetadata();
+						if (!VariableNodeMetadata)
+						{
+							continue;
+						}
+
 						// Create a variable node for the underlying data type
-						FVariableNodeConstructorParams Params;
-						Params.Literal = DataTypeRegistry.CreateDefaultLiteral(UnderlyingDataTypeName);
-						TUniquePtr<INode> VariableNode = DataTypeEntry->CreateVariableNode(MoveTemp(Params));
+						TUniquePtr<INode> VariableNode = DataTypeEntry->CreateVariableNode(
+							DataTypeRegistry.CreateDefaultLiteral(UnderlyingDataTypeName),
+							FNodeData{"VariableNodeName", FGuid::NewGuid(), VariableNodeMetadata->DefaultInterface});
 
 						if (VariableNode.IsValid())
 						{
diff --git a/Engine/Plugins/Runtime/Metasound/Source/MetasoundFrontend/Private/MetasoundDataTypeRegistrationMacro.cpp b/Engine/Plugins/Runtime/Metasound/Source/MetasoundFrontend/Private/MetasoundDataTypeRegistrationMacro.cpp
new file mode 100644
--- /dev/null
+++ b/Engine/Plugins/Runtime/Metasound/Source/MetasoundFrontend/Private/MetasoundDataTypeRegistrationMacro.cpp
@@ -0,0 +1,151 @@
+// Copyright Epic Games, Inc. All Rights Reserved.
+
+#include "MetasoundDataTypeRegistrationMacro.h"
+
+namespace Metasound::Frontend
+{
+	namespace MetasoundDataTypeRegistrationPrivate
+	{
+		// Base registry entry for any data type
+		FDataTypeRegistryEntryBase::FDataTypeRegistryEntryBase(const Frontend::FDataTypeRegistryInfo& Info, const TSharedPtr<Frontend::IEnumDataTypeInterface>& EnumInterface)
+			: Info(Info), EnumInterface(EnumInterface)
+		{
+		}
+
+		FDataTypeRegistryEntryBase::~FDataTypeRegistryEntryBase() {}
+
+		const Frontend::FDataTypeRegistryInfo& FDataTypeRegistryEntryBase::GetDataTypeInfo() const
+		{
+			return Info;
+		}
+
+		TSharedPtr<const Frontend::IEnumDataTypeInterface> FDataTypeRegistryEntryBase::GetEnumInterface() const
+		{
+			return EnumInterface;
+		}
+
+		const FMetasoundFrontendClass& FDataTypeRegistryEntryBase::GetFrontendInputClass() const
+		{
+			return InputClass;
+		}
+
+		TSharedPtr<const FNodeClassMetadata> FDataTypeRegistryEntryBase::GetInputClassMetadata() const
+		{
+			return InputClassMetadata;
+		}
+
+		const FMetasoundFrontendClass& FDataTypeRegistryEntryBase::GetFrontendConstructorInputClass() const
+		{
+			return ConstructorInputClass;
+		}
+
+		TSharedPtr<const FNodeClassMetadata> FDataTypeRegistryEntryBase::GetConstructorInputClassMetadata() const
+		{
+			return ConstructorInputClassMetadata;
+		}
+
+		const FMetasoundFrontendClass& FDataTypeRegistryEntryBase::GetFrontendLiteralClass() const
+		{
+			return LiteralClass;
+		}
+
+		const FMetasoundFrontendClass& FDataTypeRegistryEntryBase::GetFrontendOutputClass() const
+		{
+			return OutputClass;
+		}
+
+		TSharedPtr<const FNodeClassMetadata> FDataTypeRegistryEntryBase::GetOutputClassMetadata() const
+		{
+			return OutputClassMetadata;
+		}
+
+		const FMetasoundFrontendClass& FDataTypeRegistryEntryBase::GetFrontendConstructorOutputClass() const
+		{
+			return ConstructorOutputClass;
+		}
+
+		TSharedPtr<const FNodeClassMetadata> FDataTypeRegistryEntryBase::GetConstructorOutputClassMetadata() const
+		{
+			return ConstructorOutputClassMetadata;
+		}
+
+		const FMetasoundFrontendClass& FDataTypeRegistryEntryBase::GetFrontendVariableClass() const
+		{
+			return VariableClass;
+		}
+
+		TSharedPtr<const FNodeClassMetadata> FDataTypeRegistryEntryBase::GetVariableClassMetadata() const
+		{
+			return VariableClassMetadata;
+		}
+
+		const FMetasoundFrontendClass& FDataTypeRegistryEntryBase::GetFrontendVariableMutatorClass() const
+		{
+			return VariableMutatorClass;
+		}
+
+		TSharedPtr<const FNodeClassMetadata> FDataTypeRegistryEntryBase::GetVariableMutatorClassMetadata() const
+		{
+			return VariableMutatorClassMetadata;
+		}
+
+		const FMetasoundFrontendClass& FDataTypeRegistryEntryBase::GetFrontendVariableAccessorClass() const
+		{
+			return VariableAccessorClass;
+		}
+
+		TSharedPtr<const FNodeClassMetadata> FDataTypeRegistryEntryBase::GetVariableAccessorClassMetadata() const
+		{
+			return VariableAccessorClassMetadata;
+		}
+
+		const FMetasoundFrontendClass& FDataTypeRegistryEntryBase::GetFrontendVariableDeferredAccessorClass() const
+		{
+			return VariableDeferredAccessorClass;
+		}
+
+		TSharedPtr<const FNodeClassMetadata> FDataTypeRegistryEntryBase::GetVariableDeferredAccessorClassMetadata() const
+		{
+			return VariableDeferredAccessorClassMetadata;
+		}
+
+		const Frontend::IParameterAssignmentFunction& FDataTypeRegistryEntryBase::GetRawAssignmentFunction() const
+		{
+			return RawAssignmentFunction;
+		}
+
+		Frontend::FLiteralAssignmentFunction FDataTypeRegistryEntryBase::GetLiteralAssignmentFunction() const
+		{
+			return LiteralAssignmentFunction;
+		}
+
+		TUniquePtr<INode> FDataTypeRegistryEntryBase::CreateOutputNode(FNodeData InNodeData) const
+		{
+			if (OutputClassMetadata)
+			{
+				const FOutputVertexInterface& Outputs = InNodeData.Interface.GetOutputInterface();
+				if (ensure(Outputs.Num() == 1))
+				{
+					const FVertexName& VertexName = Outputs.At(0).VertexName;
+					return MakeUnique<FOutputNode>(VertexName, MoveTemp(InNodeData), OutputClassMetadata.ToSharedRef());
+				}
+			}
+
+			return TUniquePtr<INode>(nullptr);
+		}
+			
+		TUniquePtr<INode> FDataTypeRegistryEntryBase::CreateConstructorOutputNode(FNodeData InNodeData) const
+		{
+			if (ConstructorOutputClassMetadata)
+			{
+				const FOutputVertexInterface& Outputs = InNodeData.Interface.GetOutputInterface();
+				if (ensure(Outputs.Num() == 1))
+				{
+					const FVertexName& VertexName = Outputs.At(0).VertexName;
+					return MakeUnique<FOutputNode>(VertexName, MoveTemp(InNodeData), ConstructorOutputClassMetadata.ToSharedRef());
+				}
+			}
+			return TUniquePtr<INode>(nullptr);
+		}
+	}
+}
diff --git a/Engine/Plugins/Runtime/Metasound/Source/MetasoundFrontend/Private/MetasoundFrontendDataTypeRegistry.cpp b/Engine/Plugins/Runtime/Metasound/Source/MetasoundFrontend/Private/MetasoundFrontendDataTypeRegistry.cpp
--- a/Engine/Plugins/Runtime/Metasound/Source/MetasoundFrontend/Private/MetasoundFrontendDataTypeRegistry.cpp
+++ b/Engine/Plugins/Runtime/Metasound/Source/MetasoundFrontend/Private/MetasoundFrontendDataTypeRegistry.cpp
@@ -126,7 +126,7 @@ namespace Metasound
 				{
 					return nullptr;
 				}
-
+		
 				virtual const TSet<FMetasoundFrontendVersion>* GetImplementedInterfaces() const override
 				{
 					return nullptr;
@@ -170,6 +170,11 @@ namespace Metasound
 
 				using FDataTypeNodeRegistryEntry::CreateNode;
 
+				virtual TUniquePtr<INode> CreateNode(FNodeData InNodeData) const override
+				{
+					return DataTypeEntry->CreateInputNode(MoveTemp(InNodeData));
+				}
+
 				UE_DEPRECATED(5.6, "Node classes should be constructed with FNodeData")
 				virtual TUniquePtr<INode> CreateNode(FDefaultNamedVertexWithLiteralNodeConstructorParams&& InParams) const override
 				{
@@ -179,7 +184,9 @@ namespace Metasound
 					InputParams.NodeName = InParams.NodeName;
 					InputParams.VertexName = InParams.VertexName;
 
+PRAGMA_DISABLE_DEPRECATION_WARNINGS
 					return DataTypeEntry->CreateInputNode(MoveTemp(InputParams));
+PRAGMA_ENABLE_DEPRECATION_WARNINGS
 				}
 
 			private:
@@ -208,6 +215,11 @@ namespace Metasound
 
 				using FDataTypeNodeRegistryEntry::CreateNode;
 
+				virtual TUniquePtr<INode> CreateNode(FNodeData InNodeData) const override
+				{
+					return DataTypeEntry->CreateConstructorInputNode(MoveTemp(InNodeData));
+				}
+
 				UE_DEPRECATED(5.6, "Node classes should be constructed with FNodeData")
 				virtual TUniquePtr<INode> CreateNode(FDefaultNamedVertexWithLiteralNodeConstructorParams&& InParams) const override
 				{
@@ -217,7 +229,9 @@ namespace Metasound
 					InputParams.NodeName = InParams.NodeName;
 					InputParams.VertexName = InParams.VertexName;
 
+PRAGMA_DISABLE_DEPRECATION_WARNINGS
 					return DataTypeEntry->CreateConstructorInputNode(MoveTemp(InputParams));
+PRAGMA_ENABLE_DEPRECATION_WARNINGS
 				}
 
 			private:
@@ -246,10 +260,9 @@ namespace Metasound
 
 				using FDataTypeNodeRegistryEntry::CreateNode;
 
-				UE_DEPRECATED(5.6, "Node classes should be constructed with FNodeData")
-				virtual TUniquePtr<INode> CreateNode(FDefaultNamedVertexNodeConstructorParams&& InParams) const override
+				virtual TUniquePtr<INode> CreateNode(FNodeData InNodeData) const override
 				{
-					return DataTypeEntry->CreateOutputNode(MoveTemp(InParams));
+					return DataTypeEntry->CreateOutputNode(MoveTemp(InNodeData));
 				}
 
 			private:
@@ -278,10 +291,9 @@ namespace Metasound
 
 				using FDataTypeNodeRegistryEntry::CreateNode;
 
-				UE_DEPRECATED(5.6, "Node classes should be constructed with FNodeData")
-				virtual TUniquePtr<INode> CreateNode(FDefaultNamedVertexNodeConstructorParams&& InParams) const override
+				virtual TUniquePtr<INode> CreateNode(FNodeData InNodeData) const override
 				{
-					return DataTypeEntry->CreateConstructorOutputNode(MoveTemp(InParams));
+					return DataTypeEntry->CreateConstructorOutputNode(MoveTemp(InNodeData));
 				}
 
 			private:
@@ -311,10 +323,20 @@ namespace Metasound
 
 				using FDataTypeNodeRegistryEntry::CreateNode;
 
-				UE_DEPRECATED(5.6, "Node classes should be constructed with FNodeData")
+				virtual TUniquePtr<INode> CreateNode(FNodeData) const override
+				{
+					// FLiteral node should not be used.  Also cannot create it
+					// with the data in FNodeData because it's missing the default
+					// literal.
+					return nullptr;
+				}
+
+				UE_DEPRECATED(5.6, "Literal nodes should no longer be created from the registry")
 				virtual TUniquePtr<INode> CreateNode(FDefaultLiteralNodeConstructorParams&& InParams) const override
 				{
+PRAGMA_DISABLE_DEPRECATION_WARNINGS	
 					return DataTypeEntry->CreateLiteralNode(MoveTemp(InParams));
+PRAGMA_ENABLE_DEPRECATION_WARNINGS	
 				}
 
 			private:
@@ -344,10 +366,19 @@ namespace Metasound
 
 				using FDataTypeNodeRegistryEntry::CreateNode;
 
-				UE_DEPRECATED(5.6, "Node classes should be constructed with FNodeData")
+				virtual TUniquePtr<INode> CreateNode(FNodeData) const override
+				{
+					// Cannot create variable node with FNodeData because it's missing
+					// the default literal.
+					return nullptr;
+				}
+
+				UE_DEPRECATED(5.6, "Node classes should be constructed from the data type registry") 
 				virtual TUniquePtr<INode> CreateNode(FDefaultLiteralNodeConstructorParams&& InParams) const override
 				{
+PRAGMA_DISABLE_DEPRECATION_WARNINGS
 					return DataTypeEntry->CreateVariableNode(MoveTemp(InParams));
+PRAGMA_ENABLE_DEPRECATION_WARNINGS
 				}
 
 			private:
@@ -376,7 +407,14 @@ namespace Metasound
 
 				virtual TUniquePtr<INode> CreateNode(const FNodeInitData& InParams) const override
 				{
+PRAGMA_DISABLE_DEPRECATION_WARNINGS
 					return DataTypeEntry->CreateVariableMutatorNode(InParams);
+PRAGMA_ENABLE_DEPRECATION_WARNINGS
+				}
+
+				virtual TUniquePtr<INode> CreateNode(FNodeData InNodeData) const override
+				{
+					return DataTypeEntry->CreateVariableMutatorNode(MoveTemp(InNodeData));
 				}
 
 			private:
@@ -405,7 +443,14 @@ namespace Metasound
 
 				virtual TUniquePtr<INode> CreateNode(const FNodeInitData& InParams) const override
 				{
+PRAGMA_DISABLE_DEPRECATION_WARNINGS
 					return DataTypeEntry->CreateVariableAccessorNode(InParams);
+PRAGMA_ENABLE_DEPRECATION_WARNINGS
+				}
+
+				virtual TUniquePtr<INode> CreateNode(FNodeData InNodeData) const override
+				{
+					return DataTypeEntry->CreateVariableAccessorNode(MoveTemp(InNodeData));
 				}
 
 			private:
@@ -434,7 +479,14 @@ namespace Metasound
 
 				virtual TUniquePtr<INode> CreateNode(const FNodeInitData& InParams) const override
 				{
+PRAGMA_DISABLE_DEPRECATION_WARNINGS
 					return DataTypeEntry->CreateVariableDeferredAccessorNode(InParams);
+PRAGMA_ENABLE_DEPRECATION_WARNINGS
+				}
+
+				virtual TUniquePtr<INode> CreateNode(FNodeData InNodeData) const override
+				{
+					return DataTypeEntry->CreateVariableDeferredAccessorNode(MoveTemp(InNodeData));
 				}
 
 			private:
@@ -499,18 +551,70 @@ namespace Metasound
 				virtual bool GetFrontendVariableMutatorClass(const FName& InDataType, FMetasoundFrontendClass& OutClass) const override;
 				virtual bool GetFrontendVariableAccessorClass(const FName& InDataType, FMetasoundFrontendClass& OutClass) const override;
 				virtual bool GetFrontendVariableDeferredAccessorClass(const FName& InDataType, FMetasoundFrontendClass& OutClass) const override;
+				virtual TSharedPtr<const FNodeClassMetadata> GetInputClassMetadata(const FName& InDataType) const override;
+				virtual TSharedPtr<const FNodeClassMetadata> GetConstructorInputClassMetadata(const FName& InDataType) const override;
+				virtual TSharedPtr<const FNodeClassMetadata> GetOutputClassMetadata(const FName& InDataType) const override;
+				virtual TSharedPtr<const FNodeClassMetadata> GetConstructorOutputClassMetadata(const FName& InDataType) const override;
+				virtual TSharedPtr<const FNodeClassMetadata> GetVariableClassMetadata(const FName& InDataType) const override;
+				virtual TSharedPtr<const FNodeClassMetadata> GetVariableMutatorClassMetadata(const FName& InDataType) const override;
+				virtual TSharedPtr<const FNodeClassMetadata> GetVariableAccessorClassMetadata(const FName& InDataType) const override;
+				virtual TSharedPtr<const FNodeClassMetadata> GetVariableDeferredAccessorClassMetadata(const FName& InDataType) const override;
 
 				// Create a new instance of a C++ implemented node from the registry.
+#if UE_METASOUND_ENABLE_5_6_NODE_REGISTRATION_DEPRECATION_WARNINGS
+				UE_DEPRECATED(5.6, "Create input nodes with FNodeData")
+#endif
 				virtual TUniquePtr<INode> CreateInputNode(const FName& InInputType, FInputNodeConstructorParams&& InParams) const override;
+				virtual TUniquePtr<INode> CreateInputNode(const FName& InInputType, FNodeData InNodeData) const override;
+
+#if UE_METASOUND_ENABLE_5_6_NODE_REGISTRATION_DEPRECATION_WARNINGS
+				UE_DEPRECATED(5.6, "Create constructor input nodes with FNodeData")
+#endif
 				virtual TUniquePtr<INode> CreateConstructorInputNode(const FName& InInputType, FInputNodeConstructorParams&& InParams) const override;
+				virtual TUniquePtr<INode> CreateConstructorInputNode(const FName& InInputType, FNodeData InNodeData) const override;
+
+				UE_DEPRECATED(5.6, "Direct creation of literal nodes will no longer be supported")
 				virtual TUniquePtr<INode> CreateLiteralNode(const FName& InLiteralType, FLiteralNodeConstructorParams&& InParams) const override;
-				virtual TUniquePtr<INode> CreateOutputNode(const FName& InOutputType, FOutputNodeConstructorParams&& InParams) const override;
-				virtual TUniquePtr<INode> CreateConstructorOutputNode(const FName& InOutputType, FOutputNodeConstructorParams&& InParams) const override;
-				virtual TUniquePtr<INode> CreateReceiveNode(const FName& InOutputType, const FNodeInitData& InParams) const override;
+
+#if UE_METASOUND_ENABLE_5_6_NODE_REGISTRATION_DEPRECATION_WARNINGS
+				UE_DEPRECATED(5.6, "Create output nodes with FNodeData")
+#endif
+				virtual TUniquePtr<INode> CreateOutputNode(const FName& InDataTypeName, FOutputNodeConstructorParams&& InParams) const override;
+				virtual TUniquePtr<INode> CreateOutputNode(const FName& InDataTypeName, FNodeData InNodeData) const override;
+#if UE_METASOUND_ENABLE_5_6_NODE_REGISTRATION_DEPRECATION_WARNINGS
+				UE_DEPRECATED(5.6, "Create output nodes with FNodeData")
+#endif
+				virtual TUniquePtr<INode> CreateConstructorOutputNode(const FName& InDataTypeName, FOutputNodeConstructorParams&& InParams) const override;
+				virtual TUniquePtr<INode> CreateConstructorOutputNode(const FName& InDataTypeName, FNodeData InNodeData) const override;
+
+#if UE_METASOUND_ENABLE_5_6_NODE_REGISTRATION_DEPRECATION_WARNINGS
+				UE_DEPRECATED(5.6, "Direct creation of receive nodes will no longer be supported")
+#endif
+				virtual TUniquePtr<INode> CreateReceiveNode(const FName& InDataTypeName, const FNodeInitData& InParams) const override;
+
+#if UE_METASOUND_ENABLE_5_6_NODE_REGISTRATION_DEPRECATION_WARNINGS
+				UE_DEPRECATED(5.6, "Create Variable Node with FNodeData")
+#endif
 				virtual TUniquePtr<INode> CreateVariableNode(const FName& InDataType, FVariableNodeConstructorParams&&) const override;
+				virtual TUniquePtr<INode> CreateVariableNode(const FName& InDataType, FLiteral InLiteral, FNodeData InNodeData) const override;
+
+#if UE_METASOUND_ENABLE_5_6_NODE_REGISTRATION_DEPRECATION_WARNINGS
+				UE_DEPRECATED(5.6, "Create variable mutator nodes with FNodeData")
+#endif
 				virtual TUniquePtr<INode> CreateVariableMutatorNode(const FName& InDataType, const FNodeInitData&) const override;
+				virtual TUniquePtr<INode> CreateVariableMutatorNode(const FName& InDataType, FNodeData InNodeData) const override;
+
+#if UE_METASOUND_ENABLE_5_6_NODE_REGISTRATION_DEPRECATION_WARNINGS
+				UE_DEPRECATED(5.6, "Create variable accessor nodes with FNodeData")
+#endif
 				virtual TUniquePtr<INode> CreateVariableAccessorNode(const FName& InDataType, const FNodeInitData&) const override;
+				virtual TUniquePtr<INode> CreateVariableAccessorNode(const FName& InDataType, FNodeData InNodeData) const override;
+
+#if UE_METASOUND_ENABLE_5_6_NODE_REGISTRATION_DEPRECATION_WARNINGS
+				UE_DEPRECATED(5.6, "Create variable deferred accessor nodes with FNodeData")
+#endif
 				virtual TUniquePtr<INode> CreateVariableDeferredAccessorNode(const FName& InDataType, const FNodeInitData&) const override;
+				virtual TUniquePtr<INode> CreateVariableDeferredAccessorNode(const FName& InDataType, FNodeData InNodeData) const override;
 
 			private:
 
@@ -1070,11 +1174,95 @@ namespace Metasound
 				return false;
 			}
 
+			TSharedPtr<const FNodeClassMetadata> FDataTypeRegistry::GetInputClassMetadata(const FName& InDataType) const
+			{
+				if (const IDataTypeRegistryEntry* Entry = FindDataTypeEntry(InDataType))
+				{
+					return Entry->GetInputClassMetadata();
+				}
+				return nullptr;
+			}
+
+			TSharedPtr<const FNodeClassMetadata> FDataTypeRegistry::GetConstructorInputClassMetadata(const FName& InDataType) const
+			{
+				if (const IDataTypeRegistryEntry* Entry = FindDataTypeEntry(InDataType))
+				{
+					return Entry->GetConstructorInputClassMetadata();
+				}
+				return nullptr;
+			}
+
+			TSharedPtr<const FNodeClassMetadata> FDataTypeRegistry::GetOutputClassMetadata(const FName& InDataType) const
+			{
+				if (const IDataTypeRegistryEntry* Entry = FindDataTypeEntry(InDataType))
+				{
+					return Entry->GetOutputClassMetadata();
+				}
+				return nullptr;
+			}
+
+			TSharedPtr<const FNodeClassMetadata> FDataTypeRegistry::GetConstructorOutputClassMetadata(const FName& InDataType) const
+			{
+				if (const IDataTypeRegistryEntry* Entry = FindDataTypeEntry(InDataType))
+				{
+					return Entry->GetConstructorOutputClassMetadata();
+				}
+				return nullptr;
+			}
+
+			TSharedPtr<const FNodeClassMetadata> FDataTypeRegistry::GetVariableClassMetadata(const FName& InDataType) const
+			{
+				if (const IDataTypeRegistryEntry* Entry = FindDataTypeEntry(InDataType))
+				{
+					return Entry->GetVariableClassMetadata();
+				}
+				return nullptr;
+			}
+
+			TSharedPtr<const FNodeClassMetadata> FDataTypeRegistry::GetVariableMutatorClassMetadata(const FName& InDataType) const
+			{
+				if (const IDataTypeRegistryEntry* Entry = FindDataTypeEntry(InDataType))
+				{
+					return Entry->GetVariableMutatorClassMetadata();
+				}
+				return nullptr;
+			}
+
+			TSharedPtr<const FNodeClassMetadata> FDataTypeRegistry::GetVariableAccessorClassMetadata(const FName& InDataType) const
+			{
+				if (const IDataTypeRegistryEntry* Entry = FindDataTypeEntry(InDataType))
+				{
+					return Entry->GetVariableAccessorClassMetadata();
+				}
+				return nullptr;
+			}
+
+			TSharedPtr<const FNodeClassMetadata> FDataTypeRegistry::GetVariableDeferredAccessorClassMetadata(const FName& InDataType) const
+			{
+				if (const IDataTypeRegistryEntry* Entry = FindDataTypeEntry(InDataType))
+				{
+					return Entry->GetVariableDeferredAccessorClassMetadata();
+				}
+				return nullptr;
+			}
+
+
 			TUniquePtr<INode> FDataTypeRegistry::CreateInputNode(const FName& InDataType, FInputNodeConstructorParams&& InParams) const
 			{
 				if (const IDataTypeRegistryEntry* Entry = FindDataTypeEntry(InDataType))
 				{
+PRAGMA_DISABLE_DEPRECATION_WARNINGS
 					return Entry->CreateInputNode(MoveTemp(InParams));
+PRAGMA_ENABLE_DEPRECATION_WARNINGS
+				}
+				return nullptr;
+			}
+
+			TUniquePtr<INode> FDataTypeRegistry::CreateInputNode(const FName& InDataType, FNodeData InNodeData) const
+			{
+				if (const IDataTypeRegistryEntry* Entry = FindDataTypeEntry(InDataType))
+				{
+					return Entry->CreateInputNode(MoveTemp(InNodeData));
 				}
 				return nullptr;
 			}
@@ -1083,7 +1271,18 @@ namespace Metasound
 			{
 				if (const IDataTypeRegistryEntry* Entry = FindDataTypeEntry(InDataType))
 				{
+PRAGMA_DISABLE_DEPRECATION_WARNINGS
 					return Entry->CreateConstructorInputNode(MoveTemp(InParams));
+PRAGMA_ENABLE_DEPRECATION_WARNINGS
+				}
+				return nullptr;
+			}
+
+			TUniquePtr<INode> FDataTypeRegistry::CreateConstructorInputNode(const FName& InDataType, FNodeData InNodeData) const
+			{
+				if (const IDataTypeRegistryEntry* Entry = FindDataTypeEntry(InDataType))
+				{
+					return Entry->CreateConstructorInputNode(MoveTemp(InNodeData));
 				}
 				return nullptr;
 			}
@@ -1092,7 +1291,9 @@ namespace Metasound
 			{
 				if (const IDataTypeRegistryEntry* Entry = FindDataTypeEntry(InDataType))
 				{
+PRAGMA_DISABLE_DEPRECATION_WARNINGS
 					return Entry->CreateLiteralNode(MoveTemp(InParams));
+PRAGMA_ENABLE_DEPRECATION_WARNINGS
 				}
 				return nullptr;
 			}
@@ -1101,7 +1302,18 @@ namespace Metasound
 			{
 				if (const IDataTypeRegistryEntry* Entry = FindDataTypeEntry(InDataType))
 				{
+PRAGMA_DISABLE_DEPRECATION_WARNINGS
 					return Entry->CreateOutputNode(MoveTemp(InParams));
+PRAGMA_ENABLE_DEPRECATION_WARNINGS
+				}
+				return nullptr;
+			}
+
+			TUniquePtr<INode> FDataTypeRegistry::CreateOutputNode(const FName& InDataTypeName, FNodeData InNodeData) const
+			{
+				if (const IDataTypeRegistryEntry* Entry = FindDataTypeEntry(InDataTypeName))
+				{
+					return Entry->CreateOutputNode(MoveTemp(InNodeData));
 				}
 				return nullptr;
 			}
@@ -1110,7 +1322,18 @@ namespace Metasound
 			{
 				if (const IDataTypeRegistryEntry* Entry = FindDataTypeEntry(InDataType))
 				{
+PRAGMA_DISABLE_DEPRECATION_WARNINGS
 					return Entry->CreateConstructorOutputNode(MoveTemp(InParams));
+PRAGMA_ENABLE_DEPRECATION_WARNINGS
+				}
+				return nullptr;
+			}
+
+			TUniquePtr<INode> FDataTypeRegistry::CreateConstructorOutputNode(const FName& InDataTypeName, FNodeData InNodeData) const
+			{
+				if (const IDataTypeRegistryEntry* Entry = FindDataTypeEntry(InDataTypeName))
+				{
+					return Entry->CreateConstructorOutputNode(MoveTemp(InNodeData));
 				}
 				return nullptr;
 			}
@@ -1119,7 +1342,9 @@ namespace Metasound
 			{
 				if (const IDataTypeRegistryEntry* Entry = FindDataTypeEntry(InDataType))
 				{
+PRAGMA_DISABLE_DEPRECATION_WARNINGS
 					return Entry->CreateReceiveNode(InParams);
+PRAGMA_ENABLE_DEPRECATION_WARNINGS
 				}
 				return nullptr;
 			}
@@ -1128,7 +1353,18 @@ namespace Metasound
 			{
 				if (const IDataTypeRegistryEntry* Entry = FindDataTypeEntry(InDataType))
 				{
+PRAGMA_DISABLE_DEPRECATION_WARNINGS
 					return Entry->CreateVariableNode(MoveTemp(InParams));
+PRAGMA_ENABLE_DEPRECATION_WARNINGS
+				}
+				return nullptr;
+			}
+
+			TUniquePtr<INode> FDataTypeRegistry::CreateVariableNode(const FName& InDataType, FLiteral InLiteral, FNodeData InNodeData) const
+			{
+				if (const IDataTypeRegistryEntry* Entry = FindDataTypeEntry(InDataType))
+				{
+					return Entry->CreateVariableNode(MoveTemp(InLiteral), MoveTemp(InNodeData));
 				}
 				return nullptr;
 			}
@@ -1137,7 +1373,18 @@ namespace Metasound
 			{
 				if (const IDataTypeRegistryEntry* Entry = FindDataTypeEntry(InDataType))
 				{
+PRAGMA_DISABLE_DEPRECATION_WARNINGS
 					return Entry->CreateVariableMutatorNode(InParams);
+PRAGMA_ENABLE_DEPRECATION_WARNINGS
+				}
+				return nullptr;
+			}
+
+			TUniquePtr<INode> FDataTypeRegistry::CreateVariableMutatorNode(const FName& InDataType, FNodeData InNodeData) const 
+			{
+				if (const IDataTypeRegistryEntry* Entry = FindDataTypeEntry(InDataType))
+				{
+					return Entry->CreateVariableMutatorNode(MoveTemp(InNodeData));
 				}
 				return nullptr;
 			}
@@ -1146,7 +1393,18 @@ namespace Metasound
 			{
 				if (const IDataTypeRegistryEntry* Entry = FindDataTypeEntry(InDataType))
 				{
+PRAGMA_DISABLE_DEPRECATION_WARNINGS
 					return Entry->CreateVariableAccessorNode(InParams);
+PRAGMA_ENABLE_DEPRECATION_WARNINGS
+				}
+				return nullptr;
+			}
+
+			TUniquePtr<INode> FDataTypeRegistry::CreateVariableAccessorNode(const FName& InDataType, FNodeData InNodeData) const 
+			{
+				if (const IDataTypeRegistryEntry* Entry = FindDataTypeEntry(InDataType))
+				{
+					return Entry->CreateVariableAccessorNode(MoveTemp(InNodeData));
 				}
 				return nullptr;
 			}
@@ -1155,7 +1413,18 @@ namespace Metasound
 			{
 				if (const IDataTypeRegistryEntry* Entry = FindDataTypeEntry(InDataType))
 				{
+PRAGMA_DISABLE_DEPRECATION_WARNINGS
 					return Entry->CreateVariableDeferredAccessorNode(InParams);
+PRAGMA_ENABLE_DEPRECATION_WARNINGS
+				}
+				return nullptr;
+			}
+
+			TUniquePtr<INode> FDataTypeRegistry::CreateVariableDeferredAccessorNode(const FName& InDataType, FNodeData InNodeData) const 
+			{
+				if (const IDataTypeRegistryEntry* Entry = FindDataTypeEntry(InDataType))
+				{
+					return Entry->CreateVariableDeferredAccessorNode(MoveTemp(InNodeData));
 				}
 				return nullptr;
 			}
diff --git a/Engine/Plugins/Runtime/Metasound/Source/MetasoundFrontend/Private/MetasoundInputNode.cpp b/Engine/Plugins/Runtime/Metasound/Source/MetasoundFrontend/Private/MetasoundInputNode.cpp
--- a/Engine/Plugins/Runtime/Metasound/Source/MetasoundFrontend/Private/MetasoundInputNode.cpp
+++ b/Engine/Plugins/Runtime/Metasound/Source/MetasoundFrontend/Private/MetasoundInputNode.cpp
@@ -73,9 +73,16 @@ namespace Metasound
 	}
 
 
-	FVertexInterface FInputNode::CreateDefaultVertexInterface(const FVertexName& InVertexName, const FName& InDataTypeName, EVertexAccessType InVertexAccess)
+	FVertexInterface FInputNode::CreateDefaultVertexInterface(const FVertexName& InVertexName, const FName& InDataTypeName, EVertexAccessType InVertexAccess, const FLiteral* InDefaultLiteral)
 	{
-		return CreateVertexInterface(InVertexName, InDataTypeName, InVertexAccess, FLiteral());
+		if (InDefaultLiteral)
+		{
+			return CreateVertexInterface(InVertexName, InDataTypeName, InVertexAccess, *InDefaultLiteral);
+		}
+		else
+		{
+			return CreateVertexInterface(InVertexName, InDataTypeName, InVertexAccess, FLiteral());
+		}
 	}
 
 	FText FInputNode::GetInputDescription()
@@ -84,6 +91,11 @@ namespace Metasound
 	}
 
 	FNodeClassMetadata FInputNode::GetNodeMetadata(const FVertexName& InVertexName, const FName& InDataTypeName, EVertexAccessType InVertexAccess)
+	{
+		return CreateNodeClassMetadata(InVertexName, InDataTypeName, InVertexAccess);
+	}
+
+	FNodeClassMetadata FInputNode::CreateNodeClassMetadata(const FVertexName& InVertexName, const FName& InDataTypeName, EVertexAccessType InVertexAccess)
 	{
 		FNodeClassMetadata Info;
 
@@ -99,31 +111,21 @@ namespace Metasound
 	}
 
 	FInputNode::FInputNode(FInputNodeConstructorParams&& InParams, const FName& InDataTypeName, EVertexAccessType InVertexAccess, FOperatorFactorySharedRef InFactory)
-	: FNode(InParams.NodeName, InParams.InstanceID, GetNodeMetadata(InParams.VertexName, InDataTypeName, InVertexAccess))
-	, VertexName(InParams.VertexName)
-	, Interface(CreateVertexInterface(InParams.VertexName, InDataTypeName, InVertexAccess, InParams.InitParam))
-	, Factory(InFactory)
+	: FInputNode(MoveTemp(InFactory), FNodeData{InParams.NodeName, InParams.InstanceID, CreateVertexInterface(InParams.VertexName, InDataTypeName, InVertexAccess, InParams.InitParam)}, MakeShared<const FNodeClassMetadata>(CreateNodeClassMetadata(InParams.VertexName, InDataTypeName, InVertexAccess)))
 	{
 	}
 
-	const FVertexName& FInputNode::GetVertexName() const
-	{
-		return VertexName;
-	}
-
-	const FVertexInterface& FInputNode::GetVertexInterface() const
+	FInputNode::FInputNode(FOperatorFactorySharedRef InFactory, FNodeData InNodeData, TSharedRef<const FNodeClassMetadata> InClassMetadata)
+	: FBasicNode(MoveTemp(InNodeData), MoveTemp(InClassMetadata))
+	, Factory(MoveTemp(InFactory))
 	{
-		return Interface;
 	}
 
-	bool FInputNode::SetVertexInterface(const FVertexInterface& InInterface)
-	{
-		return Interface == InInterface;
-	}
-
-	bool FInputNode::IsVertexInterfaceSupported(const FVertexInterface& InInterface) const
+	const FVertexName& FInputNode::GetVertexName() const
 	{
-		return Interface == InInterface;
+		const FInputVertexInterface& Inputs = GetVertexInterface().GetInputInterface();
+		check(Inputs.Num() == 1);
+		return Inputs.At(0).VertexName;
 	}
 
 	TSharedRef<IOperatorFactory, ESPMode::ThreadSafe> FInputNode::GetDefaultOperatorFactory() const
diff --git a/Engine/Plugins/Runtime/Metasound/Source/MetasoundFrontend/Private/MetasoundOutputNode.cpp b/Engine/Plugins/Runtime/Metasound/Source/MetasoundFrontend/Private/MetasoundOutputNode.cpp
--- a/Engine/Plugins/Runtime/Metasound/Source/MetasoundFrontend/Private/MetasoundOutputNode.cpp
+++ b/Engine/Plugins/Runtime/Metasound/Source/MetasoundFrontend/Private/MetasoundOutputNode.cpp
@@ -12,22 +12,35 @@ namespace Metasound
 		static const FLazyName ValueVertexAccessVariantName("Constructor");
 		static const FLazyName ReferenceVertexAccessVariantName("");
 
-		FOutputOperator::FOutputOperator(const FVertexName& InVertexName, const FAnyDataReference& InDataReference)
-		: VertexName(InVertexName)
-		, DataReference(InDataReference)
+		class FOutputOperator : public FNoOpOperator
 		{
-		}
+			public:
+				FOutputOperator(const FVertexName& InVertexName, const FAnyDataReference& InDataReference);
+				virtual ~FOutputOperator() = default;
+
+				virtual void BindInputs(FInputVertexInterfaceData& InVertexData) override;
+				virtual void BindOutputs(FOutputVertexInterfaceData& InVertexData) override;
+
+			private:
+				FVertexName VertexName;
+				FAnyDataReference DataReference;
+		};
 
-		FDataReferenceCollection FOutputOperator::GetInputs() const
+		class FOutputOperatorFactory : public IOperatorFactory
 		{
-			// Slated for deprecation
-			return {};
-		}
+		public:
+
+			FOutputOperatorFactory(const FVertexName& InVertexName);
+			virtual TUniquePtr<IOperator> CreateOperator(const FBuildOperatorParams& InParams, FBuildResults& OutResults) override;
+
+		private:
+			FVertexName VertexName;
+		};
 
-		FDataReferenceCollection FOutputOperator::GetOutputs() const
+		FOutputOperator::FOutputOperator(const FVertexName& InVertexName, const FAnyDataReference& InDataReference)
+		: VertexName(InVertexName)
+		, DataReference(InDataReference)
 		{
-			// Slated for deprecation
-			return {};
 		}
 
 		void FOutputOperator::BindInputs(FInputVertexInterfaceData& InVertexData) 
@@ -40,33 +53,40 @@ namespace Metasound
 			InVertexData.BindVertex(VertexName, DataReference);
 		}
 
-		FOutputNode::FOutputOperatorFactory::FOutputOperatorFactory(const FVertexName& InVertexName, const FName& InDataTypeName, EVertexAccessType InVertexAccessType)
+		FOutputOperatorFactory::FOutputOperatorFactory(const FVertexName& InVertexName)
 		: VertexName(InVertexName)
-		, DataTypeName(InDataTypeName)
-		, VertexAccessType(InVertexAccessType)
 		{
 		}
 
-		TUniquePtr<IOperator> FOutputNode::FOutputOperatorFactory::CreateOperator(const FBuildOperatorParams& InParams, FBuildResults& OutResults)
+		TUniquePtr<IOperator> FOutputOperatorFactory::CreateOperator(const FBuildOperatorParams& InParams, FBuildResults& OutResults)
 		{
 			using namespace OutputNodePrivate;
 
-			const FOutputNode& OutputNode = static_cast<const FOutputNode&>(InParams.Node);
+			const FInputDataVertex& Vertex = InParams.InputData.GetVertex(VertexName);
 
 			// Use data reference if it is passed in. 
 			if (const FAnyDataReference* Ref = InParams.InputData.FindDataReference(VertexName))
 			{
-				checkf(Ref->GetDataTypeName() == DataTypeName, TEXT("Mismatched data type names on output node (%s). Expected (%s), received (%s)."), *VertexName.ToString(), *DataTypeName.ToString(), *(Ref->GetDataTypeName().ToString()));
+				checkf(Ref->GetDataTypeName() == Vertex.DataTypeName, TEXT("Mismatched data type names on output node (%s). Expected (%s), received (%s)."), *VertexName.ToString(), *Vertex.DataTypeName.ToString(), *(Ref->GetDataTypeName().ToString()));
 				return MakeUnique<FOutputOperator>(VertexName, *Ref);
 			}
 			else 
 			{
 				// Make data reference if none are passed in. 
 				Frontend::IDataTypeRegistry& DataTypeRegistry = Frontend::IDataTypeRegistry::Get();
-				if (const Frontend::IDataTypeRegistryEntry* Entry = DataTypeRegistry.FindDataTypeRegistryEntry(DataTypeName))
+
+				if (const Frontend::IDataTypeRegistryEntry* Entry = DataTypeRegistry.FindDataTypeRegistryEntry(Vertex.DataTypeName))
 				{
-					FLiteral DefaultLiteral = DataTypeRegistry.CreateDefaultLiteral(DataTypeName);
-					TOptional<FAnyDataReference> DataReference = Entry->CreateDataReference(EDataReferenceAccessType::Value, DefaultLiteral, InParams.OperatorSettings);
+					FLiteral Literal = Vertex.GetDefaultLiteral();
+					if (!Literal.IsValid())
+					{
+						// If the literal on the interface is invalid, fallback to using
+						// the default literal for the data type
+						Literal = DataTypeRegistry.CreateDefaultLiteral(Vertex.DataTypeName);
+					}
+
+					// Make a EDataReferenceAccessType::Value reference because this data is not going to be mutated by anything.
+					TOptional<FAnyDataReference> DataReference = Entry->CreateDataReference(EDataReferenceAccessType::Value, Literal, InParams.OperatorSettings);
 					if (DataReference.IsSet())
 					{
 						return MakeUnique<FOutputOperator>(VertexName, *DataReference);
@@ -78,73 +98,61 @@ namespace Metasound
 			OutResults.Errors.Emplace(MakeUnique<FMissingOutputNodeInputReferenceError>(InParams.Node));
 			return TUniquePtr<IOperator>(nullptr);
 		}
+	} // namespace OutputNodePrivate
 
-		FName FOutputNode::GetVariantName(EVertexAccessType InVertexAccessType)
-		{
-			if (EVertexAccessType::Value == InVertexAccessType)
-			{
-				return ValueVertexAccessVariantName;
-			}
-			else
-			{
-				return ReferenceVertexAccessVariantName;
-			}
-		}
 
-		FVertexInterface FOutputNode::GetVertexInterface(const FVertexName& InVertexName, const FName& InDataTypeName, EVertexAccessType InVertexAccessType)
-		{
-			static const FText VertexDescription = METASOUND_LOCTEXT("Metasound_OutputVertexDescription", "Output from the parent Metasound graph.");
-
-			return FVertexInterface(
-				FInputVertexInterface(
-					FInputDataVertex(InVertexName, InDataTypeName, FDataVertexMetadata{VertexDescription}, InVertexAccessType)
-				),
-				FOutputVertexInterface(
-					FOutputDataVertex(InVertexName, InDataTypeName, FDataVertexMetadata{VertexDescription}, InVertexAccessType)
-				)
-			);
-		}
+	FName FOutputNode::GetVariantName(EVertexAccessType InVertexAccessType)
+	{
+		using namespace OutputNodePrivate;
 
-		FNodeClassMetadata FOutputNode::GetNodeInfo(const FVertexName& InVertexName, const FName& InDataTypeName, EVertexAccessType InVertexAccessType)
+		if (EVertexAccessType::Value == InVertexAccessType)
 		{
-			FNodeClassMetadata Info;
-
-			Info.ClassName = { "Output", InDataTypeName, GetVariantName(InVertexAccessType)};
-			Info.MajorVersion = 1;
-			Info.MinorVersion = 0;
-			Info.Description = METASOUND_LOCTEXT("Metasound_OutputNodeDescription", "Output from the parent Metasound graph.");
-			Info.Author = PluginAuthor;
-			Info.PromptIfMissing = PluginNodeMissingPrompt;
-			Info.DefaultInterface = GetVertexInterface(InVertexName, InDataTypeName, InVertexAccessType);
-
-			return Info;
+			return ValueVertexAccessVariantName;
 		}
-
-		FOutputNode::FOutputNode(const FName& InInstanceName, const FGuid& InInstanceID, const FVertexName& InVertexName, const FName& InDataTypeName, EVertexAccessType InVertexAccessType)
-		:	FNode(InInstanceName, InInstanceID, GetNodeInfo(InVertexName, InDataTypeName, InVertexAccessType))
-		,	Factory(MakeShared<FOutputOperatorFactory, ESPMode::ThreadSafe>(InVertexName, InDataTypeName, InVertexAccessType))
+		else
 		{
+			return ReferenceVertexAccessVariantName;
 		}
+	}
 
-		const FVertexInterface& FOutputNode::GetVertexInterface() const
-		{
-			return GetMetadata().DefaultInterface;
-		}
+	FVertexInterface FOutputNode::CreateVertexInterface(const FVertexName& InVertexName, const FName& InDataTypeName, EVertexAccessType InVertexAccessType)
+	{
+		static const FText VertexDescription = METASOUND_LOCTEXT("Metasound_OutputVertexDescription", "Output from the parent Metasound graph.");
+
+		return FVertexInterface(
+			FInputVertexInterface(
+				FInputDataVertex(InVertexName, InDataTypeName, FDataVertexMetadata{VertexDescription}, InVertexAccessType)
+			),
+			FOutputVertexInterface(
+				FOutputDataVertex(InVertexName, InDataTypeName, FDataVertexMetadata{VertexDescription}, InVertexAccessType)
+			)
+		);
+	}
 
-		bool FOutputNode::SetVertexInterface(const FVertexInterface& InInterface)
-		{
-			return GetVertexInterface() == InInterface;
-		}
+	FNodeClassMetadata FOutputNode::CreateNodeClassMetadata(const FVertexName& InVertexName, const FName& InDataTypeName, EVertexAccessType InVertexAccessType)
+	{
+		FNodeClassMetadata Info;
 
-		bool FOutputNode::IsVertexInterfaceSupported(const FVertexInterface& InInterface) const 
-		{
-			return GetVertexInterface() == InInterface;
-		}
+		Info.ClassName = { "Output", InDataTypeName, GetVariantName(InVertexAccessType)};
+		Info.MajorVersion = 1;
+		Info.MinorVersion = 0;
+		Info.Description = METASOUND_LOCTEXT("Metasound_OutputNodeDescription", "Output from the parent Metasound graph.");
+		Info.Author = PluginAuthor;
+		Info.PromptIfMissing = PluginNodeMissingPrompt;
+		Info.DefaultInterface = CreateVertexInterface(InVertexName, InDataTypeName, InVertexAccessType);
 
-		TSharedRef<IOperatorFactory, ESPMode::ThreadSafe> FOutputNode::GetDefaultOperatorFactory() const
-		{
-			return Factory;
-		}
+		return Info;
+	}
+
+	FOutputNode::FOutputNode(const FVertexName& InVertexName, FNodeData InNodeData, TSharedRef<const FNodeClassMetadata> InClassMetadata)
+	: FBasicNode(MoveTemp(InNodeData), MoveTemp(InClassMetadata))
+	, Factory(MakeShared<OutputNodePrivate::FOutputOperatorFactory, ESPMode::ThreadSafe>(InVertexName))
+	{
+	}
+
+	TSharedRef<IOperatorFactory, ESPMode::ThreadSafe> FOutputNode::GetDefaultOperatorFactory() const
+	{
+		return Factory;
 	}
 }
 
diff --git a/Engine/Plugins/Runtime/Metasound/Source/MetasoundFrontend/Public/MetasoundDataTypeRegistrationMacro.h b/Engine/Plugins/Runtime/Metasound/Source/MetasoundFrontend/Public/MetasoundDataTypeRegistrationMacro.h
--- a/Engine/Plugins/Runtime/Metasound/Source/MetasoundFrontend/Public/MetasoundDataTypeRegistrationMacro.h
+++ b/Engine/Plugins/Runtime/Metasound/Source/MetasoundFrontend/Public/MetasoundDataTypeRegistrationMacro.h
@@ -298,6 +298,63 @@ namespace Metasound::Frontend
 		template <typename T>
 		struct HasRawParameterAssignmentOp<T, decltype(&T::AssignRawParameter, void())> { bool value = true; };
 
+		// Base registry entry for any data type
+		class METASOUNDFRONTEND_API FDataTypeRegistryEntryBase : public IDataTypeRegistryEntry
+		{
+		public:
+			FDataTypeRegistryEntryBase(const Frontend::FDataTypeRegistryInfo& Info, const TSharedPtr<Frontend::IEnumDataTypeInterface>& EnumInterface);
+			virtual ~FDataTypeRegistryEntryBase();
+
+			virtual const Frontend::FDataTypeRegistryInfo& GetDataTypeInfo() const override;
+			virtual TSharedPtr<const Frontend::IEnumDataTypeInterface> GetEnumInterface() const override;
+			virtual const FMetasoundFrontendClass& GetFrontendInputClass() const override;
+			virtual TSharedPtr<const FNodeClassMetadata> GetInputClassMetadata() const override;
+			virtual const FMetasoundFrontendClass& GetFrontendConstructorInputClass() const override;
+			virtual TSharedPtr<const FNodeClassMetadata> GetConstructorInputClassMetadata() const override;
+			virtual const FMetasoundFrontendClass& GetFrontendLiteralClass() const override;
+			virtual const FMetasoundFrontendClass& GetFrontendOutputClass() const override;
+			virtual TSharedPtr<const FNodeClassMetadata> GetOutputClassMetadata() const override;
+			virtual const FMetasoundFrontendClass& GetFrontendConstructorOutputClass() const override;
+			virtual TSharedPtr<const FNodeClassMetadata> GetConstructorOutputClassMetadata() const override;
+			virtual const FMetasoundFrontendClass& GetFrontendVariableClass() const override;
+			virtual TSharedPtr<const FNodeClassMetadata> GetVariableClassMetadata() const override;
+			virtual const FMetasoundFrontendClass& GetFrontendVariableMutatorClass() const override;
+			virtual TSharedPtr<const FNodeClassMetadata> GetVariableMutatorClassMetadata() const override;
+			virtual const FMetasoundFrontendClass& GetFrontendVariableAccessorClass() const override;
+			virtual TSharedPtr<const FNodeClassMetadata> GetVariableAccessorClassMetadata() const override;
+			virtual const FMetasoundFrontendClass& GetFrontendVariableDeferredAccessorClass() const override;
+			virtual TSharedPtr<const FNodeClassMetadata> GetVariableDeferredAccessorClassMetadata() const override;
+			virtual const Frontend::IParameterAssignmentFunction& GetRawAssignmentFunction() const override;
+			virtual Frontend::FLiteralAssignmentFunction GetLiteralAssignmentFunction() const override;
+			virtual TUniquePtr<INode> CreateOutputNode(FNodeData InNodeData) const override;
+			virtual TUniquePtr<INode> CreateConstructorOutputNode(FNodeData InNodeData) const override;
+
+		protected:
+			Frontend::FDataTypeRegistryInfo Info;
+			FMetasoundFrontendClass InputClass;
+			TSharedPtr<const FNodeClassMetadata> InputClassMetadata;
+			FMetasoundFrontendClass ConstructorInputClass;
+			TSharedPtr<const FNodeClassMetadata> ConstructorInputClassMetadata;
+			FMetasoundFrontendClass OutputClass;
+			TSharedPtr<const FNodeClassMetadata> OutputClassMetadata;
+			FMetasoundFrontendClass ConstructorOutputClass;
+			TSharedPtr<const FNodeClassMetadata> ConstructorOutputClassMetadata;
+			
+			FMetasoundFrontendClass LiteralClass;
+
+			FMetasoundFrontendClass VariableClass;
+			TSharedPtr<const FNodeClassMetadata> VariableClassMetadata;
+			FMetasoundFrontendClass VariableMutatorClass;
+			TSharedPtr<const FNodeClassMetadata> VariableMutatorClassMetadata;
+			FMetasoundFrontendClass VariableAccessorClass;
+			TSharedPtr<const FNodeClassMetadata> VariableAccessorClassMetadata;
+			FMetasoundFrontendClass VariableDeferredAccessorClass;
+			TSharedPtr<const FNodeClassMetadata> VariableDeferredAccessorClassMetadata;
+			TSharedPtr<Frontend::IEnumDataTypeInterface> EnumInterface;
+			Frontend::IParameterAssignmentFunction RawAssignmentFunction;
+			Frontend::FLiteralAssignmentFunction LiteralAssignmentFunction = nullptr;
+		};
+
 
 		/** Registers a data type with the MetaSound Frontend. This allows the data type
 		 * to be used in Input and Output nodes by informing the Frontend how to
@@ -339,99 +396,6 @@ namespace Metasound::Frontend
 
 			bAlreadyRegisteredThisDataType = true;
 
-			// Define registry entry class for this data type.
-			class FDataTypeRegistryEntryBase : public Frontend::IDataTypeRegistryEntry
-			{
-			public:
-				FDataTypeRegistryEntryBase(const Frontend::FDataTypeRegistryInfo& Info, const TSharedPtr<Frontend::IEnumDataTypeInterface>& EnumInterface)
-					: Info(Info), EnumInterface(EnumInterface)
-				{
-				}
-
-				virtual ~FDataTypeRegistryEntryBase() {}
-
-				virtual const Frontend::FDataTypeRegistryInfo& GetDataTypeInfo() const override
-				{
-					return Info;
-				}
-
-				virtual TSharedPtr<const Frontend::IEnumDataTypeInterface> GetEnumInterface() const override
-				{
-					return EnumInterface;
-				}
-
-				virtual const FMetasoundFrontendClass& GetFrontendInputClass() const override
-				{
-					return InputClass;
-				}
-
-				virtual const FMetasoundFrontendClass& GetFrontendConstructorInputClass() const override
-				{
-					return ConstructorInputClass;
-				}
-
-				virtual const FMetasoundFrontendClass& GetFrontendLiteralClass() const override
-				{
-					return LiteralClass;
-				}
-
-				virtual const FMetasoundFrontendClass& GetFrontendOutputClass() const override
-				{
-					return OutputClass;
-				}
-
-				virtual const FMetasoundFrontendClass& GetFrontendConstructorOutputClass() const override
-				{
-					return ConstructorOutputClass;
-				}
-
-				virtual const FMetasoundFrontendClass& GetFrontendVariableClass() const override
-				{
-					return VariableClass;
-				}
-
-				virtual const FMetasoundFrontendClass& GetFrontendVariableMutatorClass() const override
-				{
-					return VariableMutatorClass;
-				}
-
-				virtual const FMetasoundFrontendClass& GetFrontendVariableAccessorClass() const override
-				{
-					return VariableAccessorClass;
-				}
-
-				virtual const FMetasoundFrontendClass& GetFrontendVariableDeferredAccessorClass() const 
-				{
-					return VariableDeferredAccessorClass;
-				}
-
-				virtual TUniquePtr<IDataTypeRegistryEntry> Clone() const = 0;
-
-				virtual const Frontend::IParameterAssignmentFunction& GetRawAssignmentFunction() const override
-				{
-					return RawAssignmentFunction;
-				}
-
-				virtual Frontend::FLiteralAssignmentFunction GetLiteralAssignmentFunction() const override
-				{
-					return LiteralAssignmentFunction;
-				}
-
-			protected:
-				Frontend::FDataTypeRegistryInfo Info;
-				FMetasoundFrontendClass InputClass;
-				FMetasoundFrontendClass ConstructorInputClass;
-				FMetasoundFrontendClass OutputClass;
-				FMetasoundFrontendClass ConstructorOutputClass;
-				FMetasoundFrontendClass LiteralClass;
-				FMetasoundFrontendClass VariableClass;
-				FMetasoundFrontendClass VariableMutatorClass;
-				FMetasoundFrontendClass VariableAccessorClass;
-				FMetasoundFrontendClass VariableDeferredAccessorClass;
-				TSharedPtr<Frontend::IEnumDataTypeInterface> EnumInterface;
-				Frontend::IParameterAssignmentFunction RawAssignmentFunction;
-				Frontend::FLiteralAssignmentFunction LiteralAssignmentFunction = nullptr;
-			};
 
 			class FDataTypeRegistryEntry : public FDataTypeRegistryEntryBase
 			{
@@ -480,40 +444,44 @@ namespace Metasound::Frontend
 					}
 				}
 
-				void InitFrontendNodeClasses()
+				void InitNodeClasses()
 				{
 					// Create class info using prototype node
 					// TODO: register nodes with static class info instead of prototype instance.
 
 					if constexpr (bIsParsable)
 					{
-						TInputNode<TDataType, EVertexAccessType::Reference> InputPrototype(FInputNodeConstructorParams{ });
-						this->InputClass = Metasound::Frontend::GenerateClass(InputPrototype.GetMetadata(), EMetasoundFrontendClassType::Input);
+						const FName DataTypeName = GetMetasoundDataTypeName<TDataType>();
+						const FName UnnamedVertex;
+
 
-						TOutputNode<TDataType, EVertexAccessType::Reference> OutputPrototype(TEXT(""), FGuid(), TEXT(""));
-						this->OutputClass = Metasound::Frontend::GenerateClass(OutputPrototype.GetMetadata(), EMetasoundFrontendClassType::Output);
+						this->InputClassMetadata = MakeShared<FNodeClassMetadata>(TInputNode<TDataType, EVertexAccessType::Reference>::CreateNodeClassMetadata(UnnamedVertex));
+						this->InputClass = Metasound::Frontend::GenerateClass(*this->InputClassMetadata, EMetasoundFrontendClassType::Input);
+
+						this->OutputClassMetadata = MakeShared<FNodeClassMetadata>(FOutputNode::CreateNodeClassMetadata(UnnamedVertex, DataTypeName, EVertexAccessType::Reference));
+						this->OutputClass = Metasound::Frontend::GenerateClass(*this->OutputClassMetadata, EMetasoundFrontendClassType::Output);
 
 						TLiteralNode<TDataType> LiteralPrototype(TEXT(""), FGuid(), FLiteral());
 						this->LiteralClass = Metasound::Frontend::GenerateClass(LiteralPrototype.GetMetadata(), EMetasoundFrontendClassType::Literal);
 
-						TVariableNode<TDataType> VariablePrototype(TEXT(""), FGuid(), FLiteral());
-						this->VariableClass = Metasound::Frontend::GenerateClass(VariablePrototype.GetMetadata(), EMetasoundFrontendClassType::Variable);
+						this->VariableClassMetadata = MakeShared<FNodeClassMetadata>(TVariableNode<TDataType>::CreateNodeClassMetadata());
+						this->VariableClass = Metasound::Frontend::GenerateClass(*this->VariableClassMetadata, EMetasoundFrontendClassType::Variable);
 
-						TVariableMutatorNode<TDataType> VariableMutatorPrototype(TEXT(""), FGuid());
-						this->VariableMutatorClass = Metasound::Frontend::GenerateClass(VariableMutatorPrototype.GetMetadata(), EMetasoundFrontendClassType::VariableMutator);
+						this->VariableMutatorClassMetadata = MakeShared<FNodeClassMetadata>(TVariableMutatorNode<TDataType>::CreateNodeClassMetadata());
+						this->VariableMutatorClass = Metasound::Frontend::GenerateClass(*this->VariableMutatorClassMetadata, EMetasoundFrontendClassType::VariableMutator);
 
-						TVariableAccessorNode<TDataType> VariableAccessorPrototype(TEXT(""), FGuid());
-						this->VariableAccessorClass = Metasound::Frontend::GenerateClass(VariableAccessorPrototype.GetMetadata(), EMetasoundFrontendClassType::VariableAccessor);
+						this->VariableAccessorClassMetadata = MakeShared<FNodeClassMetadata>(TVariableAccessorNode<TDataType>::CreateNodeClassMetadata());
+						this->VariableAccessorClass = Metasound::Frontend::GenerateClass(*this->VariableAccessorClassMetadata, EMetasoundFrontendClassType::VariableAccessor);
 
-						TVariableDeferredAccessorNode<TDataType> VariableDeferredAccessorPrototype(TEXT(""), FGuid());
-						this->VariableDeferredAccessorClass = Metasound::Frontend::GenerateClass(VariableDeferredAccessorPrototype.GetMetadata(), EMetasoundFrontendClassType::VariableDeferredAccessor);
+						this->VariableDeferredAccessorClassMetadata = MakeShared<FNodeClassMetadata>(TVariableDeferredAccessorNode<TDataType>::CreateNodeClassMetadata());
+						this->VariableDeferredAccessorClass = Metasound::Frontend::GenerateClass(*this->VariableDeferredAccessorClassMetadata, EMetasoundFrontendClassType::VariableDeferredAccessor);
 
 						if constexpr (bIsConstructorType)
 						{
-							TInputNode<TDataType, EVertexAccessType::Value> ConstructorInputPrototype(FInputNodeConstructorParams{ });
-							this->ConstructorInputClass = Metasound::Frontend::GenerateClass(ConstructorInputPrototype.GetMetadata(), EMetasoundFrontendClassType::Input);
-							TOutputNode<TDataType, EVertexAccessType::Value> ConstructorOutputPrototype(TEXT(""), FGuid(), TEXT(""));
-							this->ConstructorOutputClass = Metasound::Frontend::GenerateClass(ConstructorOutputPrototype.GetMetadata(), EMetasoundFrontendClassType::Output);
+							this->ConstructorInputClassMetadata = MakeShared<FNodeClassMetadata>(TInputNode<TDataType, EVertexAccessType::Value>::CreateNodeClassMetadata(UnnamedVertex));
+							this->ConstructorInputClass = Metasound::Frontend::GenerateClass(*this->ConstructorInputClassMetadata, EMetasoundFrontendClassType::Input);
+							this->ConstructorOutputClassMetadata = MakeShared<FNodeClassMetadata>(FOutputNode::CreateNodeClassMetadata(UnnamedVertex, DataTypeName, EVertexAccessType::Value));
+							this->ConstructorOutputClass = Metasound::Frontend::GenerateClass(*this->ConstructorOutputClassMetadata, EMetasoundFrontendClassType::Output);
 						}
 					}
 				}
@@ -526,9 +494,10 @@ namespace Metasound::Frontend
 				{
 					InitRawAssignmentFunction();
 					InitLiteralAssignmentFunction();
-					InitFrontendNodeClasses();
+					InitNodeClasses();
 				}
 
+				UE_DEPRECATED(5.6, "Create input nodes using FNodeData")
 				virtual TUniquePtr<INode> CreateInputNode(FInputNodeConstructorParams&& InParams) const override
 				{
 					if constexpr (bIsParsable)
@@ -541,6 +510,24 @@ namespace Metasound::Frontend
 					}
 				}
 
+				virtual TUniquePtr<INode> CreateInputNode(FNodeData InNodeData) const override
+				{
+					if constexpr (bIsParsable)
+					{
+						if (InputClassMetadata)
+						{
+							const FInputVertexInterface& Inputs = InNodeData.Interface.GetInputInterface();
+							if (ensure(Inputs.Num() == 1))
+							{
+								const FVertexName& VertexName = Inputs.At(0).VertexName;
+								return MakeUnique<TInputNode<TDataType, EVertexAccessType::Reference>>(VertexName, MoveTemp(InNodeData), InputClassMetadata.ToSharedRef());
+							}
+						}
+					}
+					return TUniquePtr<INode>(nullptr);
+				}
+
+				UE_DEPRECATED(5.6, "Create constructor input nodes using FNodeData")
 				virtual TUniquePtr<INode> CreateConstructorInputNode(FInputNodeConstructorParams&& InParams) const override
 				{
 					if constexpr (bIsParsable && bIsConstructorType)
@@ -553,6 +540,26 @@ namespace Metasound::Frontend
 					}
 				}
 
+				virtual TUniquePtr<INode> CreateConstructorInputNode(FNodeData InNodeData) const override
+				{
+					if constexpr (bIsParsable && bIsConstructorType)
+					{
+						if (ConstructorInputClassMetadata)
+						{
+							const FInputVertexInterface& Inputs = InNodeData.Interface.GetInputInterface();
+							if (ensure(Inputs.Num() == 1))
+							{
+								const FVertexName& VertexName = Inputs.At(0).VertexName;
+								return MakeUnique<TInputNode<TDataType, EVertexAccessType::Value>>(VertexName, MoveTemp(InNodeData), ConstructorInputClassMetadata.ToSharedRef());
+							}
+						}
+					}
+					return TUniquePtr<INode>(nullptr);
+				}
+
+				using FDataTypeRegistryEntryBase::CreateOutputNode; // Unhide override
+
+				UE_DEPRECATED(5.6, "Create output nodes using FNodeData")
 				virtual TUniquePtr<INode> CreateOutputNode(FOutputNodeConstructorParams&& InParams) const override
 				{
 					if constexpr (bIsParsable)
@@ -565,6 +572,9 @@ namespace Metasound::Frontend
 					}
 				}
 
+				using FDataTypeRegistryEntryBase::CreateConstructorOutputNode; // Unhide override
+																			   
+				UE_DEPRECATED(5.6, "Create constructor output nodes using FNodeData")
 				virtual TUniquePtr<INode> CreateConstructorOutputNode(FOutputNodeConstructorParams&& InParams) const override
 				{
 					if constexpr (bIsParsable && bIsConstructorType)
@@ -577,6 +587,8 @@ namespace Metasound::Frontend
 					}
 				}
 
+
+				UE_DEPRECATED(5.6, "Direct creation of literal nodes will no longer be supported")
 				virtual TUniquePtr<INode> CreateLiteralNode(FLiteralNodeConstructorParams&& InParams) const override
 				{
 					if constexpr (bIsParsable)
@@ -589,6 +601,7 @@ namespace Metasound::Frontend
 					}
 				}
 
+				UE_DEPRECATED(5.6, "Direct creation of receive nodes will no longer be supported")
 				virtual TUniquePtr<INode> CreateReceiveNode(const FNodeInitData& InParams) const override
 				{
 					if constexpr (bIsParsable)
@@ -601,6 +614,7 @@ namespace Metasound::Frontend
 					}
 				}
 
+				UE_DEPRECATED(5.6, "Create variable nodes using FNodeData")
 				virtual TUniquePtr<INode> CreateVariableNode(FVariableNodeConstructorParams&& InParams) const override
 				{
 					if constexpr (bIsParsable)
@@ -613,6 +627,19 @@ namespace Metasound::Frontend
 					}
 				}
 
+				virtual TUniquePtr<INode> CreateVariableNode(FLiteral InLiteral, FNodeData InNodeData) const override
+				{
+					if constexpr (bIsParsable)
+					{
+						if (VariableClassMetadata)
+						{
+							return MakeUnique<TVariableNode<TDataType>>(MoveTemp(InLiteral), MoveTemp(InNodeData), VariableClassMetadata.ToSharedRef());
+						}
+					}
+					return TUniquePtr<INode>(nullptr);
+				}
+ 
+				UE_DEPRECATED(5.6, "Create variable mutator nodes using FNodeData")
 				virtual TUniquePtr<INode> CreateVariableMutatorNode(const FNodeInitData& InParams) const override
 				{
 					if constexpr (bIsParsable)
@@ -625,6 +652,19 @@ namespace Metasound::Frontend
 					}
 				}
 
+				virtual TUniquePtr<INode> CreateVariableMutatorNode(FNodeData InNodeData) const override
+				{
+					if constexpr (bIsParsable)
+					{
+						if (VariableMutatorClassMetadata)
+						{
+							return MakeUnique<TVariableMutatorNode<TDataType>>(MoveTemp(InNodeData), VariableMutatorClassMetadata.ToSharedRef());
+						}
+					}
+					return TUniquePtr<INode>(nullptr);
+				}
+
+				UE_DEPRECATED(5.6, "Create variable accessor nodes using FNodeData")
 				virtual TUniquePtr<INode> CreateVariableAccessorNode(const FNodeInitData& InParams) const override
 				{
 					if constexpr (bIsParsable)
@@ -637,6 +677,19 @@ namespace Metasound::Frontend
 					}
 				}
 
+				virtual TUniquePtr<INode> CreateVariableAccessorNode(FNodeData InNodeData) const override
+				{
+					if constexpr (bIsParsable)
+					{
+						if (VariableAccessorClassMetadata)
+						{
+							return MakeUnique<TVariableAccessorNode<TDataType>>(MoveTemp(InNodeData), VariableAccessorClassMetadata.ToSharedRef());
+						}
+					}
+					return TUniquePtr<INode>(nullptr);
+				}
+
+				UE_DEPRECATED(5.6, "Create variable deferred accessor nodes using FNodeData")
 				virtual TUniquePtr<INode> CreateVariableDeferredAccessorNode(const FNodeInitData& InParams) const override
 				{
 					if constexpr (bIsParsable)
@@ -649,6 +702,18 @@ namespace Metasound::Frontend
 					}
 				}
 
+				virtual TUniquePtr<INode> CreateVariableDeferredAccessorNode(FNodeData InNodeData) const override
+				{
+					if constexpr (bIsParsable)
+					{
+						if (VariableAccessorClassMetadata)
+						{
+							return MakeUnique<TVariableDeferredAccessorNode<TDataType>>(MoveTemp(InNodeData), VariableAccessorClassMetadata.ToSharedRef());
+						}
+					}
+					return TUniquePtr<INode>(nullptr);
+				}
+
 				virtual TSharedPtr<Audio::IProxyData> CreateProxy(UObject* InObject) const override
 				{
 					// Only attempt to create proxy if the `UClassToUse` is not void.
@@ -704,11 +769,6 @@ namespace Metasound::Frontend
 						return TSharedPtr<IDataChannel, ESPMode::ThreadSafe>(nullptr);
 					}
 				}
-
-				virtual TUniquePtr<Frontend::IDataTypeRegistryEntry> Clone() const override
-				{
-					return MakeUnique<FDataTypeRegistryEntry>();
-				}
 			};
 
 			bool bSucceeded = Frontend::IDataTypeRegistry::Get().RegisterDataType(MakeUnique<FDataTypeRegistryEntry>());
diff --git a/Engine/Plugins/Runtime/Metasound/Source/MetasoundFrontend/Public/MetasoundFrontendDataTypeRegistry.h b/Engine/Plugins/Runtime/Metasound/Source/MetasoundFrontend/Public/MetasoundFrontendDataTypeRegistry.h
--- a/Engine/Plugins/Runtime/Metasound/Source/MetasoundFrontend/Public/MetasoundFrontendDataTypeRegistry.h
+++ b/Engine/Plugins/Runtime/Metasound/Source/MetasoundFrontend/Public/MetasoundFrontendDataTypeRegistry.h
@@ -13,6 +13,13 @@
 #include "UObject/Class.h"
 #include "UObject/NameTypes.h"
 
+// In UE 5.6, registered node are expected to support the constructor signature Constructor(FNodeData, TSharedRef<const FNodeClassMetadata>)
+// Because there are many existing nodes, it may take time to update them. For convenience, the deprecations related to this change are
+// configurable via a preprocessor macro so that the deprecation warnings do not drown out other compiler errors and warnings.
+#ifndef UE_METASOUND_ENABLE_5_6_NODE_REGISTRATION_DEPRECATION_WARNINGS
+#define UE_METASOUND_ENABLE_5_6_NODE_REGISTRATION_DEPRECATION_WARNINGS (0)
+#endif
+
 
 namespace Metasound
 {
@@ -150,46 +157,88 @@ namespace Metasound
 			/** Return an FMetasoundFrontendClass representing an input node of the data type. */
 			virtual const FMetasoundFrontendClass& GetFrontendInputClass() const = 0;
 
+			/** Return an FNodeClassMetadata representing an input node of the data type. */
+			virtual TSharedPtr<const FNodeClassMetadata> GetInputClassMetadata() const = 0;
+
 			/** Return an FMetasoundFrontendClass representing an input node of the data type. */
 			virtual const FMetasoundFrontendClass& GetFrontendConstructorInputClass() const = 0;
 
+			/** Return an FNodeClassMetadata representing an input node of the data type. */
+			virtual TSharedPtr<const FNodeClassMetadata> GetConstructorInputClassMetadata() const = 0;
+
 			/** Return an FMetasoundFrontendClass representing a variable node of the data type. */
 			virtual const FMetasoundFrontendClass& GetFrontendLiteralClass() const = 0;
 
 			/** Return an FMetasoundFrontendClass representing an output node of the data type. */
 			virtual const FMetasoundFrontendClass& GetFrontendOutputClass() const = 0;
 
+			/** Return an FNodeClassMetadata representing an output node of the data type. */
+			virtual TSharedPtr<const FNodeClassMetadata> GetOutputClassMetadata() const = 0;
+
 			/** Return an FMetasoundFrontendClass representing an input node of the data type. */
 			virtual const FMetasoundFrontendClass& GetFrontendConstructorOutputClass() const = 0;
 
+			/** Return an FNodeClassMetadata representing an input node of the data type. */
+			virtual TSharedPtr<const FNodeClassMetadata> GetConstructorOutputClassMetadata() const = 0;
+
 			/** Return an FMetasoundFrontendClass representing an init variable node of the data type. */
 			virtual const FMetasoundFrontendClass& GetFrontendVariableClass() const = 0;
 
+			/** Return an FNodeClassMetadata representing an init variable node of the data type. */
+			virtual TSharedPtr<const FNodeClassMetadata> GetVariableClassMetadata() const = 0;
+
 			/** Return an FMetasoundFrontendClass representing an set variable node of the data type. */
 			virtual const FMetasoundFrontendClass& GetFrontendVariableMutatorClass() const = 0;
 
+			/** Return an FNodeClassMetadata representing an set variable node of the data type. */
+			virtual TSharedPtr<const FNodeClassMetadata> GetVariableMutatorClassMetadata() const = 0;
+
 			/** Return an FMetasoundFrontendClass representing an get variable node of the data type. */
 			virtual const FMetasoundFrontendClass& GetFrontendVariableAccessorClass() const = 0;
 
+			/** Return an FNodeClassMetadata representing an get variable node of the data type. */
+			virtual TSharedPtr<const FNodeClassMetadata> GetVariableAccessorClassMetadata() const = 0;
+
 			/** Return an FMetasoundFrontendClass representing an get delayed variable node of the data type. */
 			virtual const FMetasoundFrontendClass& GetFrontendVariableDeferredAccessorClass() const = 0;
 
+			/** Return an FNodeClassMetadata representing an get delayed variable node of the data type. */
+			virtual TSharedPtr<const FNodeClassMetadata> GetVariableDeferredAccessorClassMetadata() const = 0;
+
 			/** Create an input node */
+			UE_DEPRECATED(5.6, "Create input nodes with FNodeData")
 			virtual TUniquePtr<INode> CreateInputNode(FInputNodeConstructorParams&&) const = 0;
 
 			/** Create an input node */
+			virtual TUniquePtr<INode> CreateInputNode(FNodeData InNodeData) const = 0;
+
+			/** Create an constructor input node */
+			UE_DEPRECATED(5.6, "Create constructor input nodes with FNodeData")
 			virtual TUniquePtr<INode> CreateConstructorInputNode(FInputNodeConstructorParams&&) const = 0;
 
+			/** Create a constructor input node */
+			virtual TUniquePtr<INode> CreateConstructorInputNode(FNodeData InNodeData) const = 0;
+
 			/** Create an output node */
+			UE_DEPRECATED(5.6, "Create output nodes with FNodeData")
 			virtual TUniquePtr<INode> CreateOutputNode(FOutputNodeConstructorParams&&) const = 0;
 
 			/** Create an output node */
+			virtual TUniquePtr<INode> CreateOutputNode(FNodeData InNodeData) const = 0;
+
+			/** Create an output node */
+			UE_DEPRECATED(5.6, "Create constructor output nodes with FNodeData")
 			virtual TUniquePtr<INode> CreateConstructorOutputNode(FOutputNodeConstructorParams&&) const = 0;
-			
+
+			/** Create a constructor output node */
+			virtual TUniquePtr<INode> CreateConstructorOutputNode(FNodeData InNodeData) const = 0;
+
 			/** Create a variable node */
+			UE_DEPRECATED(5.6, "Direct creation of literal nodes will no longer be supported")
 			virtual TUniquePtr<INode> CreateLiteralNode(FLiteralNodeConstructorParams&&) const = 0;
 
 			/** Create a receive node for this data type. */
+			UE_DEPRECATED(5.6, "Direct creation of receive nodes will no longer be supported")
 			virtual TUniquePtr<INode> CreateReceiveNode(const FNodeInitData&) const = 0;
 
 			/* Get function that should be used when a parameter pack contains a setting this data type. */ 
@@ -202,16 +251,35 @@ namespace Metasound
 			 *
 			 *  @param InInitParams - Contains a literal used to create the variable.
 			 */
+			UE_DEPRECATED(5.6, "Create variable nodes with FNodeData")
 			virtual TUniquePtr<INode> CreateVariableNode(FVariableNodeConstructorParams&& InInitParams) const = 0;
 
+			/** Create a init variable node for this data type. 
+			 *
+			 *  @param InInitValue - A literal used to create the variable.
+			 */
+			virtual TUniquePtr<INode> CreateVariableNode(FLiteral InInitValue, FNodeData) const = 0;
+
 			/** Create a set variable node for this data type. */
+			UE_DEPRECATED(5.6, "Create variable mutator nodes with FNodeData")
 			virtual TUniquePtr<INode> CreateVariableMutatorNode(const FNodeInitData&) const = 0;
 
+			/** Create a set variable node for this data type. */
+			virtual TUniquePtr<INode> CreateVariableMutatorNode(FNodeData) const = 0;
+
 			/** Create a get variable node for this data type. */
+			UE_DEPRECATED(5.6, "Create variable accessor nodes with FNodeData")
 			virtual TUniquePtr<INode> CreateVariableAccessorNode(const FNodeInitData&) const = 0;
 
+			/** Create a get variable node for this data type. */
+			virtual TUniquePtr<INode> CreateVariableAccessorNode(FNodeData) const = 0;
+
 			/** Create a get delayed variable node for this data type. */
+			UE_DEPRECATED(5.6, "Create variable deferred accessor nodes with FNodeData")
 			virtual TUniquePtr<INode> CreateVariableDeferredAccessorNode(const FNodeInitData&) const = 0;
+
+			/** Create a get delayed variable node for this data type. */
+			virtual TUniquePtr<INode> CreateVariableDeferredAccessorNode(FNodeData) const = 0;
 			
 			/** Create a data reference from a literal. */
 			virtual TOptional<FAnyDataReference> CreateDataReference(EDataReferenceAccessType InAccessType, const FLiteral& InLiteral, const FOperatorSettings& InOperatorSettings) const = 0;
@@ -225,8 +293,8 @@ namespace Metasound
 			virtual TSharedPtr<IDataChannel, ESPMode::ThreadSafe> CreateDataChannel(const FOperatorSettings&) const = 0;
 
 			/** Clone this registry entry. */
-			virtual TUniquePtr<IDataTypeRegistryEntry> Clone() const = 0;
-
+			UE_DEPRECATED(5.6, "Cloning is not required for IDataTypeRegistryEntries")
+			virtual TUniquePtr<IDataTypeRegistryEntry> Clone() const { return nullptr; }
 		};
 
 		class METASOUNDFRONTEND_API IDataTypeRegistry
@@ -316,17 +384,71 @@ namespace Metasound
 			/** Return an FMetasoundFrontendClass representing an get delayed variable node of the data type. */
 			virtual bool GetFrontendVariableDeferredAccessorClass(const FName& InDataType, FMetasoundFrontendClass& OutClass) const = 0;
 
-			// Create a new instance of a C++ implemented node from the registry.
+			/** Below are class metadata accessor for the special nodes related to
+			 * a data type. */
+
+			virtual TSharedPtr<const FNodeClassMetadata> GetInputClassMetadata(const FName& InDataType) const = 0;
+			virtual TSharedPtr<const FNodeClassMetadata> GetConstructorInputClassMetadata(const FName& InDataType) const = 0;
+			virtual TSharedPtr<const FNodeClassMetadata> GetOutputClassMetadata(const FName& InDataType) const = 0;
+			virtual TSharedPtr<const FNodeClassMetadata> GetConstructorOutputClassMetadata(const FName& InDataType) const = 0;
+			virtual TSharedPtr<const FNodeClassMetadata> GetVariableClassMetadata(const FName& InDataType) const = 0;
+			virtual TSharedPtr<const FNodeClassMetadata> GetVariableMutatorClassMetadata(const FName& InDataType) const = 0;
+			virtual TSharedPtr<const FNodeClassMetadata> GetVariableAccessorClassMetadata(const FName& InDataType) const = 0;
+			virtual TSharedPtr<const FNodeClassMetadata> GetVariableDeferredAccessorClassMetadata(const FName& InDataType) const = 0;
+
+			/** Below are node creation methods for creation special nodes related to 
+			 * an individual data type. These nodes power core features of metasounds
+			 * such as Inputs, Outputs, and Variables. 
+			 */
+
+#if UE_METASOUND_ENABLE_5_6_NODE_REGISTRATION_DEPRECATION_WARNINGS
+			UE_DEPRECATED(5.6, "Create input nodes with FNodeData")
+#endif
 			virtual TUniquePtr<INode> CreateInputNode(const FName& InInputType, FInputNodeConstructorParams&& InParams) const = 0;
+			virtual TUniquePtr<INode> CreateInputNode(const FName& InInputType, FNodeData InNodeData) const = 0;
+#if UE_METASOUND_ENABLE_5_6_NODE_REGISTRATION_DEPRECATION_WARNINGS
+			UE_DEPRECATED(5.6, "Create constructor input nodes with FNodeData")
+#endif
 			virtual TUniquePtr<INode> CreateConstructorInputNode(const FName& InInputType, FInputNodeConstructorParams&& InParams) const = 0;
+			virtual TUniquePtr<INode> CreateConstructorInputNode(const FName& InInputType, FNodeData InNodeData) const = 0;
+#if UE_METASOUND_ENABLE_5_6_NODE_REGISTRATION_DEPRECATION_WARNINGS
+			UE_DEPRECATED(5.6, "Direct creation of literal nodes will no longer be supported")
+#endif
 			virtual TUniquePtr<INode> CreateLiteralNode(const FName& InLiteralType, FLiteralNodeConstructorParams&& InParams) const = 0;
-			virtual TUniquePtr<INode> CreateOutputNode(const FName& InOutputType, FOutputNodeConstructorParams&& InParams) const = 0;
-			virtual TUniquePtr<INode> CreateConstructorOutputNode(const FName& InOutputType, FOutputNodeConstructorParams&& InParams) const = 0;
+#if UE_METASOUND_ENABLE_5_6_NODE_REGISTRATION_DEPRECATION_WARNINGS
+			UE_DEPRECATED(5.6, "Create output nodes with FNodeData")
+#endif
+			virtual TUniquePtr<INode> CreateOutputNode(const FName& InDataTypeName, FOutputNodeConstructorParams&& InParams) const = 0;
+			virtual TUniquePtr<INode> CreateOutputNode(const FName& InDataTypeName, FNodeData InNodeData) const = 0;
+#if UE_METASOUND_ENABLE_5_6_NODE_REGISTRATION_DEPRECATION_WARNINGS
+			UE_DEPRECATED(5.6, "Create output nodes with FNodeData")
+#endif
+			virtual TUniquePtr<INode> CreateConstructorOutputNode(const FName& InDataTypeName, FOutputNodeConstructorParams&& InParams) const = 0;
+			virtual TUniquePtr<INode> CreateConstructorOutputNode(const FName& InDataTypeName, FNodeData InNodeData) const = 0;
+#if UE_METASOUND_ENABLE_5_6_NODE_REGISTRATION_DEPRECATION_WARNINGS
+			UE_DEPRECATED(5.6, "Direct creation of receive nodes will no longer be supported")
+#endif
 			virtual TUniquePtr<INode> CreateReceiveNode(const FName& InDataType, const FNodeInitData&) const = 0;
+#if UE_METASOUND_ENABLE_5_6_NODE_REGISTRATION_DEPRECATION_WARNINGS
+			UE_DEPRECATED(5.6, "Create Variable Node with FNodeData")
+#endif
 			virtual TUniquePtr<INode> CreateVariableNode(const FName& InDataType, FVariableNodeConstructorParams&& InParams) const = 0;
+			virtual TUniquePtr<INode> CreateVariableNode(const FName& InDataType, FLiteral InLiteral, FNodeData InNodeData) const = 0;
+#if UE_METASOUND_ENABLE_5_6_NODE_REGISTRATION_DEPRECATION_WARNINGS
+			UE_DEPRECATED(5.6, "Create variable mutator nodes with FNodeData")
+#endif
 			virtual TUniquePtr<INode> CreateVariableMutatorNode(const FName& InDataType, const FNodeInitData& InParams) const = 0;
+			virtual TUniquePtr<INode> CreateVariableMutatorNode(const FName& InDataType, FNodeData InNodeData) const = 0;
+#if UE_METASOUND_ENABLE_5_6_NODE_REGISTRATION_DEPRECATION_WARNINGS
+			UE_DEPRECATED(5.6, "Create variable accessor nodes with FNodeData")
+#endif
 			virtual TUniquePtr<INode> CreateVariableAccessorNode(const FName& InDataType, const FNodeInitData& InParams) const = 0;
+			virtual TUniquePtr<INode> CreateVariableAccessorNode(const FName& InDataType, FNodeData InNodeData) const = 0;
+#if UE_METASOUND_ENABLE_5_6_NODE_REGISTRATION_DEPRECATION_WARNINGS
+			UE_DEPRECATED(5.6, "Create variable deferred accessor nodes with FNodeData")
+#endif
 			virtual TUniquePtr<INode> CreateVariableDeferredAccessorNode(const FName& InDataType, const FNodeInitData& InParams) const = 0;
+			virtual TUniquePtr<INode> CreateVariableDeferredAccessorNode(const FName& InDataType, FNodeData InNodeData) const = 0;
 		};
 
 		/** Create data references for all inputs on the FInputVertexInterfaceData 
diff --git a/Engine/Plugins/Runtime/Metasound/Source/MetasoundFrontend/Public/MetasoundInputNode.h b/Engine/Plugins/Runtime/Metasound/Source/MetasoundFrontend/Public/MetasoundInputNode.h
--- a/Engine/Plugins/Runtime/Metasound/Source/MetasoundFrontend/Public/MetasoundInputNode.h
+++ b/Engine/Plugins/Runtime/Metasound/Source/MetasoundFrontend/Public/MetasoundInputNode.h
@@ -1,6 +1,7 @@
 // Copyright Epic Games, Inc. All Rights Reserved.
 #pragma once
 
+#include "MetasoundBasicNode.h"
 #include "MetasoundBuilderInterface.h"
 #include "MetasoundBuildError.h"
 #include "MetasoundDataReference.h"
@@ -8,7 +9,6 @@
 #include "MetasoundFrontendDataTypeTraits.h"
 #include "MetasoundNodeConstructorParams.h"
 #include "MetasoundLiteral.h"
-#include "MetasoundNode.h"
 #include "MetasoundNodeInterface.h"
 #include "MetasoundOperatorInterface.h"
 #include "MetasoundVertexData.h"
@@ -70,107 +70,17 @@ namespace Metasound
 			: FNonExecutableInputOperatorBase(InVertexName, FAnyDataReference{TDataValueReferenceLiteralFactory<DataType>::CreateExplicitArgs(InSettings, InLiteral)})
 			{
 			}
-		};
-
-		template<typename DataType>
-		class TExecutableInputOperator : public IOperator
-		{
-			static_assert(TExecutableDataType<DataType>::bIsExecutable, "TExecutableInputOperatorBase should only be used with executable data types");
-		public:
-			using FDataWriteReference = TDataWriteReference<DataType>;
-
-			UE_DEPRECATED(5.4, "The Executable data types will no longer be supported. Please use PostExecutable data types.") 
-			TExecutableInputOperator(const FVertexName& InDataReferenceName, TDataWriteReference<DataType> InValue)
-				: DataReferenceName(InDataReferenceName)
-				, InputValue(InValue)
-				, OutputValue(FDataWriteReference::CreateNew(*InValue))
-			{
-			}
-
-			virtual void BindInputs(FInputVertexInterfaceData& InOutVertexData) override
-			{
-				InOutVertexData.BindWriteVertex(DataReferenceName, InputValue);
-			}
-
-			virtual void BindOutputs(FOutputVertexInterfaceData& InOutVertexData) override
-			{
-				InOutVertexData.BindReadVertex(DataReferenceName, OutputValue);
-			}
-
-			virtual FExecuteFunction GetExecuteFunction() override
-			{
-				return &Execute;
-			}
-
-			virtual FExecuteFunction GetPostExecuteFunction() override
-			{
-				return nullptr;
-			}
-
-			virtual FResetFunction GetResetFunction() override
-			{
-				return nullptr;
-			}
-
-
-		protected:
-
-			static void Execute(IOperator* InOperator)
-			{
-				using FExecutableInputOperator = TExecutableInputOperator<DataType>;
-
-				FExecutableInputOperator* DerivedOperator = static_cast<FExecutableInputOperator*>(InOperator);
-				check(nullptr != DerivedOperator);
-
-				TExecutableDataType<DataType>::Execute(*(DerivedOperator->InputValue), *(DerivedOperator->OutputValue));
-			}
-
-			FVertexName DataReferenceName;
-
-			FDataWriteReference InputValue;
-			FDataWriteReference OutputValue;
-
-		};
 
-		template<typename DataType>
-		class TResetableExecutableInputOperator : public TExecutableInputOperator<DataType>
-		{
-		public:
-			using FDataWriteReference = TDataWriteReference<DataType>;
-			using FDataWriteReferenceFactory = TDataWriteReferenceLiteralFactory<DataType>;
-
-			TResetableExecutableInputOperator(const FVertexName& InDataReferenceName, const FOperatorSettings& InSettings, const FLiteral& InLiteral)
-			: TExecutableInputOperator<DataType>(InDataReferenceName, FDataWriteReferenceFactory::CreateExplicitArgs(InSettings, InLiteral))
-			, Literal(InLiteral)
+			TInputValueOperator(const FVertexName& InVertexName, const FOperatorSettings& InSettings, const FInputVertexInterfaceData& InInterfaceData)
+			: FNonExecutableInputOperatorBase(InVertexName, FAnyDataReference{InInterfaceData.GetOrCreateDefaultDataValueReference<DataType>(InVertexName, InSettings)})
 			{
 			}
-
-			virtual IOperator::FResetFunction GetResetFunction() override
-			{
-				return &Reset;
-			}
-
-		private:
-
-			static void Reset(IOperator* InOperator, const IOperator::FResetParams& InParams)
-			{
-				using FResetableExecutableInputOperator = TResetableExecutableInputOperator<DataType>;
-
-				FResetableExecutableInputOperator* Operator = static_cast<FResetableExecutableInputOperator*>(InOperator);
-				check(nullptr != Operator);
-
-				*Operator->InputValue = TDataTypeLiteralFactory<DataType>::CreateExplicitArgs(InParams.OperatorSettings, Operator->Literal);
-				*Operator->OutputValue = *Operator->InputValue;
-			}
-
-			FLiteral Literal;
 		};
 
 		template<typename DataType>
 		class TPostExecutableInputOperator : public IOperator
 		{
 			static_assert(TPostExecutableDataType<DataType>::bIsPostExecutable, "TPostExecutableInputOperator should only be used with post executable data types");
-			static_assert(!TExecutableDataType<DataType>::bIsExecutable, "A data type cannot be Executable and PostExecutable");
 
 		public:
 			using FDataWriteReference = TDataWriteReference<DataType>;
@@ -252,6 +162,10 @@ namespace Metasound
 			FAnyDataReference DataRef;
 		};
 
+		// To reset the state of a PostExecutable input operator, we need to reset
+		// the data to it's original state. In order to do that, the FLiteral is
+		// stored on the operator so that it can be used to reinitialize the data 
+		// when the operator reset.  
 		template<typename DataType>
 		class TResetablePostExecutableInputOperator : public TPostExecutableInputOperator<DataType>
 		{
@@ -266,6 +180,13 @@ namespace Metasound
 			{
 			}
 
+			TResetablePostExecutableInputOperator(const FVertexName& InDataReferenceName, const FOperatorSettings& InSettings, const FInputVertexInterfaceData& InData)
+			: TPostExecutableInputOperator<DataType>(InDataReferenceName, InData.GetOrCreateDefaultDataWriteReference<DataType>(InDataReferenceName, InSettings))
+			, Literal(InData.GetVertex(InDataReferenceName).GetDefaultLiteral())
+			{
+				checkf(!InData.IsVertexBound(InDataReferenceName), TEXT("Vertex %s should not be bound when using TResetablePostExecutableInputOperator"), *InDataReferenceName.ToString());
+			}
+
 			virtual IOperator::FResetFunction GetResetFunction() override
 			{
 				if (DataRef.GetAccessType() == EDataReferenceAccessType::Write)
@@ -301,27 +222,18 @@ namespace Metasound
 		/** Non owning input operator that may need execution. */
 		template<typename DataType, EVertexAccessType VertexAccess=EVertexAccessType::Reference>
 		using TNonOwningInputOperator = std::conditional_t<
-			TExecutableDataType<DataType>::bIsExecutable,
-			TExecutableInputOperator<DataType>, // Use this input operator if the data type is not owned by the input node but needs execution.
-			std::conditional_t<
-				TPostExecutableDataType<DataType>::bIsPostExecutable,
-				TPostExecutableInputOperator<DataType>, // Use this input operator if the data type is not owned by the input node but needs post execution.
-				MetasoundInputNodePrivate::FNonExecutableInputPassThroughOperator // Use this input operator if the data type is not owned by the input node and is not executable, nor post executable.
-			>
+			TPostExecutableDataType<DataType>::bIsPostExecutable,
+			TPostExecutableInputOperator<DataType>, // Use this input operator if the data type is not owned by the input node but needs post execution.
+			MetasoundInputNodePrivate::FNonExecutableInputPassThroughOperator // Use this input operator if the data type is not owned by the input node and is not executable, nor post executable.
 		>;
 	}
 
 	/** Owning input operator that may need execution. */
 	template<typename DataType, EVertexAccessType VertexAccess=EVertexAccessType::Reference>
 	using TInputOperator = std::conditional_t<
-		VertexAccess == EVertexAccessType::Value || (!TExecutableDataType<DataType>::bIsExecutable && !TPostExecutableDataType<DataType>::bIsPostExecutable),
+		VertexAccess == EVertexAccessType::Value || !TPostExecutableDataType<DataType>::bIsPostExecutable,
 		MetasoundInputNodePrivate::TInputValueOperator<DataType>, // Use this input operator if the data type is owned by the input node and is not executable, nor post executable.
-		std::conditional_t<
-			TExecutableDataType<DataType>::bIsExecutable,
-			MetasoundInputNodePrivate::TResetableExecutableInputOperator<DataType>, // Use this input operator if the data type is owned by the input node and is executable.
-			MetasoundInputNodePrivate::TResetablePostExecutableInputOperator<DataType> // Use this input operator if the data type is owned by the input node and is post executable.
-
-		>
+		MetasoundInputNodePrivate::TResetablePostExecutableInputOperator<DataType> // Use this input operator if the data type is owned by the input node and is post executable.
 	>;
 
 	/** Choose pass through operator based upon data type and access type */
@@ -332,55 +244,11 @@ namespace Metasound
 		MetasoundInputNodePrivate::FNonExecutableInputPassThroughOperator
 	>;
 
-	/** FInputNode represents an input to a metasound graph. */
-	class METASOUNDFRONTEND_API FInputNode : public FNode
-	{
-		static FLazyName ConstructorVariant;
-		// Use Variant names to differentiate between normal input nodes and constructor 
-		// input nodes.
-		static FName GetVariantName(EVertexAccessType InVertexAccess);
-
-		static FVertexInterface CreateVertexInterface(const FVertexName& InVertexName, const FName& InDataTypeName, EVertexAccessType InVertexAccess, const FLiteral& InLiteral);
-
-	protected:
-
-		static FVertexInterface CreateDefaultVertexInterface(const FVertexName& InVertexName, const FName& InDataTypeName, EVertexAccessType InVertexAccess);
-
-	public:
-		static FText GetInputDescription();
-		static FNodeClassMetadata GetNodeMetadata(const FVertexName& InVertexName, const FName& InDataTypeName, EVertexAccessType InVertexAccess);
-
-		/* Construct a TInputNode using the TInputOperatorLiteralFactory<> and moving
-		 * InParam to the TInputOperatorLiteralFactory constructor.*/
-		explicit FInputNode(FInputNodeConstructorParams&& InParams, const FName& InDataTypeName, EVertexAccessType InVertexAccess, FOperatorFactorySharedRef InFactory);
-
-		const FVertexName& GetVertexName() const;
-
-		virtual const FVertexInterface& GetVertexInterface() const override;
-		virtual bool SetVertexInterface(const FVertexInterface& InInterface) override;
-		virtual bool IsVertexInterfaceSupported(const FVertexInterface& InInterface) const override;
-		virtual TSharedRef<IOperatorFactory, ESPMode::ThreadSafe> GetDefaultOperatorFactory() const override;
-
-	private:
-		FVertexName VertexName;
-		FVertexInterface Interface;
-		FOperatorFactorySharedRef Factory;
-	};
-
-
-	/** TInputNode represents an input to a metasound graph. */
-	template<typename DataType, EVertexAccessType VertexAccess=EVertexAccessType::Reference>
-	class TInputNode : public FInputNode
+	namespace MetasoundInputNodePrivate
 	{
-		static constexpr bool bIsConstructorInput = VertexAccess == EVertexAccessType::Value;
-		static constexpr bool bIsSupportedConstructorInput = TIsConstructorVertexSupported<DataType>::Value && bIsConstructorInput;
-		static constexpr bool bIsReferenceInput = VertexAccess == EVertexAccessType::Reference;
-		static constexpr bool bIsSupportedReferenceInput = TLiteralTraits<DataType>::bIsParsableFromAnyLiteralType && bIsReferenceInput;
-
-		static constexpr bool bIsSupportedInput = bIsSupportedConstructorInput || bIsSupportedReferenceInput;
-
 		// Factory for creating input operators. 
-		class FInputNodeOperatorFactory : public IOperatorFactory
+		template<typename DataType, EVertexAccessType VertexAccess=EVertexAccessType::Reference>
+		class TInputNodeOperatorFactory : public IOperatorFactory
 		{
 			static constexpr bool bIsReferenceVertexAccess = VertexAccess == EVertexAccessType::Reference;
 			static constexpr bool bIsValueVertexAccess = VertexAccess == EVertexAccessType::Value;
@@ -388,8 +256,6 @@ namespace Metasound
 			static_assert(bIsValueVertexAccess || bIsReferenceVertexAccess, "Unsupported EVertexAccessType");
 
 			// Choose which data reference type is created based on template parameters
-			using FDataReference = std::conditional_t<bIsReferenceVertexAccess, TDataReadReference<DataType>, TDataValueReference<DataType>>;
-			using FDataReferenceFactory = std::conditional_t<bIsReferenceVertexAccess, TDataReadReferenceLiteralFactory<DataType>, TDataValueReferenceLiteralFactory<DataType>>;
 			using FPassThroughDataReference = std::conditional_t<bIsReferenceVertexAccess, TDataReadReference<DataType>, TDataValueReference<DataType>>;
 
 			// Return correct data reference type based on vertex access type for pass through scenario.
@@ -410,7 +276,8 @@ namespace Metasound
 			}
 
 		public:
-			explicit FInputNodeOperatorFactory()
+			explicit TInputNodeOperatorFactory(const FVertexName& InVertexName)
+			: VertexName(InVertexName)
 			{
 			}
 
@@ -418,52 +285,102 @@ namespace Metasound
 			{
 				using namespace MetasoundInputNodePrivate;
 
-				using FInputNodeType = TInputNode<DataType, VertexAccess>;
-
-				const FInputNodeType& InputNode = static_cast<const FInputNodeType&>(InParams.Node);
-				const FVertexName& VertexKey = InputNode.GetVertexName();
-
-				if (const FAnyDataReference* Ref = InParams.InputData.FindDataReference(VertexKey))
+				if (const FAnyDataReference* Ref = InParams.InputData.FindDataReference(VertexName))
 				{
 					if constexpr (bIsReferenceVertexAccess)
 					{
 						if (EDataReferenceAccessType::Write == Ref->GetAccessType())
 						{
-							return MakeUnique<TNonOwningInputOperator<DataType, VertexAccess>>(VertexKey, Ref->GetDataWriteReference<DataType>());
+							return MakeUnique<TNonOwningInputOperator<DataType, VertexAccess>>(VertexName, Ref->GetDataWriteReference<DataType>());
 						}
 					}
 					// Pass through input value
-					return MakeUnique<TPassThroughOperator<DataType, VertexAccess>>(VertexKey, CreatePassThroughDataReference(*Ref));
+					return MakeUnique<TPassThroughOperator<DataType, VertexAccess>>(VertexName, CreatePassThroughDataReference(*Ref));
 				}
 				else
 				{
-					const FLiteral& Literal = InputNode.GetVertexInterface().GetInputInterface()[VertexKey].GetDefaultLiteral();
 					// Owned input value
-					return MakeUnique<TInputOperator<DataType, VertexAccess>>(VertexKey, InParams.OperatorSettings, Literal);
+					return MakeUnique<TInputOperator<DataType, VertexAccess>>(VertexName, InParams.OperatorSettings, InParams.InputData);
 				}
 			}
+
+		private:
+			FVertexName VertexName;
 		};
+	}
+
+	/** FInputNode represents an input to a metasound graph. */
+	class METASOUNDFRONTEND_API FInputNode : public FBasicNode
+	{
+		static FLazyName ConstructorVariant;
+		// Use Variant names to differentiate between normal input nodes and constructor 
+		// input nodes.
+		static FName GetVariantName(EVertexAccessType InVertexAccess);
+
+		static FVertexInterface CreateVertexInterface(const FVertexName& InVertexName, const FName& InDataTypeName, EVertexAccessType InVertexAccess, const FLiteral& InLiteral);
+
+	protected:
 
+		static FVertexInterface CreateDefaultVertexInterface(const FVertexName& InVertexName, const FName& InDataTypeName, EVertexAccessType InVertexAccess, const FLiteral* InDefaultLiteral=nullptr);
+
+		explicit FInputNode(FOperatorFactorySharedRef InFactory, FNodeData InNodeData, TSharedRef<const FNodeClassMetadata> InClassMetadata);
+
+	public:
+		static FText GetInputDescription();
+		UE_DEPRECATED(5.6, "Use CreateNodeClassMetadata(...) instead")
+		static FNodeClassMetadata GetNodeMetadata(const FVertexName& InVertexName, const FName& InDataTypeName, EVertexAccessType InVertexAccess);
+		static FNodeClassMetadata CreateNodeClassMetadata(const FVertexName& InVertexName, const FName& InDataTypeName, EVertexAccessType InVertexAccess);
+
+		/* Construct a TInputNode using the TInputOperatorLiteralFactory<> and moving
+		 * InParam to the TInputOperatorLiteralFactory constructor.*/
+		explicit FInputNode(FInputNodeConstructorParams&& InParams, const FName& InDataTypeName, EVertexAccessType InVertexAccess, FOperatorFactorySharedRef InFactory);
+
+
+		UE_DEPRECATED(5.6, "You can find the vertex name by inspecting the FVertexInterface.")
+		const FVertexName& GetVertexName() const;
+
+		virtual TSharedRef<IOperatorFactory, ESPMode::ThreadSafe> GetDefaultOperatorFactory() const override;
+
+	private:
+		FOperatorFactorySharedRef Factory;
+	};
+
+
+	/** TInputNode represents an input to a metasound graph. */
+	template<typename DataType, EVertexAccessType VertexAccess=EVertexAccessType::Reference>
+	class TInputNode : public FInputNode
+	{
+		static constexpr bool bIsConstructorInput = VertexAccess == EVertexAccessType::Value;
+		static constexpr bool bIsSupportedConstructorInput = TIsConstructorVertexSupported<DataType>::Value && bIsConstructorInput;
+		static constexpr bool bIsReferenceInput = VertexAccess == EVertexAccessType::Reference;
+		static constexpr bool bIsSupportedReferenceInput = TLiteralTraits<DataType>::bIsParsableFromAnyLiteralType && bIsReferenceInput;
+
+		static constexpr bool bIsSupportedInput = bIsSupportedConstructorInput || bIsSupportedReferenceInput;
 
 	public:
 		// If true, this node can be instantiated by the Frontend.
 		static constexpr bool bCanRegister = bIsSupportedInput;
 
-		UE_DEPRECATED(5.3, "Access the default vertex interface from the input node metadata.")
-		static FVertexInterface DeclareVertexInterface(const FVertexName& InVertexName)
+		UE_DEPRECATED(5.6, "Use CreateNodeClassMetadata(...) instead")
+		static FNodeClassMetadata GetNodeInfo(const FVertexName& InVertexName)
 		{
-			return CreateDefaultVertexInterface(InVertexName, GetMetasoundDataTypeName<DataType>(), VertexAccess);
+			return FInputNode::CreateNodeClassMetadata(InVertexName, GetMetasoundDataTypeName<DataType>(), VertexAccess);
 		}
 
-		static FNodeClassMetadata GetNodeInfo(const FVertexName& InVertexName)
+		static FNodeClassMetadata CreateNodeClassMetadata(const FVertexName& InVertexName)
 		{
-			return GetNodeMetadata(InVertexName, GetMetasoundDataTypeName<DataType>(), VertexAccess);
+			return FInputNode::CreateNodeClassMetadata(InVertexName, GetMetasoundDataTypeName<DataType>(), VertexAccess);
 		}
 
 		/* Construct a TInputNode using the TInputOperatorLiteralFactory<> and moving
 		 * InParam to the TInputOperatorLiteralFactory constructor.*/
 		explicit TInputNode(FInputNodeConstructorParams&& InParams)
-		:	FInputNode(MoveTemp(InParams), GetMetasoundDataTypeName<DataType>(), VertexAccess, MakeShared<FInputNodeOperatorFactory>())
+		:	FInputNode(MoveTemp(InParams), GetMetasoundDataTypeName<DataType>(), VertexAccess, MakeShared<MetasoundInputNodePrivate::TInputNodeOperatorFactory<DataType, VertexAccess>>(InParams.VertexName))
+		{
+		}
+
+		explicit TInputNode(const FVertexName& InVertexName, FNodeData InNodeData, TSharedRef<const FNodeClassMetadata> InClassMetadata)
+		: FInputNode(MakeShared<MetasoundInputNodePrivate::TInputNodeOperatorFactory<DataType, VertexAccess>>(InVertexName), MoveTemp(InNodeData), MoveTemp(InClassMetadata))
 		{
 		}
 	};
diff --git a/Engine/Plugins/Runtime/Metasound/Source/MetasoundFrontend/Public/MetasoundOutputNode.h b/Engine/Plugins/Runtime/Metasound/Source/MetasoundFrontend/Public/MetasoundOutputNode.h
--- a/Engine/Plugins/Runtime/Metasound/Source/MetasoundFrontend/Public/MetasoundOutputNode.h
+++ b/Engine/Plugins/Runtime/Metasound/Source/MetasoundFrontend/Public/MetasoundOutputNode.h
@@ -2,6 +2,7 @@
 #pragma once
 
 #include "Internationalization/Text.h"
+#include "MetasoundBasicNode.h"
 #include "MetasoundBuildError.h"
 #include "MetasoundBuilderInterface.h"
 #include "MetasoundExecutableOperator.h"
@@ -18,6 +19,11 @@
 
 namespace Metasound
 {
+	namespace OutputNodePrivate
+	{
+		class FOutputOperatorFactory;
+	}
+
 	/** FMissingOutputNodeInputReferenceError
 	 *
 	 * Caused by Output not being able to generate an IOperator instance due to
@@ -45,81 +51,33 @@ namespace Metasound
 		virtual ~FMissingOutputNodeInputReferenceError() = default;
 	};
 
-	namespace OutputNodePrivate
+	class METASOUNDFRONTEND_API FOutputNode : public FBasicNode
 	{
-		class METASOUNDFRONTEND_API FOutputOperator : public FNoOpOperator
-		{
-			public:
-				FOutputOperator(const FVertexName& InVertexName, const FAnyDataReference& InDataReference);
-				virtual ~FOutputOperator() = default;
-
-				virtual FDataReferenceCollection GetInputs() const override;
-				virtual FDataReferenceCollection GetOutputs() const override;
-				virtual void BindInputs(FInputVertexInterfaceData& InVertexData) override;
-				virtual void BindOutputs(FOutputVertexInterfaceData& InVertexData) override;
-
-			private:
-				FVertexName VertexName;
-				FAnyDataReference DataReference;
-		};
-
-		class METASOUNDFRONTEND_API FOutputNode : public FNode
-		{
-		private:
-			class FOutputOperatorFactory : public IOperatorFactory
-			{
-			public:
-
-				FOutputOperatorFactory(const FVertexName& InVertexName, const FName& InDataTypeName, EVertexAccessType InVertexAccessType);
-				virtual TUniquePtr<IOperator> CreateOperator(const FBuildOperatorParams& InParams, FBuildResults& OutResults) override;
-
-			private:
-
-				FVertexName VertexName;
-				FName DataTypeName;
-				EVertexAccessType VertexAccessType;
-			};
-
-			static FName GetVariantName(EVertexAccessType InVertexAccessType);
-			static FVertexInterface GetVertexInterface(const FVertexName& InVertexName, const FName& InDataTypeName, EVertexAccessType InVertexAccessType);
-			static FNodeClassMetadata GetNodeInfo(const FVertexName& InVertexName, const FName& InDataTypeName, EVertexAccessType InVertexAccessType);
-
-			public:
-				FOutputNode(const FName& InInstanceName, const FGuid& InInstanceID, const FVertexName& InVertexName, const FName& InDataType, EVertexAccessType InVertexAccessType);
-
-				/** Return the current vertex interface. */
-				virtual const FVertexInterface& GetVertexInterface() const override;
-
-				/** Set the vertex interface. If the vertex was successfully changed, returns true. 
-				 *
-				 * @param InInterface - New interface for node. 
-				 *
-				 * @return True on success, false otherwise.
-				 */
-				virtual bool SetVertexInterface(const FVertexInterface& InInterface) override;
+	public:
+		FOutputNode(const FVertexName& InVertexName, FNodeData InNodeData, TSharedRef<const FNodeClassMetadata> InClassMetadata);
 
-				/** Expresses whether a specific vertex interface is supported.
-				 *
-				 * @param InInterface - New interface. 
-				 *
-				 * @return True if the interface is supported, false otherwise. 
-				 */
-				virtual bool IsVertexInterfaceSupported(const FVertexInterface& InInterface) const;
+		static FName GetVariantName(EVertexAccessType InVertexAccessType);
+		static FNodeClassMetadata CreateNodeClassMetadata(const FVertexName& InVertexName, const FName& InDataTypeName, EVertexAccessType InVertexAccessType);
+		static FVertexInterface CreateVertexInterface(const FVertexName& InVertexName, const FName& InDataTypeName, EVertexAccessType InVertexAccessType);
 
-				virtual TSharedRef<IOperatorFactory, ESPMode::ThreadSafe> GetDefaultOperatorFactory() const override;
+		virtual TSharedRef<IOperatorFactory, ESPMode::ThreadSafe> GetDefaultOperatorFactory() const override;
 
-			private:
-				TSharedRef<FOutputOperatorFactory, ESPMode::ThreadSafe> Factory;
-		};
-	}
+	private:
+		TSharedRef<OutputNodePrivate::FOutputOperatorFactory, ESPMode::ThreadSafe> Factory;
+	};
 
 	/** Output nodes are used to expose graph data to external entities. */
 	template<typename DataType, EVertexAccessType VertexAccess=EVertexAccessType::Reference>
-	class TOutputNode : public OutputNodePrivate::FOutputNode
+	class TOutputNode : public FOutputNode
 	{
 	public:
 		TOutputNode(const FVertexName& InInstanceName, const FGuid& InInstanceID, const FVertexName& InVertexName)
-		:	FOutputNode(InInstanceName, InInstanceID, InVertexName, GetMetasoundDataTypeName<DataType>(), VertexAccess)
+		:	TOutputNode(InVertexName, FNodeData{InInstanceName, InInstanceID, FOutputNode::CreateVertexInterface(InVertexName, GetMetasoundDataTypeName<DataType>(), VertexAccess)}, MakeShared<FNodeClassMetadata>(FOutputNode::CreateNodeClassMetadata(InVertexName, GetMetasoundDataTypeName<DataType>(), VertexAccess)))
+		{
+		}
+
+		TOutputNode(const FVertexName& InVertexName, FNodeData InNodeData, TSharedRef<const FNodeClassMetadata> InClassMetadata)
+		:	FOutputNode(InVertexName, MoveTemp(InNodeData), MoveTemp(InClassMetadata))
 		{
 		}
 	};
diff --git a/Engine/Plugins/Runtime/Metasound/Source/MetasoundGraphCore/Private/MetasoundVertex.cpp b/Engine/Plugins/Runtime/Metasound/Source/MetasoundGraphCore/Private/MetasoundVertex.cpp
--- a/Engine/Plugins/Runtime/Metasound/Source/MetasoundGraphCore/Private/MetasoundVertex.cpp
+++ b/Engine/Plugins/Runtime/Metasound/Source/MetasoundGraphCore/Private/MetasoundVertex.cpp
@@ -416,6 +416,12 @@ namespace Metasound
 		return InLHS.VertexName.FastLess(InRHS.VertexName);
 	}
 
+	FInputVertexInterface::FInputVertexInterface(TArray<FInputDataVertex> InVertices, TArray<VertexPrivate::FSubInterfaceLayout> InSubInterfaces)
+	: TVertexInterfaceImpl<FInputDataVertex>(MoveTemp(InVertices))
+	, SubInterfaces(MoveTemp(InSubInterfaces))
+	{
+	}
+
 	void FInputVertexInterface::ForEachSubInterfaceInstance(const FName& InSubInterfaceName, TFunctionRef<void (TConstArrayView<FInputDataVertex>)> Callable) const
 	{
 		using namespace VertexPrivate;
@@ -456,6 +462,12 @@ namespace Metasound
 		return SubInterfaces.FindByPredicate(VertexPrivate::FEqualSubInterfaceName(InName));
 	}
 
+	FOutputVertexInterface::FOutputVertexInterface(TArray<FOutputDataVertex> InVertices, TArray<VertexPrivate::FSubInterfaceLayout> InSubInterfaces)
+	: TVertexInterfaceImpl<FOutputDataVertex>(MoveTemp(InVertices))
+	, SubInterfaces(MoveTemp(InSubInterfaces))
+	{
+	}
+
 	void FOutputVertexInterface::ForEachSubInterfaceInstance(const FName& InSubInterfaceName, TFunctionRef<void (TConstArrayView<FOutputDataVertex>)> Callable) const
 	{
 		using namespace VertexPrivate;
@@ -494,6 +506,10 @@ namespace Metasound
 		return SubInterfaces.FindByPredicate(VertexPrivate::FEqualSubInterfaceName(InName));
 	}
 
+	FEnvironmentVertexInterface::FEnvironmentVertexInterface(TArray<FEnvironmentVertex> InVertices)
+	: TVertexInterfaceImpl<FEnvironmentVertex>(MoveTemp(InVertices))
+	{
+	}
 
 	FVertexInterface::FVertexInterface(const FInputVertexInterface& InInputs, const FOutputVertexInterface& InOutputs)
 	:	InputInterface(InInputs)
@@ -501,10 +517,10 @@ namespace Metasound
 	{
 	}
 
-	FVertexInterface::FVertexInterface(const FInputVertexInterface& InInputs, const FOutputVertexInterface& InOutputs, const FEnvironmentVertexInterface& InEnvironmentVariables)
-	:	InputInterface(InInputs)
-	,	OutputInterface(InOutputs)
-	,	EnvironmentInterface(InEnvironmentVariables)
+	FVertexInterface::FVertexInterface(FInputVertexInterface InInputs, FOutputVertexInterface InOutputs, FEnvironmentVertexInterface InEnvironmentVariables)
+	:	InputInterface(MoveTemp(InInputs))
+	,	OutputInterface(MoveTemp(InOutputs))
+	,	EnvironmentInterface(MoveTemp(InEnvironmentVariables))
 	{
 	}
 
diff --git a/Engine/Plugins/Runtime/Metasound/Source/MetasoundGraphCore/Public/MetasoundFacade.h b/Engine/Plugins/Runtime/Metasound/Source/MetasoundGraphCore/Public/MetasoundFacade.h
--- a/Engine/Plugins/Runtime/Metasound/Source/MetasoundGraphCore/Public/MetasoundFacade.h
+++ b/Engine/Plugins/Runtime/Metasound/Source/MetasoundGraphCore/Public/MetasoundFacade.h
@@ -149,13 +149,20 @@ namespace Metasound
 	class TNodeFacade : public FNodeFacade
 	{
 	public:
+		TNodeFacade(const FName& InNodeName, const FGuid& InNodeID)
+		: TNodeFacade(FNodeInitData{InNodeName, InNodeID})
+		{
+		}
+
 		TNodeFacade(const FNodeInitData& InInitData)
 			: FNodeFacade(FNodeData(InInitData.InstanceName, InInitData.InstanceID, OperatorType::GetNodeInfo().DefaultInterface), MakeShared<FNodeClassMetadata>(OperatorType::GetNodeInfo()), TFacadeOperatorClass<OperatorType>())
-		{}
+		{
+		}
 
 		TNodeFacade(FNodeData InNodeData, TSharedRef<const FNodeClassMetadata> InClassMetadata)
 			: FNodeFacade(MoveTemp(InNodeData), MoveTemp(InClassMetadata), TFacadeOperatorClass<OperatorType>())
-		{}
+		{
+		}
 
 		static FNodeClassMetadata CreateNodeClassMetadata()
 		{
diff --git a/Engine/Plugins/Runtime/Metasound/Source/MetasoundGraphCore/Public/MetasoundVariableNodes.h b/Engine/Plugins/Runtime/Metasound/Source/MetasoundGraphCore/Public/MetasoundVariableNodes.h
--- a/Engine/Plugins/Runtime/Metasound/Source/MetasoundGraphCore/Public/MetasoundVariableNodes.h
+++ b/Engine/Plugins/Runtime/Metasound/Source/MetasoundGraphCore/Public/MetasoundVariableNodes.h
@@ -484,7 +484,7 @@ namespace Metasound
 	/** Variable nodes initialize variable values. The output of a VariableNode
 	 * is a TVariable.  */
 	template<typename DataType>
-	class TVariableNode : public FNode
+	class TVariableNode : public FBasicNode
 	{
 		using FVariable = TVariable<DataType>;
 
@@ -520,7 +520,20 @@ namespace Metasound
 			);
 		}
 
-		static FNodeClassMetadata GetNodeMetadata()
+
+	public:
+		TVariableNode(const FVertexName& InNodeName, const FGuid& InInstanceID, FLiteral&& InLiteral)
+		: TVariableNode(MoveTemp(InLiteral), FNodeData{InNodeName, InInstanceID, DeclareVertexInterface()}, MakeShared<const FNodeClassMetadata>(CreateNodeClassMetadata()))
+		{
+		}
+
+		TVariableNode(FLiteral InLiteral, FNodeData InNodeData, TSharedRef<const FNodeClassMetadata> InClassMetadata)
+		: FBasicNode(MoveTemp(InNodeData), MoveTemp(InClassMetadata))
+		, Factory(MakeShared<FFactory>(MoveTemp(InLiteral)))
+		{
+		}
+
+		static FNodeClassMetadata CreateNodeClassMetadata()
 		{
 			FNodeClassMetadata Info;
 
@@ -541,14 +554,6 @@ namespace Metasound
 		}
 
 
-	public:
-		TVariableNode(const FVertexName& InNodeName, const FGuid& InInstanceID, FLiteral&& InLiteral)
-		: FNode(InNodeName, InInstanceID, GetNodeMetadata())
-		, Interface(DeclareVertexInterface())
-		, Factory(MakeShared<FFactory>(MoveTemp(InLiteral)))
-		{
-		}
-
 		virtual TSharedRef<IOperatorFactory, ESPMode::ThreadSafe> GetDefaultOperatorFactory() const override
 		{
 			return Factory;
@@ -556,25 +561,12 @@ namespace Metasound
 
 	private:
 
-		FVertexInterface Interface;
 		FOperatorFactorySharedRef Factory;
 	};
 
 	/** TVariableMutatorNode allows variable values to be set.  */
 	template<typename DataType>
-	class TVariableMutatorNode : public FNodeFacade
-	{
-
-	public:
-		TVariableMutatorNode(const FVertexName& InNodeName, const FGuid& InInstanceID)
-		: FNodeFacade(InNodeName, InInstanceID, TFacadeOperatorClass<MetasoundVariableNodesPrivate::TVariableMutatorOperator<DataType>>())
-		{
-		}
-
-		TVariableMutatorNode(const FNodeInitData& InInitData)
-		: TVariableMutatorNode(InInitData.InstanceName, InInitData.InstanceID)
-		{}
-	};
+	using TVariableMutatorNode = TNodeFacade<MetasoundVariableNodesPrivate::TVariableMutatorOperator<DataType>>;
 
 
 	/** TVariableDeferredAccessorNode provides access to the prior executions variable value.
@@ -582,35 +574,11 @@ namespace Metasound
 	 * order.
 	 */
 	template<typename DataType>
-	class TVariableDeferredAccessorNode : public FNodeFacade
-	{
-	public:
-
-		TVariableDeferredAccessorNode(const FVertexName& InNodeName, const FGuid& InInstanceID)
-		: FNodeFacade(InNodeName, InInstanceID, TFacadeOperatorClass<MetasoundVariableNodesPrivate::TVariableDeferredAccessorOperator<DataType>>())
-		{
-		}
-
-		TVariableDeferredAccessorNode(const FNodeInitData& InInitData)
-		: TVariableDeferredAccessorNode(InInitData.InstanceName, InInitData.InstanceID)
-		{}
-	};
+	using TVariableDeferredAccessorNode = TNodeFacade<MetasoundVariableNodesPrivate::TVariableDeferredAccessorOperator<DataType>>;
 
 	/** FGetVariable node provides delay free, cpu free access to a set variable. */
 	template<typename DataType>
-	class TVariableAccessorNode : public FNodeFacade
-	{
-	public:
-
-		TVariableAccessorNode(const FVertexName& InNodeName, const FGuid& InInstanceID)
-		: FNodeFacade(InNodeName, InInstanceID, TFacadeOperatorClass<MetasoundVariableNodesPrivate::TVariableAccessorOperator<DataType>>())
-		{
-		}
-
-		TVariableAccessorNode(const FNodeInitData& InInitData)
-		: TVariableAccessorNode(InInitData.InstanceName, InInitData.InstanceID)
-		{}
-	};
+	using TVariableAccessorNode = TNodeFacade<MetasoundVariableNodesPrivate::TVariableAccessorOperator<DataType>>;
 }
 
 #undef LOCTEXT_NAMESPACE // MetasoundGraphCore
diff --git a/Engine/Plugins/Runtime/Metasound/Source/MetasoundGraphCore/Public/MetasoundVertex.h b/Engine/Plugins/Runtime/Metasound/Source/MetasoundGraphCore/Public/MetasoundVertex.h
--- a/Engine/Plugins/Runtime/Metasound/Source/MetasoundGraphCore/Public/MetasoundVertex.h
+++ b/Engine/Plugins/Runtime/Metasound/Source/MetasoundGraphCore/Public/MetasoundVertex.h
@@ -162,6 +162,12 @@ namespace Metasound
 		{
 			return Literal;
 		}
+		
+		/** Set the default literal for this vertex */
+		void SetDefaultLiteral(const FLiteral& InLiteral)
+		{
+			Literal = InLiteral;
+		}
 
 		friend bool METASOUNDGRAPHCORE_API operator==(const FInputDataVertex& InLHS, const FInputDataVertex& InRHS);
 		friend bool METASOUNDGRAPHCORE_API operator!=(const FInputDataVertex& InLHS, const FInputDataVertex& InRHS);
@@ -307,6 +313,12 @@ namespace Metasound
 		using RangedForConstIteratorType = typename TArray<VertexType>::RangedForConstIteratorType;
 
 		TVertexInterfaceImpl() = default;
+
+		/** Construct with prebuilt array of vertices. */
+		explicit TVertexInterfaceImpl(TArray<VertexType> InVertices)
+		: Vertices(MoveTemp(InVertices))
+		{
+		}
 		
 		/** TVertexInterfaceImpl constructor with variadic list of vertex
 		 * models.
@@ -400,6 +412,18 @@ namespace Metasound
 			return Vertices.Num();
 		}
 
+		/** Return a vertex at an index. */
+		const VertexType& At(int32 InIndex) const
+		{
+			return Vertices[InIndex];
+		}
+		
+		/** Return a vertex at an index. */
+		VertexType& At(int32 InIndex)
+		{
+			return Vertices[InIndex];
+		}
+
 		/** Compare whether two vertex groups are equal. */
 		friend bool operator==(const TVertexInterfaceImpl<VertexType>& InLHS, const TVertexInterfaceImpl<VertexType>& InRHS)
 		{
@@ -622,6 +646,8 @@ namespace Metasound
 			InterfaceBuilder.Build(Forward<ArgTypes>(InArgs)...);
 		}
 
+		explicit FInputVertexInterface(TArray<FInputDataVertex> InVertices, TArray<VertexPrivate::FSubInterfaceLayout> InSubInterfaces);
+
 		/** Iterate through all repetitions of a sub interface. */
 		void ForEachSubInterfaceInstance(const FName& InSubInterfaceName, TFunctionRef<void (TConstArrayView<FInputDataVertex>)> Callable) const;
 		/** Iterate through all repetitions of a sub interface. */
@@ -662,6 +688,8 @@ namespace Metasound
 			InterfaceBuilder.Build(Forward<ArgTypes>(InArgs)...);
 		}
 
+		explicit FOutputVertexInterface(TArray<FOutputDataVertex> InVertices, TArray<VertexPrivate::FSubInterfaceLayout> InSubInterfaces);
+
 		/** Iterate through all repetitions of a sub interface. */
 		void ForEachSubInterfaceInstance(const FName& InSubInterfaceName, TFunctionRef<void (TConstArrayView<FOutputDataVertex>)> Callable) const;
 		/** Iterate through all repetitions of a sub interface. */
@@ -690,6 +718,8 @@ namespace Metasound
 			VertexPrivate::FEnvironmentDeclarationBuilder InterfaceBuilder(Vertices);
 			InterfaceBuilder.Build(Forward<ArgTypes>(InArgs)...);
 		}
+
+		explicit FEnvironmentVertexInterface(TArray<FEnvironmentVertex> InVertices);
 	};
 
 	/** FVertexInterface provides access to a collection of input and output vertex
@@ -706,7 +736,7 @@ namespace Metasound
 			FVertexInterface(const FInputVertexInterface& InInputs, const FOutputVertexInterface& InOutputs);
 
 			/** Construct with input, output and environment interface. */
-			FVertexInterface(const FInputVertexInterface& InInputs, const FOutputVertexInterface& InOutputs, const FEnvironmentVertexInterface& InEnvironmentVariables);
+			FVertexInterface(FInputVertexInterface InInputs, FOutputVertexInterface InOutputs, FEnvironmentVertexInterface InEnvironmentVariables);
 
 			/** Return the input interface. */
 			const FInputVertexInterface& GetInputInterface() const;
diff --git a/Engine/Plugins/Runtime/Metasound/Source/MetasoundGraphCore/Public/MetasoundVertexData.h b/Engine/Plugins/Runtime/Metasound/Source/MetasoundGraphCore/Public/MetasoundVertexData.h
--- a/Engine/Plugins/Runtime/Metasound/Source/MetasoundGraphCore/Public/MetasoundVertexData.h
+++ b/Engine/Plugins/Runtime/Metasound/Source/MetasoundGraphCore/Public/MetasoundVertexData.h
@@ -136,6 +136,21 @@ namespace Metasound
 			}
 		};
 
+		template<typename DataType>
+		struct TFactoryFacade<TDataValueReference<DataType>>
+		{
+			static TDataValueReference<DataType> CreateWithLiteral(const FOperatorSettings& InSettings, const FLiteral& InLiteral)
+			{
+				return TDataValueReferenceLiteralFactory<DataType>::CreateExplicitArgs(InSettings, InLiteral);
+			}
+
+			template<typename ... ArgTypes>
+			static TDataValueReference<DataType> CreateWithArgs(ArgTypes&&... Args)
+			{
+				return TDataValueReference<DataType>::CreateNew(Forward<ArgTypes>(Args)...);
+			}
+		};
+
 		template<typename T>
 		struct TTriggerDeprecatedBindAPIWarning
 		{
@@ -774,6 +789,16 @@ namespace Metasound
 			return Binding->GetOrCreateDefaultAs<DataType>(InSettings);
 		}
 
+		/**  Gets the value of the bound data reference if it exists. Otherwise
+		 * create and return a value by constructing one using the Vertex's 
+		 * default literal. */
+		template<typename DataType>
+		TDataValueReference<DataType> GetOrCreateDefaultDataValueReference(const FVertexName& InVertexName, const FOperatorSettings& InSettings) const
+		{
+			const FInputBinding* Binding = FindChecked(InVertexName);
+			return Binding->GetOrCreateDefaultAs<TDataValueReference<DataType>>(InSettings);
+		}
+
 		/** Get data read reference assuming data is bound and read or write accessible. */
 		template<typename DataType> 
 		TDataReadReference<DataType> GetDataReadReference(const FVertexName& InVertexName) const
@@ -845,7 +870,7 @@ namespace Metasound
 				UE_LOG(LogMetaSound, Warning, TEXT("Failed find vertex with name '%s'. Cannot check for existing bound data."), *InVertexName.ToString());
 			}
 
-			return TDataWriteReferenceLiteralFactory<DataType>::CreateAny(InSettings, FLiteral::CreateInvalid());
+			return TDataWriteReferenceLiteralFactory<DataType>::CreateExplicitArgs(InSettings, FLiteral::CreateInvalid());
 		}
 
 		/** Get the bound data write reference if it exists. Otherwise create and 
diff --git a/Engine/Plugins/Runtime/Metasound/Source/MetasoundStandardNodes/Private/MetasoundAudioFormats.cpp b/Engine/Plugins/Runtime/Metasound/Source/MetasoundStandardNodes/Private/MetasoundAudioFormats.cpp
--- a/Engine/Plugins/Runtime/Metasound/Source/MetasoundStandardNodes/Private/MetasoundAudioFormats.cpp
+++ b/Engine/Plugins/Runtime/Metasound/Source/MetasoundStandardNodes/Private/MetasoundAudioFormats.cpp
@@ -7,6 +7,7 @@
 #include "MetasoundAudioBuffer.h"
 #include "MetasoundDataReference.h"
 #include "MetasoundDataTypeRegistrationMacro.h"
+#include "MetasoundInputNode.h"
 #include "MetasoundLiteral.h"
 #include "MetasoundOutputNode.h"
 #include "MetasoundVertex.h"
@@ -146,8 +147,8 @@ namespace Metasound
 				using namespace StereoAudioFormatVertexKeys;
 
 				// Construct stereo from left and right audio buffers. 
-				TDataReadReference<FAudioBuffer> Left = InParams.InputData.GetOrConstructDataReadReference<FAudioBuffer>(METASOUND_GET_PARAM_NAME(LeftChannelVertex), InParams.OperatorSettings);
-				TDataReadReference<FAudioBuffer> Right = InParams.InputData.GetOrConstructDataReadReference<FAudioBuffer>(METASOUND_GET_PARAM_NAME(RightChannelVertex), InParams.OperatorSettings);
+				TDataReadReference<FAudioBuffer> Left = InParams.InputData.GetOrCreateDefaultDataReadReference<FAudioBuffer>(METASOUND_GET_PARAM_NAME(LeftChannelVertex), InParams.OperatorSettings);
+				TDataReadReference<FAudioBuffer> Right = InParams.InputData.GetOrCreateDefaultDataReadReference<FAudioBuffer>(METASOUND_GET_PARAM_NAME(RightChannelVertex), InParams.OperatorSettings);
 
 				TDataReadReference<FStereoAudioFormat> Stereo = TDataReadReferenceFactory<FStereoAudioFormat>::CreateExplicitArgs(InParams.OperatorSettings, WriteCast(Left), WriteCast(Right));
 
@@ -158,7 +159,7 @@ namespace Metasound
 			FVertexName OutputName;
 		};
 
-		static FVertexInterface GetVertexInterface(const FVertexName& InVertexName)
+		static FVertexInterface CreateVertexInterface(const FVertexName& InVertexName)
 		{
 			using namespace StereoAudioFormatVertexKeys; 
 
@@ -173,7 +174,17 @@ namespace Metasound
 			);
 		}
 
-		static FNodeClassMetadata GetNodeClassMetadata(const FVertexName& InOutputName)
+
+
+	public:
+		TOutputNode(const FVertexName& InInstanceName, const FGuid& InInstanceID, const FVertexName& InVertexName)
+		:	FNode(InInstanceName, InInstanceID, CreateNodeClassMetadata(InVertexName))
+		,	VertexInterface(CreateVertexInterface(InVertexName))
+		,	Factory(MakeShared<FOutputOperatorFactory, ESPMode::ThreadSafe>(InVertexName))
+		{
+		}
+
+		static FNodeClassMetadata CreateNodeClassMetadata(const FVertexName& InOutputName)
 		{
 			FNodeClassMetadata Info;
 
@@ -184,20 +195,11 @@ namespace Metasound
 			Info.Description = METASOUND_LOCTEXT("Metasound_OutputNodeDescription", "Output from the parent Metasound graph.");
 			Info.Author = PluginAuthor;
 			Info.PromptIfMissing = PluginNodeMissingPrompt;
-			Info.DefaultInterface = GetVertexInterface(InOutputName);
+			Info.DefaultInterface = CreateVertexInterface(InOutputName);
 
 			return Info;
 		};
 
-
-	public:
-		TOutputNode(const FVertexName& InInstanceName, const FGuid& InInstanceID, const FVertexName& InVertexName)
-		:	FNode(InInstanceName, InInstanceID, GetNodeClassMetadata(InVertexName))
-		,	VertexInterface(GetVertexInterface(InVertexName))
-		,	Factory(MakeShared<FOutputOperatorFactory, ESPMode::ThreadSafe>(InVertexName))
-		{
-		}
-
 		const FVertexInterface& GetVertexInterface() const override
 		{
 			return VertexInterface;
@@ -227,7 +229,7 @@ namespace Metasound
 
 	// Input node specialization to expose left/right audio buffers 
 	template<>
-	class METASOUNDSTANDARDNODES_API TInputNode<FStereoAudioFormat> : public FNode
+	class METASOUNDSTANDARDNODES_API TInputNode<FStereoAudioFormat> : public FInputNode
 	{
 		// Noop operator. Used to return inputs / outputs for debugging.
 		class FInputOperator : public FNoOpOperator
@@ -275,7 +277,7 @@ namespace Metasound
 			TUniquePtr<IOperator> CreateOperator(const FBuildOperatorParams& InParams, FBuildResults& OutResults) override
 			{
 				// Split a stereo signal into left/right.
-				TDataReadReference<FStereoAudioFormat> Stereo = InParams.InputData.GetOrConstructDataReadReference<FStereoAudioFormat>(InputName, InParams.OperatorSettings);
+				TDataReadReference<FStereoAudioFormat> Stereo = InParams.InputData.GetOrCreateDefaultDataReadReference<FStereoAudioFormat>(InputName, InParams.OperatorSettings);
 				TDataReadReference<FAudioBuffer> Left = Stereo->GetLeft();
 				TDataReadReference<FAudioBuffer> Right = Stereo->GetRight();
 
@@ -286,7 +288,7 @@ namespace Metasound
 			FVertexName InputName;
 		};
 
-		static FVertexInterface GetVertexInterface(const FVertexName& InVertexName)
+		static FVertexInterface CreateVertexInterface(const FVertexName& InVertexName)
 		{
 			using namespace StereoAudioFormatVertexKeys;
 			
@@ -301,58 +303,36 @@ namespace Metasound
 			);
 		}
 
-		static FNodeClassMetadata GetNodeClassMetadata(const FVertexName& InInputName)
-		{
-			FNodeClassMetadata Info;
-
-			Info.ClassName = { "Input", GetMetasoundDataTypeName<FStereoAudioFormat>(), "" };
-			Info.MajorVersion = 1;
-			Info.MinorVersion = 0;
-			Info.DisplayName = METASOUND_LOCTEXT_FORMAT("Metasound_InputNodeDisplayNameFormat", "Input {0}", GetMetasoundDataTypeDisplayText<FStereoAudioFormat>());
-			Info.Description = METASOUND_LOCTEXT("Metasound_InputNodeDescription", "Input from the parent Metasound graph.");
-			Info.Author = PluginAuthor;
-			Info.PromptIfMissing = PluginNodeMissingPrompt;
-			Info.DefaultInterface = GetVertexInterface(InInputName);
-
-			return Info;
-		};
 
 
 	public:
 		static constexpr bool bCanRegister = true;
 
 		TInputNode(FInputNodeConstructorParams&& InParams)
-			: FNode(InParams.NodeName, InParams.InstanceID, GetNodeClassMetadata(InParams.VertexName))
-			, VertexInterface(GetVertexInterface(InParams.VertexName))
-			, Factory(MakeShared<FInputOperatorFactory, ESPMode::ThreadSafe>(InParams.VertexName))
+		:	FInputNode(MoveTemp(InParams), GetMetasoundDataTypeName<FStereoAudioFormat>(), EVertexAccessType::Reference, MakeShared<FInputOperatorFactory>(InParams.VertexName))
 		{
 		}
 
-		const FVertexInterface& GetVertexInterface() const override
+		explicit TInputNode(const FVertexName& InVertexName, FNodeData InNodeData, TSharedRef<const FNodeClassMetadata> InClassMetadata)
+		: FInputNode(MakeShared<FInputOperatorFactory>(InVertexName), MoveTemp(InNodeData), MoveTemp(InClassMetadata))
 		{
-			return VertexInterface;
-		}
-
-		bool SetVertexInterface(const FVertexInterface& InInterface) override
-		{
-			return VertexInterface == InInterface;
-		}
-
-		bool IsVertexInterfaceSupported(const FVertexInterface& InInterface) const 
-		{
-			return VertexInterface == InInterface;
 		}
 
-		TSharedRef<IOperatorFactory, ESPMode::ThreadSafe> GetDefaultOperatorFactory() const override
+		static FNodeClassMetadata CreateNodeClassMetadata(const FVertexName& InInputName)
 		{
-			return Factory;
-		}
-
-	private:
-		FVertexInterface VertexInterface;
+			FNodeClassMetadata Info;
 
-		TSharedRef<FInputOperatorFactory, ESPMode::ThreadSafe> Factory;
+			Info.ClassName = { "Input", GetMetasoundDataTypeName<FStereoAudioFormat>(), "" };
+			Info.MajorVersion = 1;
+			Info.MinorVersion = 0;
+			Info.DisplayName = METASOUND_LOCTEXT_FORMAT("Metasound_InputNodeDisplayNameFormat", "Input {0}", GetMetasoundDataTypeDisplayText<FStereoAudioFormat>());
+			Info.Description = METASOUND_LOCTEXT("Metasound_InputNodeDescription", "Input from the parent Metasound graph.");
+			Info.Author = PluginAuthor;
+			Info.PromptIfMissing = PluginNodeMissingPrompt;
+			Info.DefaultInterface = CreateVertexInterface(InInputName);
 
+			return Info;
+		};
 	};
 
 	// TOutputNode<> specialization for FMonoAudioFormat. Allows an audio buffer
@@ -409,7 +389,7 @@ namespace Metasound
 
 			TUniquePtr<IOperator> CreateOperator(const FBuildOperatorParams& InParams, FBuildResults& OutResults) override
 			{
-				TDataReadReference<FAudioBuffer> Center = InParams.InputData.GetOrConstructDataReadReference<FAudioBuffer>(OutputName, InParams.OperatorSettings);
+				TDataReadReference<FAudioBuffer> Center = InParams.InputData.GetOrCreateDefaultDataReadReference<FAudioBuffer>(OutputName, InParams.OperatorSettings);
 
 				TDataReadReference<FMonoAudioFormat> Mono = TDataReadReferenceFactory<FMonoAudioFormat>::CreateExplicitArgs(InParams.OperatorSettings, WriteCast(Center));
 
@@ -420,7 +400,7 @@ namespace Metasound
 			FVertexName OutputName;
 		};
 
-		static FVertexInterface GetVertexInterface(const FVertexName& InVertexName)
+		static FVertexInterface CreateVertexInterface(const FVertexName& InVertexName)
 		{
 			return FVertexInterface(
 				FInputVertexInterface(
@@ -432,7 +412,17 @@ namespace Metasound
 			);
 		}
 
-		static FNodeClassMetadata GetNodeClassMetadata(const FVertexName& InOutputName)
+
+
+	public:
+		TOutputNode(const FVertexName& InInstanceName, const FGuid& InInstanceID, const FVertexName& InVertexName)
+		:	FNode(InInstanceName, InInstanceID, CreateNodeClassMetadata(InVertexName))
+		,	VertexInterface(CreateVertexInterface(InVertexName))
+		,	Factory(MakeShared<FOutputOperatorFactory, ESPMode::ThreadSafe>(InVertexName))
+		{
+		}
+
+		static FNodeClassMetadata CreateNodeClassMetadata(const FVertexName& InOutputName)
 		{
 			FNodeClassMetadata Info;
 
@@ -443,20 +433,11 @@ namespace Metasound
 			Info.Description = METASOUND_LOCTEXT("Metasound_OutputNodeDescription", "Output from the parent Metasound graph.");
 			Info.Author = PluginAuthor;
 			Info.PromptIfMissing = PluginNodeMissingPrompt;
-			Info.DefaultInterface = GetVertexInterface(InOutputName);
+			Info.DefaultInterface = CreateVertexInterface(InOutputName);
 
 			return Info;
 		};
 
-
-	public:
-		TOutputNode(const FVertexName& InInstanceName, const FGuid& InInstanceID, const FVertexName& InVertexName)
-		:	FNode(InInstanceName, InInstanceID, GetNodeClassMetadata(InVertexName))
-		,	VertexInterface(GetVertexInterface(InVertexName))
-		,	Factory(MakeShared<FOutputOperatorFactory, ESPMode::ThreadSafe>(InVertexName))
-		{
-		}
-
 		const FVertexInterface& GetVertexInterface() const override
 		{
 			return VertexInterface;
@@ -487,7 +468,7 @@ namespace Metasound
 	// TInputNode<> specializastion for FMonoAudioFormat. Allows an input mono audio
 	// format to be exposed as a single buffer.
 	template<>
-	class METASOUNDSTANDARDNODES_API TInputNode<FMonoAudioFormat> : public FNode
+	class METASOUNDSTANDARDNODES_API TInputNode<FMonoAudioFormat> : public FInputNode
 	{
 		class FInputOperator : public FNoOpOperator
 		{
@@ -539,7 +520,7 @@ namespace Metasound
 
 			TUniquePtr<IOperator> CreateOperator(const FBuildOperatorParams& InParams, FBuildResults& OutResults) override
 			{
-				TDataReadReference<FMonoAudioFormat> Mono = InParams.InputData.GetOrConstructDataReadReference<FMonoAudioFormat>(InputName, InParams.OperatorSettings);
+				TDataReadReference<FMonoAudioFormat> Mono = InParams.InputData.GetOrCreateDefaultDataReadReference<FMonoAudioFormat>(InputName, InParams.OperatorSettings);
 				TDataReadReference<FAudioBuffer> Center = Mono->GetCenter();
 
 				return MakeUnique<FInputOperator>(InputName, Center, Mono);
@@ -549,7 +530,7 @@ namespace Metasound
 			FVertexName InputName;
 		};
 
-		static FVertexInterface GetVertexInterface(const FVertexName& InVertexName)
+		static FVertexInterface CreateVertexInterface(const FVertexName& InVertexName)
 		{
 			return FVertexInterface(
 				FInputVertexInterface(
@@ -561,58 +542,34 @@ namespace Metasound
 			);
 		}
 
-		static FNodeClassMetadata GetNodeClassMetadata(const FVertexName& InInputName)
-		{
-			FNodeClassMetadata Info;
-
-			Info.ClassName = { "Input", GetMetasoundDataTypeName<FMonoAudioFormat>(), "" };
-			Info.MajorVersion = 1;
-			Info.MinorVersion = 0;
-			Info.DisplayName = METASOUND_LOCTEXT_FORMAT("Metasound_InputNodeDisplayNameFormat", "Input {0}", GetMetasoundDataTypeDisplayText<FMonoAudioFormat>());
-			Info.Description = METASOUND_LOCTEXT("Metasound_InputNodeDescription", "Input from the parent Metasound graph.");
-			Info.Author = PluginAuthor;
-			Info.PromptIfMissing = PluginNodeMissingPrompt;
-			Info.DefaultInterface = GetVertexInterface(InInputName);
-
-			return Info;
-		};
-
-
 	public:
 		static constexpr bool bCanRegister = true;
 
 		TInputNode(FInputNodeConstructorParams&& InParams)
-			: FNode(InParams.NodeName, InParams.InstanceID, GetNodeClassMetadata(InParams.VertexName))
-			, VertexInterface(GetVertexInterface(InParams.VertexName))
-			, Factory(MakeShared<FInputOperatorFactory, ESPMode::ThreadSafe>(InParams.VertexName))
+		:	FInputNode(MoveTemp(InParams), GetMetasoundDataTypeName<FMonoAudioFormat>(), EVertexAccessType::Reference, MakeShared<FInputOperatorFactory>(InParams.VertexName))
 		{
 		}
 
-		const FVertexInterface& GetVertexInterface() const override
-		{
-			return VertexInterface;
-		}
-
-		bool SetVertexInterface(const FVertexInterface& InInterface) override
-		{
-			return VertexInterface == InInterface;
-		}
-
-		bool IsVertexInterfaceSupported(const FVertexInterface& InInterface) const 
+		explicit TInputNode(const FVertexName& InVertexName, FNodeData InNodeData, TSharedRef<const FNodeClassMetadata> InClassMetadata)
+		: FInputNode(MakeShared<FInputOperatorFactory>(InVertexName), MoveTemp(InNodeData), MoveTemp(InClassMetadata))
 		{
-			return VertexInterface == InInterface;
 		}
 
-		TSharedRef<IOperatorFactory, ESPMode::ThreadSafe> GetDefaultOperatorFactory() const override
+		static FNodeClassMetadata CreateNodeClassMetadata(const FVertexName& InInputName)
 		{
-			return Factory;
-		}
-
-	private:
-		FVertexInterface VertexInterface;
+			FNodeClassMetadata Info;
 
-		TSharedRef<FInputOperatorFactory, ESPMode::ThreadSafe> Factory;
+			Info.ClassName = { "Input", GetMetasoundDataTypeName<FMonoAudioFormat>(), "" };
+			Info.MajorVersion = 1;
+			Info.MinorVersion = 0;
+			Info.DisplayName = METASOUND_LOCTEXT_FORMAT("Metasound_InputNodeDisplayNameFormat", "Input {0}", GetMetasoundDataTypeDisplayText<FMonoAudioFormat>());
+			Info.Description = METASOUND_LOCTEXT("Metasound_InputNodeDescription", "Input from the parent Metasound graph.");
+			Info.Author = PluginAuthor;
+			Info.PromptIfMissing = PluginNodeMissingPrompt;
+			Info.DefaultInterface = CreateVertexInterface(InInputName);
 
+			return Info;
+		};
 	};
 
 	// Disable arrays of audio formats. 
