Changing clamp on no-repeats order to be half the array size vs the whole array size.
Adding a new value of -1 to indicate "maximum" no-repeats (i.e. half array size) automatically so user doesn't have to worry about array size to maximize no repeats.

[RN] Clamped no-repeats order to be half the array size vs the whole array size using RandomGet in MetaSounds.

diff --git a/Engine/Plugins/Runtime/Metasound/Source/MetasoundFrontend/Private/MetasoundArrayRandomNode.cpp b/Engine/Plugins/Runtime/Metasound/Source/MetasoundFrontend/Private/MetasoundArrayRandomNode.cpp
--- a/Engine/Plugins/Runtime/Metasound/Source/MetasoundFrontend/Private/MetasoundArrayRandomNode.cpp
+++ b/Engine/Plugins/Runtime/Metasound/Source/MetasoundFrontend/Private/MetasoundArrayRandomNode.cpp
@@ -74,15 +74,29 @@ namespace Metasound
 
 	void FArrayRandomGet::SetNoRepeatOrder(int32 InNoRepeatOrder)
 	{
-		// Make sure the no repeat order is between 0 and one less than our max index
-		InNoRepeatOrder = FMath::Clamp(InNoRepeatOrder, 0, MaxIndex - 1);
-
-		if (InNoRepeatOrder != NoRepeatOrder)
+		// Only set the no repeats if the array length is larger than 3, otherwise, ignore no-repeats as it won't work anyway (it will end up repeating a clear pattern)
+		// Make sure the no repeat order is between 0 and half of the max index
+		// This is to prevent clear patterns in the output that causes confusion by user. By clamping to half the max index we achieve maximum no-repeats behavior, which is the intent of a user trying to maximize no repeats.
+		if (MaxIndex > 3)
 		{
-			PreviousIndicesQueue = MakeUnique<TCircularQueue<int32>>(InNoRepeatOrder + 2);
-			PreviousIndices.Reset();
-			NoRepeatOrder = InNoRepeatOrder;
+			// Check for -1 on no-repeat order. This means to utilize the maximum non-repetition, matched to the array size. This prevents user from having to adjust their no-repeats to match the array size themselves.
+			if (InNoRepeatOrder == -1)
+			{
+				InNoRepeatOrder = MaxIndex / 2;
+			}
+			else
+			{
+				InNoRepeatOrder = FMath::Clamp(InNoRepeatOrder, 0, MaxIndex / 2);
+			}
+
+			if (InNoRepeatOrder != NoRepeatOrder)
+			{
+				PreviousIndicesQueue = MakeUnique<TCircularQueue<int32>>(InNoRepeatOrder + 2);
+				PreviousIndices.Reset();
+				NoRepeatOrder = InNoRepeatOrder;
+			}
 		}
+
 	}
 
 	void FArrayRandomGet::SetRandomWeights(const TArray<float>& InRandomWeights)
diff --git a/Engine/Plugins/Runtime/Metasound/Source/MetasoundFrontend/Public/MetasoundArrayRandomNode.h b/Engine/Plugins/Runtime/Metasound/Source/MetasoundFrontend/Public/MetasoundArrayRandomNode.h
--- a/Engine/Plugins/Runtime/Metasound/Source/MetasoundFrontend/Public/MetasoundArrayRandomNode.h
+++ b/Engine/Plugins/Runtime/Metasound/Source/MetasoundFrontend/Public/MetasoundArrayRandomNode.h
@@ -63,7 +63,7 @@ namespace Metasound
 		METASOUND_PARAM(InputRandomArray, "In Array", "Input array to randomize.")
 		METASOUND_PARAM(InputWeights, "Weights", "Input array of weights to use for random selection. Will repeat if this array is shorter than the input array to select from.")
 		METASOUND_PARAM(InputSeed, "Seed", "Seed to use for the random stream. Set to -1 to use a random seed.")
-		METASOUND_PARAM(InputNoRepeatOrder, "No Repeats", "The number of elements to track to avoid repeating in a row.")
+		METASOUND_PARAM(InputNoRepeatOrder, "No Repeats", "The number of elements to track to avoid repeating in a row. This is clamped to be within half the array size. The output will end up repeating a clear pattern if set close to the array size. Set to -1 to automatically set to half the array size (which is the maximum no-repeats behavior)")
 		METASOUND_PARAM(InputEnableSharedState, "Enable Shared State", "Set to enabled to share state with other Random Get (Array) nodes. Does not apply when previewing in the MetaSound editor; use PIE or game.")
 		METASOUND_PARAM(InputSharedStateBehavior, "Shared State Behavior", "The behavior for how state is shared with other Random Get (Array) nodes. Only applied when Enable Shared State is true.")
 		METASOUND_PARAM(OutputTriggerOnNext, "On Next", "Triggers when the \"Next\" input is triggered.")
