Update MetaSound node class registry to expose default interface + update HMX node

diff --git a/Engine/Plugins/Runtime/Metasound/Source/MetasoundFrontend/Private/MetasoundAssetBase.cpp b/Engine/Plugins/Runtime/Metasound/Source/MetasoundFrontend/Private/MetasoundAssetBase.cpp
--- a/Engine/Plugins/Runtime/Metasound/Source/MetasoundFrontend/Private/MetasoundAssetBase.cpp
+++ b/Engine/Plugins/Runtime/Metasound/Source/MetasoundFrontend/Private/MetasoundAssetBase.cpp
@@ -110,6 +110,12 @@ namespace Metasound
 					return &Interfaces;
 				}
 
+				virtual FVertexInterface GetDefaultVertexInterface() const override
+				{
+					UE_LOG(LogMetaSound, Warning, TEXT("Proxy data is not available for non runtime node %s only used for serialization, so interface will not include object literals. Please ensure calling this function is intended."), *FrontendClass.Metadata.GetClassName().ToString())
+					return CreateVertexInterfaceFromClassNoProxy(FrontendClass);
+				}
+
 				virtual bool IsNative() const override
 				{
 					return false;
diff --git a/Engine/Plugins/Runtime/Metasound/Source/MetasoundFrontend/Private/MetasoundFrontend.cpp b/Engine/Plugins/Runtime/Metasound/Source/MetasoundFrontend/Private/MetasoundFrontend.cpp
--- a/Engine/Plugins/Runtime/Metasound/Source/MetasoundFrontend/Private/MetasoundFrontend.cpp
+++ b/Engine/Plugins/Runtime/Metasound/Source/MetasoundFrontend/Private/MetasoundFrontend.cpp
@@ -87,6 +87,99 @@ namespace Metasound
 			return false;
 		}
 
+		// This method is here to support back compatability for node registration.
+		// Originally, nodes did not take in a FVertexInterface on construction,
+		// but with the introduction of node configuration in 5.6, they now do take
+		// in a FVertexInterface. 
+		FDataVertexMetadata CreateVertexMetadata(const FMetasoundFrontendClassVertex& InClassVertex)
+		{
+#if WITH_EDITORONLY_DATA
+			return FDataVertexMetadata{ InClassVertex.Metadata.GetDescription(), InClassVertex.Metadata.GetDisplayName(), InClassVertex.Metadata.bIsAdvancedDisplay };
+#else
+			return {};
+#endif
+		}
+
+		FVertexInterface CreateVertexInterfaceFromClassInternal(const FMetasoundFrontendClass& InNodeClass, const bool bCreateProxies, const FProxyDataCache* InProxyDataCache = nullptr)
+		{
+			FVertexInterface Interface;
+
+			FInputVertexInterface& Inputs = Interface.GetInputInterface();
+			for (const FMetasoundFrontendClassInput& ClassInput : InNodeClass.Interface.Inputs)
+			{
+				const FMetasoundFrontendLiteral* DefaultLiteral = ClassInput.FindConstDefault(DefaultPageID);
+				if (DefaultLiteral)
+				{
+					EMetasoundFrontendLiteralType LiteralType = DefaultLiteral->GetType();
+					if (!bCreateProxies &&
+						((EMetasoundFrontendLiteralType::UObject == LiteralType) || (EMetasoundFrontendLiteralType::UObjectArray == LiteralType)))
+					{
+						UE_LOG(LogMetaSound, Warning, TEXT("Ignoring default literal set on vertex %s of node %s. Please update construct of node to use FNodeData"), *ClassInput.Name.ToString(), *InNodeClass.Metadata.GetClassName().ToString());
+						DefaultLiteral = nullptr;
+					}
+				}
+
+				if (DefaultLiteral)
+				{
+					const IDataTypeRegistry& DataTypeRegistry = IDataTypeRegistry::Get();
+					FLiteral Literal = DefaultLiteral->ToLiteral(ClassInput.TypeName, &DataTypeRegistry, InProxyDataCache);
+
+					Inputs.Add(
+						FInputDataVertex
+						{
+							ClassInput.Name,
+							ClassInput.TypeName,
+							CreateVertexMetadata(ClassInput),
+							FrontendVertexAccessTypeToCoreVertexAccessType(ClassInput.AccessType),
+							Literal
+						});
+				}
+				else
+				{
+					// No default literal or proxy creation was skipped for object literals
+					Inputs.Add(
+						FInputDataVertex
+						{
+							ClassInput.Name,
+							ClassInput.TypeName,
+							CreateVertexMetadata(ClassInput),
+							FrontendVertexAccessTypeToCoreVertexAccessType(ClassInput.AccessType)
+						});
+				}
+			}
+
+			FOutputVertexInterface& Outputs = Interface.GetOutputInterface();
+			for (const FMetasoundFrontendClassOutput& ClassOutput : InNodeClass.Interface.Outputs)
+			{
+				Outputs.Add(
+					FOutputDataVertex
+					{
+						ClassOutput.Name,
+						ClassOutput.TypeName,
+						CreateVertexMetadata(ClassOutput),
+						FrontendVertexAccessTypeToCoreVertexAccessType(ClassOutput.AccessType)
+					});
+			}
+
+			FEnvironmentVertexInterface& Environments = Interface.GetEnvironmentInterface();
+			for (const FMetasoundFrontendClassEnvironmentVariable& ClassEnvironment : InNodeClass.Interface.Environment)
+			{
+				Environments.Add(FEnvironmentVertex{ ClassEnvironment.Name, FText::GetEmpty() });
+			}
+
+			return Interface;
+		}
+
+		FVertexInterface CreateVertexInterfaceFromClass(const FMetasoundFrontendClass& InNodeClass, const Metasound::Frontend::FProxyDataCache* InProxyDataCache)
+		{
+			checkf(InProxyDataCache || IsInGameThread(), TEXT("Vertex interface creation without proxy data cache must occur on game thread to safely create UObject proxies. Populate the proxy data cache on the game thread before calling this or use CreateVertexInterfaceFromClassNoProxy."))
+				return CreateVertexInterfaceFromClassInternal(InNodeClass, /*bCreateProxies=*/true, InProxyDataCache);
+		}
+
+		FVertexInterface CreateVertexInterfaceFromClassNoProxy(const FMetasoundFrontendClass& InNodeClass)
+		{
+			return CreateVertexInterfaceFromClassInternal(InNodeClass, /*bCreateProxies=*/false);
+		}
 	}
 }
 
diff --git a/Engine/Plugins/Runtime/Metasound/Source/MetasoundFrontend/Private/MetasoundFrontendDataTypeRegistry.cpp b/Engine/Plugins/Runtime/Metasound/Source/MetasoundFrontend/Private/MetasoundFrontendDataTypeRegistry.cpp
--- a/Engine/Plugins/Runtime/Metasound/Source/MetasoundFrontend/Private/MetasoundFrontendDataTypeRegistry.cpp
+++ b/Engine/Plugins/Runtime/Metasound/Source/MetasoundFrontend/Private/MetasoundFrontendDataTypeRegistry.cpp
@@ -132,6 +132,11 @@ namespace Metasound
 					return nullptr;
 				}
 
+				virtual FVertexInterface GetDefaultVertexInterface() const override
+				{
+					return FVertexInterface();
+				}
+
 				virtual bool IsNative() const override
 				{
 					return true;
@@ -189,6 +194,16 @@ PRAGMA_DISABLE_DEPRECATION_WARNINGS
 PRAGMA_ENABLE_DEPRECATION_WARNINGS
 				}
 
+				virtual FVertexInterface GetDefaultVertexInterface() const override
+				{
+					TSharedPtr<const FNodeClassMetadata> NodeClassMetadata = DataTypeEntry->GetInputClassMetadata();
+					if (ensure(NodeClassMetadata))
+					{
+						return NodeClassMetadata->DefaultInterface;
+					}
+					return FVertexInterface();
+				}
+
 			private:
 				
 				TSharedRef<IDataTypeRegistryEntry, ESPMode::ThreadSafe> DataTypeEntry;
@@ -234,6 +249,16 @@ PRAGMA_DISABLE_DEPRECATION_WARNINGS
 PRAGMA_ENABLE_DEPRECATION_WARNINGS
 				}
 
+				virtual FVertexInterface GetDefaultVertexInterface() const override
+				{
+					TSharedPtr<const FNodeClassMetadata> NodeClassMetadata = DataTypeEntry->GetConstructorInputClassMetadata();
+					if (ensure(NodeClassMetadata))
+					{
+						return NodeClassMetadata->DefaultInterface;
+					}
+					return FVertexInterface();
+				}
+
 			private:
 				
 				TSharedRef<IDataTypeRegistryEntry, ESPMode::ThreadSafe> DataTypeEntry;
@@ -265,6 +290,16 @@ PRAGMA_ENABLE_DEPRECATION_WARNINGS
 					return DataTypeEntry->CreateOutputNode(MoveTemp(InNodeData));
 				}
 
+				virtual FVertexInterface GetDefaultVertexInterface() const override
+				{
+					TSharedPtr<const FNodeClassMetadata> NodeClassMetadata = DataTypeEntry->GetOutputClassMetadata();
+					if (ensure(NodeClassMetadata))
+					{
+						return NodeClassMetadata->DefaultInterface;
+					}
+					return FVertexInterface();
+				}
+
 			private:
 				
 				TSharedRef<IDataTypeRegistryEntry, ESPMode::ThreadSafe> DataTypeEntry;
@@ -296,6 +331,16 @@ PRAGMA_ENABLE_DEPRECATION_WARNINGS
 					return DataTypeEntry->CreateConstructorOutputNode(MoveTemp(InNodeData));
 				}
 
+				virtual FVertexInterface GetDefaultVertexInterface() const override
+				{
+					TSharedPtr<const FNodeClassMetadata> NodeClassMetadata = DataTypeEntry->GetConstructorOutputClassMetadata();
+					if (ensure(NodeClassMetadata))
+					{
+						return NodeClassMetadata->DefaultInterface;
+					}
+					return FVertexInterface();
+				}
+
 			private:
 				
 				TSharedRef<IDataTypeRegistryEntry, ESPMode::ThreadSafe> DataTypeEntry;
@@ -381,6 +426,16 @@ PRAGMA_DISABLE_DEPRECATION_WARNINGS
 PRAGMA_ENABLE_DEPRECATION_WARNINGS
 				}
 
+				virtual FVertexInterface GetDefaultVertexInterface() const override
+				{
+					TSharedPtr<const FNodeClassMetadata> NodeClassMetadata = DataTypeEntry->GetVariableClassMetadata();
+					if (ensure(NodeClassMetadata))
+					{
+						return NodeClassMetadata->DefaultInterface;
+					}
+					return FVertexInterface();
+				}
+
 			private:
 				
 				TSharedRef<IDataTypeRegistryEntry, ESPMode::ThreadSafe> DataTypeEntry;
@@ -417,6 +472,16 @@ PRAGMA_ENABLE_DEPRECATION_WARNINGS
 					return DataTypeEntry->CreateVariableMutatorNode(MoveTemp(InNodeData));
 				}
 
+				virtual FVertexInterface GetDefaultVertexInterface() const override
+				{
+					TSharedPtr<const FNodeClassMetadata> NodeClassMetadata = DataTypeEntry->GetVariableMutatorClassMetadata();
+					if (ensure(NodeClassMetadata))
+					{
+						return NodeClassMetadata->DefaultInterface;
+					}
+					return FVertexInterface();
+				}
+
 			private:
 				
 				TSharedRef<IDataTypeRegistryEntry, ESPMode::ThreadSafe> DataTypeEntry;
@@ -453,6 +518,16 @@ PRAGMA_ENABLE_DEPRECATION_WARNINGS
 					return DataTypeEntry->CreateVariableAccessorNode(MoveTemp(InNodeData));
 				}
 
+				virtual FVertexInterface GetDefaultVertexInterface() const override
+				{
+					TSharedPtr<const FNodeClassMetadata> NodeClassMetadata = DataTypeEntry->GetVariableAccessorClassMetadata();
+					if (ensure(NodeClassMetadata))
+					{
+						return NodeClassMetadata->DefaultInterface;
+					}
+					return FVertexInterface();
+				}
+
 			private:
 				
 				TSharedRef<IDataTypeRegistryEntry, ESPMode::ThreadSafe> DataTypeEntry;
@@ -489,6 +564,16 @@ PRAGMA_ENABLE_DEPRECATION_WARNINGS
 					return DataTypeEntry->CreateVariableDeferredAccessorNode(MoveTemp(InNodeData));
 				}
 
+				virtual FVertexInterface GetDefaultVertexInterface() const override
+				{
+					TSharedPtr<const FNodeClassMetadata> NodeClassMetadata = DataTypeEntry->GetVariableDeferredAccessorClassMetadata();
+					if (ensure(NodeClassMetadata))
+					{
+						return NodeClassMetadata->DefaultInterface;
+					}
+					return FVertexInterface();
+				}
+
 			private:
 				
 				TSharedRef<IDataTypeRegistryEntry, ESPMode::ThreadSafe> DataTypeEntry;
diff --git a/Engine/Plugins/Runtime/Metasound/Source/MetasoundFrontend/Private/MetasoundFrontendNodeClassRegistry.cpp b/Engine/Plugins/Runtime/Metasound/Source/MetasoundFrontend/Private/MetasoundFrontendNodeClassRegistry.cpp
--- a/Engine/Plugins/Runtime/Metasound/Source/MetasoundFrontend/Private/MetasoundFrontendNodeClassRegistry.cpp
+++ b/Engine/Plugins/Runtime/Metasound/Source/MetasoundFrontend/Private/MetasoundFrontendNodeClassRegistry.cpp
@@ -252,100 +252,30 @@ namespace Metasound::Frontend
 				return &Interfaces;
 			}
 
-			virtual bool IsNative() const override
+			virtual FVertexInterface GetDefaultVertexInterface() const override
 			{
-				return false;
-			}
-
-		private:
-
-			FMetasoundFrontendClass FrontendClass;
-			TSet<FMetasoundFrontendVersion> Interfaces;
-			FNodeClassInfo ClassInfo;
-			TSharedPtr<const IGraph> Graph;
-		};
-
-		// This method is here to support back compatability for node registration.
-		// Originally, nodes did not take in a FVertexInterface on construction,
-		// but with the introduction of node configuration in 5.6, they now do take
-		// in a FVertexInterface. 
-		FDataVertexMetadata CreateVertexMetadata(const FMetasoundFrontendClassVertex& InClassVertex)
-		{
-#if WITH_EDITORONLY_DATA
-			return FDataVertexMetadata{InClassVertex.Metadata.GetDescription(), InClassVertex.Metadata.GetDisplayName(), InClassVertex.Metadata.bIsAdvancedDisplay};
-#else
-			return {};
-#endif
-		}
-
-		// This method is here to support back compatability for node registration.
-		// Originally, nodes did not take in a FVertexInterface on construction,
-		// but with the introduction of node configuration in 5.6, they now do take
-		// in a FVertexInterface. 
-		FVertexInterface CreateVertexInterfaceFromClass(const FMetasoundFrontendClass& InNodeClass)
-		{
-			FVertexInterface Interface;
-
-			FInputVertexInterface& Inputs = Interface.GetInputInterface();
-			for (const FMetasoundFrontendClassInput& ClassInput : InNodeClass.Interface.Inputs)
-			{
-				const FMetasoundFrontendLiteral* DefaultLiteral = ClassInput.FindConstDefault(::Metasound::Frontend::DefaultPageID);
-				if (DefaultLiteral)
-				{
-					EMetasoundFrontendLiteralType LiteralType = DefaultLiteral->GetType();
-					if ((EMetasoundFrontendLiteralType::UObject == LiteralType) || (EMetasoundFrontendLiteralType::UObjectArray == LiteralType))
-					{
-						UE_LOG(LogMetaSound, Warning, TEXT("Ignoring default literal set on vertex %s of node %s. Please update construct of node to use FNodeData"), *ClassInput.Name.ToString(), *InNodeClass.Metadata.GetClassName().ToString());
-						DefaultLiteral = nullptr; 
-					}
-				}
-
-				if (DefaultLiteral)
+				if (ensure(Graph.IsValid()))
 				{
-					Inputs.Add(
-						FInputDataVertex
-						{
-							ClassInput.Name, 
-							ClassInput.TypeName, 
-							CreateVertexMetadata(ClassInput),
-							FrontendVertexAccessTypeToCoreVertexAccessType(ClassInput.AccessType),
-							DefaultLiteral->ToLiteralNoProxy()
-						});
+					return Graph->GetMetadata().DefaultInterface;
 				}
 				else
 				{
-					Inputs.Add(
-						FInputDataVertex
-						{
-							ClassInput.Name, 
-							ClassInput.TypeName, 
-							CreateVertexMetadata(ClassInput),
-							FrontendVertexAccessTypeToCoreVertexAccessType(ClassInput.AccessType)
-						});
+					return FVertexInterface();
 				}
 			}
 
-			FOutputVertexInterface& Outputs = Interface.GetOutputInterface();
-			for (const FMetasoundFrontendClassOutput& ClassOutput : InNodeClass.Interface.Outputs)
+			virtual bool IsNative() const override
 			{
-				Outputs.Add(
-					FOutputDataVertex
-					{
-						ClassOutput.Name, 
-						ClassOutput.TypeName, 
-						CreateVertexMetadata(ClassOutput),
-						FrontendVertexAccessTypeToCoreVertexAccessType(ClassOutput.AccessType)
-					});
+				return false;
 			}
 
-			FEnvironmentVertexInterface& Environments = Interface.GetEnvironmentInterface();
-			for (const FMetasoundFrontendClassEnvironmentVariable& ClassEnvironment : InNodeClass.Interface.Environment)
-			{
-				Environments.Add(FEnvironmentVertex{ ClassEnvironment.Name, FText::GetEmpty() });
-			}
+		private:
 
-			return Interface;
-		}
+			FMetasoundFrontendClass FrontendClass;
+			TSet<FMetasoundFrontendVersion> Interfaces;
+			FNodeClassInfo ClassInfo;
+			TSharedPtr<const IGraph> Graph;
+		};
 	} // namespace RegistryPrivate
 	  
 	TUniquePtr<INode> INodeClassRegistryEntry::CreateNode(FNodeData InNodeData) const
@@ -585,7 +515,7 @@ PRAGMA_ENABLE_DEPRECATION_WARNINGS
 			{
 				InInitData.InstanceName,
 				InInitData.InstanceID,
-				RegistryPrivate::CreateVertexInterfaceFromClass(FrontendClass) // This is an attempt to maintain backwards compatibility. It creates an FVertexInterface from the class definition because none was provided. 
+				Entry.GetDefaultVertexInterface()
 			};
 			Node = Entry.CreateNode(MoveTemp(NodeData)); 
 		};
@@ -1183,6 +1113,21 @@ PRAGMA_ENABLE_DEPRECATION_WARNINGS
 		return false;
 	}
 
+	bool FNodeClassRegistry::FindDefaultVertexInterface(const FNodeClassRegistryKey& InKey, FVertexInterface& OutVertexInterface) const 
+	{
+		auto GetDefaultVertexInterface = [&OutVertexInterface](const INodeClassRegistryEntry& Entry)
+		{
+			OutVertexInterface = Entry.GetDefaultVertexInterface();
+		};
+
+		if (AccessNodeEntryThreadSafe(InKey, GetDefaultVertexInterface))
+		{
+			return true;
+		}
+
+		return false;
+	}
+
 	bool FNodeClassRegistry::FindFrontendClassFromRegistered(const FNodeClassRegistryKey& InKey, FMetasoundFrontendClass& OutClass)
 	{
 		auto SetFrontendClass = [&OutClass](const INodeClassRegistryEntry& Entry)
@@ -1445,5 +1390,4 @@ PRAGMA_ENABLE_DEPRECATION_WARNINGS
 		}
 		return false;
 	}
-
 } // namespace Metasound::Frontend
diff --git a/Engine/Plugins/Runtime/Metasound/Source/MetasoundFrontend/Private/MetasoundFrontendNodeClassRegistryPrivate.h b/Engine/Plugins/Runtime/Metasound/Source/MetasoundFrontend/Private/MetasoundFrontendNodeClassRegistryPrivate.h
--- a/Engine/Plugins/Runtime/Metasound/Source/MetasoundFrontend/Private/MetasoundFrontendNodeClassRegistryPrivate.h
+++ b/Engine/Plugins/Runtime/Metasound/Source/MetasoundFrontend/Private/MetasoundFrontendNodeClassRegistryPrivate.h
@@ -95,9 +95,10 @@ namespace Metasound::Frontend
 		virtual void IterateRegistry(FIterateMetasoundFrontendClassFunction InIterFunc, EMetasoundFrontendClassType InClassType = EMetasoundFrontendClassType::Invalid) const override;
 
 		// Find Frontend class data.
+		virtual bool FindDefaultVertexInterface(const FNodeClassRegistryKey& InKey, FVertexInterface& OutVertexInterface) const override;
+		
 		virtual bool FindFrontendClassFromRegistered(const FNodeClassRegistryKey& InKey, FMetasoundFrontendClass& OutClass) override;
 
-
 		virtual bool FindImplementedInterfacesFromRegistered(const Metasound::Frontend::FNodeClassRegistryKey& InKey, TSet<FMetasoundFrontendVersion>& OutInterfaceVersions) const override;
 		
 		virtual bool FindInputNodeRegistryKeyForDataType(const FName& InDataTypeName, const EMetasoundFrontendVertexAccessType InAccessType, FNodeClassRegistryKey& OutKey) override;
diff --git a/Engine/Plugins/Runtime/Metasound/Source/MetasoundFrontend/Private/MetasoundNodeRegistrationMacro.cpp b/Engine/Plugins/Runtime/Metasound/Source/MetasoundFrontend/Private/MetasoundNodeRegistrationMacro.cpp
--- a/Engine/Plugins/Runtime/Metasound/Source/MetasoundFrontend/Private/MetasoundNodeRegistrationMacro.cpp
+++ b/Engine/Plugins/Runtime/Metasound/Source/MetasoundFrontend/Private/MetasoundNodeRegistrationMacro.cpp
@@ -37,6 +37,11 @@ namespace Metasound::Frontend
 			return nullptr;
 		}
 
+		FVertexInterface FNodeRegistryEntryBase::GetDefaultVertexInterface() const
+		{
+			return ClassMetadata->DefaultInterface;
+		}
+
 		bool FNodeRegistryEntryBase::IsNative() const
 		{
 			return true;
diff --git a/Engine/Plugins/Runtime/Metasound/Source/MetasoundFrontend/Public/MetasoundFrontend.h b/Engine/Plugins/Runtime/Metasound/Source/MetasoundFrontend/Public/MetasoundFrontend.h
--- a/Engine/Plugins/Runtime/Metasound/Source/MetasoundFrontend/Public/MetasoundFrontend.h
+++ b/Engine/Plugins/Runtime/Metasound/Source/MetasoundFrontend/Public/MetasoundFrontend.h
@@ -76,6 +76,17 @@ namespace Metasound
 		// @returns false if the file couldn't be found or parsed into a document.
 		METASOUNDFRONTEND_API bool ImportJSONAssetToMetasound(const FString& InPath, FMetasoundFrontendDocument& OutMetasoundDocument);
 
+		// These functions can be used to get the FVertexInterface from a FMetasoundFrontendClass 
+		// for node registration in special cases.
+		// Originally, nodes did not take in a FVertexInterface on construction,
+		// but with the introduction of node configuration in 5.6, they now do take
+		// in a FVertexInterface. 
+		
+		// If called on the game thread, no proxy data cache is needed for object literal creation, but if called from other threads, the proxy data cache must be provided. 
+		FVertexInterface CreateVertexInterfaceFromClass(const FMetasoundFrontendClass& InNodeClass, const FProxyDataCache* InProxyDataCache = nullptr);
+		// Does not create proxies, so no thread restrictions. 
+		FVertexInterface CreateVertexInterfaceFromClassNoProxy(const FMetasoundFrontendClass& InNodeClass);
+
 		// Struct that indicates whether an input and an output can be connected,
 		// and whether an intermediate node is necessary to connect the two.
 		struct METASOUNDFRONTEND_API FConnectability
diff --git a/Engine/Plugins/Runtime/Metasound/Source/MetasoundFrontend/Public/MetasoundFrontendNodeClassRegistry.h b/Engine/Plugins/Runtime/Metasound/Source/MetasoundFrontend/Public/MetasoundFrontendNodeClassRegistry.h
--- a/Engine/Plugins/Runtime/Metasound/Source/MetasoundFrontend/Public/MetasoundFrontendNodeClassRegistry.h
+++ b/Engine/Plugins/Runtime/Metasound/Source/MetasoundFrontend/Public/MetasoundFrontendNodeClassRegistry.h
@@ -103,6 +103,8 @@ namespace Metasound::Frontend
 			*/
 		virtual const TSet<FMetasoundFrontendVersion>* GetImplementedInterfaces() const = 0;
 
+		virtual FVertexInterface GetDefaultVertexInterface() const = 0;
+
 		/** Whether or not the node is natively defined */
 		virtual bool IsNative() const = 0;
 	};
@@ -254,12 +256,10 @@ namespace Metasound::Frontend
 
 		static bool GetFrontendClassFromRegistered(const FNodeClassRegistryKey& InKey, FMetasoundFrontendClass& OutClass);
 
-
 		static bool GetInputNodeRegistryKeyForDataType(const FName& InDataTypeName, const EMetasoundFrontendVertexAccessType InAccessType, FNodeClassRegistryKey& OutKey);
 		static bool GetVariableNodeRegistryKeyForDataType(const FName& InDataTypeName, FNodeClassRegistryKey& OutKey);
 		static bool GetOutputNodeRegistryKeyForDataType(const FName& InDataTypeName, const EMetasoundFrontendVertexAccessType InAccessType, FNodeClassRegistryKey& OutKey);
 
-
 	protected:
 		INodeClassRegistry() = default;
 	public:
@@ -316,6 +316,11 @@ namespace Metasound::Frontend
 		virtual void IterateRegistry(FIterateMetasoundFrontendClassFunction InIterFunc, EMetasoundFrontendClassType InClassType = EMetasoundFrontendClassType::Invalid) const = 0;
 
 		// Query for MetaSound Frontend document objects.
+
+		// Get the default vertex interface for the node class entry for the given key, 
+		// returning true on success, false otherwise 
+		virtual bool FindDefaultVertexInterface(const FNodeClassRegistryKey& InKey, FVertexInterface& OutVertexInterface) const = 0;
+
 		virtual bool FindFrontendClassFromRegistered(const FNodeClassRegistryKey& InKey, FMetasoundFrontendClass& OutClass) = 0;
 
 		virtual bool FindImplementedInterfacesFromRegistered(const FNodeClassRegistryKey& InKey, TSet<FMetasoundFrontendVersion>& OutInterfaceVersions) const = 0;
diff --git a/Engine/Plugins/Runtime/Metasound/Source/MetasoundFrontend/Public/MetasoundNodeRegistrationMacro.h b/Engine/Plugins/Runtime/Metasound/Source/MetasoundFrontend/Public/MetasoundNodeRegistrationMacro.h
--- a/Engine/Plugins/Runtime/Metasound/Source/MetasoundFrontend/Public/MetasoundNodeRegistrationMacro.h
+++ b/Engine/Plugins/Runtime/Metasound/Source/MetasoundFrontend/Public/MetasoundNodeRegistrationMacro.h
@@ -39,6 +39,8 @@ namespace Metasound::Frontend
 
 			virtual const TSet<FMetasoundFrontendVersion>* GetImplementedInterfaces() const override;
 
+			virtual FVertexInterface GetDefaultVertexInterface() const override;
+			
 			virtual bool IsNative() const override;
 
 		protected:
