Exclude deprecated nodes from automated node tests.
#rnx

diff --git a/Engine/Plugins/Runtime/Metasound/Source/MetasoundEngineTest/Private/EngineTestMetaSoundAutomatedNodeTest.cpp b/Engine/Plugins/Runtime/Metasound/Source/MetasoundEngineTest/Private/EngineTestMetaSoundAutomatedNodeTest.cpp
--- a/Engine/Plugins/Runtime/Metasound/Source/MetasoundEngineTest/Private/EngineTestMetaSoundAutomatedNodeTest.cpp
+++ b/Engine/Plugins/Runtime/Metasound/Source/MetasoundEngineTest/Private/EngineTestMetaSoundAutomatedNodeTest.cpp
@@ -735,7 +735,7 @@ namespace Metasound::EngineTest{
 	}
 
 
-	void GetAllRegisteredNodes(TArray<FString>& OutBeautifiedNames, TArray <FString>& OutTestCommands)
+	void GetAllRegisteredNodes(TArray<FString>& OutBeautifiedNames, TArray <FString>& OutTestCommands, bool bInIncludeDeprecatedNodes)
 	{
 		using namespace Metasound;
 
@@ -750,6 +750,11 @@ namespace Metasound::EngineTest{
 			{
 				continue;
 			}
+			else if (!bInIncludeDeprecatedNodes && NodeClass.Metadata.GetIsDeprecated())
+			{
+				// Optionally skip deprecated classes. 
+				continue;
+			}
 
 			OutBeautifiedNames.Add(FString::Printf(TEXT("%s %s"), *NodeClass.Metadata.GetClassName().ToString(), *NodeClass.Metadata.GetVersion().ToString()));
 
@@ -759,7 +764,7 @@ namespace Metasound::EngineTest{
 		}
 	}
 
-	void GetAllRegisteredNativeNodes(TArray<FString>& OutBeautifiedNames, TArray <FString>& OutTestCommands)
+	void GetAllRegisteredNativeNodes(TArray<FString>& OutBeautifiedNames, TArray <FString>& OutTestCommands, bool bInIncludeDeprecatedNodes)
 	{
 		using namespace Metasound;
 
@@ -777,6 +782,11 @@ namespace Metasound::EngineTest{
 			{
 				continue;
 			}
+			else if (!bInIncludeDeprecatedNodes && NodeClass.Metadata.GetIsDeprecated())
+			{
+				// Optionally skip deprecated classes. 
+				continue;
+			}
 
 			Frontend::FNodeRegistryKey NodeRegistryKey(NodeClass.Metadata);
 
diff --git a/Engine/Plugins/Runtime/Metasound/Source/MetasoundEngineTest/Private/EngineTestMetaSoundAutomatedNodeTest.h b/Engine/Plugins/Runtime/Metasound/Source/MetasoundEngineTest/Private/EngineTestMetaSoundAutomatedNodeTest.h
--- a/Engine/Plugins/Runtime/Metasound/Source/MetasoundEngineTest/Private/EngineTestMetaSoundAutomatedNodeTest.h
+++ b/Engine/Plugins/Runtime/Metasound/Source/MetasoundEngineTest/Private/EngineTestMetaSoundAutomatedNodeTest.h
@@ -139,10 +139,10 @@ namespace Metasound::EngineTest
 	TUniquePtr<INode> CreateNodeFromRegistry(const Frontend::FNodeRegistryKey& InNodeRegistryKey);
 
 	// Return an array of all currently registered MetaSound nodes. 
-	void GetAllRegisteredNodes(TArray<FString>& OutBeautifiedNames, TArray<FString>& OutNodeRegistryKeys);
+	void GetAllRegisteredNodes(TArray<FString>& OutBeautifiedNames, TArray<FString>& OutNodeRegistryKeys, bool bInIncludeDeprecatedNodes);
 
 	// Return an array of all currently registered native MetaSound nodes. 
-	void GetAllRegisteredNativeNodes(TArray<FString>& OutBeautifiedNames, TArray<FString>& OutNodeRegistryKeys);
+	void GetAllRegisteredNativeNodes(TArray<FString>& OutBeautifiedNames, TArray<FString>& OutNodeRegistryKeys, bool bInIncludeDeprecatedNodes);
 
 	// Create any variables that exist on the vertex interface.
 	void CreateVariables(const FOperatorSettings& InOperatorSettings, FInputVertexInterfaceData& OutVertexData);
diff --git a/Engine/Plugins/Runtime/Metasound/Source/MetasoundEngineTest/Private/EngineTestMetaSoundAutomatedNodeTest_Bind.cpp b/Engine/Plugins/Runtime/Metasound/Source/MetasoundEngineTest/Private/EngineTestMetaSoundAutomatedNodeTest_Bind.cpp
--- a/Engine/Plugins/Runtime/Metasound/Source/MetasoundEngineTest/Private/EngineTestMetaSoundAutomatedNodeTest_Bind.cpp
+++ b/Engine/Plugins/Runtime/Metasound/Source/MetasoundEngineTest/Private/EngineTestMetaSoundAutomatedNodeTest_Bind.cpp
@@ -23,7 +23,8 @@ IMPLEMENT_COMPLEX_AUTOMATION_TEST(FMetasoundAutomatedNodeTest_Bind, "Audio.Metas
 void FMetasoundAutomatedNodeTest_Bind::GetTests(TArray<FString>& OutBeautifiedNames, TArray<FString>& OutTestCommands) const
 {
 	using namespace Metasound::EngineTest;
-	GetAllRegisteredNativeNodes(OutBeautifiedNames, OutTestCommands);
+	constexpr bool bIncludeDeprecatedNodes = false;
+	GetAllRegisteredNativeNodes(OutBeautifiedNames, OutTestCommands, bIncludeDeprecatedNodes);
 
 	UE_LOG(LogMetaSound, Verbose, TEXT("Found %d metasound nodes to test"), OutTestCommands.Num());
 }
diff --git a/Engine/Plugins/Runtime/Metasound/Source/MetasoundEngineTest/Private/EngineTestMetaSoundAutomatedNodeTest_Reset.cpp b/Engine/Plugins/Runtime/Metasound/Source/MetasoundEngineTest/Private/EngineTestMetaSoundAutomatedNodeTest_Reset.cpp
--- a/Engine/Plugins/Runtime/Metasound/Source/MetasoundEngineTest/Private/EngineTestMetaSoundAutomatedNodeTest_Reset.cpp
+++ b/Engine/Plugins/Runtime/Metasound/Source/MetasoundEngineTest/Private/EngineTestMetaSoundAutomatedNodeTest_Reset.cpp
@@ -22,7 +22,8 @@ IMPLEMENT_COMPLEX_AUTOMATION_TEST(FMetasoundAutomatedNodeTest_Reset, "Audio.Meta
 void FMetasoundAutomatedNodeTest_Reset::GetTests(TArray<FString>& OutBeautifiedNames, TArray<FString>& OutTestCommands) const
 {
 	using namespace Metasound::EngineTest;
-	GetAllRegisteredNativeNodes(OutBeautifiedNames, OutTestCommands);
+	constexpr bool bIncludeDeprecatedNodes = false;
+	GetAllRegisteredNativeNodes(OutBeautifiedNames, OutTestCommands, bIncludeDeprecatedNodes);
 
 	UE_LOG(LogMetaSound, Verbose, TEXT("Found %d metasound nodes to test"), OutTestCommands.Num());
 }
diff --git a/Engine/Plugins/Runtime/Metasound/Source/MetasoundEngineTest/Private/EngineTestMetaSoundAutomatedNodeTest_ValidOutputs.cpp b/Engine/Plugins/Runtime/Metasound/Source/MetasoundEngineTest/Private/EngineTestMetaSoundAutomatedNodeTest_ValidOutputs.cpp
--- a/Engine/Plugins/Runtime/Metasound/Source/MetasoundEngineTest/Private/EngineTestMetaSoundAutomatedNodeTest_ValidOutputs.cpp
+++ b/Engine/Plugins/Runtime/Metasound/Source/MetasoundEngineTest/Private/EngineTestMetaSoundAutomatedNodeTest_ValidOutputs.cpp
@@ -22,7 +22,8 @@ IMPLEMENT_COMPLEX_AUTOMATION_TEST(FMetasoundAutomatedNodeTest_ValidOutputs, "Aud
 void FMetasoundAutomatedNodeTest_ValidOutputs::GetTests(TArray<FString>& OutBeautifiedNames, TArray<FString>& OutTestCommands) const
 {
 	using namespace Metasound::EngineTest;
-	GetAllRegisteredNativeNodes(OutBeautifiedNames, OutTestCommands);
+	constexpr bool bIncludeDeprecatedNodes = false;
+	GetAllRegisteredNativeNodes(OutBeautifiedNames, OutTestCommands, bIncludeDeprecatedNodes);
 
 	UE_LOG(LogMetaSound, Verbose, TEXT("Found %d metasound nodes to test"), OutTestCommands.Num());
 }
diff --git a/Engine/Plugins/Runtime/Metasound/Source/MetasoundStandardNodes/Private/MetasoundAudioFormats.cpp b/Engine/Plugins/Runtime/Metasound/Source/MetasoundStandardNodes/Private/MetasoundAudioFormats.cpp
--- a/Engine/Plugins/Runtime/Metasound/Source/MetasoundStandardNodes/Private/MetasoundAudioFormats.cpp
+++ b/Engine/Plugins/Runtime/Metasound/Source/MetasoundStandardNodes/Private/MetasoundAudioFormats.cpp
@@ -196,6 +196,7 @@ namespace Metasound
 			Info.Author = PluginAuthor;
 			Info.PromptIfMissing = PluginNodeMissingPrompt;
 			Info.DefaultInterface = CreateVertexInterface(InOutputName);
+			Info.bDeprecated = true;
 
 			return Info;
 		};
@@ -330,6 +331,7 @@ namespace Metasound
 			Info.Author = PluginAuthor;
 			Info.PromptIfMissing = PluginNodeMissingPrompt;
 			Info.DefaultInterface = CreateVertexInterface(InInputName);
+			Info.bDeprecated = true;
 
 			return Info;
 		};
@@ -434,6 +436,7 @@ namespace Metasound
 			Info.Author = PluginAuthor;
 			Info.PromptIfMissing = PluginNodeMissingPrompt;
 			Info.DefaultInterface = CreateVertexInterface(InOutputName);
+			Info.bDeprecated = true;
 
 			return Info;
 		};
@@ -567,6 +570,7 @@ namespace Metasound
 			Info.Author = PluginAuthor;
 			Info.PromptIfMissing = PluginNodeMissingPrompt;
 			Info.DefaultInterface = CreateVertexInterface(InInputName);
+			Info.bDeprecated = true;
 
 			return Info;
 		};
@@ -623,6 +627,20 @@ namespace Metasound
 	{
 		static constexpr bool Value = false;
 	};
+
+	// Disable constructor inputs/outputs
+	template<>
+	struct TEnableConstructorVertex<FMonoAudioFormat>
+	{
+		static constexpr bool Value = false;
+	};
+
+	// Disable constructor inputs/outputs
+	template<>
+	struct TEnableConstructorVertex<FStereoAudioFormat>
+	{
+		static constexpr bool Value = false;
+	};
 }
 
 // Data type registration has to happen after TInputNode<> and TOutputNode<> specializations
