Fix for crash when calling BuildNewMetaSound & BuildAndOverwriteMetaSound via MetaSound Builder Subsystem
[RN] Fix for crash when calling BuildNewMetaSound & BuildAndOverwriteMetaSound via MetaSound Builder Subsystem

diff --git a/Engine/Plugins/Runtime/Metasound/Source/MetasoundEngine/Private/MetasoundBuilderSubsystem.cpp b/Engine/Plugins/Runtime/Metasound/Source/MetasoundEngine/Private/MetasoundBuilderSubsystem.cpp
--- a/Engine/Plugins/Runtime/Metasound/Source/MetasoundEngine/Private/MetasoundBuilderSubsystem.cpp
+++ b/Engine/Plugins/Runtime/Metasound/Source/MetasoundEngine/Private/MetasoundBuilderSubsystem.cpp
@@ -99,6 +99,8 @@ namespace Metasound::Engine
 
 void UMetaSoundPatchBuilder::BuildAndOverwriteMetaSoundInternal(TScriptInterface<IMetaSoundDocumentInterface> ExistingMetaSound, bool bForceUniqueClassName) const
 {
+	checkf(ExistingMetaSound.GetObject(), TEXT("ExistingMetaSound interface must point to valid MetaSound object"));
+
 	FMetaSoundBuilderOptions Options;
 	Options.ExistingMetaSound = ExistingMetaSound;
 	Options.bForceUniqueClassName = bForceUniqueClassName;
@@ -270,6 +272,8 @@ bool UMetaSoundSourceBuilder::ExecuteAuditionableTransaction(FAuditionableTransa
 
 void UMetaSoundSourceBuilder::BuildAndOverwriteMetaSoundInternal(TScriptInterface<IMetaSoundDocumentInterface> ExistingMetaSound, bool bForceUniqueClassName) const
 {
+	checkf(ExistingMetaSound.GetObject(), TEXT("ExistingMetaSound interface must point to valid MetaSound object"));
+
 	FMetaSoundBuilderOptions Options;
 	Options.ExistingMetaSound = ExistingMetaSound;
 	Options.bForceUniqueClassName = bForceUniqueClassName;
diff --git a/Engine/Plugins/Runtime/Metasound/Source/MetasoundFrontend/Private/MetasoundFrontendDocumentBuilder.cpp b/Engine/Plugins/Runtime/Metasound/Source/MetasoundFrontend/Private/MetasoundFrontendDocumentBuilder.cpp
--- a/Engine/Plugins/Runtime/Metasound/Source/MetasoundFrontend/Private/MetasoundFrontendDocumentBuilder.cpp
+++ b/Engine/Plugins/Runtime/Metasound/Source/MetasoundFrontend/Private/MetasoundFrontendDocumentBuilder.cpp
@@ -3269,18 +3269,22 @@ void FMetaSoundFrontendDocumentBuilder::InitDocument(const FMetasoundFrontendDoc
 	METASOUND_TRACE_CPUPROFILER_EVENT_SCOPE(FMetaSoundFrontendDocumentBuilder::InitDocument);
 
 	FMetasoundFrontendDocument& Document = GetDocumentChecked();
-	Document.RootGraph.InitDefaultGraphPage();
 
 	// 1. Set default class Metadata.
 	if (InDocumentTemplate)
 	{
 		// 1a. If template provided, copy that.
 		Document = *InDocumentTemplate;
+		if (Document.RootGraph.GetConstGraphPages().IsEmpty())
+		{
+			Document.RootGraph.InitDefaultGraphPage();
+		}
 		InitGraphClassMetadata(bResetVersion, InNewClassName);
 	}
 	else
 	{
 		// 1a. Initialize class using default data
+		Document.RootGraph.InitDefaultGraphPage();
 		FMetasoundFrontendClassMetadata& ClassMetadata = Document.RootGraph.Metadata;
 		InitGraphClassMetadata(Document.RootGraph.Metadata, bResetVersion, InNewClassName);
 
@@ -4482,15 +4486,28 @@ void FMetaSoundFrontendDocumentBuilder::ResetGraphPages(bool bClearDefaultGraph)
 	using namespace Metasound;
 
 	FMetasoundFrontendGraphClass& RootGraph = GetDocumentChecked().RootGraph;
-	RootGraph.IterateGraphPages([this](FMetasoundFrontendGraph& Graph)
+	TArray<FGuid> PageDelegatesToRemove;
+	RootGraph.IterateGraphPages([this, &PageDelegatesToRemove](FMetasoundFrontendGraph& Graph)
 	{
 		if (Graph.PageID != Frontend::DefaultPageID)
 		{
 			DocumentDelegates->PageDelegates.OnRemovingPage.Broadcast(Frontend::FDocumentMutatePageArgs{ Graph.PageID });
+			PageDelegatesToRemove.Add(Graph.PageID);
 		}
 	});
 
 	RootGraph.ResetGraphPages(bClearDefaultGraph);
+
+	// Must be called after reset to avoid re-initializing delegates
+	// prematurely, which is handled by delegate responding to prior
+	// OnRemovingPage broadcast.
+	constexpr bool bBroadcastNotify = false;
+	for (const FGuid& PageID : PageDelegatesToRemove)
+	{
+		DocumentDelegates->RemovePageDelegates(PageID, bBroadcastNotify);
+	}
+
+	Reload(DocumentDelegates);
 	SetBuildPageID(Frontend::DefaultPageID);
 }
 #endif // WITH_EDITORONLY_DATA
diff --git a/Engine/Plugins/Runtime/Metasound/Source/MetasoundFrontend/Private/MetasoundFrontendDocumentModifyDelegates.cpp b/Engine/Plugins/Runtime/Metasound/Source/MetasoundFrontend/Private/MetasoundFrontendDocumentModifyDelegates.cpp
--- a/Engine/Plugins/Runtime/Metasound/Source/MetasoundFrontend/Private/MetasoundFrontendDocumentModifyDelegates.cpp
+++ b/Engine/Plugins/Runtime/Metasound/Source/MetasoundFrontend/Private/MetasoundFrontendDocumentModifyDelegates.cpp
@@ -93,9 +93,12 @@ PRAGMA_DISABLE_DEPRECATION_WARNINGS
 PRAGMA_ENABLE_DEPRECATION_WARNINGS
 	}
 
-	void FDocumentModifyDelegates::RemovePageDelegates(const FGuid& InPageID)
+	void FDocumentModifyDelegates::RemovePageDelegates(const FGuid& InPageID, bool bBroadcastNotify)
 	{
-		PageDelegates.OnRemovingPage.Broadcast(FDocumentMutatePageArgs { InPageID });
+		if (bBroadcastNotify)
+		{
+			PageDelegates.OnRemovingPage.Broadcast(FDocumentMutatePageArgs{ InPageID });
+		}
 
 		PageNodeDelegates.Remove(InPageID);
 		PageEdgeDelegates.Remove(InPageID);
diff --git a/Engine/Plugins/Runtime/Metasound/Source/MetasoundFrontend/Public/MetasoundFrontendDocumentModifyDelegates.h b/Engine/Plugins/Runtime/Metasound/Source/MetasoundFrontend/Public/MetasoundFrontendDocumentModifyDelegates.h
--- a/Engine/Plugins/Runtime/Metasound/Source/MetasoundFrontend/Public/MetasoundFrontendDocumentModifyDelegates.h
+++ b/Engine/Plugins/Runtime/Metasound/Source/MetasoundFrontend/Public/MetasoundFrontendDocumentModifyDelegates.h
@@ -97,7 +97,7 @@ namespace Metasound::Frontend
 		FEdgeModifyDelegates EdgeDelegates;
 
 		void AddPageDelegates(const FGuid& InPageID);
-		void RemovePageDelegates(const FGuid& InPageID);
+		void RemovePageDelegates(const FGuid& InPageID, bool bBroadcastNotify = true);
 
 	private:
 		TSortedMap<FGuid, FNodeModifyDelegates> PageNodeDelegates;
