- Remove name selector as it was confusing for page defaults
- Utilize audition fallback when adding new paged input default value instead of data type default
- Move Input Default Reset to proper reset button (rather than trash can which was confusing b/c every input needs at least one default which can't be deleted)
#rnx

diff --git a/Engine/Plugins/Runtime/Metasound/Source/MetasoundEditor/Private/MetasoundDefaultLiteralCustomization.cpp b/Engine/Plugins/Runtime/Metasound/Source/MetasoundEditor/Private/MetasoundDefaultLiteralCustomization.cpp
--- a/Engine/Plugins/Runtime/Metasound/Source/MetasoundEditor/Private/MetasoundDefaultLiteralCustomization.cpp
+++ b/Engine/Plugins/Runtime/Metasound/Source/MetasoundEditor/Private/MetasoundDefaultLiteralCustomization.cpp
@@ -137,8 +137,24 @@ namespace Metasound::Editor
 					MetaSound.Modify();
 					LiteralPtr->Modify();
 
+					const UMetasoundEditorSettings* EditorSettings = GetDefault<UMetasoundEditorSettings>();
+					check(EditorSettings);
+
+					TArray<FGuid> ImplementedGuids;
+					LiteralPtr->IterateDefaults([&ImplementedGuids](const FGuid& PageID, FMetasoundFrontendLiteral)
+					{
+						ImplementedGuids.Add(PageID);
+					});
+
+					const FGuid FallbackPageID = EditorSettings->ResolveAuditionPage(ImplementedGuids, PageSettings->UniqueId);
+
 					LiteralPtr->InitDefault(PageSettings->UniqueId);
 
+					FMetasoundFrontendLiteral InitValue;
+					if (LiteralPtr->TryFindDefault(InitValue, &FallbackPageID))
+					{
+						LiteralPtr->SetFromLiteral(InitValue, PageSettings->UniqueId);
+					}
 					constexpr bool bPostTransaction = false;
 					Member->UpdateFrontendDefaultLiteral(bPostTransaction, &PageSettings->UniqueId);
 					FMetasoundAssetBase& MetasoundAsset = Editor::FGraphBuilder::GetOutermostMetaSoundChecked(*LiteralPtr);
@@ -170,49 +186,45 @@ namespace Metasound::Editor
 			[
 				PageDefaultComboBox->AsShared()
 			]
-			+ SHorizontalBox::Slot()
-			.Padding(2.0f)
-			.HAlign(HAlign_Left)
-			.VAlign(VAlign_Center)
-			.AutoWidth()
-			[
-				PropertyCustomizationHelpers::MakeDeleteButton(FSimpleDelegate::CreateLambda([this, LiteralPtr]()
-				{
-					using namespace Frontend;
+		]
+		.ResetToDefaultContent()
+		[
+			PropertyCustomizationHelpers::MakeResetButton(FSimpleDelegate::CreateLambda([this, LiteralPtr]()
+			{
+				using namespace Frontend;
 
-					if (!LiteralPtr.IsValid())
-					{
-						return;
-					}
+				if (!LiteralPtr.IsValid())
+				{
+					return;
+				}
 
-					UMetasoundEditorGraphMember* Member = LiteralPtr->FindMember();
-					if (!Member)
-					{
-						return;
-					}
+				UMetasoundEditorGraphMember* Member = LiteralPtr->FindMember();
+				if (!Member)
+				{
+					return;
+				}
 
-					const FScopedTransaction Transaction(LOCTEXT("ResetPageDefaultsTransaction", "Reset Paged Defaults"));
+				const FScopedTransaction Transaction(LOCTEXT("ResetPageDefaultsTransaction", "Reset Paged Defaults"));
 
-					const UMetaSoundSettings* Settings = GetDefault<UMetaSoundSettings>();
-					check(Settings);
+				const UMetaSoundSettings* Settings = GetDefault<UMetaSoundSettings>();
+				check(Settings);
 
-					FMetaSoundFrontendDocumentBuilder& Builder = Member->GetFrontendBuilderChecked();
-					UObject& MetaSound = Builder.CastDocumentObjectChecked<UObject>();
-					MetaSound.Modify();
-					LiteralPtr->Modify();
+				FMetaSoundFrontendDocumentBuilder& Builder = Member->GetFrontendBuilderChecked();
+				UObject& MetaSound = Builder.CastDocumentObjectChecked<UObject>();
+				MetaSound.Modify();
+				LiteralPtr->Modify();
 
-					LiteralPtr->ResetDefaults();
+				LiteralPtr->ResetDefaults();
 
-					constexpr bool bPostTransaction = false;
-					Member->UpdateFrontendDefaultLiteral(bPostTransaction);
+				constexpr bool bPostTransaction = false;
+				Member->UpdateFrontendDefaultLiteral(bPostTransaction);
 
-					FMetasoundAssetBase& MetasoundAsset = Editor::FGraphBuilder::GetOutermostMetaSoundChecked(*LiteralPtr);
-					MetasoundAsset.GetModifyContext().AddMemberIDsModified({ Member->GetMemberID() });
-					UpdatePagePickerNames(LiteralPtr);
-					PageDefaultComboBox->RefreshOptions();
-					FGraphBuilder::RegisterGraphWithFrontend(MetaSound);
-				}), LOCTEXT("ResetPageDefaultsTooltip", "Resets page defaults for the given member, leaving just the initial value for the required 'Default' page."))
-			]
+				FMetasoundAssetBase& MetasoundAsset = Editor::FGraphBuilder::GetOutermostMetaSoundChecked(*LiteralPtr);
+				MetasoundAsset.GetModifyContext().AddMemberIDsModified({ Member->GetMemberID() });
+				UpdatePagePickerNames(LiteralPtr);
+				PageDefaultComboBox->RefreshOptions();
+				FGraphBuilder::RegisterGraphWithFrontend(MetaSound);
+			}), LOCTEXT("ResetPageDefaultsTooltip", "Resets page defaults for the given member, leaving just the initial value for the required 'Default' page."))
 		];
 	}
 
@@ -224,13 +236,18 @@ namespace Metasound::Editor
 			return SNullWidget::NullWidget;
 		}
 
+		FText PageNameText;
+		PageNameProperty->GetValueAsFormattedText(PageNameText);
+
 		TSharedRef<SHorizontalBox> NameBox = SNew(SHorizontalBox)
 		+ SHorizontalBox::Slot()
 		.AutoWidth()
 		.VAlign(VAlign_Center)
 		.Padding(2, 0)
 		[
-			PageNameProperty->CreatePropertyValueWidget()
+			SNew(STextBlock)
+				.Font(IDetailLayoutBuilder::GetDetailFont())
+				.Text(PageNameText)
 		];
 
 		const UMetasoundEditorGraphMember* Member = Literal.FindMember();
diff --git a/Engine/Plugins/Runtime/Metasound/Source/MetasoundEditor/Public/MetasoundEditorGraphMemberDefaults.h b/Engine/Plugins/Runtime/Metasound/Source/MetasoundEditor/Public/MetasoundEditorGraphMemberDefaults.h
--- a/Engine/Plugins/Runtime/Metasound/Source/MetasoundEditor/Public/MetasoundEditorGraphMemberDefaults.h
+++ b/Engine/Plugins/Runtime/Metasound/Source/MetasoundEditor/Public/MetasoundEditorGraphMemberDefaults.h
@@ -49,7 +49,7 @@ struct FMetasoundEditorMemberPageDefault
 	static const FGuid& GetNewEntryID();
 
 	// Selectable PageName
-	UPROPERTY(EditAnywhere, Category = DefaultValue, meta = (NoResetToDefault, GetOptions = "MetasoundEngine.MetaSoundSettings.GetPageNames"))
+	UPROPERTY(VisibleAnywhere, Category = DefaultValue)
 	FName PageName;
 
 	// Used for hash and mirrors document-stored value. Defaults to random value to allow for assignment in
