Fix for inability to copy and paste output between pages
#rnx

diff --git a/Engine/Plugins/Runtime/Metasound/Source/MetasoundEditor/Private/MetasoundEditorDocumentClipboardUtils.cpp b/Engine/Plugins/Runtime/Metasound/Source/MetasoundEditor/Private/MetasoundEditorDocumentClipboardUtils.cpp
--- a/Engine/Plugins/Runtime/Metasound/Source/MetasoundEditor/Private/MetasoundEditorDocumentClipboardUtils.cpp
+++ b/Engine/Plugins/Runtime/Metasound/Source/MetasoundEditor/Private/MetasoundEditorDocumentClipboardUtils.cpp
@@ -151,17 +151,20 @@ namespace Metasound::Editor
 
 			if (OutputNode->Output && Graph.ContainsOutput(*OutputNode->Output))
 			{
-				auto NodeMatches = [OutputNodeID = OutputNode->GetNodeID()](const TObjectPtr<UEdGraphNode>& EdNode)
+				auto IsOtherMatchingNode = [&OutputNode](const TObjectPtr<UEdGraphNode>& EdNode)
 				{
-					if (UMetasoundEditorGraphOutputNode* OutputNode = Cast<UMetasoundEditorGraphOutputNode>(EdNode))
+					if (OutputNode != EdNode.Get())
 					{
-						return OutputNodeID == OutputNode->GetNodeID();
+						if (UMetasoundEditorGraphOutputNode* OtherOutputNode = Cast<UMetasoundEditorGraphOutputNode>(EdNode))
+						{
+							return OutputNode->GetNodeID() == OtherOutputNode->GetNodeID();
+						}
 					}
 					return false;
 				};
 
 				// Can only have one output reference node
-				if (TObjectPtr<UEdGraphNode>* MatchingOutput = Graph.Nodes.FindByPredicate(NodeMatches))
+				if (Graph.Nodes.ContainsByPredicate(IsOtherMatchingNode))
 				{
 					OutNotifications.bPastedNodesAddMultipleOutputNodes = true;
 					Graph.RemoveNode(OutputNode);
